Horde: Add node hashes to bundle imports. This makes it possible to implement the NodeHandle interface without needing an externally provided hash value.

diff --git a/Engine/Source/Programs/Horde/Horde.Server.Tests/BlobStoreTests.cs b/Engine/Source/Programs/Horde/Horde.Server.Tests/BlobStoreTests.cs
--- a/Engine/Source/Programs/Horde/Horde.Server.Tests/BlobStoreTests.cs
+++ b/Engine/Source/Programs/Horde/Horde.Server.Tests/BlobStoreTests.cs
@@ -52,7 +52,7 @@ namespace Horde.Server.Tests
 			List<BundleExportRef> exportRefs = new List<BundleExportRef>();
 			for(int idx = 0; idx < refs.Count; idx++)
 			{
-				exportRefs.Add(new BundleExportRef(idx, 0));
+				exportRefs.Add(new BundleExportRef(idx, 0, IoHash.Zero));
 			}
 
 			List<BundleExport> exports = new List<BundleExport>();
diff --git a/Engine/Source/Programs/Shared/EpicGames.Horde.Tests/BundleTests.cs b/Engine/Source/Programs/Shared/EpicGames.Horde.Tests/BundleTests.cs
--- a/Engine/Source/Programs/Shared/EpicGames.Horde.Tests/BundleTests.cs
+++ b/Engine/Source/Programs/Shared/EpicGames.Horde.Tests/BundleTests.cs
@@ -145,8 +145,8 @@ namespace EpicGames.Horde.Tests
 				imports.Add(new BlobLocator("import2"));
 
 				List<BundleExport> exports = new List<BundleExport>();
-				exports.Add(new BundleExport(0, IoHash.Compute(Encoding.UTF8.GetBytes("export1")), 0, 0, 2, new BundleExportRef[] { new BundleExportRef(0, 5), new BundleExportRef(0, 6) }));
-				exports.Add(new BundleExport(0, IoHash.Compute(Encoding.UTF8.GetBytes("export2")), 1, 0, 3, new BundleExportRef[] { new BundleExportRef(-1, 0) }));
+				exports.Add(new BundleExport(0, IoHash.Compute(Encoding.UTF8.GetBytes("export1")), 0, 0, 2, new BundleExportRef[] { new BundleExportRef(0, 5, IoHash.Zero), new BundleExportRef(0, 6, IoHash.Zero) }));
+				exports.Add(new BundleExport(0, IoHash.Compute(Encoding.UTF8.GetBytes("export2")), 1, 0, 3, new BundleExportRef[] { new BundleExportRef(-1, 0, IoHash.Zero) }));
 
 				List<BundlePacket> packets = new List<BundlePacket>();
 				packets.Add(new BundlePacket(BundleCompressionFormat.LZ4, 0, 20, 40));
diff --git a/Engine/Source/Programs/Shared/EpicGames.Horde/Storage/Bundle.cs b/Engine/Source/Programs/Shared/EpicGames.Horde/Storage/Bundle.cs
--- a/Engine/Source/Programs/Shared/EpicGames.Horde/Storage/Bundle.cs
+++ b/Engine/Source/Programs/Shared/EpicGames.Horde/Storage/Bundle.cs
@@ -46,6 +46,11 @@ namespace EpicGames.Horde.Storage
 		/// </summary>
 		InPlace = 3,
 
+		/// <summary>
+		/// Add import hashes to imported nodes
+		/// </summary>
+		ImportHashes = 4,
+
 		/// <summary>
 		/// Last item in the enum. Used for <see cref="Latest"/>
 		/// </summary>
@@ -250,7 +255,7 @@ namespace EpicGames.Horde.Storage
 			Data = data;
 
 			ReadOnlySpan<byte> span = data.Span;
-			// BundleVersion version = (BundleVersion)span[3];
+			BundleVersion version = (BundleVersion)span[3];
 			int headerLength = BinaryPrimitives.ReadInt32LittleEndian(span.Slice(4));
 
 			ReadOnlyMemory<byte> exportData = ReadOnlyMemory<byte>.Empty;
@@ -287,7 +292,7 @@ namespace EpicGames.Horde.Storage
 				offset += length;
 			}
 
-			Exports = new BundleExportCollection(exportData, exportRefData);
+			Exports = new BundleExportCollection(exportData, exportRefData, version);
 		}
 
 		/// <summary>
@@ -474,7 +479,7 @@ namespace EpicGames.Horde.Storage
 				int[] exportIndexes = reader.ReadVariableLengthArray(() => (int)reader.ReadUnsignedVarInt());
 				for (int exportIdx = 0; exportIdx < exportIndexes.Length; exportIdx++)
 				{
-					BundleExportRef exportReference = new BundleExportRef(imports.Count - 1, exportIndexes[exportIdx]);
+					BundleExportRef exportReference = new BundleExportRef(imports.Count - 1, exportIndexes[exportIdx], IoHash.Zero);
 					allExportReferences.Add(exportReference);
 				}
 			}
@@ -485,7 +490,7 @@ namespace EpicGames.Horde.Storage
 
 			for (int exportIdx = 0; exportIdx < numExports; exportIdx++)
 			{
-				allExportReferences.Add(new BundleExportRef(-1, exportIdx));
+				allExportReferences.Add(new BundleExportRef(-1, exportIdx, IoHash.Zero));
 
 				int typeIdx = (int)reader.ReadUnsignedVarInt();
 				IoHash hash = reader.ReadIoHash();
@@ -1009,11 +1014,33 @@ namespace EpicGames.Horde.Storage
 		/// <summary>
 		/// Constructor
 		/// </summary>
-		public BundleExportCollection(ReadOnlyMemory<byte> data, ReadOnlyMemory<byte> refs)
+		public BundleExportCollection(ReadOnlyMemory<byte> data, ReadOnlyMemory<byte> refs, BundleVersion version)
 		{
 			_count = data.Length / BundleExport.NumBytes;
 			_data = data;
 			_refs = refs;
+
+			if (version < BundleVersion.ImportHashes && _refs.Length > 0)
+			{
+				const int OldRefNumBytes = 4;
+
+				int headerLength = BinaryPrimitives.ReadInt32LittleEndian(refs.Span);
+				int count = (refs.Length - headerLength) / OldRefNumBytes;
+
+				byte[] newRefs = new byte[headerLength + (count * BundleExportRef.NumBytes)];
+				for (int headerOffset = 0; headerOffset < headerLength; headerOffset += sizeof(int))
+				{
+					int index = (BinaryPrimitives.ReadInt32LittleEndian(refs.Slice(headerOffset).Span) - headerLength) / OldRefNumBytes;
+					BinaryPrimitives.WriteInt32LittleEndian(newRefs.AsSpan(headerOffset), headerLength + (index * BundleExportRef.NumBytes));
+				}
+				for (int index = 0; index < count; index++)
+				{
+					ReadOnlyMemory<byte> source = refs.Slice(headerLength + (index * OldRefNumBytes), OldRefNumBytes);
+					source.CopyTo(newRefs.AsMemory(headerLength + (index * BundleExportRef.NumBytes)));
+				}
+
+				_refs = newRefs;
+			}
 		}
 
 		/// <inheritdoc/>
@@ -1081,19 +1108,20 @@ namespace EpicGames.Horde.Storage
 	/// </summary>
 	/// <param name="ImportIdx">Index into the import table of the blob containing the referenced node. Can be -1 for references within the same bundle.</param>
 	/// <param name="NodeIdx">Node imported from the bundle</param>
-	public record struct BundleExportRef(int ImportIdx, int NodeIdx)
+	/// <param name="Hash">Hash of the referenced node</param>
+	public record struct BundleExportRef(int ImportIdx, int NodeIdx, IoHash Hash)
 	{
 		/// <summary>
 		/// Number of bytes in the serialized object
 		/// </summary>
-		public const int NumBytes = 4;
+		public const int NumBytes = 4 + IoHash.NumBytes;
 
 		/// <summary>
 		/// Deserialize this object from memory
 		/// </summary>
 		public static BundleExportRef Read(ReadOnlySpan<byte> data)
 		{
-			return new BundleExportRef(BinaryPrimitives.ReadInt16LittleEndian(data), BinaryPrimitives.ReadInt16LittleEndian(data.Slice(2)));
+			return new BundleExportRef(BinaryPrimitives.ReadInt16LittleEndian(data), BinaryPrimitives.ReadInt16LittleEndian(data.Slice(2)), new IoHash(data.Slice(4)));
 		}
 
 		/// <summary>
@@ -1103,6 +1131,7 @@ namespace EpicGames.Horde.Storage
 		{
 			BinaryPrimitives.WriteInt16LittleEndian(data, (short)ImportIdx);
 			BinaryPrimitives.WriteUInt16LittleEndian(data.Slice(2), (ushort)NodeIdx);
+			Hash.CopyTo(data.Slice(4));
 		}
 	}
 
diff --git a/Engine/Source/Programs/Shared/EpicGames.Horde/Storage/TreeWriter.cs b/Engine/Source/Programs/Shared/EpicGames.Horde/Storage/TreeWriter.cs
--- a/Engine/Source/Programs/Shared/EpicGames.Horde/Storage/TreeWriter.cs
+++ b/Engine/Source/Programs/Shared/EpicGames.Horde/Storage/TreeWriter.cs
@@ -538,7 +538,8 @@ namespace EpicGames.Horde.Storage
 				Dictionary<NodeHandle, BundleExportRef> nodeHandleToExportRef = new Dictionary<NodeHandle, BundleExportRef>();
 				for (int exportIdx = 0; exportIdx < _queue.Count; exportIdx++)
 				{
-					nodeHandleToExportRef[_queue[exportIdx]] = new BundleExportRef(-1, exportIdx);
+					NodeHandle handle = _queue[exportIdx];
+					nodeHandleToExportRef[handle] = new BundleExportRef(-1, exportIdx, handle.Hash);
 				}
 
 				// Create the export list
@@ -566,7 +567,7 @@ namespace EpicGames.Horde.Storage
 								imports.Add(handle.Locator.Blob);
 								importToIndex.Add(handle.Locator.Blob, importIdx);
 							}
-							exportRef = new BundleExportRef(importIdx, handle.Locator.ExportIdx);
+							exportRef = new BundleExportRef(importIdx, handle.Locator.ExportIdx, handle.Hash);
 						}
 						exportRefs.Add(exportRef);
 					}
