EpicGames.Core: Set a large max degree of parallelism in Json tracer tests.

The default value here is -1, which is equivalent to Environment.ProcessorCount when calling Parallel.ForEachAsync(). This results in a large number of async tasks calling Task.Delay() in sequence, rather than actually running asynchronously, causing the test to take > 60 seconds on machines with low core counts (such as a VM). Increasing this setting to eliminate sequencing allows allows test to complete in < 50ms.

diff --git a/Engine/Source/Programs/Shared/EpicGames.Core.Tests/JsonTracerTests.cs b/Engine/Source/Programs/Shared/EpicGames.Core.Tests/JsonTracerTests.cs
--- a/Engine/Source/Programs/Shared/EpicGames.Core.Tests/JsonTracerTests.cs
+++ b/Engine/Source/Programs/Shared/EpicGames.Core.Tests/JsonTracerTests.cs
@@ -124,7 +124,7 @@ namespace EpicGames.Core.Tests
 		private static async Task ForEachAsyncSpanTest(JsonTracer Tracer, IScope ParentScope)
 		{
 			int[] Values = Enumerable.Range(0, 10000).ToArray();
-			await Parallel.ForEachAsync(Values, async (i, token) =>
+			await Parallel.ForEachAsync(Values, new ParallelOptions { MaxDegreeOfParallelism = 10000 }, async (i, token) =>
 			{
 				using (IScope ChildScope = Tracer.BuildSpan("Child").StartActive())
 				{
