Horde: Make NodeReader/NodeWriter implementation classes public, rather than abstracting them behind an interface. The pertinent abstraction details are already encapsulated in NodeHandle. Also add support for nodes having a constructor taking a NodeData parameter rather than requiring a NodeReader instance to be created.

diff --git a/Engine/Source/Programs/Horde/Horde.Server/Perforce/PerforceReplicator.cs b/Engine/Source/Programs/Horde/Horde.Server/Perforce/PerforceReplicator.cs
--- a/Engine/Source/Programs/Horde/Horde.Server/Perforce/PerforceReplicator.cs
+++ b/Engine/Source/Programs/Horde/Horde.Server/Perforce/PerforceReplicator.cs
@@ -62,7 +62,7 @@ namespace Horde.Server.Perforce
 			}
 
 			/// <inheritdoc/>
-			public override void Serialize(ITreeNodeWriter writer)
+			public override void Serialize(NodeWriter writer)
 			{
 				writer.WriteUnsignedVarInt(Change);
 				writer.WriteUnsignedVarInt(ParentChange);
diff --git a/Engine/Source/Programs/Shared/EpicGames.Horde.Tests/BundleStoreTests.cs b/Engine/Source/Programs/Shared/EpicGames.Horde.Tests/BundleStoreTests.cs
--- a/Engine/Source/Programs/Shared/EpicGames.Horde.Tests/BundleStoreTests.cs
+++ b/Engine/Source/Programs/Shared/EpicGames.Horde.Tests/BundleStoreTests.cs
@@ -45,7 +45,7 @@ namespace EpicGames.Horde.Tests
 				Text = reader.ReadString();
 			}
 
-			public override void Serialize(ITreeNodeWriter writer)
+			public override void Serialize(NodeWriter writer)
 			{
 				writer.WriteString(Text);
 			}
@@ -130,7 +130,7 @@ namespace EpicGames.Horde.Tests
 				Refs = reader.ReadVariableLengthArray(() => reader.ReadRef<SimpleNode>());
 			}
 
-			public override void Serialize(ITreeNodeWriter writer)
+			public override void Serialize(NodeWriter writer)
 			{
 				writer.WriteVariableLengthBytes(Data);
 				writer.WriteVariableLengthArray(Refs, x => writer.WriteRef(x));
diff --git a/Engine/Source/Programs/Shared/EpicGames.Horde/Logs/LogChunkNode.cs b/Engine/Source/Programs/Shared/EpicGames.Horde/Logs/LogChunkNode.cs
--- a/Engine/Source/Programs/Shared/EpicGames.Horde/Logs/LogChunkNode.cs
+++ b/Engine/Source/Programs/Shared/EpicGames.Horde/Logs/LogChunkNode.cs
@@ -119,7 +119,7 @@ namespace EpicGames.Horde.Logs
 		}
 
 		/// <inheritdoc/>
-		public override void Serialize(ITreeNodeWriter writer)
+		public override void Serialize(NodeWriter writer)
 		{
 			writer.WriteVariableLengthBytes(Data.Span);
 		}
@@ -241,7 +241,7 @@ namespace EpicGames.Horde.Logs
 		}
 
 		/// <inheritdoc/>
-		public override void Serialize(ITreeNodeWriter writer)
+		public override void Serialize(NodeWriter writer)
 		{
 			base.Serialize(writer);
 
diff --git a/Engine/Source/Programs/Shared/EpicGames.Horde/Logs/LogIndexNode.cs b/Engine/Source/Programs/Shared/EpicGames.Horde/Logs/LogIndexNode.cs
--- a/Engine/Source/Programs/Shared/EpicGames.Horde/Logs/LogIndexNode.cs
+++ b/Engine/Source/Programs/Shared/EpicGames.Horde/Logs/LogIndexNode.cs
@@ -101,7 +101,7 @@ namespace EpicGames.Horde.Logs
 		}
 
 		/// <inheritdoc/>
-		public override void Serialize(ITreeNodeWriter writer)
+		public override void Serialize(NodeWriter writer)
 		{
 			writer.WriteUnsignedVarInt(CurrentVersion);
 			writer.WriteNgramSet(_ngramSet);
diff --git a/Engine/Source/Programs/Shared/EpicGames.Horde/Logs/LogNode.cs b/Engine/Source/Programs/Shared/EpicGames.Horde/Logs/LogNode.cs
--- a/Engine/Source/Programs/Shared/EpicGames.Horde/Logs/LogNode.cs
+++ b/Engine/Source/Programs/Shared/EpicGames.Horde/Logs/LogNode.cs
@@ -101,7 +101,7 @@ namespace EpicGames.Horde.Logs
 		}
 
 		/// <inheritdoc/>
-		public override void Serialize(ITreeNodeWriter writer)
+		public override void Serialize(NodeWriter writer)
 		{
 			writer.WriteUInt8((byte)Format);
 			writer.WriteUnsignedVarInt(LineCount);
diff --git a/Engine/Source/Programs/Shared/EpicGames.Horde/Storage/Node.cs b/Engine/Source/Programs/Shared/EpicGames.Horde/Storage/Node.cs
--- a/Engine/Source/Programs/Shared/EpicGames.Horde/Storage/Node.cs
+++ b/Engine/Source/Programs/Shared/EpicGames.Horde/Storage/Node.cs
@@ -3,6 +3,7 @@
 using System;
 using System.Collections.Concurrent;
 using System.Collections.Generic;
+using System.Diagnostics;
 using System.Diagnostics.CodeAnalysis;
 using System.Reflection;
 using System.Reflection.Emit;
@@ -88,7 +89,7 @@ namespace EpicGames.Horde.Storage
 		/// Serialize the contents of this node
 		/// </summary>
 		/// <returns>Data for the node</returns>
-		public abstract void Serialize(ITreeNodeWriter writer);
+		public abstract void Serialize(NodeWriter writer);
 
 		/// <summary>
 		/// Enumerate all outward references from this node
@@ -101,7 +102,7 @@ namespace EpicGames.Horde.Storage
 		static readonly object s_writeLock = new object();
 		static readonly ConcurrentDictionary<Type, NodeType> s_typeToNodeType = new ConcurrentDictionary<Type, NodeType>();
 		static readonly ConcurrentDictionary<Guid, Type> s_guidToType = new ConcurrentDictionary<Guid, Type>();
-		static readonly ConcurrentDictionary<Guid, Func<NodeReader, Node>> s_guidToDeserializer = new ConcurrentDictionary<Guid, Func<NodeReader, Node>>();
+		static readonly ConcurrentDictionary<Guid, Func<NodeData, Node>> s_guidToDeserializer = new ConcurrentDictionary<Guid, Func<NodeData, Node>>();
 
 		static NodeType CreateNodeType(NodeTypeAttribute attribute)
 		{
@@ -152,19 +153,19 @@ namespace EpicGames.Horde.Storage
 		/// <summary>
 		/// Deserialize a node from the given reader
 		/// </summary>
-		/// <param name="reader">Data to deserialize from</param>
+		/// <param name="nodeData">Data to deserialize from</param>
 		/// <returns>New node instance</returns>
-		public static Node Deserialize(NodeReader reader)
+		public static Node Deserialize(NodeData nodeData)
 		{
-			return s_guidToDeserializer[reader.Type.Guid](reader);
+			return s_guidToDeserializer[nodeData.Type.Guid](nodeData);
 		}
 
 		/// <summary>
 		/// Deserialize a node from the given reader
 		/// </summary>
-		/// <param name="reader">Data to deserialize from</param>
+		/// <param name="nodeData">Data to deserialize from</param>
 		/// <returns>New node instance</returns>
-		public static TNode Deserialize<TNode>(NodeReader reader) where TNode : Node => (TNode)Deserialize(reader);
+		public static TNode Deserialize<TNode>(NodeData nodeData) where TNode : Node => (TNode)Deserialize(nodeData);
 
 		/// <summary>
 		/// Static constructor. Registers all the types in the current assembly.
@@ -223,44 +224,179 @@ namespace EpicGames.Horde.Storage
 			NodeType nodeType = CreateNodeType(attribute);
 			s_typeToNodeType.TryAdd(type, nodeType);
 
-			Func<NodeReader, Node> deserializer = CreateDeserializer(type);
+			Func<NodeData, Node> deserializer = CreateDeserializer(type);
 			s_guidToType.TryAdd(nodeType.Guid, type);
 			s_guidToDeserializer.TryAdd(nodeType.Guid, deserializer);
 		}
 
-		static Func<NodeReader, Node> CreateDeserializer(Type type)
+		static readonly ConstructorInfo? s_nodeReaderTypeCtor = typeof(NodeReader).GetConstructor(new[] { typeof(NodeData) });
+
+		static Func<NodeData, Node> CreateDeserializer(Type type)
 		{
-			Type[] signature = new[] { typeof(NodeReader) };
+			Type[] signature = new[] { typeof(NodeData) };
 
-			ConstructorInfo? constructorInfo = type.GetConstructor(signature);
-			if (constructorInfo == null)
+			ConstructorInfo? nodeDataCtor = type.GetConstructor(signature);
+			if (nodeDataCtor != null)
 			{
-				throw new InvalidOperationException($"Type {type.Name} does not have a constructor taking an {typeof(NodeReader).Name} as parameter.");
+				DynamicMethod method = new DynamicMethod($"Create_{type.Name}", type, signature, true);
+
+				ILGenerator generator = method.GetILGenerator();
+				generator.Emit(OpCodes.Ldarg_0);
+				generator.Emit(OpCodes.Newobj, nodeDataCtor);
+				generator.Emit(OpCodes.Ret);
+
+				return (Func<NodeData, Node>)method.CreateDelegate(typeof(Func<NodeData, Node>));
 			}
 
-			DynamicMethod method = new DynamicMethod($"Create_{type.Name}", type, signature, true);
+			ConstructorInfo? nodeReaderCtor = type.GetConstructor(new[] { typeof(NodeReader) });
+			if (nodeReaderCtor != null)
+			{
+				DynamicMethod method = new DynamicMethod($"Create_{type.Name}", type, signature, true);
 
-			ILGenerator generator = method.GetILGenerator();
-			generator.Emit(OpCodes.Ldarg_0);
-			generator.Emit(OpCodes.Newobj, constructorInfo);
-			generator.Emit(OpCodes.Ret);
+				ILGenerator generator = method.GetILGenerator();
+				generator.Emit(OpCodes.Ldarg_0);
+				generator.Emit(OpCodes.Newobj, s_nodeReaderTypeCtor!);
+				generator.Emit(OpCodes.Newobj, nodeReaderCtor);
+				generator.Emit(OpCodes.Ret);
 
-			return (Func<NodeReader, Node>)method.CreateDelegate(typeof(Func<NodeReader, Node>));
+				return (Func<NodeData, Node>)method.CreateDelegate(typeof(Func<NodeData, Node>));
+			}
+
+			throw new InvalidOperationException($"Type {type.Name} does not have a constructor taking a {typeof(NodeData).Name} or {typeof(NodeReader).Name} instance as parameter.");
 		}
 
 		#endregion
 	}
 
 	/// <summary>
-	/// Writer for tree nodes
+	/// Data for an individual node
+	/// </summary>
+	public record struct NodeData(NodeType Type, IoHash Hash, ReadOnlyMemory<byte> Data, IReadOnlyList<NodeHandle> Refs);
+
+	/// <summary>
+	/// Reader for tree nodes
 	/// </summary>
-	public interface ITreeNodeWriter : IMemoryWriter
+	public sealed class NodeReader : MemoryReader
 	{
 		/// <summary>
-		/// Writes a reference to another node
+		/// Type to deserialize
+		/// </summary>
+		public NodeType Type => _nodeData.Type;
+
+		/// <summary>
+		/// Version of the current node, as specified via <see cref="NodeTypeAttribute"/>
+		/// </summary>
+		public int Version => Type.Version;
+
+		/// <summary>
+		/// Total length of the data in this node
+		/// </summary>
+		public int Length => _nodeData.Data.Length;
+
+		/// <summary>
+		/// Hash of the node being deserialized
 		/// </summary>
-		/// <param name="handle">Handle to the target node</param>
-		void WriteNodeHandle(NodeHandle handle);
+		public IoHash Hash => _nodeData.Hash;
+
+		/// <summary>
+		/// 
+		/// </summary>
+		public ReadOnlyMemory<byte> Data => _nodeData.Data;
+
+		/// <summary>
+		/// Locations of all referenced nodes.
+		/// </summary>
+		public IReadOnlyList<NodeHandle> References => _nodeData.Refs;
+
+		readonly NodeData _nodeData;
+		int _refIdx;
+
+		/// <summary>
+		/// Constructor
+		/// </summary>
+		public NodeReader(NodeData nodeData)
+			: base(nodeData.Data)
+		{
+			_nodeData = nodeData;
+		}
+
+		/// <summary>
+		/// Reads the next reference to another node
+		/// </summary>
+		public NodeHandle ReadNodeHandle()
+		{
+			IoHash hash = this.ReadIoHash();
+			return GetNodeHandle(_refIdx++, hash);
+		}
+
+		/// <summary>
+		/// Gets a node handle with the given index and hash
+		/// </summary>
+		/// <param name="index"></param>
+		/// <param name="hash"></param>
+		/// <returns></returns>
+		public NodeHandle GetNodeHandle(int index, IoHash hash)
+		{
+			Debug.Assert(_nodeData.Refs[index].Hash == hash);
+			return _nodeData.Refs[index];
+		}
+	}
+
+	/// <summary>
+	/// Writer for node objects, which tracks references to other nodes
+	/// </summary>
+	public sealed class NodeWriter : IMemoryWriter
+	{
+		readonly IStorageWriter _treeWriter;
+
+		Memory<byte> _memory;
+		readonly List<NodeHandle> _refs = new List<NodeHandle>();
+		int _length;
+
+		/// <summary>
+		/// List of serialized references
+		/// </summary>
+		public IReadOnlyList<NodeHandle> References => _refs;
+
+		/// <inheritdoc/>
+		public int Length => _length;
+
+		/// <summary>
+		/// Constructor
+		/// </summary>
+		/// <param name="treeWriter"></param>
+		public NodeWriter(IStorageWriter treeWriter)
+		{
+			_treeWriter = treeWriter;
+			_memory = treeWriter.GetOutputBuffer(0, 256 * 1024);
+		}
+
+		/// <summary>
+		/// Writes a handle to another node
+		/// </summary>
+		public void WriteNodeHandle(NodeHandle target)
+		{
+			this.WriteIoHash(target.Hash);
+			_refs.Add(target);
+		}
+
+		/// <inheritdoc/>
+		public Span<byte> GetSpan(int sizeHint = 0) => GetMemory(sizeHint).Span;
+
+		/// <inheritdoc/>
+		public Memory<byte> GetMemory(int sizeHint = 0)
+		{
+			int newLength = _length + Math.Max(sizeHint, 1);
+			if (newLength > _memory.Length)
+			{
+				newLength = _length + Math.Max(sizeHint, 1024);
+				_memory = _treeWriter.GetOutputBuffer(_length, Math.Max(_memory.Length * 2, newLength));
+			}
+			return _memory.Slice(_length);
+		}
+
+		/// <inheritdoc/>
+		public void Advance(int length) => _length += length;
 	}
 
 	/// <summary>
@@ -277,7 +413,7 @@ namespace EpicGames.Horde.Storage
 		/// <returns></returns>
 		public static async ValueTask<TNode> ReadAsync<TNode>(this NodeHandle handle, CancellationToken cancellationToken = default) where TNode : Node
 		{
-			NodeReader nodeData = await handle.ReadAsync(cancellationToken);
+			NodeData nodeData = await handle.ReadAsync(cancellationToken);
 			return Node.Deserialize<TNode>(nodeData);
 		}
 
@@ -341,7 +477,7 @@ namespace EpicGames.Horde.Storage
 		/// </summary>
 		/// <param name="writer">Writer to serialize to</param>
 		/// <param name="value">Value to write</param>
-		public static void WriteRef(this ITreeNodeWriter writer, NodeRef value)
+		public static void WriteRef(this NodeWriter writer, NodeRef value)
 		{
 			value.Serialize(writer);
 		}
@@ -351,7 +487,7 @@ namespace EpicGames.Horde.Storage
 		/// </summary>
 		/// <param name="writer">Writer to serialize to</param>
 		/// <param name="value">Value to write</param>
-		public static void WriteOptionalRef(this ITreeNodeWriter writer, NodeRef? value)
+		public static void WriteOptionalRef(this NodeWriter writer, NodeRef? value)
 		{
 			if (value == null)
 			{
@@ -409,7 +545,7 @@ namespace EpicGames.Horde.Storage
 				return null;
 			}
 
-			NodeReader nodeData = await refTarget.ReadAsync(cancellationToken);
+			NodeData nodeData = await refTarget.ReadAsync(cancellationToken);
 			return Node.Deserialize<TNode>(nodeData);
 		}
 
diff --git a/Engine/Source/Programs/Shared/EpicGames.Horde/Storage/NodeRef.cs b/Engine/Source/Programs/Shared/EpicGames.Horde/Storage/NodeRef.cs
--- a/Engine/Source/Programs/Shared/EpicGames.Horde/Storage/NodeRef.cs
+++ b/Engine/Source/Programs/Shared/EpicGames.Horde/Storage/NodeRef.cs
@@ -118,7 +118,7 @@ namespace EpicGames.Horde.Storage
 		/// Serialize the node to the given writer
 		/// </summary>
 		/// <param name="writer"></param>
-		public virtual void Serialize(ITreeNodeWriter writer)
+		public virtual void Serialize(NodeWriter writer)
 		{
 			Debug.Assert(Handle != null);
 			writer.WriteNodeHandle(Handle);
@@ -133,8 +133,8 @@ namespace EpicGames.Horde.Storage
 		{
 			if (Target == null)
 			{
-				NodeReader nodeReader = await Handle!.ReadAsync(cancellationToken);
-				Target = Node.Deserialize(nodeReader);
+				NodeData nodeData = await Handle!.ReadAsync(cancellationToken);
+				Target = Node.Deserialize(nodeData);
 				OnExpand();
 			}
 			return Target;
@@ -224,7 +224,7 @@ namespace EpicGames.Horde.Storage
 				throw new InvalidOperationException("TreeNodeRef has not been serialized to storage");
 			}
 
-			NodeReader nodeData = await Handle!.ReadAsync(cancellationToken);
+			NodeData nodeData = await Handle!.ReadAsync(cancellationToken);
 			return Node.Deserialize<T>(nodeData);
 		}
 	}
@@ -234,52 +234,6 @@ namespace EpicGames.Horde.Storage
 	/// </summary>
 	public static class NodeRefExtensions
 	{
-		// Implementation of INodeWriter that tracks refs
-		class NodeWriter : ITreeNodeWriter
-		{
-			readonly IStorageWriter _treeWriter;
-
-			Memory<byte> _memory;
-			int _length;
-			readonly IReadOnlyList<NodeHandle> _refs;
-			int _refIdx;
-
-			public int Length => _length;
-
-			public NodeWriter(IStorageWriter treeWriter, IReadOnlyList<NodeHandle> refs)
-			{
-				_treeWriter = treeWriter;
-				_memory = treeWriter.GetOutputBuffer(0, 256 * 1024);
-				_refs = refs;
-			}
-
-			public void WriteNodeHandle(NodeHandle target)
-			{
-				if (_refs[_refIdx] != target)
-				{
-					throw new InvalidOperationException("Referenced node does not match the handle returned by owner's EnumerateRefs method.");
-				}
-
-				this.WriteIoHash(target.Hash);
-				_refIdx++;
-			}
-
-			public Span<byte> GetSpan(int sizeHint = 0) => GetMemory(sizeHint).Span;
-
-			public Memory<byte> GetMemory(int sizeHint = 0)
-			{
-				int newLength = _length + Math.Max(sizeHint, 1);
-				if (newLength > _memory.Length)
-				{
-					newLength = _length + Math.Max(sizeHint, 1024);
-					_memory = _treeWriter.GetOutputBuffer(_length, Math.Max(_memory.Length * 2, newLength));
-				}
-				return _memory.Slice(_length);
-			}
-
-			public void Advance(int length) => _length += length;
-		}
-
 		/// <summary>
 		/// Class used to track nodes which are pending write (and the state of the object when the write was started)
 		/// </summary>
@@ -346,7 +300,7 @@ namespace EpicGames.Horde.Storage
 			}
 
 			// Serialize the node
-			NodeWriter nodeWriter = new NodeWriter(writer, nextRefHandles);
+			NodeWriter nodeWriter = new NodeWriter(writer);
 			target.Serialize(nodeWriter);
 
 			// Write the final data
diff --git a/Engine/Source/Programs/Shared/EpicGames.Horde/Storage/Nodes/CbNode.cs b/Engine/Source/Programs/Shared/EpicGames.Horde/Storage/Nodes/CbNode.cs
--- a/Engine/Source/Programs/Shared/EpicGames.Horde/Storage/Nodes/CbNode.cs
+++ b/Engine/Source/Programs/Shared/EpicGames.Horde/Storage/Nodes/CbNode.cs
@@ -82,7 +82,7 @@ namespace EpicGames.Horde.Storage.Nodes
 		}
 
 		/// <inheritdoc/>
-		public override void Serialize(ITreeNodeWriter writer)
+		public override void Serialize(NodeWriter writer)
 		{
 			writer.WriteFixedLengthBytes(Object.GetView().Span);
 		}
diff --git a/Engine/Source/Programs/Shared/EpicGames.Horde/Storage/Nodes/ChunkedDataNode.cs b/Engine/Source/Programs/Shared/EpicGames.Horde/Storage/Nodes/ChunkedDataNode.cs
--- a/Engine/Source/Programs/Shared/EpicGames.Horde/Storage/Nodes/ChunkedDataNode.cs
+++ b/Engine/Source/Programs/Shared/EpicGames.Horde/Storage/Nodes/ChunkedDataNode.cs
@@ -38,7 +38,7 @@ namespace EpicGames.Horde.Storage.Nodes
 		/// <param name="cancellationToken">Cancellation token for the operation</param>
 		public static async Task CopyToStreamAsync(NodeHandle handle, Stream outputStream, CancellationToken cancellationToken)
 		{
-			NodeReader nodeData = await handle.ReadAsync(cancellationToken);
+			NodeData nodeData = await handle.ReadAsync(cancellationToken);
 			if (nodeData.Type.Guid == s_leafNodeGuid)
 			{
 				await LeafChunkedDataNode.CopyToStreamAsync(nodeData, outputStream, cancellationToken);
@@ -114,7 +114,7 @@ namespace EpicGames.Horde.Storage.Nodes
 		}
 
 		/// <inheritdoc/>
-		public override void Serialize(ITreeNodeWriter writer)
+		public override void Serialize(NodeWriter writer)
 		{
 			writer.WriteFixedLengthBytes(Data.Span);
 		}
@@ -134,7 +134,7 @@ namespace EpicGames.Horde.Storage.Nodes
 		/// <param name="nodeData">The raw node data</param>
 		/// <param name="outputStream">The output stream to receive the data</param>
 		/// <param name="cancellationToken">Cancellation token for the operation</param>
-		public static async Task CopyToStreamAsync(NodeReader nodeData, Stream outputStream, CancellationToken cancellationToken)
+		public static async Task CopyToStreamAsync(NodeData nodeData, Stream outputStream, CancellationToken cancellationToken)
 		{
 			// Keep this code in sync with the constructor
 			await outputStream.WriteAsync(nodeData.Data, cancellationToken);
@@ -315,7 +315,7 @@ namespace EpicGames.Horde.Storage.Nodes
 		}
 
 		/// <inheritdoc/>
-		public override void Serialize(ITreeNodeWriter writer)
+		public override void Serialize(NodeWriter writer)
 		{
 			foreach (NodeRef<ChunkedDataNode> child in Children)
 			{
@@ -404,11 +404,12 @@ namespace EpicGames.Horde.Storage.Nodes
 		/// <summary>
 		/// Copy the contents of the node to the output stream without creating the intermediate FileNodes
 		/// </summary>
-		/// <param name="nodeReader">Source data</param>
+		/// <param name="nodeData">Source data</param>
 		/// <param name="outputStream">The output stream to receive the data</param>
 		/// <param name="cancellationToken">Cancellation token for the operation</param>
-		public static async Task CopyToStreamAsync(NodeReader nodeReader, Stream outputStream, CancellationToken cancellationToken)
+		public static async Task CopyToStreamAsync(NodeData nodeData, Stream outputStream, CancellationToken cancellationToken)
 		{
+			NodeReader nodeReader = new NodeReader(nodeData);
 			while (nodeReader.GetMemory(0).Length > 0)
 			{
 				NodeRef nodeRef = nodeReader.ReadRef();
diff --git a/Engine/Source/Programs/Shared/EpicGames.Horde/Storage/Nodes/CommitNode.cs b/Engine/Source/Programs/Shared/EpicGames.Horde/Storage/Nodes/CommitNode.cs
--- a/Engine/Source/Programs/Shared/EpicGames.Horde/Storage/Nodes/CommitNode.cs
+++ b/Engine/Source/Programs/Shared/EpicGames.Horde/Storage/Nodes/CommitNode.cs
@@ -100,7 +100,7 @@ namespace EpicGames.Horde.Storage.Nodes
 		}
 
 		/// <inheritdoc/>
-		public override void Serialize(ITreeNodeWriter writer)
+		public override void Serialize(NodeWriter writer)
 		{
 			writer.WriteUnsignedVarInt(Number);
 			writer.WriteOptionalRef(Parent);
diff --git a/Engine/Source/Programs/Shared/EpicGames.Horde/Storage/Nodes/DirectoryEntry.cs b/Engine/Source/Programs/Shared/EpicGames.Horde/Storage/Nodes/DirectoryEntry.cs
--- a/Engine/Source/Programs/Shared/EpicGames.Horde/Storage/Nodes/DirectoryEntry.cs
+++ b/Engine/Source/Programs/Shared/EpicGames.Horde/Storage/Nodes/DirectoryEntry.cs
@@ -45,7 +45,7 @@ namespace EpicGames.Horde.Storage.Nodes
 		/// Serialize this directory entry to disk
 		/// </summary>
 		/// <param name="writer"></param>
-		public override void Serialize(ITreeNodeWriter writer)
+		public override void Serialize(NodeWriter writer)
 		{
 			base.Serialize(writer);
 
diff --git a/Engine/Source/Programs/Shared/EpicGames.Horde/Storage/Nodes/DirectoryNode.cs b/Engine/Source/Programs/Shared/EpicGames.Horde/Storage/Nodes/DirectoryNode.cs
--- a/Engine/Source/Programs/Shared/EpicGames.Horde/Storage/Nodes/DirectoryNode.cs
+++ b/Engine/Source/Programs/Shared/EpicGames.Horde/Storage/Nodes/DirectoryNode.cs
@@ -150,7 +150,7 @@ namespace EpicGames.Horde.Storage.Nodes
 		}
 
 		/// <inheritdoc/>
-		public override void Serialize(ITreeNodeWriter writer)
+		public override void Serialize(NodeWriter writer)
 		{
 			writer.WriteUnsignedVarInt((ulong)Flags);
 
diff --git a/Engine/Source/Programs/Shared/EpicGames.Horde/Storage/Nodes/DirectoryNodeRef.cs b/Engine/Source/Programs/Shared/EpicGames.Horde/Storage/Nodes/DirectoryNodeRef.cs
--- a/Engine/Source/Programs/Shared/EpicGames.Horde/Storage/Nodes/DirectoryNodeRef.cs
+++ b/Engine/Source/Programs/Shared/EpicGames.Horde/Storage/Nodes/DirectoryNodeRef.cs
@@ -52,7 +52,7 @@ namespace EpicGames.Horde.Storage.Nodes
 		/// Serialize this directory entry to disk
 		/// </summary>
 		/// <param name="writer"></param>
-		public override void Serialize(ITreeNodeWriter writer)
+		public override void Serialize(NodeWriter writer)
 		{
 			base.Serialize(writer);
 
diff --git a/Engine/Source/Programs/Shared/EpicGames.Horde/Storage/Nodes/FileEntry.cs b/Engine/Source/Programs/Shared/EpicGames.Horde/Storage/Nodes/FileEntry.cs
--- a/Engine/Source/Programs/Shared/EpicGames.Horde/Storage/Nodes/FileEntry.cs
+++ b/Engine/Source/Programs/Shared/EpicGames.Horde/Storage/Nodes/FileEntry.cs
@@ -111,7 +111,7 @@ namespace EpicGames.Horde.Storage.Nodes
 		/// Serialize this entry
 		/// </summary>
 		/// <param name="writer"></param>
-		public override void Serialize(ITreeNodeWriter writer)
+		public override void Serialize(NodeWriter writer)
 		{
 			base.Serialize(writer);
 
diff --git a/Engine/Source/Programs/Shared/EpicGames.Horde/Storage/TreeReader.cs b/Engine/Source/Programs/Shared/EpicGames.Horde/Storage/TreeReader.cs
--- a/Engine/Source/Programs/Shared/EpicGames.Horde/Storage/TreeReader.cs
+++ b/Engine/Source/Programs/Shared/EpicGames.Horde/Storage/TreeReader.cs
@@ -14,81 +14,6 @@ using System.Diagnostics;
 
 namespace EpicGames.Horde.Storage
 {
-	/// <summary>
-	/// Data for an individual node
-	/// </summary>
-	public record struct NodeData(NodeType Type, IoHash Hash, ReadOnlyMemory<byte> Data, IReadOnlyList<NodeLocator> Refs);
-
-	/// <summary>
-	/// Reader for tree nodes
-	/// </summary>
-	public class NodeReader : MemoryReader
-	{
-		/// <summary>
-		/// Type to deserialize
-		/// </summary>
-		public NodeType Type => _nodeData.Type;
-
-		/// <summary>
-		/// Version of the current node, as specified via <see cref="NodeTypeAttribute"/>
-		/// </summary>
-		public int Version => Type.Version;
-
-		/// <summary>
-		/// Total length of the data in this node
-		/// </summary>
-		public int Length => _nodeData.Data.Length;
-
-		/// <summary>
-		/// Hash of the node being deserialized
-		/// </summary>
-		public IoHash Hash => _nodeData.Hash;
-
-		/// <summary>
-		/// 
-		/// </summary>
-		public ReadOnlyMemory<byte> Data => _nodeData.Data;
-
-		/// <summary>
-		/// Locations of all referenced nodes.
-		/// </summary>
-		public IReadOnlyList<NodeLocator> References => _nodeData.Refs;
-
-		readonly TreeReader _treeReader;
-		readonly NodeData _nodeData;
-		int _refIdx;
-
-		/// <summary>
-		/// Constructor
-		/// </summary>
-		public NodeReader(TreeReader treeReader, NodeData nodeData)
-			: base(nodeData.Data)
-		{
-			_treeReader = treeReader;
-			_nodeData = nodeData;
-		}
-
-		/// <summary>
-		/// Reads the next reference to another node
-		/// </summary>
-		public NodeHandle ReadNodeHandle()
-		{
-			IoHash hash = this.ReadIoHash();
-			return GetNodeHandle(_refIdx++, hash);
-		}
-
-		/// <summary>
-		/// Gets a node handle with the given index and hash
-		/// </summary>
-		/// <param name="index"></param>
-		/// <param name="hash"></param>
-		/// <returns></returns>
-		public NodeHandle GetNodeHandle(int index, IoHash hash)
-		{
-			return new NodeHandle(_treeReader, hash, _nodeData.Refs[index]);
-		}
-	}
-
 	/// <summary>
 	/// Writes nodes from bundles in an <see cref="IStorageClient"/> instance.
 	/// </summary>
@@ -619,7 +544,7 @@ namespace EpicGames.Horde.Storage
 			BundleInfo bundleInfo = await GetBundleInfoAsync(locator.Blob, cancellationToken);
 			BundleExport export = bundleInfo.Header.Exports[locator.ExportIdx];
 
-			List<NodeLocator> refs = new List<NodeLocator>(export.References.Count);
+			List<NodeHandle> refs = new List<NodeHandle>(export.References.Count);
 			foreach (BundleExportRef reference in export.References)
 			{
 				BlobLocator importBlob;
@@ -632,7 +557,7 @@ namespace EpicGames.Horde.Storage
 					importBlob = bundleInfo.Header.Imports[reference.ImportIdx];
 				}
 				Debug.Assert(importBlob.IsValid());
-				refs.Add(new NodeLocator(importBlob, reference.NodeIdx));
+				refs.Add(new NodeHandle(this, reference.Hash, new NodeLocator(importBlob, reference.NodeIdx)));
 			}
 
 			ReadOnlyMemory<byte> nodeData = ReadOnlyMemory<byte>.Empty;
@@ -655,7 +580,7 @@ namespace EpicGames.Horde.Storage
 		public async ValueTask<Node> ReadNodeAsync(NodeLocator locator, CancellationToken cancellationToken = default)
 		{
 			NodeData nodeData = await ReadNodeDataAsync(locator, cancellationToken);
-			return Node.Deserialize(new NodeReader(this, nodeData));
+			return Node.Deserialize(nodeData);
 		}
 
 		/// <summary>
diff --git a/Engine/Source/Programs/Shared/EpicGames.Horde/Storage/TreeWriter.cs b/Engine/Source/Programs/Shared/EpicGames.Horde/Storage/TreeWriter.cs
--- a/Engine/Source/Programs/Shared/EpicGames.Horde/Storage/TreeWriter.cs
+++ b/Engine/Source/Programs/Shared/EpicGames.Horde/Storage/TreeWriter.cs
@@ -117,10 +117,9 @@ namespace EpicGames.Horde.Storage
 		}
 
 		/// <inheritdoc/>
-		public async ValueTask<NodeReader> ReadAsync(CancellationToken cancellationToken = default)
+		public async ValueTask<NodeData> ReadAsync(CancellationToken cancellationToken = default)
 		{
-			NodeData nodeData = await _reader.ReadNodeDataAsync(Locator, cancellationToken);
-			return new NodeReader(_reader, nodeData);
+			return await _reader.ReadNodeDataAsync(Locator, cancellationToken);
 		}
 
 		/// <inheritdoc/>
