Horde: Make NodeHandle an abstract base class, and add separate implementations for flushed and pending nodes.

diff --git a/Engine/Source/Programs/Horde/Horde.Server/Storage/StorageService.cs b/Engine/Source/Programs/Horde/Horde.Server/Storage/StorageService.cs
--- a/Engine/Source/Programs/Horde/Horde.Server/Storage/StorageService.cs
+++ b/Engine/Source/Programs/Horde/Horde.Server/Storage/StorageService.cs
@@ -240,7 +240,7 @@ namespace Horde.Server.Storage
 			{
 				await foreach (HashedNodeLocator locator in _outer.FindNodesAsync(NamespaceId, alias, cancellationToken))
 				{
-					yield return new NodeHandle(TreeReader, locator);
+					yield return new FlushedNodeHandle(TreeReader, locator);
 				}
 			}
 
@@ -256,7 +256,7 @@ namespace Horde.Server.Storage
 				{
 					return null;
 				}
-				return new NodeHandle(TreeReader, locator.Value);
+				return new FlushedNodeHandle(TreeReader, locator.Value);
 			}
 
 			/// <inheritdoc/>
diff --git a/Engine/Source/Programs/Shared/EpicGames.Horde/Storage/Backends/FileStorageClient.cs b/Engine/Source/Programs/Shared/EpicGames.Horde/Storage/Backends/FileStorageClient.cs
--- a/Engine/Source/Programs/Shared/EpicGames.Horde/Storage/Backends/FileStorageClient.cs
+++ b/Engine/Source/Programs/Shared/EpicGames.Horde/Storage/Backends/FileStorageClient.cs
@@ -42,7 +42,7 @@ namespace EpicGames.Horde.Storage.Backends
 		public async ValueTask<NodeHandle> ReadRefAsync(FileReference file)
 		{
 			string text = await FileReference.ReadAllTextAsync(file);
-			return new NodeHandle(TreeReader, HashedNodeLocator.Parse(text));
+			return new FlushedNodeHandle(TreeReader, HashedNodeLocator.Parse(text));
 		}
 
 		FileReference GetRefFile(RefName name) => FileReference.Combine(_rootDir, name.ToString() + ".ref");
@@ -127,7 +127,7 @@ namespace EpicGames.Horde.Storage.Backends
 
 			_logger.LogInformation("Reading {File}", file);
 			string text = await FileReference.ReadAllTextAsync(file, cancellationToken);
-			return new NodeHandle(TreeReader, HashedNodeLocator.Parse(text));
+			return new FlushedNodeHandle(TreeReader, HashedNodeLocator.Parse(text));
 		}
 
 		/// <inheritdoc/>
diff --git a/Engine/Source/Programs/Shared/EpicGames.Horde/Storage/Backends/HttpStorageClient.cs b/Engine/Source/Programs/Shared/EpicGames.Horde/Storage/Backends/HttpStorageClient.cs
--- a/Engine/Source/Programs/Shared/EpicGames.Horde/Storage/Backends/HttpStorageClient.cs
+++ b/Engine/Source/Programs/Shared/EpicGames.Horde/Storage/Backends/HttpStorageClient.cs
@@ -220,7 +220,7 @@ namespace EpicGames.Horde.Storage.Backends
 						FindNodesResponse? message = await response.Content.ReadFromJsonAsync<FindNodesResponse>(cancellationToken: cancellationToken);
 						foreach (FindNodeResponse node in message!.Nodes)
 						{
-							yield return new NodeHandle(TreeReader, node.Hash, new NodeLocator(node.Blob, node.ExportIdx));
+							yield return new FlushedNodeHandle(TreeReader, node.Hash, new NodeLocator(node.Blob, node.ExportIdx));
 						}
 					}
 				}
@@ -271,7 +271,7 @@ namespace EpicGames.Horde.Storage.Backends
 							response.EnsureSuccessStatusCode();
 							ReadRefResponse? data = await response.Content.ReadFromJsonAsync<ReadRefResponse>(cancellationToken: cancellationToken);
 							_logger.LogDebug("Read ref {RefName} -> {Blob}#{ExportIdx}", name, data!.Blob, data!.ExportIdx);
-							return new NodeHandle(TreeReader, data.Hash, new NodeLocator(data!.Blob, data!.ExportIdx));
+							return new FlushedNodeHandle(TreeReader, data.Hash, new NodeLocator(data!.Blob, data!.ExportIdx));
 						}
 					}
 				}
@@ -284,7 +284,7 @@ namespace EpicGames.Horde.Storage.Backends
 			_logger.LogDebug("Writing ref {RefName} -> {RefTarget}", name, target);
 			using (HttpClient httpClient = _createClient())
 			{
-				NodeLocator locator = target.GetLocator();
+				NodeLocator locator = await target.FlushAsync(cancellationToken);
 				using (HttpResponseMessage response = await httpClient.PutAsync($"refs/{name}", new { blob = locator.Blob, exportIdx = locator.ExportIdx, options }, cancellationToken))
 				{
 					response.EnsureSuccessStatusCode();
diff --git a/Engine/Source/Programs/Shared/EpicGames.Horde/Storage/BundleStorageClient.cs b/Engine/Source/Programs/Shared/EpicGames.Horde/Storage/BundleStorageClient.cs
--- a/Engine/Source/Programs/Shared/EpicGames.Horde/Storage/BundleStorageClient.cs
+++ b/Engine/Source/Programs/Shared/EpicGames.Horde/Storage/BundleStorageClient.cs
@@ -77,7 +77,7 @@ namespace EpicGames.Horde.Storage
 		public virtual async Task<NodeHandle> WriteRefAsync(RefName name, Bundle bundle, int exportIdx = 0, Utf8String prefix = default, RefOptions? options = null, CancellationToken cancellationToken = default)
 		{
 			BlobLocator locator = await this.WriteBundleAsync(bundle, prefix, cancellationToken);
-			NodeHandle target = new NodeHandle(TreeReader, bundle.Header.Exports[exportIdx].Hash, new NodeLocator(locator, exportIdx));
+			NodeHandle target = new FlushedNodeHandle(TreeReader, bundle.Header.Exports[exportIdx].Hash, new NodeLocator(locator, exportIdx));
 			await WriteRefTargetAsync(name, target, options, cancellationToken);
 			return target;
 		}
diff --git a/Engine/Source/Programs/Shared/EpicGames.Horde/Storage/NodeRef.cs b/Engine/Source/Programs/Shared/EpicGames.Horde/Storage/NodeRef.cs
--- a/Engine/Source/Programs/Shared/EpicGames.Horde/Storage/NodeRef.cs
+++ b/Engine/Source/Programs/Shared/EpicGames.Horde/Storage/NodeRef.cs
@@ -237,12 +237,12 @@ namespace EpicGames.Horde.Storage
 		/// <summary>
 		/// Class used to track nodes which are pending write (and the state of the object when the write was started)
 		/// </summary>
-		internal class NodeWriteCallback : TreeWriter.WriteCallback
+		internal class NodeRefWriteCallback : NodeWriteCallback
 		{
 			readonly NodeRef _nodeRef;
 			readonly NodeHandle _handle;
 
-			public NodeWriteCallback(NodeRef nodeRef, NodeHandle handle)
+			public NodeRefWriteCallback(NodeRef nodeRef, NodeHandle handle)
 			{
 				_nodeRef = nodeRef;
 				_handle = handle;
@@ -307,7 +307,7 @@ namespace EpicGames.Horde.Storage
 			NodeHandle handle = await writer.WriteNodeAsync(nodeWriter.Length, nextRefHandles, target.NodeType, cancellationToken);
 			target.Hash = handle.Hash;
 
-			NodeWriteCallback writeState = new NodeWriteCallback(nodeRef, handle);
+			NodeRefWriteCallback writeState = new NodeRefWriteCallback(nodeRef, handle);
 			handle.AddWriteCallback(writeState);
 
 			nodeRef.MarkAsPendingWrite(handle);
diff --git a/Engine/Source/Programs/Shared/EpicGames.Horde/Storage/TreeReader.cs b/Engine/Source/Programs/Shared/EpicGames.Horde/Storage/TreeReader.cs
--- a/Engine/Source/Programs/Shared/EpicGames.Horde/Storage/TreeReader.cs
+++ b/Engine/Source/Programs/Shared/EpicGames.Horde/Storage/TreeReader.cs
@@ -557,7 +557,7 @@ namespace EpicGames.Horde.Storage
 					importBlob = bundleInfo.Header.Imports[reference.ImportIdx];
 				}
 				Debug.Assert(importBlob.IsValid());
-				refs.Add(new NodeHandle(this, reference.Hash, new NodeLocator(importBlob, reference.NodeIdx)));
+				refs.Add(new FlushedNodeHandle(this, reference.Hash, new NodeLocator(importBlob, reference.NodeIdx)));
 			}
 
 			ReadOnlyMemory<byte> nodeData = ReadOnlyMemory<byte>.Empty;
diff --git a/Engine/Source/Programs/Shared/EpicGames.Horde/Storage/TreeWriter.cs b/Engine/Source/Programs/Shared/EpicGames.Horde/Storage/TreeWriter.cs
--- a/Engine/Source/Programs/Shared/EpicGames.Horde/Storage/TreeWriter.cs
+++ b/Engine/Source/Programs/Shared/EpicGames.Horde/Storage/TreeWriter.cs
@@ -53,100 +53,105 @@ namespace EpicGames.Horde.Storage
 	public record NodeKey(IoHash Hash, NodeType Type);
 
 	/// <summary>
-	/// Handle to a node. Can be used to reference nodes that have not been flushed yet.
+	/// Object to receive notifications on a node being written
 	/// </summary>
-	public class NodeHandle
+	public abstract class NodeWriteCallback
 	{
-		readonly TreeReader _reader;
-		NodeLocator _locator;
+		internal NodeWriteCallback? _next;
 
 		/// <summary>
-		/// Hash of the target node
+		/// Callback for the node being written
 		/// </summary>
-		public IoHash Hash { get; }
+		public abstract void OnWrite();
+	}
 
+	/// <summary>
+	/// Handle to a node. Can be used to reference nodes that have not been flushed yet.
+	/// </summary>
+	public abstract class NodeHandle
+	{
 		/// <summary>
-		/// Used to track the bundle that is being written to
+		/// Hash of the target node
 		/// </summary>
-		internal PendingBundle? PendingBundle { get; set; }
+		public IoHash Hash { get; }
 
 		/// <summary>
 		/// Constructor
 		/// </summary>
-		/// <param name="reader">Reader for new node data</param>
 		/// <param name="hash">Hash of the target node</param>
-		/// <param name="locator">Location of the node in storage</param>
-		public NodeHandle(TreeReader reader, IoHash hash, NodeLocator locator)
-		{
-			_reader = reader;
-			Hash = hash;
-			_locator = locator;
-		}
-
-		/// <summary>
-		/// Constructor
-		/// </summary>
-		/// <param name="reader">Reader for new node data</param>
-		/// <param name="hashedLocator">Location of the node in storage</param>
-		public NodeHandle(TreeReader reader, HashedNodeLocator hashedLocator)
-			: this(reader, hashedLocator.Hash, hashedLocator.Locator)
-		{
-		}
+		protected NodeHandle(IoHash hash) => Hash = hash;
 
 		/// <summary>
 		/// Determines if the node has been written to storage
 		/// </summary>
-		/// <returns></returns>
-		public bool HasLocator() => _locator.IsValid();
+		public abstract bool HasLocator();
 
 		/// <summary>
 		/// Gets the node locator. May throw if the node has not been written to storage yet.
 		/// </summary>
 		/// <returns>Locator for the node</returns>
-		public NodeLocator GetLocator()
-		{
-			if (!_locator.IsValid())
-			{
-				throw new InvalidOperationException("Node has not been flushed to disk; cannot retrieve locator.");
-			}
-			return _locator;
-		}
+		public abstract NodeLocator GetLocator();
 
 		/// <summary>
-		/// Sets the current locator
+		/// Adds a callback to be executed once the node has been written. Triggers immediately if the node has already been written.
+		/// </summary>
+		/// <param name="callback">Action to be executed after the write</param>
+		public abstract void AddWriteCallback(NodeWriteCallback callback);
+
+		/// <summary>
+		/// Creates a reader for this node's data
 		/// </summary>
-		/// <param name="locator"></param>
-		protected void SetLocator(NodeLocator locator)
+		public abstract ValueTask<NodeData> ReadAsync(CancellationToken cancellationToken = default);
+
+		/// <summary>
+		/// Flush the node to storage and retrieve its locator
+		/// </summary>
+		public abstract ValueTask<NodeLocator> FlushAsync(CancellationToken cancellationToken = default);
+
+		/// <inheritdoc/>
+		public override string ToString() => HasLocator()? new HashedNodeLocator(Hash, GetLocator()).ToString() : Hash.ToString();
+	}
+
+	/// <summary>
+	/// Implementation of <see cref="NodeHandle"/> for nodes which can be read from storage
+	/// </summary>
+	public sealed class FlushedNodeHandle : NodeHandle
+	{
+		readonly TreeReader _reader;
+		readonly NodeLocator _locator;
+
+		/// <summary>
+		/// Constructor
+		/// </summary>
+		public FlushedNodeHandle(TreeReader reader, HashedNodeLocator hashedLocator)
+			: this(reader, hashedLocator.Hash, hashedLocator.Locator)
 		{
-			if (_locator.IsValid())
-			{
-				throw new InvalidOperationException("Node has already been flushed.");
-			}
-			_locator = locator;
 		}
 
 		/// <summary>
-		/// Adds a callback to be executed once the node has been written. Triggers immediately if the node has already been written.
+		/// Constructor
 		/// </summary>
-		/// <param name="callback">Action to be executed after the write</param>
-		internal void AddWriteCallback(WriteCallback callback)
+		public FlushedNodeHandle(TreeReader reader, IoHash hash, NodeLocator locator)
+			: base(hash)
 		{
-			PendingBundle? pendingBundle = PendingBundle;
-			if (pendingBundle == null || !pendingBundle.TryAddWriteCallback(callback))
-			{
-				Debug.Assert(_locator.IsValid());
-				callback.OnWrite();
-			}
+			_reader = reader;
+			_locator = locator;
 		}
 
 		/// <inheritdoc/>
-		public async ValueTask<NodeData> ReadAsync(CancellationToken cancellationToken = default)
-		{
-			return await _reader.ReadNodeDataAsync(GetLocator(), cancellationToken);
-		}
+		public override bool HasLocator() => true;
+
+		/// <inheritdoc/>
+		public override NodeLocator GetLocator() => _locator;
 
 		/// <inheritdoc/>
-		public override string ToString() => _locator.IsValid()? new HashedNodeLocator(Hash, _locator).ToString() : Hash.ToString();
+		public override void AddWriteCallback(NodeWriteCallback callback) => callback.OnWrite();
+
+		/// <inheritdoc/>
+		public override ValueTask<NodeData> ReadAsync(CancellationToken cancellationToken = default) => _reader.ReadNodeDataAsync(_locator, cancellationToken);
+
+		/// <inheritdoc/>
+		public override ValueTask<NodeLocator> FlushAsync(CancellationToken cancellationToken = default) => new ValueTask<NodeLocator>(_locator);
 	}
 
 	/// <summary>
@@ -214,47 +219,113 @@ namespace EpicGames.Horde.Storage
 	/// </summary>
 	public sealed class TreeWriter : IStorageWriter
 	{
-		/// <summary>
-		/// Object to receive notifications on a node being written
-		/// </summary>
-		internal abstract class WriteCallback
-		{
-			internal WriteCallback? _next;
-
-			public abstract void OnWrite();
-		}
-
 		// Information about a unique output node. Note that multiple node refs may de-duplicate to the same output node.
 		internal class PendingNode : NodeHandle
 		{
+			readonly TreeReader _reader;
+
+			object LockObject => _reader;
+
+			NodeLocator _locator;
+			PendingBundle? _pendingBundle;
+
 			public readonly NodeKey Key;
 			public readonly int Packet;
 			public readonly int Offset;
 			public readonly int Length;
 			public readonly NodeHandle[] Refs;
 
+			public PendingBundle? PendingBundle => _pendingBundle;
+
 			public PendingNode(TreeReader reader, NodeKey key, int packet, int offset, int length, IReadOnlyList<NodeHandle> refs, PendingBundle pendingBundle)
-				: base(reader, key.Hash, default)
+				: base(key.Hash)
 			{
+				_reader = reader;
+
 				Key = key;
 				Packet = packet;
 				Offset = offset;
 				Length = length;
 				Refs = refs.ToArray();
-				PendingBundle = pendingBundle;
+
+				_pendingBundle = pendingBundle;
+			}
+
+			/// <inheritdoc/>
+			public override bool HasLocator() => _locator.IsValid();
+
+			/// <inheritdoc/>
+			public override NodeLocator GetLocator()
+			{
+				if (!_locator.IsValid())
+				{
+					throw new InvalidOperationException();
+				}
+				return _locator;
 			}
 
 			public void MarkAsWritten(NodeLocator locator)
 			{
-				SetLocator(locator);
-				PendingBundle = null;
+				lock (LockObject)
+				{
+					Debug.Assert(!_locator.IsValid());
+					_locator = locator;
+					_pendingBundle = null;
+				}
+			}
+
+			/// <inheritdoc/>
+			public override void AddWriteCallback(NodeWriteCallback callback)
+			{
+				PendingBundle? pendingBundle = PendingBundle;
+				if (pendingBundle == null || !pendingBundle.TryAddWriteCallback(callback))
+				{
+					Debug.Assert(_locator.IsValid());
+					callback.OnWrite();
+				}
+			}
+
+			/// <inheritdoc/>
+			public override async ValueTask<NodeData> ReadAsync(CancellationToken cancellationToken = default)
+			{
+				if (!_locator.IsValid())
+				{
+					lock (LockObject)
+					{
+						if (!_locator.IsValid() && _pendingBundle != null)
+						{
+							ReadOnlyMemory<byte> data = _pendingBundle.GetNodeData(Packet, Offset, Length);
+							if (data.Length == Length)
+							{
+								return new NodeData(Key.Type, Hash, data, Refs);
+							}
+						}
+					}
+				}
+
+				return await _reader.ReadNodeDataAsync(_locator, cancellationToken);
+			}
+
+			/// <inheritdoc/>
+			public override async ValueTask<NodeLocator> FlushAsync(CancellationToken cancellationToken = default)
+			{
+				Task? completeTask = _pendingBundle?.CompleteTask;
+				if (completeTask != null)
+				{
+					await completeTask;
+				}
+				if (!_locator.IsValid())
+				{
+					throw new InvalidOperationException("Locator should have been flushed already");
+				}
+				return _locator;
 			}
 		}
 
 		// Information about a bundle being built. Metadata operations are synchronous, compression/writes are asynchronous.
 		internal class PendingBundle : IDisposable
 		{
-			class WriteCallbackSentinel : WriteCallback
+			class WriteCallbackSentinel : NodeWriteCallback
 			{
 				public override void OnWrite() => throw new NotImplementedException();
 			}
@@ -268,6 +339,8 @@ namespace EpicGames.Horde.Storage
 			readonly int _maxPacketSize;
 			readonly BundleCompressionFormat _compressionFormat;
 
+			readonly object _lockObject = new object();
+
 			// The current packet being built
 			IMemoryOwner<byte>? _currentPacket;
 
@@ -299,7 +372,7 @@ namespace EpicGames.Horde.Storage
 			public bool IsReadOnly { get; private set; }
 
 			// List of post-write callbacks
-			WriteCallback? _callbacks = null;
+			NodeWriteCallback? _callbacks = null;
 
 			// Task used to compress data in the background
 			Task _writeTask = Task.CompletedTask;
@@ -352,11 +425,11 @@ namespace EpicGames.Horde.Storage
 			}
 
 			// Adds a callback after writing
-			public bool TryAddWriteCallback(WriteCallback callback)
+			public bool TryAddWriteCallback(NodeWriteCallback callback)
 			{
 				for (; ; )
 				{
-					WriteCallback? tail = _callbacks;
+					NodeWriteCallback? tail = _callbacks;
 					if (tail == s_callbackSentinel)
 					{
 						return false;
@@ -422,6 +495,37 @@ namespace EpicGames.Horde.Storage
 				return pendingNode;
 			}
 
+			public ReadOnlyMemory<byte> GetNodeData(int packetIdx, int offset, int length)
+			{
+				lock (_lockObject)
+				{
+					if (packetIdx < _packets.Count)
+					{
+						BundlePacket packet = _packets[packetIdx];
+
+						ReadOnlySequence<byte> sequence = _encodedPacketWriter.AsSequence(packet.EncodedOffset, packet.EncodedLength);
+						if (!sequence.IsSingleSegment)
+						{
+							throw new InvalidOperationException();
+						}
+
+						using (IMemoryOwner<byte> decodeBuffer = MemoryPool<byte>.Shared.Rent(packet.DecodedLength))
+						{
+							BundleData.Decompress(_compressionFormat, sequence.First, decodeBuffer.Memory);
+							return decodeBuffer.Memory.Slice(offset, length).ToArray();
+						}
+					}
+					else if (packetIdx == _currentPacketIdx)
+					{
+						return _currentPacket!.Memory.Slice(offset, length).Slice(offset, length).ToArray();
+					}
+					else
+					{
+						return ReadOnlyMemory<byte>.Empty;
+					}
+				}
+			}
+
 			// Compresses the current packet and schedule it to be written to storage
 			void FlushPacket()
 			{
@@ -447,13 +551,16 @@ namespace EpicGames.Horde.Storage
 
 			void CompressPacket(IMemoryOwner<byte> packetData, int length)
 			{
-				int encodedLength = BundleData.Compress(_compressionFormat, packetData.Memory.Slice(0, length), _encodedPacketWriter);
+				lock (_lockObject)
+				{
+					int encodedLength = BundleData.Compress(_compressionFormat, packetData.Memory.Slice(0, length), _encodedPacketWriter);
 
-				int encodedOffset = (_packets.Count == 0) ? 0 : _packets[^1].EncodedOffset + _packets[^1].EncodedLength;
-				BundlePacket packet = new BundlePacket(_compressionFormat, encodedOffset, encodedLength, length);
-				_packets.Add(packet);
+					int encodedOffset = (_packets.Count == 0) ? 0 : _packets[^1].EncodedOffset + _packets[^1].EncodedLength;
+					BundlePacket packet = new BundlePacket(_compressionFormat, encodedOffset, encodedLength, length);
+					_packets.Add(packet);
 
-				packetData.Dispose();
+					packetData.Dispose();
+				}
 			}
 
 			// Mark the bundle as complete
@@ -497,7 +604,7 @@ namespace EpicGames.Horde.Storage
 						_queue[idx].MarkAsWritten(nodeLocator);
 					}
 
-					WriteCallback? callback = Interlocked.Exchange(ref _callbacks, s_callbackSentinel);
+					NodeWriteCallback? callback = Interlocked.Exchange(ref _callbacks, s_callbackSentinel);
 					while (callback != null)
 					{
 						callback.OnWrite();
@@ -521,6 +628,14 @@ namespace EpicGames.Horde.Storage
 
 			// Mark the bundle as complete and create a bundle with the current state
 			public Bundle CreateBundle()
+			{
+				lock (_lockObject)
+				{
+					return CreateBundleInternal();
+				}
+			}
+
+			Bundle CreateBundleInternal()
 			{
 				// Create a set from the nodes to be written. We use this to determine references that are imported.
 				HashSet<NodeHandle> nodeSet = new HashSet<NodeHandle>(_queue);
@@ -763,9 +878,10 @@ namespace EpicGames.Horde.Storage
 			// Add dependencies on all bundles containing a dependent node
 			foreach (NodeHandle reference in references)
 			{
-				if (reference.PendingBundle != null)
+				PendingNode? pendingReference = reference as PendingNode;
+				if (pendingReference?.PendingBundle != null)
 				{
-					currentBundle.AddDependencyOn(reference.PendingBundle, _traceLogger);
+					currentBundle.AddDependencyOn(pendingReference.PendingBundle, _traceLogger);
 				}
 			}
 
