Horde: Initiate a flush on the tree writer whenever flushing a pending node.

diff --git a/Engine/Source/Programs/Shared/EpicGames.Horde/Storage/TreeWriter.cs b/Engine/Source/Programs/Shared/EpicGames.Horde/Storage/TreeWriter.cs
--- a/Engine/Source/Programs/Shared/EpicGames.Horde/Storage/TreeWriter.cs
+++ b/Engine/Source/Programs/Shared/EpicGames.Horde/Storage/TreeWriter.cs
@@ -309,10 +309,10 @@ namespace EpicGames.Horde.Storage
 			/// <inheritdoc/>
 			public override async ValueTask<NodeLocator> FlushAsync(CancellationToken cancellationToken = default)
 			{
-				Task? completeTask = _pendingBundle?.CompleteTask;
-				if (completeTask != null)
+				PendingBundle? pendingBundle = _pendingBundle;
+				if (!_locator.IsValid() && pendingBundle != null)
 				{
-					await completeTask;
+					await pendingBundle.FlushAsync(cancellationToken);
 				}
 				if (!_locator.IsValid())
 				{
@@ -336,6 +336,7 @@ namespace EpicGames.Horde.Storage
 			public int BundleId { get; } = Interlocked.Increment(ref s_lastBundleId);
 
 			readonly TreeReader _treeReader;
+			readonly TreeWriter _treeWriter;
 			readonly int _maxPacketSize;
 			readonly BundleCompressionFormat _compressionFormat;
 
@@ -383,9 +384,10 @@ namespace EpicGames.Horde.Storage
 			// Task signalled after the write is complete
 			public Task CompleteTask => _completeEvent.Task;
 
-			public PendingBundle(TreeReader treeReader, int maxPacketSize, int maxBlobSize, BundleCompressionFormat compressionFormat)
+			public PendingBundle(TreeReader treeReader, TreeWriter treeWriter, int maxPacketSize, int maxBlobSize, BundleCompressionFormat compressionFormat)
 			{
 				_treeReader = treeReader;
+				_treeWriter = treeWriter;
 				_maxPacketSize = maxPacketSize;
 				_compressionFormat = compressionFormat;
 				_encodedPacketWriter = new ChunkedMemoryWriter(maxBlobSize);
@@ -620,11 +622,7 @@ namespace EpicGames.Horde.Storage
 			}
 
 			// Flush the current write state
-			public Task FlushAsync()
-			{
-				FlushPacket();
-				return _writeTask;
-			}
+			public Task FlushAsync(CancellationToken cancellationToken) => _treeWriter.FlushAsync(cancellationToken);
 
 			// Mark the bundle as complete and create a bundle with the current state
 			public Bundle CreateBundle()
@@ -905,7 +903,7 @@ namespace EpicGames.Horde.Storage
 			if (_currentBundle == null || _currentBundle.IsReadOnly)
 			{
 				int bufferSize = (int)(_options.MinCompressionPacketSize * 1.2);
-				_currentBundle = new PendingBundle(_treeReader, bufferSize, _options.MaxBlobSize, _options.CompressionFormat);
+				_currentBundle = new PendingBundle(_treeReader, this, bufferSize, _options.MaxBlobSize, _options.CompressionFormat);
 			}
 			return _currentBundle;
 		}
