[UHT] Added "Auto" as a Getter specifier option that generates Get boilerplate func

diff --git a/Engine/Source/Programs/Shared/EpicGames.UHT/Exporters/CodeGen/UhtHeaderCodeGenerator.cs b/Engine/Source/Programs/Shared/EpicGames.UHT/Exporters/CodeGen/UhtHeaderCodeGenerator.cs
--- a/Engine/Source/Programs/Shared/EpicGames.UHT/Exporters/CodeGen/UhtHeaderCodeGenerator.cs
+++ b/Engine/Source/Programs/Shared/EpicGames.UHT/Exporters/CodeGen/UhtHeaderCodeGenerator.cs
@@ -33,6 +33,7 @@ namespace EpicGames.UHT.Exporters.CodeGen
 		public const string InClassIInterfaceNoPureDeclsMacroSuffix = "INCLASS_IINTERFACE_NO_PURE_DECLS";
 		public const string PrologMacroSuffix = "PROLOG";
 		public const string DelegateMacroSuffix = "DELEGATE";
+		public const string AutoGettersSettersMacroSuffix = "AUTOGETTERSETTER_DECLS";
 		#endregion
 
 		public readonly UhtHeaderFile HeaderFile;
@@ -693,6 +694,84 @@ namespace EpicGames.UHT.Exporters.CodeGen
 			return builder;
 		}
 		#endregion
+
+		#region AutoGettersSetters support
+		protected static bool NeedAutoGetterSetterCodeGen(UhtClass classObj)
+		{
+			return classObj.ClassExportFlags.HasAnyFlags(UhtClassExportFlags.HasAutoGettersSetters);
+		}
+		
+		protected static void GetAutoGetterSetterStats(UhtClass classObj, out bool hasProperties, out bool hasEditorFields, out bool allEditorFields)
+		{
+			// Scan the children to see what we have
+			hasProperties = false;
+			hasEditorFields = false;
+			allEditorFields = true;
+			foreach (UhtType type in classObj.Children)
+			{
+				if (type is UhtProperty property)
+				{
+					if (property.PropertyExportFlags.HasAnyFlags(UhtPropertyExportFlags.GetterSpecifiedAuto | UhtPropertyExportFlags.SetterSpecifiedAuto))
+					{
+						hasProperties = true;
+						hasEditorFields |= property.IsEditorOnlyProperty;
+						allEditorFields &= property.IsEditorOnlyProperty;
+					}
+				}
+			}
+
+			// If we have no editor fields, then by definition, all fields can't be editor fields
+			allEditorFields &= hasEditorFields;
+		}
+		
+		protected static StringBuilder AppendAutoGettersSetters(StringBuilder builder, UhtClass classObj,
+			UhtPropertyExportFlags specifierFlag,
+			bool hasProperties, bool hasEditorFields, bool allEditorFields,
+			bool includeEditorOnlyFields, bool appendDefine, Action<StringBuilder, UhtClass, UhtProperty> appendAction)
+		{
+			if (hasProperties && !allEditorFields)
+			{
+				foreach (UhtType child in classObj.Children)
+				{
+					if (child is UhtProperty property)
+					{
+						if (property.PropertyExportFlags.HasAnyFlags(specifierFlag) && !property.IsEditorOnlyProperty)
+						{
+							appendAction(builder, classObj, property);
+						}
+					}
+				}
+			}
+
+			if (hasEditorFields && includeEditorOnlyFields)
+			{
+				if (!allEditorFields && appendDefine)
+				{
+					builder.Append("#if WITH_EDITORONLY_DATA\r\n");
+				}
+
+				if (hasProperties)
+				{
+					foreach (UhtType child in classObj.Children)
+					{
+						if (child is UhtProperty property)
+						{
+							if (property.PropertyExportFlags.HasAnyFlags(specifierFlag) && property.IsEditorOnlyProperty)
+							{
+								appendAction(builder, classObj, property);
+							}
+						}
+					}
+				}
+
+				if (!allEditorFields && appendDefine)
+				{
+					builder.Append("#endif // WITH_EDITORONLY_DATA\r\n");
+				}
+			}
+			return builder;
+		}
+		#endregion
 	}
 
 	internal static class UhtHaederCodeGeneratorStringBuilderExtensions
diff --git a/Engine/Source/Programs/Shared/EpicGames.UHT/Exporters/CodeGen/UhtHeaderCodeGeneratorCppFile.cs b/Engine/Source/Programs/Shared/EpicGames.UHT/Exporters/CodeGen/UhtHeaderCodeGeneratorCppFile.cs
--- a/Engine/Source/Programs/Shared/EpicGames.UHT/Exporters/CodeGen/UhtHeaderCodeGeneratorCppFile.cs
+++ b/Engine/Source/Programs/Shared/EpicGames.UHT/Exporters/CodeGen/UhtHeaderCodeGeneratorCppFile.cs
@@ -1133,6 +1133,9 @@ namespace EpicGames.UHT.Exporters.CodeGen
 				AppendRpcFunctions(builder, classObj, reversedFunctions, true);
 				builder.AppendEndEditorOnlyGuard();
 			}
+			
+			// Add the auto getters/setters
+			AppendAutoGettersSetters(builder, classObj);
 
 			// Add the accessors
 			AppendPropertyAccessors(builder, classObj);
@@ -1241,6 +1244,71 @@ namespace EpicGames.UHT.Exporters.CodeGen
 			return builder;
 		}
 
+		private static StringBuilder AppendAutoGettersSetters(StringBuilder builder, UhtClass classObj)
+		{
+			if (!NeedAutoGetterSetterCodeGen(classObj))
+			{
+				return builder;
+			}
+			
+			// Scan the children to see what we have
+			GetAutoGetterSetterStats(classObj, out bool hasProperties, out bool hasEditorFields, out bool allEditorFields);
+			
+			// If we only have editor fields or no editor fields, then we only emit one block
+			if (hasEditorFields)
+			{
+				builder.Append("#if WITH_EDITORONLY_DATA\r\n");
+				AppendAutoGettersSetters(builder, classObj, hasProperties, hasEditorFields, allEditorFields, true);
+				builder.Append("#else //WITH_EDITORONLY_DATA\r\n");
+				AppendAutoGettersSetters(builder, classObj, hasProperties, hasEditorFields, allEditorFields, false);
+				builder.Append("#endif // WITH_EDITORONLY_DATA\r\n");
+			}
+			else
+			{
+				AppendAutoGettersSetters(builder, classObj, hasProperties, hasEditorFields, allEditorFields, false);
+			}
+			return builder;
+		}
+		
+		private static StringBuilder AppendAutoGettersSetters(StringBuilder builder, UhtClass classObj,
+			bool hasProperties, bool hasEditorFields, bool allEditorFields,
+			bool includeEditorOnlyFields)
+		{
+			AppendAutoGettersSetters(builder, classObj, UhtPropertyExportFlags.GetterSpecifiedAuto, hasProperties, hasEditorFields, allEditorFields,
+				includeEditorOnlyFields, false, (StringBuilder parentBuilder, UhtClass classObj, UhtProperty property) =>
+				{
+					using BorrowStringBuilder borrower = new(StringBuilderCache.Small);
+					StringBuilder propertyStringBuilder = borrower.StringBuilder;
+					propertyStringBuilder.AppendPropertyText(property, UhtPropertyTextType.GetterRetVal);
+					string getterRetText = propertyStringBuilder.ToString();
+					string getterCallText = property.Getter ?? "Get" + property.SourceName;
+					parentBuilder.Append('\t').Append(getterRetText).Append(classObj.SourceName).Append("::").Append(getterCallText).Append("() const\r\n");
+					parentBuilder.Append("\t{\r\n");
+					parentBuilder.Append('\t', 2).Append("return ").Append(property.SourceName).Append(";\r\n");
+					parentBuilder.Append("\t}\r\n");
+				});
+			
+			AppendAutoGettersSetters(builder, classObj, UhtPropertyExportFlags.SetterSpecifiedAuto, hasProperties, hasEditorFields, allEditorFields,
+				includeEditorOnlyFields, false, (StringBuilder parentBuilder, UhtClass classObj, UhtProperty property) =>
+				{
+					using BorrowStringBuilder borrower = new(StringBuilderCache.Small);
+					StringBuilder propertyStringBuilder = borrower.StringBuilder;
+					propertyStringBuilder.AppendPropertyText(property, UhtPropertyTextType.SetterParameterArgType);
+					string setterArgText = propertyStringBuilder + "InValue";
+					string setterCallText = property.Setter ?? "Set" + property.SourceName;
+					parentBuilder.Append("\tvoid ").Append(classObj.SourceName).Append("::").Append(setterCallText).Append('(').Append(setterArgText).Append(")\r\n");
+					parentBuilder.Append("\t{\r\n");
+					
+					// @todo: setter defn
+					
+					parentBuilder.Append("\t}\r\n");
+					
+					
+				});
+
+			return builder;
+		}
+
 		private static StringBuilder AppendPropertyAccessors(StringBuilder builder, UhtClass classObj)
 		{
 			foreach (UhtType type in classObj.Children)
diff --git a/Engine/Source/Programs/Shared/EpicGames.UHT/Exporters/CodeGen/UhtHeaderCodeGeneratorHFile.cs b/Engine/Source/Programs/Shared/EpicGames.UHT/Exporters/CodeGen/UhtHeaderCodeGeneratorHFile.cs
--- a/Engine/Source/Programs/Shared/EpicGames.UHT/Exporters/CodeGen/UhtHeaderCodeGeneratorHFile.cs
+++ b/Engine/Source/Programs/Shared/EpicGames.UHT/Exporters/CodeGen/UhtHeaderCodeGeneratorHFile.cs
@@ -546,6 +546,9 @@ namespace EpicGames.UHT.Exporters.CodeGen
 			List<UhtFunction> callbackFunctions = new(classObj.Functions.Where(x => x.FunctionFlags.HasAnyFlags(EFunctionFlags.Event) && x.SuperFunction == null));
 			callbackFunctions.Sort((x, y) => StringComparerUE.OrdinalIgnoreCase.Compare(x.EngineName, y.EngineName));
 			bool hasCallbacks = callbackFunctions.Count > 0;
+			
+			// Determine if auto getters/setters have been generated
+			bool hasAutoGettersSetters = classObj.Properties.Any(x => x.PropertyExportFlags.HasAnyFlags(UhtPropertyExportFlags.GetterSpecifiedAuto | UhtPropertyExportFlags.SetterSpecifiedAuto));
 
 			// Generate the RPC wrappers for the callbacks
 			AppendCallbackRpcWrapperDecls(builder, classObj, callbackFunctions);
@@ -606,7 +609,7 @@ namespace EpicGames.UHT.Exporters.CodeGen
 				}
 
 				AppendProlog(builder, classObj);
-				AppendGeneratedBodyMacroBlock(builder, classObj, nativeInterface!, alternateUsesLegacy, hasEditorRpc, hasCallbacks, null);
+				AppendGeneratedBodyMacroBlock(builder, classObj, nativeInterface!, alternateUsesLegacy, hasEditorRpc, hasCallbacks, hasAutoGettersSetters, null);
 			}
 			else
 			{
@@ -628,8 +631,9 @@ namespace EpicGames.UHT.Exporters.CodeGen
 				}
 
 				AppendFieldNotify(builder, classObj);
+				AppendAutoGettersSetters(builder, classObj);
 				AppendProlog(builder, classObj);
-				AppendGeneratedBodyMacroBlock(builder, classObj, classObj, usesLegacy, hasEditorRpc, hasCallbacks, usesLegacy ? "GENERATED_UCLASS_BODY" : null);
+				AppendGeneratedBodyMacroBlock(builder, classObj, classObj, usesLegacy, hasEditorRpc, hasCallbacks, hasAutoGettersSetters, usesLegacy ? "GENERATED_UCLASS_BODY" : null);
 			}
 
 			// Forward declare the StaticClass specialization in the header
@@ -717,6 +721,69 @@ namespace EpicGames.UHT.Exporters.CodeGen
 			return builder;
 		}
 
+		private StringBuilder AppendAutoGettersSetters(StringBuilder builder, UhtClass classObj)
+		{
+			if (!NeedAutoGetterSetterCodeGen(classObj))
+			{
+				return builder;
+			}
+
+			// Scan the children to see what we have
+			GetAutoGetterSetterStats(classObj, out bool hasProperties, out bool hasEditorFields, out bool allEditorFields);
+
+			// If we only have editor fields or no editor fields, then we only emit one block
+			if (hasEditorFields)
+			{
+				builder.Append("#if WITH_EDITORONLY_DATA\r\n");
+				using (UhtMacroCreator macro = new(builder, this, classObj, AutoGettersSettersMacroSuffix))
+				{
+					AppendAutoGetterSetter(builder, classObj, hasProperties, hasEditorFields, allEditorFields, true);
+				}
+				builder.Append("#else //WITH_EDITORONLY_DATA\r\n");
+				using (UhtMacroCreator macro = new(builder, this, classObj, AutoGettersSettersMacroSuffix))
+				{
+					AppendAutoGetterSetter(builder, classObj, hasProperties, hasEditorFields, allEditorFields, false);
+				}
+				builder.Append("#endif // WITH_EDITORONLY_DATA\r\n");
+			}
+			else
+			{
+				using UhtMacroCreator macro = new(builder, this, classObj, AutoGettersSettersMacroSuffix);
+				AppendAutoGetterSetter(builder, classObj, hasProperties, hasEditorFields, allEditorFields, false);
+			}
+			return builder;
+		}
+		
+		private StringBuilder AppendAutoGetterSetter(StringBuilder builder, UhtClass classObj,
+			bool hasProperties, bool hasEditorFields, bool allEditorFields,
+			bool includeEditorOnlyFields)
+		{
+			builder.Append("public: \\\r\n");
+            AppendAutoGettersSetters(builder, classObj, UhtPropertyExportFlags.GetterSpecifiedAuto, hasProperties, hasEditorFields, allEditorFields,
+				includeEditorOnlyFields, false, (StringBuilder parentBuilder, UhtClass classObj, UhtProperty property) =>
+				{
+					using BorrowStringBuilder borrower = new(StringBuilderCache.Small);
+					StringBuilder propertyStringBuilder = borrower.StringBuilder;
+					propertyStringBuilder.AppendPropertyText(property, UhtPropertyTextType.GetterRetVal);
+					string getterRetText = propertyStringBuilder.ToString();
+					string getterCallText = property.Getter ?? "Get" + property.SourceName;
+					parentBuilder.Append('\t').Append(getterRetText).Append(getterCallText).Append("() const; \\\r\n");
+				});
+            
+            AppendAutoGettersSetters(builder, classObj, UhtPropertyExportFlags.SetterSpecifiedAuto, hasProperties, hasEditorFields, allEditorFields,
+	            includeEditorOnlyFields, false, (StringBuilder parentBuilder, UhtClass classObj, UhtProperty property) =>
+	            {
+		            using BorrowStringBuilder borrower = new(StringBuilderCache.Small);
+		            StringBuilder propertyStringBuilder = borrower.StringBuilder;
+		            propertyStringBuilder.AppendPropertyText(property, UhtPropertyTextType.SetterParameterArgType);
+		            string setterArgText = propertyStringBuilder + "InValue";
+		            string setterCallText = property.Setter ?? "Set" + property.SourceName;
+		            parentBuilder.Append("\tvoid ").Append(setterCallText).Append('(').Append(setterArgText).Append("); \\\r\n");
+	            });
+
+            return builder;
+		}
+		
 		private StringBuilder AppendSparseDeclarations(StringBuilder builder, UhtClass classObj)
 		{
 			// Format the sparse data
@@ -1457,7 +1524,7 @@ namespace EpicGames.UHT.Exporters.CodeGen
 			return builder;
 		}
 
-		private StringBuilder AppendGeneratedBodyMacroBlock(StringBuilder builder, UhtClass classObj, UhtClass bodyClassObj, bool isLegacy, bool hasEditorRpc, bool hasCallbacks, string? deprecatedMacroName)
+		private StringBuilder AppendGeneratedBodyMacroBlock(StringBuilder builder, UhtClass classObj, UhtClass bodyClassObj, bool isLegacy, bool hasEditorRpc, bool hasCallbacks, bool hasAutoGettersSetters, string? deprecatedMacroName)
 		{
 			bool isInterface = classObj.ClassFlags.HasAnyFlags(EClassFlags.Interface);
 			using (UhtMacroCreator macro = new(builder, this, bodyClassObj, isLegacy ? GeneratedBodyLegacyMacroSuffix : GeneratedBodyMacroSuffix))
@@ -1489,6 +1556,10 @@ namespace EpicGames.UHT.Exporters.CodeGen
 					}
 				}
 				builder.Append('\t').AppendMacroName(this, classObj, AccessorsMacroSuffix).Append(" \\\r\n");
+				if (hasAutoGettersSetters)
+				{
+					builder.Append('\t').AppendMacroName(this, classObj, AutoGettersSettersMacroSuffix).Append(" \\\r\n");
+				}
 				if (hasCallbacks)
 				{
 					builder.Append('\t').AppendMacroName(this, classObj, CallbackWrappersMacroSuffix).Append(" \\\r\n");
diff --git a/Engine/Source/Programs/Shared/EpicGames.UHT/Specifiers/UhtPropertyMemberSpecifiers.cs b/Engine/Source/Programs/Shared/EpicGames.UHT/Specifiers/UhtPropertyMemberSpecifiers.cs
--- a/Engine/Source/Programs/Shared/EpicGames.UHT/Specifiers/UhtPropertyMemberSpecifiers.cs
+++ b/Engine/Source/Programs/Shared/EpicGames.UHT/Specifiers/UhtPropertyMemberSpecifiers.cs
@@ -400,6 +400,10 @@ namespace EpicGames.UHT.Parsers
 					{
 						context.PropertySettings.PropertyExportFlags |= UhtPropertyExportFlags.GetterSpecifiedNone;
 					}
+					else if (temp.Equals("Auto", StringComparison.OrdinalIgnoreCase))
+					{
+						context.PropertySettings.PropertyExportFlags |= UhtPropertyExportFlags.GetterSpecifiedAuto;
+					}
 					else
 					{
 						context.PropertySettings.Getter = value.ToString();
@@ -426,6 +430,10 @@ namespace EpicGames.UHT.Parsers
 					{
 						context.PropertySettings.PropertyExportFlags |= UhtPropertyExportFlags.SetterSpecifiedNone;
 					}
+					else if (temp.Equals("Auto", StringComparison.OrdinalIgnoreCase))
+					{
+						context.PropertySettings.PropertyExportFlags |= UhtPropertyExportFlags.SetterSpecifiedAuto;
+					}
 					else
 					{
 						context.PropertySettings.Setter = value.ToString();
diff --git a/Engine/Source/Programs/Shared/EpicGames.UHT/Types/UhtClass.cs b/Engine/Source/Programs/Shared/EpicGames.UHT/Types/UhtClass.cs
--- a/Engine/Source/Programs/Shared/EpicGames.UHT/Types/UhtClass.cs
+++ b/Engine/Source/Programs/Shared/EpicGames.UHT/Types/UhtClass.cs
@@ -94,6 +94,11 @@ namespace EpicGames.UHT.Types
 		/// legacy generated body.  If this flag is set, generate the legacy instead of the GENERATED_BODY macros.
 		/// </summary>
 		UsesGeneratedBodyLegacy = 1 << 12,
+		
+		/// <summary>
+		/// Class has one or more auto getters/setters
+		/// </summary>
+		HasAutoGettersSetters = 1 << 13,
 	}
 
 	/// <summary>
@@ -872,7 +877,17 @@ namespace EpicGames.UHT.Types
 							{
 								if (property.PropertyExportFlags.HasAnyFlags(UhtPropertyExportFlags.GetterSpecified))
 								{
-									gsToResolve = AddGetterSetter(gsToResolve, property.Getter ?? $"Get{property.SourceName}", property, false);
+									// if it's Auto, it will be generated so no need to look for a user implementation
+									if (property.PropertyExportFlags.HasAnyFlags(UhtPropertyExportFlags.GetterSpecifiedAuto))
+									{
+										ClassExportFlags |= UhtClassExportFlags.HasAutoGettersSetters;
+										property.PropertyExportFlags |= UhtPropertyExportFlags.GetterFound;
+										property.Getter = "Get" + property.SourceName;
+									}
+									else
+									{
+										gsToResolve = AddGetterSetter(gsToResolve, property.Getter ?? $"Get{property.SourceName}", property, false);	
+									}
 								}
 							}
 
@@ -880,7 +895,17 @@ namespace EpicGames.UHT.Types
 							{
 								if (property.PropertyExportFlags.HasAnyFlags(UhtPropertyExportFlags.SetterSpecified))
 								{
-									gsToResolve = AddGetterSetter(gsToResolve, property.Setter ?? $"Set{property.SourceName}", property, true);
+									// if it's Auto, it will be generated so no need to look for a user implementation
+									if (property.PropertyExportFlags.HasAnyFlags(UhtPropertyExportFlags.SetterSpecifiedAuto))
+									{
+										ClassExportFlags |= UhtClassExportFlags.HasAutoGettersSetters;
+										property.PropertyExportFlags |= UhtPropertyExportFlags.SetterFound;
+										property.Setter = "Set" + property.SourceName;
+									}
+									else
+									{
+										gsToResolve = AddGetterSetter(gsToResolve, property.Setter ?? $"Set{property.SourceName}", property, true);
+									}
 								}
 							}
 						}
@@ -1836,8 +1861,7 @@ namespace EpicGames.UHT.Types
 					}
 
 					// Check for getter/setter
-					if (property.PropertyExportFlags.HasExactFlags(UhtPropertyExportFlags.GetterSpecified | UhtPropertyExportFlags.GetterSpecifiedNone | 
-						UhtPropertyExportFlags.GetterFound, UhtPropertyExportFlags.GetterSpecified))
+					if (property.PropertyExportFlags.HasExactFlags(UhtPropertyExportFlags.GetterSpecified | UhtPropertyExportFlags.GetterSpecifiedNone | UhtPropertyExportFlags.GetterFound | UhtPropertyExportFlags.GetterSpecifiedAuto, UhtPropertyExportFlags.GetterSpecified))
 					{
 						using BorrowStringBuilder borrower = new(StringBuilderCache.Small);
 						StringBuilder builder = borrower.StringBuilder;
@@ -1850,7 +1874,7 @@ namespace EpicGames.UHT.Types
 						string expectedType = builder.ToString();
 						property.LogError($"Property '{property.SourceName}' expected a getter function '[const] {expectedType} [&] {expectedName}() const', but it was not found");
 					}
-					if (property.PropertyExportFlags.HasExactFlags(UhtPropertyExportFlags.SetterSpecified | UhtPropertyExportFlags.SetterSpecifiedNone | UhtPropertyExportFlags.SetterFound, UhtPropertyExportFlags.SetterSpecified))
+					if (property.PropertyExportFlags.HasExactFlags(UhtPropertyExportFlags.SetterSpecified | UhtPropertyExportFlags.SetterSpecifiedNone | UhtPropertyExportFlags.SetterFound | UhtPropertyExportFlags.SetterSpecifiedAuto, UhtPropertyExportFlags.SetterSpecified))
 					{
 						using BorrowStringBuilder borrower = new(StringBuilderCache.Small);
 						StringBuilder builder = borrower.StringBuilder;
diff --git a/Engine/Source/Programs/Shared/EpicGames.UHT/Types/UhtProperty.cs b/Engine/Source/Programs/Shared/EpicGames.UHT/Types/UhtProperty.cs
--- a/Engine/Source/Programs/Shared/EpicGames.UHT/Types/UhtProperty.cs
+++ b/Engine/Source/Programs/Shared/EpicGames.UHT/Types/UhtProperty.cs
@@ -76,6 +76,16 @@ namespace EpicGames.UHT.Types
 		/// Property is marked as a field notify
 		/// </summary>
 		FieldNotify = 0x00000400,
+
+		/// <summary>
+		/// If true, the property should have a generated getter function
+		/// </summary>
+		GetterSpecifiedAuto = 0x00001000,
+		
+		/// <summary>
+		/// If true, the property should have a generated setter function
+		/// </summary>
+		SetterSpecifiedAuto = 0x00002000,
 	};
 
 	/// <summary>
@@ -516,6 +526,16 @@ namespace EpicGames.UHT.Types
 		/// Type expected in a getter/setter argument list
 		/// </summary>
 		GetterSetterArg,
+		
+		/// <summary>
+		/// Type expected from a getter
+		/// </summary>
+		GetterRetVal,
+		
+		/// <summary>
+		/// Type expected as a setter argument
+		/// </summary>
+		SetterParameterArgType,
 	}
 
 	/// <summary>
@@ -536,7 +556,21 @@ namespace EpicGames.UHT.Types
 				textType == UhtPropertyTextType.GenericFunctionArgOrRetValImpl ||
 				textType == UhtPropertyTextType.ClassFunctionArgOrRetVal ||
 				textType == UhtPropertyTextType.EventFunctionArgOrRetVal ||
-				textType == UhtPropertyTextType.InterfaceFunctionArgOrRetVal;
+				textType == UhtPropertyTextType.InterfaceFunctionArgOrRetVal ||
+				textType == UhtPropertyTextType.SetterParameterArgType;
+		}
+		
+		/// <summary>
+		/// Test to see if the text type is for a getter retVal or setter argType
+		/// </summary>
+		/// <param name="textType">Type of text</param>
+		/// <returns>True if the text type is a retVal or argType</returns>
+		public static bool IsGetOrSet(this UhtPropertyTextType textType)
+		{
+			return
+				textType == UhtPropertyTextType.GetterSetterArg ||
+				textType == UhtPropertyTextType.GetterRetVal ||
+				textType == UhtPropertyTextType.SetterParameterArgType;
 		}
 	}
 
@@ -1061,6 +1095,7 @@ namespace EpicGames.UHT.Types
 			UhtPropertyCaps caps = PropertyCaps;
 
 			bool isParameter = textType.IsParameter();
+			bool isGetOrSet = textType.IsGetOrSet();
 			bool isInterfaceProp = this is UhtInterfaceProperty;
 
 			// When do we need a leading const:
@@ -1079,7 +1114,7 @@ namespace EpicGames.UHT.Types
 			}
 			bool shouldHaveRef = PropertyFlags.HasAnyFlags(EPropertyFlags.OutParm | EPropertyFlags.ReferenceParm);
 
-			bool constAtTheBeginning = isOnConstClass || isConstArgsByRef || (isConstParam && !shouldHaveRef);
+			bool constAtTheBeginning = isGetOrSet | isOnConstClass || isConstArgsByRef || (isConstParam && !shouldHaveRef);
 			if (constAtTheBeginning)
 			{
 				builder.Append("const ");
@@ -1101,7 +1136,8 @@ namespace EpicGames.UHT.Types
 				builder.Append(" const");
 			}
 
-			if (isParameter && ArrayDimensions == null && (passCppArgsByRef || PropertyFlags.HasAnyFlags(EPropertyFlags.OutParm | EPropertyFlags.ReferenceParm)))
+			shouldHaveRef = textType == UhtPropertyTextType.SetterParameterArgType || (textType == UhtPropertyTextType.GetterRetVal && passCppArgsByRef); 
+			if (shouldHaveRef || isParameter && ArrayDimensions == null && (passCppArgsByRef || PropertyFlags.HasAnyFlags(EPropertyFlags.OutParm | EPropertyFlags.ReferenceParm)))
 			{
 				builder.Append('&');
 			}
@@ -1981,6 +2017,11 @@ namespace EpicGames.UHT.Types
 		/// <returns>Destination builder</returns>
 		public static StringBuilder AppendPropertyText(this StringBuilder builder, UhtProperty property, UhtPropertyTextType textType, bool isTemplateArgument = false)
 		{
+			if (textType == UhtPropertyTextType.GetterRetVal || textType == UhtPropertyTextType.SetterParameterArgType)
+			{
+				return builder.AppendFullDecl(property, textType, true);
+			}
+
 			return property.AppendText(builder, textType, isTemplateArgument);
 		}
 
