Horde: Perform flushes recursively when flushing a node. This allows using nodes created between different writers interchangably.

Also produce a more descriptive exception when a writer if flushed after having been disposed.

diff --git a/Engine/Source/Programs/Shared/EpicGames.Horde/Storage/TreeWriter.cs b/Engine/Source/Programs/Shared/EpicGames.Horde/Storage/TreeWriter.cs
--- a/Engine/Source/Programs/Shared/EpicGames.Horde/Storage/TreeWriter.cs
+++ b/Engine/Source/Programs/Shared/EpicGames.Horde/Storage/TreeWriter.cs
@@ -10,7 +10,6 @@ using System.Threading;
 using System.Threading.Tasks;
 using EpicGames.Core;
 using Microsoft.Extensions.Logging;
-using static EpicGames.Horde.Storage.TreeWriter;
 
 namespace EpicGames.Horde.Storage
 {
@@ -309,16 +308,23 @@ namespace EpicGames.Horde.Storage
 			/// <inheritdoc/>
 			public override async ValueTask<NodeLocator> FlushAsync(CancellationToken cancellationToken = default)
 			{
-				PendingBundle? pendingBundle = _pendingBundle;
-				if (!_locator.IsValid() && pendingBundle != null)
+				if (_locator.IsValid())
 				{
-					await pendingBundle.FlushAsync(cancellationToken);
+					return _locator;
 				}
-				if (!_locator.IsValid())
+
+				foreach (NodeHandle nodeRef in Refs)
 				{
-					throw new InvalidOperationException("Locator should have been flushed already");
+					await nodeRef.FlushAsync(cancellationToken);
 				}
-				return _locator;
+
+				PendingBundle? pendingBundle = _pendingBundle;
+				if (pendingBundle != null)
+				{
+					await pendingBundle.FlushAsync(cancellationToken);
+				}
+
+				return GetLocator();
 			}
 		}
 
@@ -750,6 +756,8 @@ namespace EpicGames.Horde.Storage
 		PendingBundle? _currentBundle;
 		int _memoryFootprint;
 
+		bool _disposed;
+
 		internal readonly ILogger? _traceLogger;
 
 		/// <summary>
@@ -814,6 +822,7 @@ namespace EpicGames.Horde.Storage
 				_currentBundle.Dispose();
 				_currentBundle = null;
 			}
+			_disposed = true;
 		}
 
 		/// <summary>
@@ -924,6 +933,11 @@ namespace EpicGames.Horde.Storage
 		/// <returns></returns>
 		public async Task FlushAsync(CancellationToken cancellationToken = default)
 		{
+			if (_disposed)
+			{
+				throw new ObjectDisposedException(GetType().Name);
+			}
+
 			Complete();
 			while (_writeQueue.Count > 0)
 			{
