Horde: Fix invalid cast during ComputeSocket shutdown.

diff --git a/Engine/Source/Programs/Shared/EpicGames.Horde/Compute/ComputeSocket.cs b/Engine/Source/Programs/Shared/EpicGames.Horde/Compute/ComputeSocket.cs
--- a/Engine/Source/Programs/Shared/EpicGames.Horde/Compute/ComputeSocket.cs
+++ b/Engine/Source/Programs/Shared/EpicGames.Horde/Compute/ComputeSocket.cs
@@ -166,9 +166,9 @@ namespace EpicGames.Horde.Compute
 			lock (_lockObject)
 			{
 				_complete = true;
-				foreach (IComputeBuffer buffer in _recvBufferWriters.Values)
+				foreach (IComputeBufferWriter writer in _recvBufferWriters.Values)
 				{
-					buffer.Writer.MarkComplete();
+					writer.MarkComplete();
 				}
 			}
 
