Horde: Add a custom type to store the cache validity time for ref queries, and propagate it from HTTP queries.

diff --git a/Engine/Source/Programs/Horde/Horde.Server/Artifacts/ArtifactsController.cs b/Engine/Source/Programs/Horde/Horde.Server/Artifacts/ArtifactsController.cs
--- a/Engine/Source/Programs/Horde/Horde.Server/Artifacts/ArtifactsController.cs
+++ b/Engine/Source/Programs/Horde/Horde.Server/Artifacts/ArtifactsController.cs
@@ -212,7 +212,7 @@ namespace Horde.Server.Artifacts
 				return Forbid(ArtifactAclAction.ReadArtifact, artifact.AclScope);
 			}
 
-			return StorageController.ReadRefInternalAsync(_storageService, artifact.NamespaceId, artifact.RefName, cancellationToken);
+			return StorageController.ReadRefInternalAsync(_storageService, artifact.NamespaceId, artifact.RefName, Request.Headers, cancellationToken);
 		}
 
 		/// <summary>
diff --git a/Engine/Source/Programs/Horde/Horde.Server/Commands/Bundles/BundleWriteTest.cs b/Engine/Source/Programs/Horde/Horde.Server/Commands/Bundles/BundleWriteTest.cs
--- a/Engine/Source/Programs/Horde/Horde.Server/Commands/Bundles/BundleWriteTest.cs
+++ b/Engine/Source/Programs/Horde/Horde.Server/Commands/Bundles/BundleWriteTest.cs
@@ -33,7 +33,7 @@ namespace Horde.Server.Commands.Bundles
 			public Task AddAliasAsync(Utf8String name, NodeHandle handle, CancellationToken cancellationToken = default) => throw new NotImplementedException();
 			public Task RemoveAliasAsync(Utf8String name, NodeHandle handle, CancellationToken cancellationToken = default) => throw new NotImplementedException();
 			public IAsyncEnumerable<NodeHandle> FindNodesAsync(Utf8String alias, CancellationToken cancellationToken = default) => throw new NotImplementedException();
-			public Task<NodeHandle?> TryReadRefTargetAsync(RefName name, DateTime cacheTime = default, CancellationToken cancellationToken = default) => throw new NotImplementedException();
+			public Task<NodeHandle?> TryReadRefTargetAsync(RefName name, RefCacheTime cacheTime = default, CancellationToken cancellationToken = default) => throw new NotImplementedException();
 			public Task<BlobLocator> WriteBlobAsync(Stream stream, Utf8String prefix = default, CancellationToken cancellationToken = default) => Task.FromResult(BlobLocator.Create(HostId.Empty));
 			public Task<NodeHandle> WriteRefAsync(RefName name, Bundle bundle, int exportIdx, Utf8String prefix = default, RefOptions? options = null, CancellationToken cancellationToken = default) => throw new NotImplementedException();
 			public Task WriteRefTargetAsync(RefName name, NodeHandle target, RefOptions? options = null, CancellationToken cancellationToken = default) => Task.CompletedTask;
diff --git a/Engine/Source/Programs/Horde/Horde.Server/Storage/Backends/AwsStorageBackend.cs b/Engine/Source/Programs/Horde/Horde.Server/Storage/Backends/AwsStorageBackend.cs
--- a/Engine/Source/Programs/Horde/Horde.Server/Storage/Backends/AwsStorageBackend.cs
+++ b/Engine/Source/Programs/Horde/Horde.Server/Storage/Backends/AwsStorageBackend.cs
@@ -324,6 +324,7 @@ namespace Horde.Server.Storage.Backends
 				newGetRequest.Key = fullPath;
 				newGetRequest.Verb = verb;
 				newGetRequest.Expires = DateTime.UtcNow.AddHours(3.0);
+				newGetRequest.ResponseHeaderOverrides.CacheControl = "private, max-age=2592000, immutable"; // 30 days
 
 				string url = _client.GetPreSignedURL(newGetRequest);
 				return new Uri(url);
diff --git a/Engine/Source/Programs/Horde/Horde.Server/Storage/StorageController.cs b/Engine/Source/Programs/Horde/Horde.Server/Storage/StorageController.cs
--- a/Engine/Source/Programs/Horde/Horde.Server/Storage/StorageController.cs
+++ b/Engine/Source/Programs/Horde/Horde.Server/Storage/StorageController.cs
@@ -4,6 +4,7 @@ using System;
 using System.Collections.Generic;
 using System.IO;
 using System.Linq;
+using System.Net.Http.Headers;
 using System.Reflection;
 using System.Security.Claims;
 using System.Text.RegularExpressions;
@@ -389,17 +390,26 @@ namespace Horde.Server.Storage
 				return Forbid(StorageAclAction.ReadRefs, namespaceId);
 			}
 
-			return await ReadRefInternalAsync(_storageService, namespaceId, refName, cancellationToken);
+			return await ReadRefInternalAsync(_storageService, namespaceId, refName, Request.Headers, cancellationToken);
 		}
 
 		/// <summary>
 		/// Reads a ref from storage, without performing namespace access checks.
 		/// </summary>
-		internal static async Task<ActionResult<ReadRefResponse>> ReadRefInternalAsync(StorageService storageService, NamespaceId namespaceId, RefName refName, CancellationToken cancellationToken)
+		internal static async Task<ActionResult<ReadRefResponse>> ReadRefInternalAsync(StorageService storageService, NamespaceId namespaceId, RefName refName, IHeaderDictionary headers, CancellationToken cancellationToken)
 		{
 			IStorageClient client = await storageService.GetClientAsync(namespaceId, cancellationToken);
 
-			NodeHandle? target = await client.TryReadRefTargetAsync(refName, cancellationToken: cancellationToken);
+			RefCacheTime cacheTime = new RefCacheTime();
+			foreach (string entry in headers.CacheControl)
+			{
+				if (CacheControlHeaderValue.TryParse(entry, out CacheControlHeaderValue? value) && value?.MaxAge != null)
+				{
+					cacheTime = new RefCacheTime(value.MaxAge.Value);
+				}
+			}
+
+			NodeHandle? target = await client.TryReadRefTargetAsync(refName, cacheTime, cancellationToken: cancellationToken);
 			if (target == null)
 			{
 				return new NotFoundResult();
diff --git a/Engine/Source/Programs/Horde/Horde.Server/Storage/StorageService.cs b/Engine/Source/Programs/Horde/Horde.Server/Storage/StorageService.cs
--- a/Engine/Source/Programs/Horde/Horde.Server/Storage/StorageService.cs
+++ b/Engine/Source/Programs/Horde/Horde.Server/Storage/StorageService.cs
@@ -249,7 +249,7 @@ namespace Horde.Server.Storage
 			#region Refs
 
 			/// <inheritdoc/>
-			public override async Task<NodeHandle?> TryReadRefTargetAsync(RefName name, DateTime cacheTime = default, CancellationToken cancellationToken = default)
+			public override async Task<NodeHandle?> TryReadRefTargetAsync(RefName name, RefCacheTime cacheTime = default, CancellationToken cancellationToken = default)
 			{
 				HashedNodeLocator? locator = await _outer.TryReadRefTargetAsync(NamespaceId, name, cacheTime, cancellationToken);
 				if (locator == null)
@@ -855,10 +855,10 @@ namespace Horde.Server.Storage
 		}
 
 		/// <inheritdoc/>
-		async Task<HashedNodeLocator?> TryReadRefTargetAsync(NamespaceId namespaceId, RefName name, DateTime cacheTime = default, CancellationToken cancellationToken = default)
+		async Task<HashedNodeLocator?> TryReadRefTargetAsync(NamespaceId namespaceId, RefName name, RefCacheTime cacheTime = default, CancellationToken cancellationToken = default)
 		{
 			RefCacheValue entry;
-			if (!_cache.TryGetValue(name, out entry) || entry.Time < cacheTime)
+			if (!_cache.TryGetValue(name, out entry) || RefCacheTime.IsStaleCacheEntry(entry.Time, cacheTime))
 			{
 				RefInfo? refDocument = await _refCollection.Find(x => x.NamespaceId == namespaceId && x.Name == name).FirstOrDefaultAsync(cancellationToken);
 				entry = AddRefToCache(namespaceId, name, refDocument);
diff --git a/Engine/Source/Programs/Shared/EpicGames.Horde/Compute/AgentStorageClient.cs b/Engine/Source/Programs/Shared/EpicGames.Horde/Compute/AgentStorageClient.cs
--- a/Engine/Source/Programs/Shared/EpicGames.Horde/Compute/AgentStorageClient.cs
+++ b/Engine/Source/Programs/Shared/EpicGames.Horde/Compute/AgentStorageClient.cs
@@ -97,7 +97,7 @@ namespace EpicGames.Horde.Compute
 		}
 
 		/// <inheritdoc/>
-		public override Task<NodeHandle?> TryReadRefTargetAsync(RefName name, DateTime cacheTime = default, CancellationToken cancellationToken = default)
+		public override Task<NodeHandle?> TryReadRefTargetAsync(RefName name, RefCacheTime cacheTime = default, CancellationToken cancellationToken = default)
 		{
 			throw new NotSupportedException();
 		}
diff --git a/Engine/Source/Programs/Shared/EpicGames.Horde/Storage/Backends/FileStorageClient.cs b/Engine/Source/Programs/Shared/EpicGames.Horde/Storage/Backends/FileStorageClient.cs
--- a/Engine/Source/Programs/Shared/EpicGames.Horde/Storage/Backends/FileStorageClient.cs
+++ b/Engine/Source/Programs/Shared/EpicGames.Horde/Storage/Backends/FileStorageClient.cs
@@ -117,7 +117,7 @@ namespace EpicGames.Horde.Storage.Backends
 		}
 
 		/// <inheritdoc/>
-		public override async Task<NodeHandle?> TryReadRefTargetAsync(RefName name, DateTime cacheTime = default, CancellationToken cancellationToken = default)
+		public override async Task<NodeHandle?> TryReadRefTargetAsync(RefName name, RefCacheTime cacheTime = default, CancellationToken cancellationToken = default)
 		{
 			FileReference file = GetRefFile(name);
 			if (!FileReference.Exists(file))
diff --git a/Engine/Source/Programs/Shared/EpicGames.Horde/Storage/Backends/HttpStorageClient.cs b/Engine/Source/Programs/Shared/EpicGames.Horde/Storage/Backends/HttpStorageClient.cs
--- a/Engine/Source/Programs/Shared/EpicGames.Horde/Storage/Backends/HttpStorageClient.cs
+++ b/Engine/Source/Programs/Shared/EpicGames.Horde/Storage/Backends/HttpStorageClient.cs
@@ -248,12 +248,17 @@ namespace EpicGames.Horde.Storage.Backends
 		}
 
 		/// <inheritdoc/>
-		public override async Task<NodeHandle?> TryReadRefTargetAsync(RefName name, DateTime cacheTime = default, CancellationToken cancellationToken = default)
+		public override async Task<NodeHandle?> TryReadRefTargetAsync(RefName name, RefCacheTime cacheTime = default, CancellationToken cancellationToken = default)
 		{
 			using (HttpClient httpClient = _createClient())
 			{
 				using (HttpRequestMessage request = new HttpRequestMessage(HttpMethod.Get, $"refs/{name}"))
 				{
+					if (cacheTime.IsSet())
+					{
+						request.Headers.CacheControl = new CacheControlHeaderValue { MaxAge = cacheTime.MaxAge };
+					}
+
 					using (HttpResponseMessage response = await httpClient.SendAsync(request, cancellationToken))
 					{
 						if (response.StatusCode == HttpStatusCode.NotFound)
diff --git a/Engine/Source/Programs/Shared/EpicGames.Horde/Storage/Backends/MemoryStorageClient.cs b/Engine/Source/Programs/Shared/EpicGames.Horde/Storage/Backends/MemoryStorageClient.cs
--- a/Engine/Source/Programs/Shared/EpicGames.Horde/Storage/Backends/MemoryStorageClient.cs
+++ b/Engine/Source/Programs/Shared/EpicGames.Horde/Storage/Backends/MemoryStorageClient.cs
@@ -116,7 +116,7 @@ namespace EpicGames.Horde.Storage.Backends
 		public override Task DeleteRefAsync(RefName name, CancellationToken cancellationToken) => Task.FromResult(_refs.TryRemove(name, out _));
 
 		/// <inheritdoc/>
-		public override Task<NodeHandle?> TryReadRefTargetAsync(RefName name, DateTime cacheTime = default, CancellationToken cancellationToken = default)
+		public override Task<NodeHandle?> TryReadRefTargetAsync(RefName name, RefCacheTime cacheTime = default, CancellationToken cancellationToken = default)
 		{
 			NodeHandle? refTarget;
 			_refs.TryGetValue(name, out refTarget);
diff --git a/Engine/Source/Programs/Shared/EpicGames.Horde/Storage/BundleStorageClient.cs b/Engine/Source/Programs/Shared/EpicGames.Horde/Storage/BundleStorageClient.cs
--- a/Engine/Source/Programs/Shared/EpicGames.Horde/Storage/BundleStorageClient.cs
+++ b/Engine/Source/Programs/Shared/EpicGames.Horde/Storage/BundleStorageClient.cs
@@ -71,7 +71,7 @@ namespace EpicGames.Horde.Storage
 		public abstract Task DeleteRefAsync(RefName name, CancellationToken cancellationToken = default);
 
 		/// <inheritdoc/>
-		public abstract Task<NodeHandle?> TryReadRefTargetAsync(RefName name, DateTime cacheTime = default, CancellationToken cancellationToken = default);
+		public abstract Task<NodeHandle?> TryReadRefTargetAsync(RefName name, RefCacheTime cacheTime = default, CancellationToken cancellationToken = default);
 
 		/// <inheritdoc/>
 		public virtual async Task<NodeHandle> WriteRefAsync(RefName name, Bundle bundle, int exportIdx = 0, Utf8String prefix = default, RefOptions? options = null, CancellationToken cancellationToken = default)
diff --git a/Engine/Source/Programs/Shared/EpicGames.Horde/Storage/IStorageClient.cs b/Engine/Source/Programs/Shared/EpicGames.Horde/Storage/IStorageClient.cs
--- a/Engine/Source/Programs/Shared/EpicGames.Horde/Storage/IStorageClient.cs
+++ b/Engine/Source/Programs/Shared/EpicGames.Horde/Storage/IStorageClient.cs
@@ -158,7 +158,7 @@ namespace EpicGames.Horde.Storage
 		/// <param name="cacheTime">Minimum coherency for any cached value to be returned</param>
 		/// <param name="cancellationToken">Cancellation token for the operation</param>
 		/// <returns>Node pointed to by the ref</returns>
-		Task<NodeHandle?> TryReadRefTargetAsync(RefName name, DateTime cacheTime = default, CancellationToken cancellationToken = default);
+		Task<NodeHandle?> TryReadRefTargetAsync(RefName name, RefCacheTime cacheTime = default, CancellationToken cancellationToken = default);
 
 		/// <summary>
 		/// Writes a new ref to the store which points to a new blob
@@ -247,6 +247,46 @@ namespace EpicGames.Horde.Storage
 		ValueTask<IStorageClient> GetClientAsync(NamespaceId namespaceId, CancellationToken cancellationToken = default);
 	}
 
+	/// <summary>
+	/// Indicates the maximum age of a entry returned from a cache in the hierarchy
+	/// </summary>
+	/// <param name="Utc">Oldest allowed timestamp for a returned result</param>
+	public record struct RefCacheTime(DateTime Utc)
+	{
+		/// <summary>
+		/// Maximum age for a cached value to be returned
+		/// </summary>
+		public TimeSpan MaxAge => DateTime.UtcNow - Utc;
+
+		/// <summary>
+		/// Sets the earliest time at which the entry must have been valid
+		/// </summary>
+		/// <param name="age">Maximum age of any returned cache value. Taken from the moment that this object was created.</param>
+		public RefCacheTime(TimeSpan age) : this(DateTime.UtcNow - age) { }
+
+		/// <summary>
+		/// Tests whether this value is set
+		/// </summary>
+		public bool IsSet() => Utc != default;
+
+		/// <summary>
+		/// Determines if this cache time deems a particular cache entry stale
+		/// </summary>
+		/// <param name="entryTime">Time at which the cache entry was valid</param>
+		/// <param name="cacheTime">Maximum cache time to test against</param>
+		public static bool IsStaleCacheEntry(DateTime entryTime, RefCacheTime cacheTime) => cacheTime.IsSet() && cacheTime.Utc < entryTime;
+
+		/// <summary>
+		/// Implicit conversion operator from datetime values.
+		/// </summary>
+		public static implicit operator RefCacheTime(DateTime time) => new RefCacheTime(time);
+
+		/// <summary>
+		/// Implicit conversion operator from timespan values.
+		/// </summary>
+		public static implicit operator RefCacheTime(TimeSpan age) => new RefCacheTime(age);
+	}
+
 	/// <summary>
 	/// Extension methods for <see cref="IStorageClient"/>
 	/// </summary>
@@ -326,38 +366,12 @@ namespace EpicGames.Horde.Storage
 		/// <param name="cacheTime">Minimum coherency for any cached value to be returned</param>
 		/// <param name="cancellationToken">Cancellation token for the operation</param>
 		/// <returns>True if the ref exists, false if it did not exist</returns>
-		public static async Task<bool> HasRefAsync(this IStorageClient store, RefName name, DateTime cacheTime = default, CancellationToken cancellationToken = default)
+		public static async Task<bool> HasRefAsync(this IStorageClient store, RefName name, RefCacheTime cacheTime = default, CancellationToken cancellationToken = default)
 		{
 			NodeHandle? target = await store.TryReadRefTargetAsync(name, cacheTime, cancellationToken);
 			return target != null;
 		}
 
-		/// <summary>
-		/// Checks if the given ref exists
-		/// </summary>
-		/// <param name="store">The store instance to read from</param>
-		/// <param name="name">Name of the reference to look for</param>
-		/// <param name="maxAge">Maximum age of any cached ref</param>
-		/// <param name="cancellationToken">Cancellation token for the operation</param>
-		/// <returns>True if the ref exists, false if it did not exist</returns>
-		public static Task<bool> HasRefAsync(this IStorageClient store, RefName name, TimeSpan maxAge, CancellationToken cancellationToken = default)
-		{
-			return HasRefAsync(store, name, DateTime.UtcNow - maxAge, cancellationToken);
-		}
-
-		/// <summary>
-		/// Attempts to reads a ref from the store
-		/// </summary>
-		/// <param name="store">The store instance to read from</param>
-		/// <param name="name">Id for the ref</param>
-		/// <param name="maxAge">Maximum age of any cached ref</param>
-		/// <param name="cancellationToken">Cancellation token for the operation</param>
-		/// <returns>The ref target</returns>
-		public static Task<NodeHandle?> TryReadRefTargetAsync(this IStorageClient store, RefName name, TimeSpan maxAge, CancellationToken cancellationToken = default)
-		{
-			return store.TryReadRefTargetAsync(name, DateTime.UtcNow - maxAge, cancellationToken);
-		}
-
 		/// <summary>
 		/// Reads a ref from the store, throwing an exception if it does not exist
 		/// </summary>
@@ -366,7 +380,7 @@ namespace EpicGames.Horde.Storage
 		/// <param name="cacheTime">Minimum coherency of any cached result</param>
 		/// <param name="cancellationToken">Cancellation token for the operation</param>
 		/// <returns>The ref target</returns>
-		public static async Task<NodeHandle> ReadRefTargetAsync(this IStorageClient store, RefName name, DateTime cacheTime = default, CancellationToken cancellationToken = default)
+		public static async Task<NodeHandle> ReadRefTargetAsync(this IStorageClient store, RefName name, RefCacheTime cacheTime = default, CancellationToken cancellationToken = default)
 		{
 			NodeHandle? refTarget = await store.TryReadRefTargetAsync(name, cacheTime, cancellationToken);
 			if (refTarget == null)
@@ -376,19 +390,6 @@ namespace EpicGames.Horde.Storage
 			return refTarget;
 		}
 
-		/// <summary>
-		/// Reads a ref from the store, throwing an exception if it does not exist
-		/// </summary>
-		/// <param name="store">The store instance to read from</param>
-		/// <param name="name">Id for the ref</param>
-		/// <param name="maxAge">Maximum age for any cached result</param>
-		/// <param name="cancellationToken">Cancellation token for the operation</param>
-		/// <returns>The blob instance</returns>
-		public static Task<NodeHandle> ReadRefTargetAsync(this IStorageClient store, RefName name, TimeSpan maxAge, CancellationToken cancellationToken = default)
-		{
-			return ReadRefTargetAsync(store, name, DateTime.UtcNow - maxAge, cancellationToken);
-		}
-
 		#endregion
 	}
 }
