Treat .native.verse files as code.

diff --git a/Engine/Source/Programs/Shared/EpicGames.Perforce/PerforceUtils.cs b/Engine/Source/Programs/Shared/EpicGames.Perforce/PerforceUtils.cs
--- a/Engine/Source/Programs/Shared/EpicGames.Perforce/PerforceUtils.cs
+++ b/Engine/Source/Programs/Shared/EpicGames.Perforce/PerforceUtils.cs
@@ -34,7 +34,8 @@ namespace EpicGames.Perforce
 			".ush",
 			".uproject",
 			".uplugin",
-			".sln"
+			".sln",
+			".native.verse"
 		};
 
 		/// <summary>
