Horde: Move NodeHandle into its own file to separate it from bundle implementation details.

diff --git a/Engine/Source/Programs/Shared/EpicGames.Horde/Storage/NodeHandle.cs b/Engine/Source/Programs/Shared/EpicGames.Horde/Storage/NodeHandle.cs
new file mode 100644
--- /dev/null
+++ b/Engine/Source/Programs/Shared/EpicGames.Horde/Storage/NodeHandle.cs
@@ -0,0 +1,68 @@
+// Copyright Epic Games, Inc. All Rights Reserved.
+
+using System.Threading;
+using System.Threading.Tasks;
+using EpicGames.Core;
+
+namespace EpicGames.Horde.Storage
+{
+	/// <summary>
+	/// Handle to a node. Can be used to reference nodes that have not been flushed yet.
+	/// </summary>
+	public abstract class NodeHandle
+	{
+		/// <summary>
+		/// Hash of the target node
+		/// </summary>
+		public IoHash Hash { get; }
+
+		/// <summary>
+		/// Constructor
+		/// </summary>
+		/// <param name="hash">Hash of the target node</param>
+		protected NodeHandle(IoHash hash) => Hash = hash;
+
+		/// <summary>
+		/// Determines if the node has been written to storage
+		/// </summary>
+		public abstract bool HasLocator();
+
+		/// <summary>
+		/// Gets the node locator. May throw if the node has not been written to storage yet.
+		/// </summary>
+		/// <returns>Locator for the node</returns>
+		public abstract NodeLocator GetLocator();
+
+		/// <summary>
+		/// Adds a callback to be executed once the node has been written. Triggers immediately if the node has already been written.
+		/// </summary>
+		/// <param name="callback">Action to be executed after the write</param>
+		public abstract void AddWriteCallback(NodeWriteCallback callback);
+
+		/// <summary>
+		/// Creates a reader for this node's data
+		/// </summary>
+		public abstract ValueTask<NodeData> ReadAsync(CancellationToken cancellationToken = default);
+
+		/// <summary>
+		/// Flush the node to storage and retrieve its locator
+		/// </summary>
+		public abstract ValueTask<NodeLocator> FlushAsync(CancellationToken cancellationToken = default);
+
+		/// <inheritdoc/>
+		public override string ToString() => HasLocator() ? new HashedNodeLocator(Hash, GetLocator()).ToString() : Hash.ToString();
+	}
+
+	/// <summary>
+	/// Object to receive notifications on a node being written
+	/// </summary>
+	public abstract class NodeWriteCallback
+	{
+		internal NodeWriteCallback? _next;
+
+		/// <summary>
+		/// Callback for the node being written
+		/// </summary>
+		public abstract void OnWrite();
+	}
+}
diff --git a/Engine/Source/Programs/Shared/EpicGames.Horde/Storage/TreeWriter.cs b/Engine/Source/Programs/Shared/EpicGames.Horde/Storage/TreeWriter.cs
--- a/Engine/Source/Programs/Shared/EpicGames.Horde/Storage/TreeWriter.cs
+++ b/Engine/Source/Programs/Shared/EpicGames.Horde/Storage/TreeWriter.cs
@@ -51,66 +51,6 @@ namespace EpicGames.Horde.Storage
 	/// <param name="Type">Type of the node</param>
 	public record NodeKey(IoHash Hash, NodeType Type);
 
-	/// <summary>
-	/// Object to receive notifications on a node being written
-	/// </summary>
-	public abstract class NodeWriteCallback
-	{
-		internal NodeWriteCallback? _next;
-
-		/// <summary>
-		/// Callback for the node being written
-		/// </summary>
-		public abstract void OnWrite();
-	}
-
-	/// <summary>
-	/// Handle to a node. Can be used to reference nodes that have not been flushed yet.
-	/// </summary>
-	public abstract class NodeHandle
-	{
-		/// <summary>
-		/// Hash of the target node
-		/// </summary>
-		public IoHash Hash { get; }
-
-		/// <summary>
-		/// Constructor
-		/// </summary>
-		/// <param name="hash">Hash of the target node</param>
-		protected NodeHandle(IoHash hash) => Hash = hash;
-
-		/// <summary>
-		/// Determines if the node has been written to storage
-		/// </summary>
-		public abstract bool HasLocator();
-
-		/// <summary>
-		/// Gets the node locator. May throw if the node has not been written to storage yet.
-		/// </summary>
-		/// <returns>Locator for the node</returns>
-		public abstract NodeLocator GetLocator();
-
-		/// <summary>
-		/// Adds a callback to be executed once the node has been written. Triggers immediately if the node has already been written.
-		/// </summary>
-		/// <param name="callback">Action to be executed after the write</param>
-		public abstract void AddWriteCallback(NodeWriteCallback callback);
-
-		/// <summary>
-		/// Creates a reader for this node's data
-		/// </summary>
-		public abstract ValueTask<NodeData> ReadAsync(CancellationToken cancellationToken = default);
-
-		/// <summary>
-		/// Flush the node to storage and retrieve its locator
-		/// </summary>
-		public abstract ValueTask<NodeLocator> FlushAsync(CancellationToken cancellationToken = default);
-
-		/// <inheritdoc/>
-		public override string ToString() => HasLocator()? new HashedNodeLocator(Hash, GetLocator()).ToString() : Hash.ToString();
-	}
-
 	/// <summary>
 	/// Implementation of <see cref="NodeHandle"/> for nodes which can be read from storage
 	/// </summary>
