Horde: Rename TreeReader -> BundleNodeReader and TreeWriter -> BundleNodeWriter.

diff --git a/Engine/Source/Programs/Horde/Horde.Agent.Tests/JsonRpcLoggerTests.cs b/Engine/Source/Programs/Horde/Horde.Agent.Tests/JsonRpcLoggerTests.cs
--- a/Engine/Source/Programs/Horde/Horde.Agent.Tests/JsonRpcLoggerTests.cs
+++ b/Engine/Source/Programs/Horde/Horde.Agent.Tests/JsonRpcLoggerTests.cs
@@ -144,7 +144,7 @@ namespace Horde.Agent.Tests
 			using MemoryCache cache = new MemoryCache(new MemoryCacheOptions());
 			MemoryStorageClient store = new MemoryStorageClient();
 
-			TreeReader reader = new TreeReader(store, cache, NullLogger.Instance);
+			BundleReader reader = new BundleReader(store, cache, NullLogger.Instance);
 
 			await using FakeLogSink innerSink = new FakeLogSink();
 
diff --git a/Engine/Source/Programs/Horde/Horde.Agent/Execution/JobExecutor.cs b/Engine/Source/Programs/Horde/Horde.Agent/Execution/JobExecutor.cs
--- a/Engine/Source/Programs/Horde/Horde.Agent/Execution/JobExecutor.cs
+++ b/Engine/Source/Programs/Horde/Horde.Agent/Execution/JobExecutor.cs
@@ -792,7 +792,7 @@ namespace Horde.Agent.Execution
 
 			using MemoryCache cache = new MemoryCache(new MemoryCacheOptions { });
 			IStorageClient storage = _storageFactory.CreateStorageClient(_session, _namespaceId, _token);
-			TreeReader reader = new TreeReader(storage, cache, _logger);
+			BundleReader reader = new BundleReader(storage, cache, _logger);
 
 			// Create the mapping of tag names to file sets
 			Dictionary<string, HashSet<FileReference>> tagNameToFileSet = new Dictionary<string, HashSet<FileReference>>();
@@ -954,7 +954,7 @@ namespace Horde.Agent.Execution
 				Stopwatch timer = Stopwatch.StartNew();
 				RefName refName = TempStorage.GetRefNameForNode(_storagePrefix, step.Name);
 
-				TreeOptions treeOptions = new TreeOptions();
+				BundleOptions treeOptions = new BundleOptions();
 				await using IStorageWriter treeWriter = storage.CreateWriter(refName, treeOptions);
 
 				DirectoryNode outputNode = new DirectoryNode();
diff --git a/Engine/Source/Programs/Horde/Horde.Server/Commands/Bundles/BundleWriteTest.cs b/Engine/Source/Programs/Horde/Horde.Server/Commands/Bundles/BundleWriteTest.cs
--- a/Engine/Source/Programs/Horde/Horde.Server/Commands/Bundles/BundleWriteTest.cs
+++ b/Engine/Source/Programs/Horde/Horde.Server/Commands/Bundles/BundleWriteTest.cs
@@ -20,11 +20,11 @@ namespace Horde.Server.Commands.Bundles
 	{
 		class FakeStorageClient : IStorageClient
 		{
-			readonly TreeReader _reader;
+			readonly BundleReader _reader;
 
 			public FakeStorageClient()
 			{
-				_reader = new TreeReader(this, null, NullLogger.Instance);
+				_reader = new BundleReader(this, null, NullLogger.Instance);
 			}
 
 			public Task DeleteRefAsync(RefName name, CancellationToken cancellationToken = default) => Task.CompletedTask;
@@ -37,14 +37,14 @@ namespace Horde.Server.Commands.Bundles
 			public Task<BlobLocator> WriteBlobAsync(Stream stream, Utf8String prefix = default, CancellationToken cancellationToken = default) => Task.FromResult(BlobLocator.Create(HostId.Empty));
 			public Task WriteRefTargetAsync(RefName name, NodeHandle target, RefOptions? options = null, CancellationToken cancellationToken = default) => Task.CompletedTask;
 
-			public IStorageWriter CreateWriter(RefName refName = default, TreeOptions? options = null) => new TreeWriter(this, _reader, options);
+			public IStorageWriter CreateWriter(RefName refName = default, BundleOptions? options = null) => new BundleWriter(this, _reader, options);
 		}
 
 		public override async Task<int> ExecuteAsync(ILogger logger)
 		{
 			FakeStorageClient store = new FakeStorageClient();
 
-			TreeOptions options = new TreeOptions();
+			BundleOptions options = new BundleOptions();
 			options.CompressionFormat = BundleCompressionFormat.None;
 			await using IStorageWriter writer = store.CreateWriter(options: options);
 
diff --git a/Engine/Source/Programs/Horde/Horde.Server/Perforce/PerforceReplicator.cs b/Engine/Source/Programs/Horde/Horde.Server/Perforce/PerforceReplicator.cs
--- a/Engine/Source/Programs/Horde/Horde.Server/Perforce/PerforceReplicator.cs
+++ b/Engine/Source/Programs/Horde/Horde.Server/Perforce/PerforceReplicator.cs
@@ -27,7 +27,7 @@ namespace Horde.Server.Perforce
 	class PerforceReplicationOptions
 	{
 		public bool IncludeContent { get; set; }
-		public TreeOptions TreeOptions { get; set; } = new TreeOptions();
+		public BundleOptions TreeOptions { get; set; } = new BundleOptions();
 		public ChunkingOptions ChunkingOptions { get; set; } = new ChunkingOptions();
 		public RefOptions RefOptions { get; set; } = new RefOptions();
 	}
diff --git a/Engine/Source/Programs/Horde/Horde.Server/ServerSettings.cs b/Engine/Source/Programs/Horde/Horde.Server/ServerSettings.cs
--- a/Engine/Source/Programs/Horde/Horde.Server/ServerSettings.cs
+++ b/Engine/Source/Programs/Horde/Horde.Server/ServerSettings.cs
@@ -107,7 +107,7 @@ namespace Horde.Server
 		/// <summary>
 		/// Options for creating bundles
 		/// </summary>
-		TreeOptions Bundle { get; }
+		BundleOptions Bundle { get; }
 
 		/// <summary>
 		/// Options for chunking content
@@ -121,7 +121,7 @@ namespace Horde.Server
 	public class TreeStoreOptions : BlobStoreOptions, ITreeStoreOptions
 	{
 		/// <inheritdoc/>
-		public TreeOptions Bundle { get; set; } = new TreeOptions();
+		public BundleOptions Bundle { get; set; } = new BundleOptions();
 
 		/// <inheritdoc/>
 		public ChunkingOptions Chunking { get; set; } = new ChunkingOptions();
@@ -281,7 +281,7 @@ namespace Horde.Server
 		/// <summary>
 		/// Options for how objects are packed together
 		/// </summary>
-		public TreeOptions Bundle { get; set; } = new TreeOptions();
+		public BundleOptions Bundle { get; set; } = new BundleOptions();
 
 		/// <summary>
 		/// Options for how objects are sliced
diff --git a/Engine/Source/Programs/Horde/Horde.Server/Storage/StorageController.cs b/Engine/Source/Programs/Horde/Horde.Server/Storage/StorageController.cs
--- a/Engine/Source/Programs/Horde/Horde.Server/Storage/StorageController.cs
+++ b/Engine/Source/Programs/Horde/Horde.Server/Storage/StorageController.cs
@@ -484,7 +484,7 @@ namespace Horde.Server.Storage
 			}
 
 			IStorageClient storageClient = await _storageService.GetClientAsync(namespaceId, cancellationToken);
-			TreeReader reader = new TreeReader(storageClient, _memoryCache, _logger);
+			BundleReader reader = new BundleReader(storageClient, _memoryCache, _logger);
 
 			BundleHeader header = await reader.ReadBundleHeaderAsync(locator, cancellationToken);
 
@@ -564,7 +564,7 @@ namespace Horde.Server.Storage
 			}
 
 			IStorageClient storageClient = await _storageService.GetClientAsync(namespaceId, cancellationToken);
-			TreeReader reader = new TreeReader(storageClient, _memoryCache, _logger);
+			BundleReader reader = new BundleReader(storageClient, _memoryCache, _logger);
 
 			BundleHeader header = await reader.ReadBundleHeaderAsync(locator, cancellationToken);
 			BundleExport export = header.Exports[exportIdx];
diff --git a/Engine/Source/Programs/Horde/Horde.Server/Storage/StorageService.cs b/Engine/Source/Programs/Horde/Horde.Server/Storage/StorageService.cs
--- a/Engine/Source/Programs/Horde/Horde.Server/Storage/StorageService.cs
+++ b/Engine/Source/Programs/Horde/Horde.Server/Storage/StorageService.cs
@@ -570,18 +570,6 @@ namespace Horde.Server.Storage
 		/// <inheritdoc/>
 		async ValueTask<IStorageClient> IStorageClientFactory.GetClientAsync(NamespaceId namespaceId, CancellationToken cancellationToken) => await GetClientAsync(namespaceId, cancellationToken);
 
-		/// <summary>
-		/// Gets a <see cref="TreeReader"/> for a particular namespace
-		/// </summary>
-		/// <param name="namespaceId">Namespace identififer</param>
-		/// <param name="cancellationToken">Cancellation token for the operation</param>
-		/// <returns></returns>
-		public async ValueTask<TreeReader> GetReaderAsync(NamespaceId namespaceId, CancellationToken cancellationToken)
-		{
-			IStorageClientImpl storageClient = await GetClientAsync(namespaceId, cancellationToken);
-			return new TreeReader(storageClient, _cache, _logger);
-		}
-
 		#region Config
 
 		State GetNextState()
diff --git a/Engine/Source/Programs/Horde/Horde/Commands/Bundles/BundleExtract.cs b/Engine/Source/Programs/Horde/Horde/Commands/Bundles/BundleExtract.cs
--- a/Engine/Source/Programs/Horde/Horde/Commands/Bundles/BundleExtract.cs
+++ b/Engine/Source/Programs/Horde/Horde/Commands/Bundles/BundleExtract.cs
@@ -42,7 +42,7 @@ namespace Horde.Commands.Bundles
 			}
 
 			using MemoryCache cache = new MemoryCache(new MemoryCacheOptions());
-			TreeReader reader = new TreeReader(store, cache, logger);
+			BundleReader reader = new BundleReader(store, cache, logger);
 
 			Stopwatch timer = Stopwatch.StartNew();
 
diff --git a/Engine/Source/Programs/Horde/Horde/Commands/Vcs/VcsLog.cs b/Engine/Source/Programs/Horde/Horde/Commands/Vcs/VcsLog.cs
--- a/Engine/Source/Programs/Horde/Horde/Commands/Vcs/VcsLog.cs
+++ b/Engine/Source/Programs/Horde/Horde/Commands/Vcs/VcsLog.cs
@@ -30,7 +30,7 @@ namespace Horde.Commands.Vcs
 			IStorageClient store = await GetStorageClientAsync();
 
 			using MemoryCache cache = new MemoryCache(new MemoryCacheOptions());
-			TreeReader reader = new TreeReader(store, cache, logger);
+			BundleReader reader = new BundleReader(store, cache, logger);
 
 			List<CommitNode> commits = new List<CommitNode>();
 
diff --git a/Engine/Source/Programs/Shared/EpicGames.Horde.Tests/BundleStoreTests.cs b/Engine/Source/Programs/Shared/EpicGames.Horde.Tests/BundleStoreTests.cs
--- a/Engine/Source/Programs/Shared/EpicGames.Horde.Tests/BundleStoreTests.cs
+++ b/Engine/Source/Programs/Shared/EpicGames.Horde.Tests/BundleStoreTests.cs
@@ -61,7 +61,7 @@ namespace EpicGames.Horde.Tests
 		static async Task<Bundle> CreateBundleNormalAsync()
 		{
 			MemoryStorageClient store = new MemoryStorageClient();
-			await using IStorageWriter writer = store.CreateWriter(options: new TreeOptions { CompressionFormat = BundleCompressionFormat.None });
+			await using IStorageWriter writer = store.CreateWriter(options: new BundleOptions { CompressionFormat = BundleCompressionFormat.None });
 
 			TextNode node = new TextNode("Hello world");
 			NodeHandle handle = await writer.FlushAsync(node, CancellationToken.None);
@@ -94,7 +94,7 @@ namespace EpicGames.Horde.Tests
 			using IMemoryCache cache = new MemoryCache(new MemoryCacheOptions());
 
 			MemoryStorageClient blobStore = new MemoryStorageClient();
-			await TestTreeAsync(blobStore, new TreeOptions { MaxBlobSize = 1024 * 1024 });
+			await TestTreeAsync(blobStore, new BundleOptions { MaxBlobSize = 1024 * 1024 });
 
 			Assert.AreEqual(1, blobStore.Blobs.Count);
 			Assert.AreEqual(1, blobStore.Refs.Count);
@@ -106,7 +106,7 @@ namespace EpicGames.Horde.Tests
 			using IMemoryCache cache = new MemoryCache(new MemoryCacheOptions());
 
 			MemoryStorageClient blobStore = new MemoryStorageClient();
-			await TestTreeAsync(blobStore, new TreeOptions { MaxBlobSize = 1 });
+			await TestTreeAsync(blobStore, new BundleOptions { MaxBlobSize = 1 });
 
 			Assert.AreEqual(5, blobStore.Blobs.Count);
 			Assert.AreEqual(1, blobStore.Refs.Count);
@@ -139,7 +139,7 @@ namespace EpicGames.Horde.Tests
 			public override IEnumerable<NodeRef> EnumerateRefs() => Refs;
 		}
 
-		static async Task TestTreeAsync(MemoryStorageClient store, TreeOptions options)
+		static async Task TestTreeAsync(MemoryStorageClient store, BundleOptions options)
 		{
 			// Generate a tree
 			{
@@ -154,7 +154,7 @@ namespace EpicGames.Horde.Tests
 
 				await store.WriteRefTargetAsync(new RefName("test"), await writer.WriteNodeAsync(root));
 
-				TreeReader reader = new TreeReader(store, null, NullLogger.Instance);
+				BundleReader reader = new BundleReader(store, null, NullLogger.Instance);
 				await CheckTree(root);
 			}
 
@@ -262,7 +262,7 @@ namespace EpicGames.Horde.Tests
 		{
 			using IMemoryCache cache = new MemoryCache(new MemoryCacheOptions());
 			MemoryStorageClient store = new MemoryStorageClient();
-			TreeReader reader = new TreeReader(store, null, NullLogger.Instance);
+			BundleReader reader = new BundleReader(store, null, NullLogger.Instance);
 
 			// Generate a tree
 			{
@@ -319,7 +319,7 @@ namespace EpicGames.Horde.Tests
 			// Generate a tree
 			NodeRef<ChunkedDataNode> nodeRef;
 			{
-				await using IStorageWriter writer = store.CreateWriter(options: new TreeOptions { MaxBlobSize = 1024 });
+				await using IStorageWriter writer = store.CreateWriter(options: new BundleOptions { MaxBlobSize = 1024 });
 
 				ChunkingOptions options = new ChunkingOptions();
 				options.LeafOptions = new LeafChunkedDataNodeOptions(128, 256, 64 * 1024);
@@ -366,7 +366,7 @@ namespace EpicGames.Horde.Tests
 			// Generate a tree
 			DirectoryNode root;
 			{
-				await using IStorageWriter writer = store.CreateWriter(options: new TreeOptions { MaxBlobSize = 1024 });
+				await using IStorageWriter writer = store.CreateWriter(options: new BundleOptions { MaxBlobSize = 1024 });
 
 				ChunkingOptions options = new ChunkingOptions();
 				options.LeafOptions = new LeafChunkedDataNodeOptions(128, 256, 64 * 1024);
diff --git a/Engine/Source/Programs/Shared/EpicGames.Horde.Tests/BundleTests.cs b/Engine/Source/Programs/Shared/EpicGames.Horde.Tests/BundleTests.cs
--- a/Engine/Source/Programs/Shared/EpicGames.Horde.Tests/BundleTests.cs
+++ b/Engine/Source/Programs/Shared/EpicGames.Horde.Tests/BundleTests.cs
@@ -68,7 +68,7 @@ namespace EpicGames.Horde.Tests
 		public async Task BasicChunkingTests()
 		{
 			RefName refName = new RefName("test");
-			TreeReader reader = new TreeReader(_storage, null, NullLogger.Instance);
+			BundleReader reader = new BundleReader(_storage, null, NullLogger.Instance);
 			await using IStorageWriter writer = _storage.CreateWriter(refName);
 
 			ChunkingOptions options = new ChunkingOptions();
@@ -244,7 +244,7 @@ namespace EpicGames.Horde.Tests
 			Assert.AreEqual(0, _storage.Refs.Count);
 			Assert.AreEqual(0, _storage.Blobs.Count);
 
-			TreeOptions options = new TreeOptions();
+			BundleOptions options = new BundleOptions();
 			options.MaxBlobSize = 1;
 
 			await using (IStorageWriter writer = _storage.CreateWriter())
@@ -268,7 +268,7 @@ namespace EpicGames.Horde.Tests
 		[TestMethod]
 		public async Task ReloadTests()
 		{
-			TreeOptions options = new TreeOptions();
+			BundleOptions options = new BundleOptions();
 			options.MaxBlobSize = 1;
 
 			RefName refName = new RefName("ref");
diff --git a/Engine/Source/Programs/Shared/EpicGames.Horde.Tests/ChunkingTests.cs b/Engine/Source/Programs/Shared/EpicGames.Horde.Tests/ChunkingTests.cs
--- a/Engine/Source/Programs/Shared/EpicGames.Horde.Tests/ChunkingTests.cs
+++ b/Engine/Source/Programs/Shared/EpicGames.Horde.Tests/ChunkingTests.cs
@@ -72,7 +72,7 @@ namespace EpicGames.Horde.Tests
 			using MemoryCache cache = new MemoryCache(new MemoryCacheOptions());
 
 			MemoryStorageClient store = new MemoryStorageClient();
-			TreeReader reader = new TreeReader(store, cache, NullLogger.Instance);
+			BundleReader reader = new BundleReader(store, cache, NullLogger.Instance);
 			await using IStorageWriter writer = store.CreateWriter();
 
 			byte[] data = new byte[4096];
@@ -117,7 +117,7 @@ namespace EpicGames.Horde.Tests
 			await CheckSizes(reader, root, options, true);
 		}
 
-		async Task CheckSizes(TreeReader reader, ChunkedDataNode node, ChunkingOptions options, bool rightmost)
+		async Task CheckSizes(BundleReader reader, ChunkedDataNode node, ChunkingOptions options, bool rightmost)
 		{
 			if (node is LeafChunkedDataNode leafNode)
 			{
diff --git a/Engine/Source/Programs/Shared/EpicGames.Horde/Compute/AgentMessageHandler.cs b/Engine/Source/Programs/Shared/EpicGames.Horde/Compute/AgentMessageHandler.cs
--- a/Engine/Source/Programs/Shared/EpicGames.Horde/Compute/AgentMessageHandler.cs
+++ b/Engine/Source/Programs/Shared/EpicGames.Horde/Compute/AgentMessageHandler.cs
@@ -119,7 +119,7 @@ namespace EpicGames.Horde.Compute
 		async Task WriteFilesAsync(AgentMessageChannel channel, string path, NodeLocator locator, CancellationToken cancellationToken)
 		{
 			using AgentStorageClient store = new AgentStorageClient(channel);
-			TreeReader reader = new TreeReader(store, _memoryCache, _logger);
+			BundleReader reader = new BundleReader(store, _memoryCache, _logger);
 
 			DirectoryNode directoryNode = await reader.ReadNodeAsync<DirectoryNode>(locator, cancellationToken);
 
diff --git a/Engine/Source/Programs/Shared/EpicGames.Horde/Logs/LogNode.cs b/Engine/Source/Programs/Shared/EpicGames.Horde/Logs/LogNode.cs
--- a/Engine/Source/Programs/Shared/EpicGames.Horde/Logs/LogNode.cs
+++ b/Engine/Source/Programs/Shared/EpicGames.Horde/Logs/LogNode.cs
@@ -324,7 +324,7 @@ namespace EpicGames.Horde.Logs
 		/// <param name="reader">Reader to pull nodes from</param>
 		/// <param name="cancellationToken">Cancellation token</param>
 		/// <returns>Sequence of line buffers</returns>
-		public static async IAsyncEnumerable<ReadOnlyMemory<byte>> ReadLogAsync(this LogNode logNode, TreeReader reader, [EnumeratorCancellation] CancellationToken cancellationToken)
+		public static async IAsyncEnumerable<ReadOnlyMemory<byte>> ReadLogAsync(this LogNode logNode, BundleReader reader, [EnumeratorCancellation] CancellationToken cancellationToken)
 		{
 			foreach (LogChunkRef textChunkRef in logNode.TextChunkRefs)
 			{
diff --git a/Engine/Source/Programs/Shared/EpicGames.Horde/Storage/BundleStorageClient.cs b/Engine/Source/Programs/Shared/EpicGames.Horde/Storage/BundleStorageClient.cs
--- a/Engine/Source/Programs/Shared/EpicGames.Horde/Storage/BundleStorageClient.cs
+++ b/Engine/Source/Programs/Shared/EpicGames.Horde/Storage/BundleStorageClient.cs
@@ -19,14 +19,14 @@ namespace EpicGames.Horde.Storage
 		/// <summary>
 		/// Reader for node data
 		/// </summary>
-		protected TreeReader TreeReader { get; }
+		protected BundleReader TreeReader { get; }
 
 		/// <summary>
 		/// Constructor
 		/// </summary>
 		protected BundleStorageClient(IMemoryCache? memoryCache, ILogger logger)
 		{
-			TreeReader = new TreeReader(this, memoryCache, logger);
+			TreeReader = new BundleReader(this, memoryCache, logger);
 		}
 
 		#region Blobs
@@ -45,9 +45,9 @@ namespace EpicGames.Horde.Storage
 		#region Nodes
 
 		/// <inheritdoc/>
-		public IStorageWriter CreateWriter(RefName refName = default, TreeOptions? options = null)
+		public IStorageWriter CreateWriter(RefName refName = default, BundleOptions? options = null)
 		{
-			return new TreeWriter(this, TreeReader, refName, options);
+			return new BundleWriter(this, TreeReader, refName, options);
 		}
 
 		#endregion
diff --git a/Engine/Source/Programs/Shared/EpicGames.Horde/Storage/IStorageClient.cs b/Engine/Source/Programs/Shared/EpicGames.Horde/Storage/IStorageClient.cs
--- a/Engine/Source/Programs/Shared/EpicGames.Horde/Storage/IStorageClient.cs
+++ b/Engine/Source/Programs/Shared/EpicGames.Horde/Storage/IStorageClient.cs
@@ -117,7 +117,7 @@ namespace EpicGames.Horde.Storage
 		/// <param name="refName">Name of the ref being written.</param>
 		/// <param name="options"></param>
 		/// <returns>New writer instance. Must be disposed after use.</returns>
-		IStorageWriter CreateWriter(RefName refName = default, TreeOptions? options = null);
+		IStorageWriter CreateWriter(RefName refName = default, BundleOptions? options = null);
 
 		#endregion
 
diff --git a/Engine/Source/Programs/Shared/EpicGames.Horde/Storage/Node.cs b/Engine/Source/Programs/Shared/EpicGames.Horde/Storage/Node.cs
--- a/Engine/Source/Programs/Shared/EpicGames.Horde/Storage/Node.cs
+++ b/Engine/Source/Programs/Shared/EpicGames.Horde/Storage/Node.cs
@@ -427,7 +427,7 @@ namespace EpicGames.Horde.Storage
 		/// <param name="refOptions">Options for the ref</param>
 		/// <param name="cancellationToken">Cancellation token for the operation</param>
 		/// <returns>Location of node targetted by the ref</returns>
-		public static async Task<NodeHandle> WriteNodeAsync(this IStorageClient store, RefName name, Node node, TreeOptions? options = null, RefOptions? refOptions = null, CancellationToken cancellationToken = default)
+		public static async Task<NodeHandle> WriteNodeAsync(this IStorageClient store, RefName name, Node node, BundleOptions? options = null, RefOptions? refOptions = null, CancellationToken cancellationToken = default)
 		{
 			await using IStorageWriter writer = store.CreateWriter(name, options);
 			NodeRef<Node> nodeRef = await writer.WriteNodeAsync(node, cancellationToken);
diff --git a/Engine/Source/Programs/Shared/EpicGames.Horde/Storage/Nodes/ChunkedDataWriter.cs b/Engine/Source/Programs/Shared/EpicGames.Horde/Storage/Nodes/ChunkedDataWriter.cs
--- a/Engine/Source/Programs/Shared/EpicGames.Horde/Storage/Nodes/ChunkedDataWriter.cs
+++ b/Engine/Source/Programs/Shared/EpicGames.Horde/Storage/Nodes/ChunkedDataWriter.cs
@@ -55,7 +55,7 @@ namespace EpicGames.Horde.Storage.Nodes
 	}
 
 	/// <summary>
-	/// Utility class for generating FileNode data directly into <see cref="TreeWriter"/> instances, without constructing node representations first.
+	/// Utility class for generating FileNode data directly into <see cref="BundleWriter"/> instances, without constructing node representations first.
 	/// </summary>
 	public class ChunkedDataWriter
 	{
