Horde: Missing files.

diff --git a/Engine/Source/Programs/Shared/EpicGames.Horde/Storage/BundleReader.cs b/Engine/Source/Programs/Shared/EpicGames.Horde/Storage/BundleReader.cs
new file mode 100644
--- /dev/null
+++ b/Engine/Source/Programs/Shared/EpicGames.Horde/Storage/BundleReader.cs
@@ -0,0 +1,594 @@
+// Copyright Epic Games, Inc. All Rights Reserved.
+
+using System;
+using System.Buffers;
+using System.Collections.Generic;
+using System.Linq;
+using System.Threading;
+using System.Threading.Tasks;
+using EpicGames.Core;
+using Microsoft.Extensions.Caching.Memory;
+using Microsoft.Extensions.Logging;
+using Microsoft.CodeAnalysis;
+using System.Diagnostics;
+
+namespace EpicGames.Horde.Storage
+{
+	/// <summary>
+	/// Writes nodes from bundles in an <see cref="IStorageClient"/> instance.
+	/// </summary>
+	public class BundleReader
+	{
+		/// <summary>
+		/// Computed information about a bundle
+		/// </summary>
+		class BundleInfo
+		{
+			public readonly BlobLocator Locator;
+			public readonly BundleHeader Header;
+			public readonly int HeaderLength;
+
+			public BundleInfo(BlobLocator locator, BundleHeader header, int headerLength)
+			{
+				Locator = locator;
+				Header = header;
+				HeaderLength = headerLength;
+			}
+		}
+
+		/// <summary>
+		/// Bundle header queued to be read
+		/// </summary>
+		class QueuedHeader
+		{
+			public readonly BlobLocator Blob;
+			public readonly TaskCompletionSource<BundleInfo> CompletionSource = new TaskCompletionSource<BundleInfo>();
+
+			public BlobId BlobId => Blob.BlobId;
+
+			public QueuedHeader(BlobLocator blob)
+			{
+				Blob = blob;
+			}
+		}
+
+		/// <summary>
+		/// Encoded bundle packet queued to be read
+		/// </summary>
+		class QueuedPacket
+		{
+			public readonly BundleInfo Bundle;
+			public readonly int PacketIdx;
+			public readonly TaskCompletionSource<ReadOnlyMemory<byte>> CompletionSource = new TaskCompletionSource<ReadOnlyMemory<byte>>();
+
+			public BlobId BlobId => Bundle.Locator.BlobId;
+
+			public QueuedPacket(BundleInfo bundle, int packetIdx)
+			{
+				Bundle = bundle;
+				PacketIdx = packetIdx;
+			}
+		}
+
+		// Size of data to fetch by default. This is larger than the minimum request size to reduce number of reads.
+		const int DefaultFetchSize = 15 * 1024 * 1024;
+
+		// When reader is uncached, use a smaller default fetch size
+		const int DefaultUncachedFetchSize = 1 * 1024 * 1024;
+
+		readonly IStorageClient _store;
+		readonly IMemoryCache? _cache;
+		readonly ILogger _logger;
+
+		readonly object _queueLock = new object();
+		readonly List<QueuedHeader> _queuedHeaders = new List<QueuedHeader>();
+		readonly List<QueuedPacket> _queuedPackets = new List<QueuedPacket>();
+		readonly Dictionary<string, Task<ReadOnlyMemory<byte>>> _decodeTasks = new Dictionary<string, Task<ReadOnlyMemory<byte>>>(StringComparer.Ordinal);
+		Task? _readTask;
+
+		/// <summary>
+		/// Constructor
+		/// </summary>
+		/// <param name="store"></param>
+		/// <param name="cache">Cache for data</param>
+		/// <param name="logger">Logger for output</param>
+		public BundleReader(IStorageClient store, IMemoryCache? cache, ILogger logger)
+		{
+			_store = store;
+			_cache = cache;
+			_logger = logger;
+		}
+
+		#region Bundles
+
+		static string GetBundleInfoCacheKey(BlobId blobId) => $"bundle:{blobId}";
+		static string GetEncodedPacketCacheKey(BlobId blobId, int packetIdx) => $"encoded-packet:{blobId}#{packetIdx}";
+		static string GetDecodedPacketCacheKey(BlobId blobId, int packetIdx) => $"decoded-packet:{blobId}#{packetIdx}";
+
+		/// <summary>
+		/// Adds an object to the storage cache
+		/// </summary>
+		/// <param name="cacheKey">Key for the item</param>
+		/// <param name="value">Value to add</param>
+		/// <param name="size">Size of the value</param>
+		void AddToCache(string cacheKey, object value, int size)
+		{
+			if (_cache != null)
+			{
+				using (ICacheEntry entry = _cache.CreateEntry(cacheKey))
+				{
+					entry.SetValue(value);
+					entry.SetSize(size);
+				}
+			}
+		}
+
+		/// <summary>
+		/// Starts the background task for reading data from the store
+		/// </summary>
+		void StartReadTask()
+		{
+			if (_readTask == null)
+			{
+				_readTask = Task.Run(() => ServiceReadQueueAsync(CancellationToken.None));
+			}
+		}
+
+		/// <summary>
+		/// Dispatches requests in the read queue
+		/// </summary>
+		/// <param name="cancellationToken">Cancellation token for the background task</param>
+		async Task ServiceReadQueueAsync(CancellationToken cancellationToken)
+		{
+			const int MaxConcurrentReads = 4;
+
+			List<(BlobId BlobId, Task Task)> currentTasks = new List<(BlobId, Task)>();
+			for (; ; )
+			{
+				// Start any new reads
+				lock (_queueLock)
+				{
+					while (currentTasks.Count < MaxConcurrentReads)
+					{
+						HashSet<BlobId> currentBlobIds = new HashSet<BlobId>(currentTasks.Select(x => x.BlobId));
+
+						// Try to start another header read
+						QueuedHeader? queuedHeader = _queuedHeaders.FirstOrDefault(x => !currentBlobIds.Contains(x.BlobId));
+						if (queuedHeader != null)
+						{
+							Task task = Task.Run(() => PerformHeaderReadGuardedAsync(queuedHeader, cancellationToken), cancellationToken);
+							currentTasks.Add((queuedHeader.BlobId, task));
+							continue;
+						}
+
+						// Try to start another packet read
+						QueuedPacket? queuedPacket = _queuedPackets.FirstOrDefault(x => !currentBlobIds.Contains(x.BlobId));
+						if (queuedPacket != null)
+						{
+							Task task = Task.Run(() => PerformPacketReadGuardedAsync(queuedPacket, cancellationToken), cancellationToken);
+							currentTasks.Add((queuedPacket.BlobId, task));
+							continue;
+						}
+
+						// If we're not waiting for anything else and there are no more requests, end the task thread.
+						if (currentTasks.Count == 0)
+						{
+							_readTask = null;
+							return;
+						}
+
+						// Break out of the loop
+						break;
+					}
+				}
+
+				// Wait for any read task to complete
+				await Task.WhenAny(currentTasks.Select(x => x.Task));
+
+				// Remove any tasks which are complete
+				for (int idx = 0; idx < currentTasks.Count; idx++)
+				{
+					Task task = currentTasks[idx].Task;
+					if (task.IsCompleted)
+					{
+						if (task.Exception != null)
+						{
+							_logger.LogError(task.Exception, "Exception while reading from blob {BlobId}.", currentTasks[idx].BlobId);
+						}
+						currentTasks.RemoveAt(idx--);
+					}
+				}
+			}
+		}
+
+		/// <summary>
+		/// Reads a bundle header from the queue
+		/// </summary>
+		/// <param name="queuedHeader">The header to read</param>
+		/// <param name="cancellationToken">Cancellation token for the operation</param>
+		async Task PerformHeaderReadGuardedAsync(QueuedHeader queuedHeader, CancellationToken cancellationToken)
+		{
+			try
+			{
+				await PerformHeaderReadAsync(queuedHeader, cancellationToken);
+			}
+			catch (Exception ex)
+			{
+				queuedHeader.CompletionSource.TrySetException(ex);
+			}
+		}
+
+		/// <summary>
+		/// Reads a bundle header from the queue
+		/// </summary>
+		/// <param name="queuedHeader">The header to read</param>
+		/// <param name="cancellationToken">Cancellation token for the operation</param>
+		async Task PerformHeaderReadAsync(QueuedHeader queuedHeader, CancellationToken cancellationToken)
+		{
+			int prefetchSize = _cache != null ? DefaultFetchSize : DefaultUncachedFetchSize;
+			for (; ; )
+			{
+				using (IMemoryOwner<byte> owner = MemoryPool<byte>.Shared.Rent(prefetchSize))
+				{
+					// Read the prefetch size from the blob
+					Memory<byte> memory = owner.Memory.Slice(0, prefetchSize);
+					memory = await _store.ReadBlobRangeAsync(queuedHeader.Blob, 0, memory, cancellationToken);
+
+					// Make sure it's large enough to hold the header
+					int headerSize = BundleHeader.ReadPrelude(memory.Span);
+					if (headerSize > prefetchSize)
+					{
+						prefetchSize = headerSize;
+						continue;
+					}
+
+					// Parse the header and construct the bundle info from it
+					BundleHeader header = BundleHeader.Read(memory.ToArray());
+
+					// Construct the bundle info
+					BundleInfo bundleInfo = new BundleInfo(queuedHeader.Blob, header, headerSize);
+
+					if (_cache != null)
+					{
+						// Also add any encoded packets we prefetched
+						List<ReadOnlyMemory<byte>> packets = new List<ReadOnlyMemory<byte>>();
+						for (int packetOffset = headerSize; packets.Count < header.Packets.Count;)
+						{
+							int packetIdx = packets.Count;
+
+							int packetLength = header.Packets[packetIdx].EncodedLength;
+							if (packetOffset + packetLength > memory.Length)
+							{
+								break;
+							}
+
+							ReadOnlyMemory<byte> packetData = memory.Slice(packetOffset, packetLength).ToArray();
+							AddToCache(GetEncodedPacketCacheKey(queuedHeader.Blob.BlobId, packetIdx), packetData, packetData.Length);
+							packets.Add(packetData);
+							packetOffset += packetLength;
+						}
+
+						// Add the info to the cache
+						string cacheKey = GetBundleInfoCacheKey(queuedHeader.Blob.BlobId);
+						AddToCache(cacheKey, bundleInfo, headerSize);
+					}
+
+					// Update the task
+					queuedHeader.CompletionSource.SetResult(bundleInfo);
+
+					// Remove it from the queue
+					lock (_queueLock)
+					{
+						_queuedHeaders.Remove(queuedHeader);
+					}
+					break;
+				}
+			}
+		}
+
+		/// <summary>
+		/// Reads a packet from storage
+		/// </summary>
+		/// <param name="queuedPacket">The packet to read</param>
+		/// <param name="cancellationToken">Cancellation token for the operation</param>
+		async Task PerformPacketReadGuardedAsync(QueuedPacket queuedPacket, CancellationToken cancellationToken)
+		{
+			try
+			{
+				await PerformPacketReadAsync(queuedPacket, cancellationToken);
+			}
+			catch (Exception ex)
+			{
+				queuedPacket.CompletionSource.SetException(ex);
+			}
+		}
+
+		/// <summary>
+		/// Reads a packet from storage
+		/// </summary>
+		/// <param name="queuedPacket">The packet to read</param>
+		/// <param name="cancellationToken">Cancellation token for the operation</param>
+		async Task PerformPacketReadAsync(QueuedPacket queuedPacket, CancellationToken cancellationToken)
+		{
+			BundleInfo bundleInfo = queuedPacket.Bundle;
+			int minPacketIdx = queuedPacket.PacketIdx;
+
+			BundlePacket packet = bundleInfo.Header.Packets[minPacketIdx];
+			int readLength = packet.EncodedLength;
+
+			int maxPacketIdx = minPacketIdx + 1;
+			for (; maxPacketIdx < bundleInfo.Header.Packets.Count; maxPacketIdx++)
+			{
+				int nextReadLength = readLength + bundleInfo.Header.Packets[maxPacketIdx].EncodedLength;
+				if (nextReadLength > DefaultFetchSize)
+				{
+					break;
+				}
+				readLength = nextReadLength;
+			}
+
+			using (IMemoryOwner<byte> owner = MemoryPool<byte>.Shared.Rent(readLength))
+			{
+				Memory<byte> buffer = owner.Memory.Slice(0, readLength);
+				buffer = await _store.ReadBlobRangeAsync(bundleInfo.Locator, bundleInfo.HeaderLength + bundleInfo.Header.Packets[minPacketIdx].EncodedOffset, buffer, cancellationToken);
+
+				// Copy all the packets that have been read into separate buffers, so we can cache them individually.
+				ReadOnlyMemory<byte>?[] packets = new ReadOnlyMemory<byte>?[maxPacketIdx - minPacketIdx];
+				if (_cache != null)
+				{
+					for (int idx = minPacketIdx; idx < maxPacketIdx; idx++)
+					{
+						string cacheKey = GetEncodedPacketCacheKey(bundleInfo.Locator.BlobId, idx);
+						ReadOnlyMemory<byte> data = buffer.Slice(bundleInfo.Header.Packets[idx].EncodedOffset - bundleInfo.Header.Packets[minPacketIdx].EncodedOffset, bundleInfo.Header.Packets[idx].EncodedLength).ToArray();
+						packets[idx - minPacketIdx] = data;
+						AddToCache(cacheKey, data, data.Length);
+					}
+				}
+
+				// Find all the packets we can mark as complete
+				List<QueuedPacket> updatePackets = new List<QueuedPacket>();
+				lock (_queueLock)
+				{
+					updatePackets.AddRange(_queuedPackets.Where(x => x.BlobId == bundleInfo.Locator.BlobId && (x.PacketIdx >= minPacketIdx && x.PacketIdx < maxPacketIdx)));
+				}
+
+				// Mark them all as complete
+				foreach (QueuedPacket updatePacket in updatePackets)
+				{
+					ReadOnlyMemory<byte>? data = packets[updatePacket.PacketIdx - minPacketIdx];
+					if (data == null)
+					{
+						data = buffer.Slice(bundleInfo.Header.Packets[updatePacket.PacketIdx].EncodedOffset - bundleInfo.Header.Packets[minPacketIdx].EncodedOffset, bundleInfo.Header.Packets[updatePacket.PacketIdx].EncodedLength).ToArray();
+					}
+					updatePacket.CompletionSource.SetResult((ReadOnlyMemory<byte>)data);
+				}
+
+				// Remove all the completed packets from the queue
+				lock (_queueLock)
+				{
+					_queuedPackets.RemoveAll(x => updatePackets.Contains(x));
+				}
+			}
+		}
+
+		/// <summary>
+		/// Reads a bundle header from the given blob locator, or retrieves it from the cache
+		/// </summary>
+		/// <param name="locator"></param>
+		/// <param name="cancellationToken"></param>
+		/// <returns></returns>
+		public async Task<BundleHeader> ReadBundleHeaderAsync(BlobLocator locator, CancellationToken cancellationToken = default)
+		{
+			BundleInfo info = await GetBundleInfoAsync(locator, cancellationToken);
+			return info.Header;
+		}
+
+		/// <summary>
+		/// Reads the header and structural metadata about the bundle
+		/// </summary>
+		/// <param name="locator">The bundle location</param>
+		/// <param name="cancellationToken">Cancellation token for the operation</param>
+		/// <returns>Information about the bundle</returns>
+		async ValueTask<BundleInfo> GetBundleInfoAsync(BlobLocator locator, CancellationToken cancellationToken = default)
+		{
+			Debug.Assert(locator.IsValid());
+
+			string cacheKey = GetBundleInfoCacheKey(locator.BlobId);
+			if (_cache != null && _cache.TryGetValue(cacheKey, out BundleInfo bundleInfo))
+			{
+				return bundleInfo;
+			}
+
+			QueuedHeader? queuedHeader;
+			lock (_queueLock)
+			{
+				// Check the cache again inside lock scope to avoid races
+				if (_cache != null && _cache.TryGetValue(cacheKey, out bundleInfo))
+				{
+					return bundleInfo;
+				}
+
+				// Find or start the read
+				queuedHeader = _queuedHeaders.FirstOrDefault(x => x.Blob.BlobId == locator.BlobId);
+				if (queuedHeader == null)
+				{
+					queuedHeader = new QueuedHeader(locator);
+					_queuedHeaders.Add(queuedHeader);
+					StartReadTask();
+				}
+			}
+			return await queuedHeader.CompletionSource.Task.AbandonOnCancel(cancellationToken);
+		}
+
+		/// <summary>
+		/// Gets a decoded block from the store
+		/// </summary>
+		/// <param name="bundleInfo">Information about the bundle</param>
+		/// <param name="packetIdx">Index of the packet</param>
+		/// <param name="cancellationToken">Cancellation token for the operation</param>
+		/// <returns>The decoded data</returns>
+		async ValueTask<ReadOnlyMemory<byte>> ReadBundlePacketAsync(BundleInfo bundleInfo, int packetIdx, CancellationToken cancellationToken)
+		{
+			char c = bundleInfo.Locator.BlobId.ToString()[0];
+			Debug.Assert(Char.IsLetterOrDigit(c));
+
+			if (packetIdx < 0 || packetIdx >= bundleInfo.Header.Packets.Count)
+			{
+				throw new ArgumentException("Packet index is out of range", nameof(packetIdx));
+			}
+
+			string decodedCacheKey = GetDecodedPacketCacheKey(bundleInfo.Locator.BlobId, packetIdx);
+			if (_cache != null && _cache.TryGetValue(decodedCacheKey, out ReadOnlyMemory<byte> decodedPacket))
+			{
+				return decodedPacket;
+			}
+
+			Task<ReadOnlyMemory<byte>>? decodeTask;
+			lock (_queueLock)
+			{
+				// Query the cache again, to eliminate races between cache checks and decode tasks finishing.
+				if (_cache != null && _cache.TryGetValue(decodedCacheKey, out decodedPacket))
+				{
+					return decodedPacket;
+				}
+
+				// Create an async task to read the data
+				if (!_decodeTasks.TryGetValue(decodedCacheKey, out decodeTask))
+				{
+					decodeTask = Task.Run(() => ReadAndDecodePacketAsync(bundleInfo, packetIdx), CancellationToken.None);
+					_decodeTasks.Add(decodedCacheKey, decodeTask);
+				}
+			}
+			return await decodeTask.AbandonOnCancel(cancellationToken);
+		}
+
+		/// <summary>
+		/// Reads and decodes a packet from a bundle
+		/// </summary>
+		/// <param name="bundleInfo">Bundle to read from</param>
+		/// <param name="packetIdx">Index of the packet to return</param>
+		/// <returns>The decoded packet data</returns>
+		async Task<ReadOnlyMemory<byte>> ReadAndDecodePacketAsync(BundleInfo bundleInfo, int packetIdx)
+		{
+			ReadOnlyMemory<byte> encodedPacket = await ReadEncodedPacketAsync(bundleInfo, packetIdx);
+
+			BundlePacket packet = bundleInfo.Header.Packets[packetIdx];
+			byte[] decodedPacket = new byte[packet.DecodedLength];
+
+			BundleData.Decompress(packet.CompressionFormat, encodedPacket, decodedPacket);
+
+			string decodedCacheKey = GetDecodedPacketCacheKey(bundleInfo.Locator.BlobId, packetIdx);
+			AddToCache(decodedCacheKey, (ReadOnlyMemory<byte>)decodedPacket, decodedPacket.Length);
+
+			lock (_queueLock)
+			{
+				_decodeTasks.Remove(decodedCacheKey);
+			}
+
+			return decodedPacket;
+		}
+
+		/// <summary>
+		/// Reads an encoded packet from a bundle
+		/// </summary>
+		/// <param name="bundleInfo">Bundle to read from</param>
+		/// <param name="packetIdx">Index of the packet to return</param>
+		/// <returns>The encoded packet data</returns>
+		async ValueTask<ReadOnlyMemory<byte>> ReadEncodedPacketAsync(BundleInfo bundleInfo, int packetIdx)
+		{
+			char c = bundleInfo.Locator.BlobId.ToString()[0];
+			Debug.Assert(Char.IsLetterOrDigit(c));
+
+			if (packetIdx < 0 || packetIdx >= bundleInfo.Header.Packets.Count)
+			{
+				throw new ArgumentException("Packet index is out of range", nameof(packetIdx));
+			}
+
+			string encodedCacheKey = GetEncodedPacketCacheKey(bundleInfo.Locator.BlobId, packetIdx);
+			if (_cache != null && _cache.TryGetValue(encodedCacheKey, out ReadOnlyMemory<byte> encodedPacket))
+			{
+				return encodedPacket;
+			}
+
+			QueuedPacket? queuedPacket;
+			lock (_queueLock)
+			{
+				// Query the cache again, to eliminate races between cache checks and decode tasks finishing.
+				if (_cache != null && _cache.TryGetValue(encodedCacheKey, out encodedPacket))
+				{
+					return encodedPacket;
+				}
+
+				// Add a read to the queue
+				queuedPacket = _queuedPackets.FirstOrDefault(x => x.Bundle.Locator.BlobId == bundleInfo.Locator.BlobId && x.PacketIdx == packetIdx);
+				if (queuedPacket == null)
+				{
+					queuedPacket = new QueuedPacket(bundleInfo, packetIdx);
+					_queuedPackets.Add(queuedPacket);
+					StartReadTask();
+				}
+			}
+			return await queuedPacket.CompletionSource.Task;
+		}
+
+		#endregion
+
+		/// <summary>
+		/// Reads a node from a bundle
+		/// </summary>
+		/// <param name="locator">Locator for the node</param>
+		/// <param name="cancellationToken">Cancellation token for the operation</param>
+		/// <returns>Node data read from the given bundle</returns>
+		public async ValueTask<NodeData> ReadNodeDataAsync(NodeLocator locator, CancellationToken cancellationToken = default)
+		{
+			BundleInfo bundleInfo = await GetBundleInfoAsync(locator.Blob, cancellationToken);
+			BundleExport export = bundleInfo.Header.Exports[locator.ExportIdx];
+
+			List<NodeHandle> refs = new List<NodeHandle>(export.References.Count);
+			foreach (BundleExportRef reference in export.References)
+			{
+				BlobLocator importBlob;
+				if (reference.ImportIdx == -1)
+				{
+					importBlob = locator.Blob;
+				}
+				else
+				{
+					importBlob = bundleInfo.Header.Imports[reference.ImportIdx];
+				}
+				Debug.Assert(importBlob.IsValid());
+				refs.Add(new FlushedNodeHandle(this, reference.Hash, new NodeLocator(importBlob, reference.NodeIdx)));
+			}
+
+			ReadOnlyMemory<byte> nodeData = ReadOnlyMemory<byte>.Empty;
+			if (export.Length > 0)
+			{
+				ReadOnlyMemory<byte> packetData = await ReadBundlePacketAsync(bundleInfo, export.Packet, cancellationToken);
+				nodeData = packetData.Slice(export.Offset, export.Length);
+			}
+
+			NodeType nodeType = bundleInfo.Header.Types[export.TypeIdx];
+			return new NodeData(nodeType, export.Hash, nodeData, refs);
+		}
+
+		/// <summary>
+		/// Reads a node from a bundle
+		/// </summary>
+		/// <param name="locator">Locator for the node</param>
+		/// <param name="cancellationToken">Cancellation token for the operation</param>
+		/// <returns>Node data read from the given bundle</returns>
+		public async ValueTask<Node> ReadNodeAsync(NodeLocator locator, CancellationToken cancellationToken = default)
+		{
+			NodeData nodeData = await ReadNodeDataAsync(locator, cancellationToken);
+			return Node.Deserialize(nodeData);
+		}
+
+		/// <summary>
+		/// Reads a node from a bundle
+		/// </summary>
+		/// <param name="locator">Locator for the node</param>
+		/// <param name="cancellationToken">Cancellation token for the operation</param>
+		/// <returns>Node data read from the given bundle</returns>
+		public async ValueTask<TNode> ReadNodeAsync<TNode>(NodeLocator locator, CancellationToken cancellationToken = default) where TNode : Node => (TNode)await ReadNodeAsync(locator, cancellationToken);
+	}
+}
diff --git a/Engine/Source/Programs/Shared/EpicGames.Horde/Storage/BundleWriter.cs b/Engine/Source/Programs/Shared/EpicGames.Horde/Storage/BundleWriter.cs
new file mode 100644
--- /dev/null
+++ b/Engine/Source/Programs/Shared/EpicGames.Horde/Storage/BundleWriter.cs
@@ -0,0 +1,893 @@
+// Copyright Epic Games, Inc. All Rights Reserved.
+
+using System;
+using System.Buffers;
+using System.Collections.Generic;
+using System.Diagnostics;
+using System.Diagnostics.CodeAnalysis;
+using System.Linq;
+using System.Threading;
+using System.Threading.Tasks;
+using EpicGames.Core;
+using Microsoft.Extensions.Logging;
+
+namespace EpicGames.Horde.Storage
+{
+	/// <summary>
+	/// Options for configuring a bundle serializer
+	/// </summary>
+	public class BundleOptions
+	{
+		/// <summary>
+		/// Maximum payload size fo a blob
+		/// </summary>
+		public int MaxBlobSize { get; set; } = 10 * 1024 * 1024;
+
+		/// <summary>
+		/// Compression format to use
+		/// </summary>
+		public BundleCompressionFormat CompressionFormat { get; set; } = BundleCompressionFormat.LZ4;
+
+		/// <summary>
+		/// Minimum size of a block to be compressed
+		/// </summary>
+		public int MinCompressionPacketSize { get; set; } = 16 * 1024;
+
+		/// <summary>
+		/// Maximum amount of data to store in memory. This includes any background writes as well as bundles being built.
+		/// </summary>
+		public int MaxInMemoryDataLength { get; set; } = 128 * 1024 * 1024;
+
+		/// <summary>
+		/// Number of nodes to cache
+		/// </summary>
+		public int NodeCacheSize { get; set; } = 1024;
+	}
+
+	/// <summary>
+	/// Unique identifier for a node
+	/// </summary>
+	/// <param name="Hash">Hash of the node</param>
+	/// <param name="Type">Type of the node</param>
+	public record NodeKey(IoHash Hash, NodeType Type);
+
+	/// <summary>
+	/// Implementation of <see cref="NodeHandle"/> for nodes which can be read from storage
+	/// </summary>
+	public sealed class FlushedNodeHandle : NodeHandle
+	{
+		readonly BundleReader _reader;
+		readonly NodeLocator _locator;
+
+		/// <summary>
+		/// Constructor
+		/// </summary>
+		public FlushedNodeHandle(BundleReader reader, HashedNodeLocator hashedLocator)
+			: this(reader, hashedLocator.Hash, hashedLocator.Locator)
+		{
+		}
+
+		/// <summary>
+		/// Constructor
+		/// </summary>
+		public FlushedNodeHandle(BundleReader reader, IoHash hash, NodeLocator locator)
+			: base(hash)
+		{
+			_reader = reader;
+			_locator = locator;
+		}
+
+		/// <inheritdoc/>
+		public override bool HasLocator() => true;
+
+		/// <inheritdoc/>
+		public override NodeLocator GetLocator() => _locator;
+
+		/// <inheritdoc/>
+		public override void AddWriteCallback(NodeWriteCallback callback) => callback.OnWrite();
+
+		/// <inheritdoc/>
+		public override ValueTask<NodeData> ReadAsync(CancellationToken cancellationToken = default) => _reader.ReadNodeDataAsync(_locator, cancellationToken);
+
+		/// <inheritdoc/>
+		public override ValueTask<NodeLocator> FlushAsync(CancellationToken cancellationToken = default) => new ValueTask<NodeLocator>(_locator);
+	}
+
+	/// <summary>
+	/// Index of known nodes that can be used for deduplication.
+	/// </summary>
+	public class NodeCache
+	{
+		readonly object _lockObject = new object();
+		readonly int _maxKeys;
+		readonly Queue<NodeKey> _nodeKeys = new Queue<NodeKey>();
+		readonly Dictionary<NodeKey, NodeHandle> _nodeKeyToHandle = new Dictionary<NodeKey, NodeHandle>();
+
+		/// <summary>
+		/// Constructor
+		/// </summary>
+		/// <param name="maxKeys">Maximum number of node keys to keep in the cache</param>
+		public NodeCache(int maxKeys)
+		{
+			_maxKeys = maxKeys;
+			_nodeKeys = new Queue<NodeKey>(maxKeys);
+			_nodeKeyToHandle = new Dictionary<NodeKey, NodeHandle>(maxKeys);
+		}
+
+		/// <summary>
+		/// Adds a new node handle to the cache
+		/// </summary>
+		/// <param name="key">Unique node key</param>
+		/// <param name="handle">Handle to the node</param>
+		public void Add(NodeKey key, NodeHandle handle)
+		{
+			lock (_lockObject)
+			{
+				AddInternal(key, handle);
+			}
+		}
+
+		void AddInternal(NodeKey key, NodeHandle handle)
+		{
+			NodeKey? prevKey;
+			if (_nodeKeys.Count == _maxKeys && _nodeKeys.TryDequeue(out prevKey))
+			{
+				_nodeKeyToHandle.Remove(prevKey);
+			}
+			_nodeKeyToHandle.TryAdd(key, handle);
+		}
+
+		/// <summary>
+		/// Find a node within the cache
+		/// </summary>
+		/// <param name="key">Key to look up in the cache</param>
+		/// <param name="handle">Handle for the node</param>
+		/// <returns>True if the node was found</returns>
+		public bool TryGetNode(NodeKey key, [NotNullWhen(true)] out NodeHandle? handle)
+		{
+			lock (_lockObject)
+			{
+				return _nodeKeyToHandle.TryGetValue(key, out handle);
+			}
+		}
+	}
+
+	/// <summary>
+	/// Writes nodes of a tree to an <see cref="IStorageClient"/>, packed into bundles. Each <see cref="BundleWriter"/> instance is single threaded,
+	/// but multiple instances may be written to in parallel.
+	/// </summary>
+	public sealed class BundleWriter : IStorageWriter
+	{
+		// Information about a unique output node. Note that multiple node refs may de-duplicate to the same output node.
+		internal class PendingNode : NodeHandle
+		{
+			readonly BundleReader _reader;
+
+			object LockObject => _reader;
+
+			NodeLocator _locator;
+			PendingBundle? _pendingBundle;
+
+			public readonly NodeKey Key;
+			public readonly int Packet;
+			public readonly int Offset;
+			public readonly int Length;
+			public readonly NodeHandle[] Refs;
+
+			public PendingBundle? PendingBundle => _pendingBundle;
+
+			public PendingNode(BundleReader reader, NodeKey key, int packet, int offset, int length, IReadOnlyList<NodeHandle> refs, PendingBundle pendingBundle)
+				: base(key.Hash)
+			{
+				_reader = reader;
+
+				Key = key;
+				Packet = packet;
+				Offset = offset;
+				Length = length;
+				Refs = refs.ToArray();
+
+				_pendingBundle = pendingBundle;
+			}
+
+			/// <inheritdoc/>
+			public override bool HasLocator() => _locator.IsValid();
+
+			/// <inheritdoc/>
+			public override NodeLocator GetLocator()
+			{
+				if (!_locator.IsValid())
+				{
+					throw new InvalidOperationException();
+				}
+				return _locator;
+			}
+
+			public void MarkAsWritten(NodeLocator locator)
+			{
+				lock (LockObject)
+				{
+					Debug.Assert(!_locator.IsValid());
+					_locator = locator;
+					_pendingBundle = null;
+				}
+			}
+
+			/// <inheritdoc/>
+			public override void AddWriteCallback(NodeWriteCallback callback)
+			{
+				PendingBundle? pendingBundle = PendingBundle;
+				if (pendingBundle == null || !pendingBundle.TryAddWriteCallback(callback))
+				{
+					Debug.Assert(_locator.IsValid());
+					callback.OnWrite();
+				}
+			}
+
+			/// <inheritdoc/>
+			public override async ValueTask<NodeData> ReadAsync(CancellationToken cancellationToken = default)
+			{
+				if (!_locator.IsValid())
+				{
+					lock (LockObject)
+					{
+						if (!_locator.IsValid() && _pendingBundle != null)
+						{
+							ReadOnlyMemory<byte> data = _pendingBundle.GetNodeData(Packet, Offset, Length);
+							if (data.Length == Length)
+							{
+								return new NodeData(Key.Type, Hash, data, Refs);
+							}
+						}
+					}
+				}
+
+				return await _reader.ReadNodeDataAsync(_locator, cancellationToken);
+			}
+
+			/// <inheritdoc/>
+			public override async ValueTask<NodeLocator> FlushAsync(CancellationToken cancellationToken = default)
+			{
+				if (_locator.IsValid())
+				{
+					return _locator;
+				}
+
+				foreach (NodeHandle nodeRef in Refs)
+				{
+					await nodeRef.FlushAsync(cancellationToken);
+				}
+
+				PendingBundle? pendingBundle = _pendingBundle;
+				if (pendingBundle != null)
+				{
+					await pendingBundle.FlushAsync(cancellationToken);
+				}
+
+				return GetLocator();
+			}
+		}
+
+		// Information about a bundle being built. Metadata operations are synchronous, compression/writes are asynchronous.
+		internal class PendingBundle : IDisposable
+		{
+			class WriteCallbackSentinel : NodeWriteCallback
+			{
+				public override void OnWrite() => throw new NotImplementedException();
+			}
+
+			static readonly WriteCallbackSentinel s_callbackSentinel = new WriteCallbackSentinel();
+
+			static int s_lastBundleId = 0;
+			public int BundleId { get; } = Interlocked.Increment(ref s_lastBundleId);
+
+			readonly BundleReader _treeReader;
+			readonly BundleWriter _treeWriter;
+			readonly int _maxPacketSize;
+			readonly BundleCompressionFormat _compressionFormat;
+
+			readonly object _lockObject = new object();
+
+			// The current packet being built
+			IMemoryOwner<byte>? _currentPacket;
+
+			// Index of the current packet being written to
+			int _currentPacketIdx;
+
+			// Length of the packet that has been written
+			int _currentPacketLength;
+
+			// Map of keys to nodes in the queue
+			public readonly Dictionary<NodeKey, PendingNode> _nodeKeyToInfo = new Dictionary<NodeKey, PendingNode>();
+
+			// Queue of nodes for the current bundle
+			readonly List<PendingNode> _queue = new List<PendingNode>();
+
+			// List of packets in the current bundle
+			readonly List<BundlePacket> _packets = new List<BundlePacket>();
+
+			// List of compressed packets
+			readonly ChunkedMemoryWriter _encodedPacketWriter;
+
+			// Set of all direct dependencies from this bundle
+			readonly HashSet<Task> _dependencies = new HashSet<Task>();
+
+			// Total size of uncompressed data in the current bundle
+			public int UncompressedLength { get; private set; }
+
+			// Whether all nodes have been added to a bundle
+			public bool IsReadOnly { get; private set; }
+
+			// List of post-write callbacks
+			NodeWriteCallback? _callbacks = null;
+
+			// Task used to compress data in the background
+			Task _writeTask = Task.CompletedTask;
+
+			// Event which is signalled after the bundle is written to storage
+			readonly TaskCompletionSource<bool> _completeEvent = new TaskCompletionSource<bool>();
+
+			// Task signalled after the write is complete
+			public Task CompleteTask => _completeEvent.Task;
+
+			public PendingBundle(BundleReader treeReader, BundleWriter treeWriter, int maxPacketSize, int maxBlobSize, BundleCompressionFormat compressionFormat)
+			{
+				_treeReader = treeReader;
+				_treeWriter = treeWriter;
+				_maxPacketSize = maxPacketSize;
+				_compressionFormat = compressionFormat;
+				_encodedPacketWriter = new ChunkedMemoryWriter(maxBlobSize);
+			}
+
+			/// <inheritdoc/>
+			public void Dispose()
+			{
+				if (!CompleteTask.IsCompleted)
+				{
+					_completeEvent.TrySetException(new ObjectDisposedException($"Bundle {BundleId} has been disposed. Make sure to call FlushAsync() on forked writers that may have dependencies before disposing.", $"Bundle {BundleId}"));
+				}
+				if (_currentPacket != null)
+				{
+					_currentPacket.Dispose();
+					_currentPacket = null;
+				}
+				_encodedPacketWriter.Dispose();
+			}
+
+			// Whether this bundle has finished writing
+			public bool IsComplete() => CompleteTask.IsCompleted;
+
+			// Whether this bundle can be written
+			public bool CanComplete() => !IsReadOnly && _dependencies.Count == 0;
+
+			// Add a dependency onto another bundle
+			public void AddDependencyOn(PendingBundle bundle, ILogger? traceLogger)
+			{
+				if (bundle != this)
+				{
+					if (_dependencies.Add(bundle.CompleteTask))
+					{
+						traceLogger?.LogInformation("Added dependency from bundle {BundleId} on {OtherBundleId}", BundleId, bundle.BundleId);
+					}
+				}
+			}
+
+			// Adds a callback after writing
+			public bool TryAddWriteCallback(NodeWriteCallback callback)
+			{
+				for (; ; )
+				{
+					NodeWriteCallback? tail = _callbacks;
+					if (tail == s_callbackSentinel)
+					{
+						return false;
+					}
+
+					callback._next = tail;
+
+					if (Interlocked.CompareExchange(ref _callbacks, callback, tail) == tail)
+					{
+						return true;
+					}
+				}
+			}
+
+			// Try to add a ref to an existing output node
+			public bool TryAddRefToExistingNode(NodeKey nodeKey, [NotNullWhen(true)] out PendingNode? handle)
+			{
+				PendingNode? pendingNode;
+				if (_nodeKeyToInfo.TryGetValue(nodeKey, out pendingNode))
+				{
+					handle = pendingNode;
+					return true;
+				}
+				else
+				{
+					handle = null;
+					return false;
+				}
+			}
+
+			// Starts writing a node to the buffer
+			public Memory<byte> GetBuffer(int usedSize, int desiredSize)
+			{
+				if (_currentPacket == null)
+				{
+					_currentPacket = MemoryPool<byte>.Shared.Rent(Math.Max(_maxPacketSize, desiredSize));
+				}
+				else if (_currentPacketLength + desiredSize > _currentPacket.Memory.Length)
+				{
+					IMemoryOwner<byte> nextPacket = MemoryPool<byte>.Shared.Rent(Math.Max(_maxPacketSize, desiredSize));
+					if (usedSize > 0)
+					{
+						_currentPacket.Memory.Slice(_currentPacketLength, usedSize).CopyTo(nextPacket.Memory);
+					}
+					FlushPacket();
+					_currentPacket = nextPacket;
+				}
+				return _currentPacket.Memory.Slice(_currentPacketLength, desiredSize);
+			}
+
+			// Finish a node write
+			public PendingNode WriteNode(NodeKey nodeKey, int size, IReadOnlyList<NodeHandle> refs)
+			{
+				int offset = _currentPacketLength;
+
+				_currentPacketLength += size;
+				UncompressedLength += size;
+
+				PendingNode pendingNode = new PendingNode(_treeReader, nodeKey, _currentPacketIdx, offset, (int)size, refs, this);
+				_queue.Add(pendingNode);
+				_nodeKeyToInfo.Add(nodeKey, pendingNode);
+
+				return pendingNode;
+			}
+
+			public ReadOnlyMemory<byte> GetNodeData(int packetIdx, int offset, int length)
+			{
+				lock (_lockObject)
+				{
+					if (packetIdx < _packets.Count)
+					{
+						BundlePacket packet = _packets[packetIdx];
+
+						ReadOnlySequence<byte> sequence = _encodedPacketWriter.AsSequence(packet.EncodedOffset, packet.EncodedLength);
+						if (!sequence.IsSingleSegment)
+						{
+							throw new InvalidOperationException();
+						}
+
+						using (IMemoryOwner<byte> decodeBuffer = MemoryPool<byte>.Shared.Rent(packet.DecodedLength))
+						{
+							BundleData.Decompress(_compressionFormat, sequence.First, decodeBuffer.Memory);
+							return decodeBuffer.Memory.Slice(offset, length).ToArray();
+						}
+					}
+					else if (packetIdx == _currentPacketIdx)
+					{
+						return _currentPacket!.Memory.Slice(offset, length).Slice(offset, length).ToArray();
+					}
+					else
+					{
+						return ReadOnlyMemory<byte>.Empty;
+					}
+				}
+			}
+
+			// Compresses the current packet and schedule it to be written to storage
+			void FlushPacket()
+			{
+				if (_currentPacket != null)
+				{
+					IMemoryOwner<byte> currentPacket = _currentPacket;
+					int currentPacketLength = _currentPacketLength;
+
+					if (currentPacketLength > 0)
+					{
+						_writeTask = _writeTask.ContinueWith(x => CompressPacket(currentPacket, currentPacketLength), TaskScheduler.Default);
+						_currentPacketIdx++;
+					}
+					else
+					{
+						currentPacket.Dispose();
+					}
+
+					_currentPacket = null;
+					_currentPacketLength = 0;
+				}
+			}
+
+			void CompressPacket(IMemoryOwner<byte> packetData, int length)
+			{
+				lock (_lockObject)
+				{
+					int encodedLength = BundleData.Compress(_compressionFormat, packetData.Memory.Slice(0, length), _encodedPacketWriter);
+
+					int encodedOffset = (_packets.Count == 0) ? 0 : _packets[^1].EncodedOffset + _packets[^1].EncodedLength;
+					BundlePacket packet = new BundlePacket(_compressionFormat, encodedOffset, encodedLength, length);
+					_packets.Add(packet);
+
+					packetData.Dispose();
+				}
+			}
+
+			// Mark the bundle as complete
+			public void MarkAsComplete(IStorageClient store, Utf8String prefix, ILogger? traceLogger)
+			{
+				if (!IsReadOnly)
+				{
+					traceLogger?.LogInformation("Marking bundle {BundleId} as complete ({NumNodes} nodes); adding to write queue.", BundleId, _queue.Count);
+					FlushPacket();
+					Task prevWriteTask = _writeTask;
+					_writeTask = Task.Run(() => CompleteAsync(prevWriteTask, store, prefix, traceLogger));
+					IsReadOnly = true;
+				}
+			}
+
+			async Task CompleteAsync(Task prevWriteTask, IStorageClient store, Utf8String prefix, ILogger? traceLogger)
+			{
+				try
+				{
+					await prevWriteTask;
+					if (traceLogger != null)
+					{
+						foreach (Task dependency in _dependencies)
+						{
+							if (!dependency.IsCompleted)
+							{
+								traceLogger.LogInformation("Bundle {BundleId} is stalling waiting for dependencies to flush first", BundleId);
+							}
+						}
+					}
+					await Task.WhenAll(_dependencies);
+
+					Bundle bundle = CreateBundle();
+					BlobLocator locator = await store.WriteBundleAsync(bundle, prefix);
+					traceLogger?.LogInformation("Written bundle {BundleId} as {Locator}", BundleId, locator);
+
+					for (int idx = 0; idx < _queue.Count; idx++)
+					{
+						NodeLocator nodeLocator = new NodeLocator(locator, idx);
+						traceLogger?.LogInformation("Updated pending node {Hash} with locator {Locator}", _queue[idx].Key.Hash, nodeLocator);
+						_queue[idx].MarkAsWritten(nodeLocator);
+					}
+
+					NodeWriteCallback? callback = Interlocked.Exchange(ref _callbacks, s_callbackSentinel);
+					while (callback != null)
+					{
+						callback.OnWrite();
+						callback = callback._next;
+					}
+
+					_completeEvent.SetResult(true);
+				}
+				catch (Exception ex)
+				{
+					_completeEvent.SetException(ex);
+				}
+			}
+
+			// Flush the current write state
+			public Task FlushAsync(CancellationToken cancellationToken) => _treeWriter.FlushAsync(cancellationToken);
+
+			// Mark the bundle as complete and create a bundle with the current state
+			public Bundle CreateBundle()
+			{
+				lock (_lockObject)
+				{
+					return CreateBundleInternal();
+				}
+			}
+
+			Bundle CreateBundleInternal()
+			{
+				// Create a set from the nodes to be written. We use this to determine references that are imported.
+				HashSet<NodeHandle> nodeSet = new HashSet<NodeHandle>(_queue);
+
+				// Find all the imported nodes by bundle
+				Dictionary<BlobLocator, List<(int, NodeHandle)>> bundleToImports = new Dictionary<BlobLocator, List<(int, NodeHandle)>>();
+				foreach (PendingNode nodeInfo in _queue)
+				{
+					foreach (NodeHandle handle in nodeInfo.Refs)
+					{
+						if (nodeSet.Add(handle))
+						{
+							NodeLocator refLocator = handle.GetLocator();
+
+							List<(int, NodeHandle)>? importedNodes;
+							if (!bundleToImports.TryGetValue(refLocator.Blob, out importedNodes))
+							{
+								importedNodes = new List<(int, NodeHandle)>();
+								bundleToImports.Add(refLocator.Blob, importedNodes);
+							}
+
+							importedNodes.Add((refLocator.ExportIdx, handle));
+						}
+					}
+				}
+
+				List<BlobLocator> imports = new List<BlobLocator>();
+				Dictionary<BlobLocator, int> importToIndex = new Dictionary<BlobLocator, int>();
+
+				// List of types in the bundle
+				List<NodeType> types = new List<NodeType>();
+				Dictionary<NodeType, int> typeToIndex = new Dictionary<NodeType, int>();
+
+				// Map of node handle to reference
+				Dictionary<NodeHandle, BundleExportRef> nodeHandleToExportRef = new Dictionary<NodeHandle, BundleExportRef>();
+				for (int exportIdx = 0; exportIdx < _queue.Count; exportIdx++)
+				{
+					NodeHandle handle = _queue[exportIdx];
+					nodeHandleToExportRef[handle] = new BundleExportRef(-1, exportIdx, handle.Hash);
+				}
+
+				// Create the export list
+				List<BundleExport> exports = new List<BundleExport>(_queue.Count);
+				foreach (PendingNode nodeInfo in _queue)
+				{
+					int typeIdx;
+					if (!typeToIndex.TryGetValue(nodeInfo.Key.Type, out typeIdx))
+					{
+						typeIdx = types.Count;
+						typeToIndex.Add(nodeInfo.Key.Type, typeIdx);
+						types.Add(nodeInfo.Key.Type);
+					}
+
+					List<BundleExportRef> exportRefs = new List<BundleExportRef>();
+					foreach (NodeHandle handle in nodeInfo.Refs)
+					{
+						BundleExportRef exportRef;
+						if (!nodeHandleToExportRef.TryGetValue(handle, out exportRef))
+						{
+							NodeLocator locator = handle.GetLocator();
+
+							int importIdx;
+							if (!importToIndex.TryGetValue(locator.Blob, out importIdx))
+							{
+								importIdx = imports.Count;
+								imports.Add(locator.Blob);
+								importToIndex.Add(locator.Blob, importIdx);
+							}
+
+							exportRef = new BundleExportRef(importIdx, locator.ExportIdx, handle.Hash);
+						}
+						exportRefs.Add(exportRef);
+					}
+
+					BundleExport export = new BundleExport(typeIdx, nodeInfo.Key.Hash, nodeInfo.Packet, nodeInfo.Offset, nodeInfo.Length, exportRefs);
+					exports.Add(export);
+				}
+
+				// Get the memory for each packet
+				List<ReadOnlyMemory<byte>> packetData = new List<ReadOnlyMemory<byte>>(_packets.Count);
+				foreach (ReadOnlyMemory<byte> segment in _encodedPacketWriter.AsSequence())
+				{
+					int offset = 0;
+					while(packetData.Count < _packets.Count)
+					{
+						BundlePacket packet = _packets[packetData.Count];
+						int next = offset + packet.EncodedLength;
+
+						if (next > segment.Length)
+						{
+							break;
+						}
+						packetData.Add(segment.Slice(offset, packet.EncodedLength));
+						offset = next;
+					}
+				}
+
+				// Create the bundle
+				BundleHeader header = BundleHeader.Create(types, imports, exports, _packets);
+				return new Bundle(header, packetData);
+			}
+		}
+
+		static readonly BundleOptions s_defaultOptions = new BundleOptions();
+
+		readonly IStorageClient _store;
+		readonly BundleReader _treeReader;
+		readonly BundleOptions _options;
+		readonly Utf8String _prefix;
+
+		readonly NodeCache _nodeCache;
+		readonly Queue<PendingBundle> _writeQueue = new Queue<PendingBundle>();
+
+		PendingBundle? _currentBundle;
+		int _memoryFootprint;
+
+		bool _disposed;
+
+		internal readonly ILogger? _traceLogger;
+
+		/// <summary>
+		/// Accessor for the store backing this writer
+		/// </summary>
+		public IStorageClient Store => _store;
+
+		/// <summary>
+		/// Cache of nodes to deduplicate against
+		/// </summary>
+		public NodeCache NodeCache => _nodeCache;
+
+		/// <summary>
+		/// Constructor
+		/// </summary>
+		/// <param name="store">Store to write data to</param>
+		/// <param name="treeReader">Reader for serialized node data</param>
+		/// <param name="options">Options for the writer</param>
+		/// <param name="prefix">Prefix for blobs written to the store</param>
+		/// <param name="nodeCache">Cache of nodes for deduplication</param>
+		/// <param name="traceLogger">Optional logger for trace information</param>
+		public BundleWriter(IStorageClient store, BundleReader treeReader, BundleOptions? options = null, Utf8String prefix = default, NodeCache? nodeCache = null, ILogger? traceLogger = null)
+		{
+			_store = store;
+			_treeReader = treeReader;
+			_options = options ?? s_defaultOptions;
+			_prefix = prefix;
+			_nodeCache = nodeCache ?? new NodeCache(_options.NodeCacheSize);
+			_traceLogger = traceLogger;
+		}
+
+		/// <summary>
+		/// Constructor
+		/// </summary>
+		/// <param name="store">Store to write data to</param>
+		/// <param name="reader">Reader for serialized node data</param>
+		/// <param name="refName">Ref being written. Will be used as a prefix for storing blobs.</param>
+		/// <param name="options">Options for the writer</param>
+		/// <param name="traceLogger">Optional logger for trace information</param>
+		public BundleWriter(IStorageClient store, BundleReader reader, RefName refName, BundleOptions? options = null, ILogger? traceLogger = null)
+			: this(store, reader, options, refName.Text, traceLogger: traceLogger)
+		{
+		}
+
+		/// <summary>
+		/// Copy constructor
+		/// </summary>
+		/// <param name="other"></param>
+		public BundleWriter(BundleWriter other)
+			: this(other._store, other._treeReader, other._options, other._prefix, other._nodeCache, other._traceLogger)
+		{
+		}
+
+		/// <inheritdoc/>
+		IStorageWriter IStorageWriter.Fork() => new BundleWriter(this);
+
+		/// <inheritdoc/>
+		public async ValueTask DisposeAsync()
+		{
+			if (!_disposed)
+			{
+				await FlushAsync();
+				_disposed = true;
+			}
+
+			if (_currentBundle != null)
+			{
+				_currentBundle.Dispose();
+				_currentBundle = null;
+			}
+		}
+
+		/// <summary>
+		/// Mark this writer as complete, allowing its data to be serialized.
+		/// </summary>
+		public void Complete()
+		{
+			if (_currentBundle != null)
+			{
+				_currentBundle.MarkAsComplete(_store, _prefix, _traceLogger);
+				_writeQueue.Enqueue(_currentBundle);
+				_currentBundle = null;
+			}
+		}
+
+		/// <summary>
+		/// Gets an output buffer for writing.
+		/// </summary>
+		/// <param name="usedSize">Current size in the existing buffer that has been written to</param>
+		/// <param name="desiredSize">Desired size of the returned buffer</param>
+		/// <returns>Buffer to be written into.</returns>
+		public Memory<byte> GetOutputBuffer(int usedSize, int desiredSize)
+		{
+			PendingBundle currentBundle = GetCurrentBundle();
+			return currentBundle.GetBuffer(usedSize, desiredSize);
+		}
+
+		/// <summary>
+		/// Finish writing a node.
+		/// </summary>
+		/// <param name="size">Used size of the buffer</param>
+		/// <param name="references">References to other nodes</param>
+		/// <param name="type">Type of the node that was written</param>
+		/// <param name="cancellationToken">Cancellation token for the operation</param>
+		/// <returns>Handle to the written node</returns>
+		public async ValueTask<NodeHandle> WriteNodeAsync(int size, IReadOnlyList<NodeHandle> references, NodeType type, CancellationToken cancellationToken = default)
+		{
+			PendingBundle currentBundle = GetCurrentBundle();
+
+			// Get the hash for the new blob
+			ReadOnlyMemory<byte> memory = currentBundle.GetBuffer(size, size);
+			IoHash hash = IoHash.Compute(memory.Span);
+
+			// Create a unique key for the new node
+			NodeKey nodeKey = new NodeKey(hash, type);
+
+			// Check if we have a matching node already in storage
+			if (_nodeCache.TryGetNode(nodeKey, out NodeHandle? handle))
+			{
+				_traceLogger?.LogInformation("Returning cached handle for {NodeKey} -> {Handle}", nodeKey, handle);
+				return handle;
+			}
+
+			// Append this node data
+			PendingNode pendingNode = currentBundle.WriteNode(nodeKey, size, references);
+			_memoryFootprint += size;
+			_nodeCache.Add(nodeKey, pendingNode);
+			_traceLogger?.LogInformation("Added new node for {NodeKey} in bundle {BundleId}", nodeKey, currentBundle.BundleId);
+
+			// Add dependencies on all bundles containing a dependent node
+			foreach (NodeHandle reference in references)
+			{
+				PendingNode? pendingReference = reference as PendingNode;
+				if (pendingReference?.PendingBundle != null)
+				{
+					currentBundle.AddDependencyOn(pendingReference.PendingBundle, _traceLogger);
+				}
+			}
+
+			// If the bundle is full, start the process of writing it to disk
+			if (currentBundle.UncompressedLength > _options.MaxBlobSize)
+			{
+				Complete();
+			}
+
+			// Remove any complete bundle writes
+			while (_writeQueue.Count > 0 && (_writeQueue.Peek().IsComplete() || _memoryFootprint > _options.MaxInMemoryDataLength))
+			{
+				await WaitForWriteAsync(cancellationToken);
+			}
+
+			return pendingNode;
+		}
+
+		PendingBundle GetCurrentBundle()
+		{
+			if (_currentBundle == null || _currentBundle.IsReadOnly)
+			{
+				int bufferSize = (int)(_options.MinCompressionPacketSize * 1.2);
+				_currentBundle = new PendingBundle(_treeReader, this, bufferSize, _options.MaxBlobSize, _options.CompressionFormat);
+			}
+			return _currentBundle;
+		}
+
+		async Task WaitForWriteAsync(CancellationToken cancellationToken)
+		{
+			PendingBundle writtenBundle = _writeQueue.Dequeue();
+			_traceLogger?.LogInformation("Waiting for bundle {BundleId} to finish writing", writtenBundle.BundleId);
+			await await Task.WhenAny(writtenBundle.CompleteTask, Task.Delay(-1, cancellationToken));
+			_memoryFootprint -= writtenBundle.UncompressedLength;
+			writtenBundle.Dispose();
+		}
+
+		/// <summary>
+		/// Flushes all the current nodes to storage
+		/// </summary>
+		/// <param name="cancellationToken">Cancellation token for the operation</param>
+		/// <returns></returns>
+		public async Task FlushAsync(CancellationToken cancellationToken = default)
+		{
+			if (_disposed)
+			{
+				throw new ObjectDisposedException(GetType().Name);
+			}
+
+			Complete();
+			while (_writeQueue.Count > 0)
+			{
+				await WaitForWriteAsync(cancellationToken);
+			}
+		}
+	}
+}
diff --git a/Engine/Source/Programs/Shared/EpicGames.Horde/Storage/TreeReader.cs b/Engine/Source/Programs/Shared/EpicGames.Horde/Storage/TreeReader.cs
deleted file mode 100644
diff --git a/Engine/Source/Programs/Shared/EpicGames.Horde/Storage/TreeWriter.cs b/Engine/Source/Programs/Shared/EpicGames.Horde/Storage/TreeWriter.cs
deleted file mode 100644
