Horde: Remove bundle-specific options object from IStorageClient.

diff --git a/Engine/Source/Programs/Horde/Horde.Agent/Execution/JobExecutor.cs b/Engine/Source/Programs/Horde/Horde.Agent/Execution/JobExecutor.cs
--- a/Engine/Source/Programs/Horde/Horde.Agent/Execution/JobExecutor.cs
+++ b/Engine/Source/Programs/Horde/Horde.Agent/Execution/JobExecutor.cs
@@ -954,8 +954,7 @@ namespace Horde.Agent.Execution
 				Stopwatch timer = Stopwatch.StartNew();
 				RefName refName = TempStorage.GetRefNameForNode(_storagePrefix, step.Name);
 
-				BundleOptions treeOptions = new BundleOptions();
-				await using IStorageWriter treeWriter = storage.CreateWriter(refName, treeOptions);
+				await using IStorageWriter treeWriter = storage.CreateWriter(refName);
 
 				DirectoryNode outputNode = new DirectoryNode();
 
diff --git a/Engine/Source/Programs/Horde/Horde.Server/Commands/Bundles/BundleWriteTest.cs b/Engine/Source/Programs/Horde/Horde.Server/Commands/Bundles/BundleWriteTest.cs
--- a/Engine/Source/Programs/Horde/Horde.Server/Commands/Bundles/BundleWriteTest.cs
+++ b/Engine/Source/Programs/Horde/Horde.Server/Commands/Bundles/BundleWriteTest.cs
@@ -37,7 +37,9 @@ namespace Horde.Server.Commands.Bundles
 			public Task<BlobLocator> WriteBlobAsync(Stream stream, Utf8String prefix = default, CancellationToken cancellationToken = default) => Task.FromResult(BlobLocator.Create(HostId.Empty));
 			public Task WriteRefTargetAsync(RefName name, NodeHandle target, RefOptions? options = null, CancellationToken cancellationToken = default) => Task.CompletedTask;
 
-			public IStorageWriter CreateWriter(RefName refName = default, BundleOptions? options = null) => new BundleWriter(this, _reader, options);
+			public BundleWriter CreateWriter(RefName refName = default, BundleOptions? options = null) => new BundleWriter(this, _reader, options);
+
+			IStorageWriter IStorageClient.CreateWriter(EpicGames.Horde.Storage.RefName refName) => CreateWriter(refName);
 		}
 
 		public override async Task<int> ExecuteAsync(ILogger logger)
diff --git a/Engine/Source/Programs/Horde/Horde.Server/Perforce/PerforceReplicator.cs b/Engine/Source/Programs/Horde/Horde.Server/Perforce/PerforceReplicator.cs
--- a/Engine/Source/Programs/Horde/Horde.Server/Perforce/PerforceReplicator.cs
+++ b/Engine/Source/Programs/Horde/Horde.Server/Perforce/PerforceReplicator.cs
@@ -362,7 +362,7 @@ namespace Horde.Server.Perforce
 			_logger.LogInformation("Total sync size: {Size:n1}mb", totalSize / (1024.0 * 1024.0));
 
 			// Create the tree writer
-			await using IStorageWriter writer = store.CreateWriter(refName, options.TreeOptions);
+			await using IStorageWriter writer = store.CreateWriter(refName);
 
 			// Keep track of changes to make to the directory structure
 			DirectoryUpdate rootUpdate = new DirectoryUpdate();
diff --git a/Engine/Source/Programs/Shared/EpicGames.Horde.Tests/BundleTests.cs b/Engine/Source/Programs/Shared/EpicGames.Horde.Tests/BundleTests.cs
--- a/Engine/Source/Programs/Shared/EpicGames.Horde.Tests/BundleTests.cs
+++ b/Engine/Source/Programs/Shared/EpicGames.Horde.Tests/BundleTests.cs
@@ -247,7 +247,7 @@ namespace EpicGames.Horde.Tests
 			BundleOptions options = new BundleOptions();
 			options.MaxBlobSize = 1;
 
-			await using (IStorageWriter writer = _storage.CreateWriter())
+			await using (BundleWriter writer = _storage.CreateWriter())
 			{
 				DirectoryNode root = new DirectoryNode();
 				for (int idx = 1; idx <= 3; idx++)
@@ -258,7 +258,8 @@ namespace EpicGames.Horde.Tests
 				}
 
 				RefName refName = new RefName("ref");
-				await _storage.WriteNodeAsync(refName, root, options);
+				NodeRef<DirectoryNode> rootRef = await writer.WriteNodeAsync(root);
+				await _storage.WriteRefTargetAsync(refName, rootRef);
 			}
 
 			Assert.AreEqual(1, _storage.Refs.Count);
diff --git a/Engine/Source/Programs/Shared/EpicGames.Horde/Storage/BundleStorageClient.cs b/Engine/Source/Programs/Shared/EpicGames.Horde/Storage/BundleStorageClient.cs
--- a/Engine/Source/Programs/Shared/EpicGames.Horde/Storage/BundleStorageClient.cs
+++ b/Engine/Source/Programs/Shared/EpicGames.Horde/Storage/BundleStorageClient.cs
@@ -45,11 +45,14 @@ namespace EpicGames.Horde.Storage
 		#region Nodes
 
 		/// <inheritdoc/>
-		public IStorageWriter CreateWriter(RefName refName = default, BundleOptions? options = null)
+		public BundleWriter CreateWriter(RefName refName = default, BundleOptions? options = null)
 		{
 			return new BundleWriter(this, TreeReader, refName, options);
 		}
 
+		/// <inheritdoc/>
+		IStorageWriter IStorageClient.CreateWriter(RefName refName) => CreateWriter(refName);
+
 		#endregion
 
 		#region Aliases
diff --git a/Engine/Source/Programs/Shared/EpicGames.Horde/Storage/IStorageClient.cs b/Engine/Source/Programs/Shared/EpicGames.Horde/Storage/IStorageClient.cs
--- a/Engine/Source/Programs/Shared/EpicGames.Horde/Storage/IStorageClient.cs
+++ b/Engine/Source/Programs/Shared/EpicGames.Horde/Storage/IStorageClient.cs
@@ -115,9 +115,8 @@ namespace EpicGames.Horde.Storage
 		/// Creates a new writer for storage nodes
 		/// </summary>
 		/// <param name="refName">Name of the ref being written.</param>
-		/// <param name="options"></param>
 		/// <returns>New writer instance. Must be disposed after use.</returns>
-		IStorageWriter CreateWriter(RefName refName = default, BundleOptions? options = null);
+		IStorageWriter CreateWriter(RefName refName = default);
 
 		#endregion
 
diff --git a/Engine/Source/Programs/Shared/EpicGames.Horde/Storage/Node.cs b/Engine/Source/Programs/Shared/EpicGames.Horde/Storage/Node.cs
--- a/Engine/Source/Programs/Shared/EpicGames.Horde/Storage/Node.cs
+++ b/Engine/Source/Programs/Shared/EpicGames.Horde/Storage/Node.cs
@@ -423,13 +423,12 @@ namespace EpicGames.Horde.Storage
 		/// <param name="store">Store instance to write to</param>
 		/// <param name="name">Name of the ref containing this node</param>
 		/// <param name="node">Node to be written</param>
-		/// <param name="options">Options for the node writer</param>
 		/// <param name="refOptions">Options for the ref</param>
 		/// <param name="cancellationToken">Cancellation token for the operation</param>
 		/// <returns>Location of node targetted by the ref</returns>
-		public static async Task<NodeHandle> WriteNodeAsync(this IStorageClient store, RefName name, Node node, BundleOptions? options = null, RefOptions? refOptions = null, CancellationToken cancellationToken = default)
+		public static async Task<NodeHandle> WriteNodeAsync(this IStorageClient store, RefName name, Node node, RefOptions? refOptions = null, CancellationToken cancellationToken = default)
 		{
-			await using IStorageWriter writer = store.CreateWriter(name, options);
+			await using IStorageWriter writer = store.CreateWriter(name);
 			NodeRef<Node> nodeRef = await writer.WriteNodeAsync(node, cancellationToken);
 			await store.WriteRefTargetAsync(name, nodeRef.Handle, refOptions, cancellationToken);
 			return nodeRef.Handle;
