Horde: Move extension methods for serializing bundles into an implementation-specific file.

diff --git a/Engine/Source/Programs/Shared/EpicGames.Horde/Storage/BundleStorageClient.cs b/Engine/Source/Programs/Shared/EpicGames.Horde/Storage/BundleStorageClient.cs
--- a/Engine/Source/Programs/Shared/EpicGames.Horde/Storage/BundleStorageClient.cs
+++ b/Engine/Source/Programs/Shared/EpicGames.Horde/Storage/BundleStorageClient.cs
@@ -81,4 +81,43 @@ namespace EpicGames.Horde.Storage
 
 		#endregion
 	}
+
+	/// <summary>
+	/// Extension methods for serializing bundles
+	/// </summary>
+	public static class BundleStorageClientExtensions
+	{
+		#region Bundles
+
+		/// <summary>
+		/// Reads a bundle from the given blob id, or retrieves it from the cache
+		/// </summary>
+		/// <param name="store">The store instance to read from</param>
+		/// <param name="locator"></param>
+		/// <param name="cancellationToken"></param>
+		/// <returns></returns>
+		public static async Task<Bundle> ReadBundleAsync(this IStorageClient store, BlobLocator locator, CancellationToken cancellationToken = default)
+		{
+			using (Stream stream = await store.ReadBlobAsync(locator, cancellationToken))
+			{
+				return await Bundle.FromStreamAsync(stream, cancellationToken);
+			}
+		}
+
+		/// <summary>
+		/// Writes a new bundle to the store
+		/// </summary>
+		/// <param name="store">The store instance to write to</param>
+		/// <param name="bundle">Bundle data</param>
+		/// <param name="cancellationToken">Cancellation token for the operation</param>
+		/// <param name="prefix">Prefix for blob names. While the returned BlobId is guaranteed to be unique, this name can be used as a prefix to aid debugging.</param>
+		/// <returns>Unique identifier for the blob</returns>
+		public static async Task<BlobLocator> WriteBundleAsync(this IStorageClient store, Bundle bundle, Utf8String prefix = default, CancellationToken cancellationToken = default)
+		{
+			using ReadOnlySequenceStream stream = new ReadOnlySequenceStream(bundle.AsSequence());
+			return await store.WriteBlobAsync(stream, prefix, cancellationToken);
+		}
+
+		#endregion
+	}
 }
diff --git a/Engine/Source/Programs/Shared/EpicGames.Horde/Storage/IStorageClient.cs b/Engine/Source/Programs/Shared/EpicGames.Horde/Storage/IStorageClient.cs
--- a/Engine/Source/Programs/Shared/EpicGames.Horde/Storage/IStorageClient.cs
+++ b/Engine/Source/Programs/Shared/EpicGames.Horde/Storage/IStorageClient.cs
@@ -2,13 +2,7 @@
 
 using System;
 using System.Collections.Generic;
-using System.Diagnostics;
-using System.Diagnostics.CodeAnalysis;
-using System.Globalization;
 using System.IO;
-using System.Reflection;
-using System.Reflection.Metadata;
-using System.Text;
 using System.Threading;
 using System.Threading.Tasks;
 using EpicGames.Core;
@@ -279,39 +273,6 @@ namespace EpicGames.Horde.Storage
 	/// </summary>
 	public static class StorageClientExtensions
 	{
-		#region Bundles
-
-		/// <summary>
-		/// Reads a bundle from the given blob id, or retrieves it from the cache
-		/// </summary>
-		/// <param name="store">The store instance to read from</param>
-		/// <param name="locator"></param>
-		/// <param name="cancellationToken"></param>
-		/// <returns></returns>
-		public static async Task<Bundle> ReadBundleAsync(this IStorageClient store, BlobLocator locator, CancellationToken cancellationToken = default)
-		{
-			using (Stream stream = await store.ReadBlobAsync(locator, cancellationToken))
-			{
-				return await Bundle.FromStreamAsync(stream, cancellationToken);
-			}
-		}
-
-		/// <summary>
-		/// Writes a new bundle to the store
-		/// </summary>
-		/// <param name="store">The store instance to write to</param>
-		/// <param name="bundle">Bundle data</param>
-		/// <param name="cancellationToken">Cancellation token for the operation</param>
-		/// <param name="prefix">Prefix for blob names. While the returned BlobId is guaranteed to be unique, this name can be used as a prefix to aid debugging.</param>
-		/// <returns>Unique identifier for the blob</returns>
-		public static async Task<BlobLocator> WriteBundleAsync(this IStorageClient store, Bundle bundle, Utf8String prefix = default, CancellationToken cancellationToken = default)
-		{
-			using ReadOnlySequenceStream stream = new ReadOnlySequenceStream(bundle.AsSequence());
-			return await store.WriteBlobAsync(stream, prefix, cancellationToken);
-		}
-
-		#endregion
-
 		#region Blobs
 
 		/// <summary>
