Horde: Fix server tests.

diff --git a/Engine/Source/Programs/Shared/EpicGames.Horde.Tests/BundleTests.cs b/Engine/Source/Programs/Shared/EpicGames.Horde.Tests/BundleTests.cs
--- a/Engine/Source/Programs/Shared/EpicGames.Horde.Tests/BundleTests.cs
+++ b/Engine/Source/Programs/Shared/EpicGames.Horde.Tests/BundleTests.cs
@@ -247,7 +247,7 @@ namespace EpicGames.Horde.Tests
 			BundleOptions options = new BundleOptions();
 			options.MaxBlobSize = 1;
 
-			await using (BundleWriter writer = _storage.CreateWriter())
+			await using (BundleWriter writer = _storage.CreateWriter(options: options))
 			{
 				DirectoryNode root = new DirectoryNode();
 				for (int idx = 1; idx <= 3; idx++)
