Horde: Fix paths within bundles having an additional directory.

diff --git a/Engine/Source/Programs/Shared/EpicGames.Horde/Storage/Nodes/DirectoryNode.cs b/Engine/Source/Programs/Shared/EpicGames.Horde/Storage/Nodes/DirectoryNode.cs
--- a/Engine/Source/Programs/Shared/EpicGames.Horde/Storage/Nodes/DirectoryNode.cs
+++ b/Engine/Source/Programs/Shared/EpicGames.Horde/Storage/Nodes/DirectoryNode.cs
@@ -771,7 +771,7 @@ namespace EpicGames.Horde.Storage.Nodes
 			string[] fragments = path.Split(new char[] { '/', '\\' });
 
 			DirectoryUpdate lastDir = this;
-			for (int idx = 0; idx < fragments.Length; idx++)
+			for (int idx = 0; idx + 1 < fragments.Length; idx++)
 			{
 				DirectoryUpdate? nextDir;
 				if (!lastDir.Directories.TryGetValue(fragments[idx], out nextDir) || nextDir == null)
