Horde: Remove IStorageClient accessor from IStorageWriter. Add a method to flush the current ref instead.

diff --git a/Engine/Source/Programs/Horde/Horde.Server/Commands/Bundles/BundleWriteTest.cs b/Engine/Source/Programs/Horde/Horde.Server/Commands/Bundles/BundleWriteTest.cs
--- a/Engine/Source/Programs/Horde/Horde.Server/Commands/Bundles/BundleWriteTest.cs
+++ b/Engine/Source/Programs/Horde/Horde.Server/Commands/Bundles/BundleWriteTest.cs
@@ -37,7 +37,7 @@ namespace Horde.Server.Commands.Bundles
 			public Task<BlobLocator> WriteBlobAsync(Stream stream, Utf8String prefix = default, CancellationToken cancellationToken = default) => Task.FromResult(BlobLocator.Create(HostId.Empty));
 			public Task WriteRefTargetAsync(RefName name, NodeHandle target, RefOptions? options = null, CancellationToken cancellationToken = default) => Task.CompletedTask;
 
-			public BundleWriter CreateWriter(RefName refName = default, BundleOptions? options = null) => new BundleWriter(this, _reader, options);
+			public BundleWriter CreateWriter(RefName refName = default, BundleOptions? options = null) => new BundleWriter(this, _reader, refName, options);
 
 			IStorageWriter IStorageClient.CreateWriter(EpicGames.Horde.Storage.RefName refName) => CreateWriter(refName);
 		}
diff --git a/Engine/Source/Programs/Shared/EpicGames.Horde/Storage/BundleWriter.cs b/Engine/Source/Programs/Shared/EpicGames.Horde/Storage/BundleWriter.cs
--- a/Engine/Source/Programs/Shared/EpicGames.Horde/Storage/BundleWriter.cs
+++ b/Engine/Source/Programs/Shared/EpicGames.Horde/Storage/BundleWriter.cs
@@ -686,9 +686,9 @@ namespace EpicGames.Horde.Storage
 		static readonly BundleOptions s_defaultOptions = new BundleOptions();
 
 		readonly IStorageClient _store;
-		readonly BundleReader _treeReader;
+		readonly BundleReader _reader;
 		readonly BundleOptions _options;
-		readonly Utf8String _prefix;
+		readonly RefName _refName;
 
 		readonly NodeCache _nodeCache;
 		readonly Queue<PendingBundle> _writeQueue = new Queue<PendingBundle>();
@@ -700,11 +700,6 @@ namespace EpicGames.Horde.Storage
 
 		internal readonly ILogger? _traceLogger;
 
-		/// <summary>
-		/// Accessor for the store backing this writer
-		/// </summary>
-		public IStorageClient Store => _store;
-
 		/// <summary>
 		/// Cache of nodes to deduplicate against
 		/// </summary>
@@ -714,40 +709,27 @@ namespace EpicGames.Horde.Storage
 		/// Constructor
 		/// </summary>
 		/// <param name="store">Store to write data to</param>
-		/// <param name="treeReader">Reader for serialized node data</param>
+		/// <param name="reader">Reader for serialized node data</param>
+		/// <param name="refName">Name of the ref being written</param>
 		/// <param name="options">Options for the writer</param>
-		/// <param name="prefix">Prefix for blobs written to the store</param>
 		/// <param name="nodeCache">Cache of nodes for deduplication</param>
 		/// <param name="traceLogger">Optional logger for trace information</param>
-		public BundleWriter(IStorageClient store, BundleReader treeReader, BundleOptions? options = null, Utf8String prefix = default, NodeCache? nodeCache = null, ILogger? traceLogger = null)
+		public BundleWriter(IStorageClient store, BundleReader reader, RefName refName, BundleOptions? options = null, NodeCache? nodeCache = null, ILogger? traceLogger = null)
 		{
 			_store = store;
-			_treeReader = treeReader;
+			_reader = reader;
+			_refName = refName;
 			_options = options ?? s_defaultOptions;
-			_prefix = prefix;
 			_nodeCache = nodeCache ?? new NodeCache(_options.NodeCacheSize);
 			_traceLogger = traceLogger;
 		}
 
-		/// <summary>
-		/// Constructor
-		/// </summary>
-		/// <param name="store">Store to write data to</param>
-		/// <param name="reader">Reader for serialized node data</param>
-		/// <param name="refName">Ref being written. Will be used as a prefix for storing blobs.</param>
-		/// <param name="options">Options for the writer</param>
-		/// <param name="traceLogger">Optional logger for trace information</param>
-		public BundleWriter(IStorageClient store, BundleReader reader, RefName refName, BundleOptions? options = null, ILogger? traceLogger = null)
-			: this(store, reader, options, refName.Text, traceLogger: traceLogger)
-		{
-		}
-
 		/// <summary>
 		/// Copy constructor
 		/// </summary>
 		/// <param name="other"></param>
 		public BundleWriter(BundleWriter other)
-			: this(other._store, other._treeReader, other._options, other._prefix, other._nodeCache, other._traceLogger)
+			: this(other._store, other._reader, other._refName, other._options, other._nodeCache, other._traceLogger)
 		{
 		}
 
@@ -777,7 +759,7 @@ namespace EpicGames.Horde.Storage
 		{
 			if (_currentBundle != null)
 			{
-				_currentBundle.MarkAsComplete(_store, _prefix, _traceLogger);
+				_currentBundle.MarkAsComplete(_store, _refName.Text, _traceLogger);
 				_writeQueue.Enqueue(_currentBundle);
 				_currentBundle = null;
 			}
@@ -857,7 +839,7 @@ namespace EpicGames.Horde.Storage
 			if (_currentBundle == null || _currentBundle.IsReadOnly)
 			{
 				int bufferSize = (int)(_options.MinCompressionPacketSize * 1.2);
-				_currentBundle = new PendingBundle(_treeReader, this, bufferSize, _options.MaxBlobSize, _options.CompressionFormat);
+				_currentBundle = new PendingBundle(_reader, this, bufferSize, _options.MaxBlobSize, _options.CompressionFormat);
 			}
 			return _currentBundle;
 		}
@@ -889,5 +871,11 @@ namespace EpicGames.Horde.Storage
 				await WaitForWriteAsync(cancellationToken);
 			}
 		}
+
+		/// <inheritdoc/>
+		public async ValueTask WriteRefAsync(NodeHandle target, RefOptions? options = null, CancellationToken cancellationToken = default)
+		{
+			await _store.WriteRefTargetAsync(_refName, target, options, cancellationToken);
+		}
 	}
 }
diff --git a/Engine/Source/Programs/Shared/EpicGames.Horde/Storage/IStorageClient.cs b/Engine/Source/Programs/Shared/EpicGames.Horde/Storage/IStorageClient.cs
--- a/Engine/Source/Programs/Shared/EpicGames.Horde/Storage/IStorageClient.cs
+++ b/Engine/Source/Programs/Shared/EpicGames.Horde/Storage/IStorageClient.cs
@@ -178,11 +178,6 @@ namespace EpicGames.Horde.Storage
 	/// </summary>
 	public interface IStorageWriter : IAsyncDisposable
 	{
-		/// <summary>
-		/// Accessor for the store instance
-		/// </summary>
-		IStorageClient Store { get; }
-
 		/// <summary>
 		/// Create another writer instance, allowing multiple threads to write in parallel.
 		/// </summary>
@@ -212,6 +207,14 @@ namespace EpicGames.Horde.Storage
 		/// <param name="cancellationToken">Cancellation token for the operation</param>
 		/// <returns>Handle to the written node</returns>
 		ValueTask<NodeHandle> WriteNodeAsync(int size, IReadOnlyList<NodeHandle> references, NodeType type, CancellationToken cancellationToken = default);
+
+		/// <summary>
+		/// Writes the reference using the given target node
+		/// </summary>
+		/// <param name="target">The target node</param>
+		/// <param name="options">Options for the new ref</param>
+		/// <param name="cancellationToken">Cancellation token for the operation</param>
+		ValueTask WriteRefAsync(NodeHandle target, RefOptions? options = null, CancellationToken cancellationToken = default);
 	}
 
 	/// <summary>
diff --git a/Engine/Source/Programs/Shared/EpicGames.Horde/Storage/Node.cs b/Engine/Source/Programs/Shared/EpicGames.Horde/Storage/Node.cs
--- a/Engine/Source/Programs/Shared/EpicGames.Horde/Storage/Node.cs
+++ b/Engine/Source/Programs/Shared/EpicGames.Horde/Storage/Node.cs
@@ -424,7 +424,7 @@ namespace EpicGames.Horde.Storage
 		{
 			await using IStorageWriter writer = store.CreateWriter(name);
 			NodeRef<Node> nodeRef = await writer.WriteNodeAsync(node, cancellationToken);
-			await store.WriteRefTargetAsync(name, nodeRef.Handle, refOptions, cancellationToken);
+			await writer.WriteRefAsync(nodeRef.Handle, refOptions, cancellationToken);
 			return nodeRef.Handle;
 		}
 
