EpicGames.Core: Prevent NullReferenceException in ThreadPoolWorkQueue

#rnx

diff --git a/Engine/Source/Programs/Shared/EpicGames.Core/ThreadPoolWorkQueue.cs b/Engine/Source/Programs/Shared/EpicGames.Core/ThreadPoolWorkQueue.cs
--- a/Engine/Source/Programs/Shared/EpicGames.Core/ThreadPoolWorkQueue.cs
+++ b/Engine/Source/Programs/Shared/EpicGames.Core/ThreadPoolWorkQueue.cs
@@ -1,7 +1,8 @@
 // Copyright Epic Games, Inc. All Rights Reserved.
 
 using System;
-using System.Collections.Generic;
+using System.Collections.Concurrent;
+using System.Linq;
 using System.Threading;
 
 namespace EpicGames.Core
@@ -12,9 +13,9 @@ namespace EpicGames.Core
 	public sealed class ThreadPoolWorkQueue : IDisposable
 	{
 		/// <summary>
-		/// Object used for controlling access to NumOutstandingJobs and updating EmptyEvent
+		/// Object used for controlling access to updating _lockObject
 		/// </summary>
-		readonly object _lockObject = new object();
+		readonly object _lockObject = new();
 
 		/// <summary>
 		/// Number of jobs remaining in the queue. This is updated in an atomic way.
@@ -24,12 +25,12 @@ namespace EpicGames.Core
 		/// <summary>
 		/// Event which indicates whether the queue is empty.
 		/// </summary>
-		ManualResetEvent _emptyEvent = new ManualResetEvent(true);
+		ManualResetEvent? _emptyEvent = new(true);
 
 		/// <summary>
 		/// Exceptions which occurred while executing tasks
 		/// </summary>
-		readonly List<Exception> _exceptions = new List<Exception>();
+		readonly ConcurrentBag<Exception> _exceptions = new();
 
 		/// <summary>
 		/// Default constructor
@@ -43,12 +44,12 @@ namespace EpicGames.Core
 		/// </summary>
 		public void Dispose()
 		{
-			if(_emptyEvent != null)
+			if (_emptyEvent != null)
 			{
 				Wait();
 
-				_emptyEvent.Dispose();
-				_emptyEvent = null!;
+				_emptyEvent?.Dispose();
+				_emptyEvent = null;
 			}
 		}
 
@@ -85,12 +86,9 @@ namespace EpicGames.Core
 			{
 				((Action)actionToExecute!)();
 			}
-			catch(Exception ex)
+			catch (Exception ex)
 			{
-				lock(_lockObject)
-				{
-					_exceptions.Add(ex);
-				}
+				_exceptions.Add(ex);
 			}
 			finally
 			{
@@ -106,7 +104,7 @@ namespace EpicGames.Core
 		/// </summary>
 		public void Wait()
 		{
-			_emptyEvent.WaitOne();
+			_emptyEvent?.WaitOne();
 			RethrowExceptions();
 		}
 
@@ -127,7 +125,7 @@ namespace EpicGames.Core
 		/// <returns>True if the queue completed, false if the timeout elapsed</returns>
 		public bool Wait(TimeSpan timeout)
 		{
-			bool bResult = _emptyEvent.WaitOne(timeout);
+			bool bResult = _emptyEvent?.WaitOne(timeout) ?? false;
 			if (bResult)
 			{
 				RethrowExceptions();
@@ -140,12 +138,9 @@ namespace EpicGames.Core
 		/// </summary>
 		public void RethrowExceptions()
 		{
-			lock(_lockObject)
+			if (!_exceptions.IsEmpty)
 			{
-				if(_exceptions.Count > 0)
-				{
-					throw new AggregateException(_exceptions.ToArray());
-				}
+				throw new AggregateException(_exceptions.AsEnumerable());
 			}
 		}
 
@@ -160,11 +155,11 @@ namespace EpicGames.Core
 			{
 				if (_numOutstandingJobs > 0)
 				{
-					_emptyEvent.Reset();
+					_emptyEvent?.Reset();
 				}
 				else
 				{
-					_emptyEvent.Set();
+					_emptyEvent?.Set();
 				}
 			}
 		}
