Horde: Remove unused code for creating bundles.

diff --git a/Engine/Source/Programs/Shared/EpicGames.Horde/Storage/BundleWriter.cs b/Engine/Source/Programs/Shared/EpicGames.Horde/Storage/BundleWriter.cs
--- a/Engine/Source/Programs/Shared/EpicGames.Horde/Storage/BundleWriter.cs
+++ b/Engine/Source/Programs/Shared/EpicGames.Horde/Storage/BundleWriter.cs
@@ -573,31 +573,7 @@ namespace EpicGames.Horde.Storage
 
 			Bundle CreateBundleInternal()
 			{
-				// Create a set from the nodes to be written. We use this to determine references that are imported.
-				HashSet<BlobHandle> nodeSet = new HashSet<BlobHandle>(_queue);
-
-				// Find all the imported nodes by bundle
-				Dictionary<BlobLocator, List<(int, BlobHandle)>> bundleToImports = new Dictionary<BlobLocator, List<(int, BlobHandle)>>();
-				foreach (PendingNode nodeInfo in _queue)
-				{
-					foreach (BlobHandle handle in nodeInfo.Refs)
-					{
-						if (nodeSet.Add(handle))
-						{
-							NodeLocator refLocator = handle.GetLocator();
-
-							List<(int, BlobHandle)>? importedNodes;
-							if (!bundleToImports.TryGetValue(refLocator.Blob, out importedNodes))
-							{
-								importedNodes = new List<(int, BlobHandle)>();
-								bundleToImports.Add(refLocator.Blob, importedNodes);
-							}
-
-							importedNodes.Add((refLocator.ExportIdx, handle));
-						}
-					}
-				}
-
+				// List of imported blobs
 				List<BlobLocator> imports = new List<BlobLocator>();
 				Dictionary<BlobLocator, int> importToIndex = new Dictionary<BlobLocator, int>();
 
@@ -617,13 +593,7 @@ namespace EpicGames.Horde.Storage
 				List<BundleExport> exports = new List<BundleExport>(_queue.Count);
 				foreach (PendingNode nodeInfo in _queue)
 				{
-					int typeIdx;
-					if (!typeToIndex.TryGetValue(nodeInfo.Key.Type, out typeIdx))
-					{
-						typeIdx = types.Count;
-						typeToIndex.Add(nodeInfo.Key.Type, typeIdx);
-						types.Add(nodeInfo.Key.Type);
-					}
+					int typeIdx = FindOrAddItemIndex(nodeInfo.Key.Type, types, typeToIndex);
 
 					List<BundleExportRef> exportRefs = new List<BundleExportRef>();
 					foreach (BlobHandle handle in nodeInfo.Refs)
@@ -633,14 +603,7 @@ namespace EpicGames.Horde.Storage
 						{
 							NodeLocator locator = handle.GetLocator();
 
-							int importIdx;
-							if (!importToIndex.TryGetValue(locator.Blob, out importIdx))
-							{
-								importIdx = imports.Count;
-								imports.Add(locator.Blob);
-								importToIndex.Add(locator.Blob, importIdx);
-							}
-
+							int importIdx = FindOrAddItemIndex(locator.Blob, imports, importToIndex);
 							exportRef = new BundleExportRef(importIdx, locator.ExportIdx, handle.Hash);
 						}
 						exportRefs.Add(exportRef);
@@ -673,6 +636,18 @@ namespace EpicGames.Horde.Storage
 				BundleHeader header = BundleHeader.Create(types, imports, exports, _packets);
 				return new Bundle(header, packetData);
 			}
+
+			static int FindOrAddItemIndex<TItem>(TItem item, List<TItem> items, Dictionary<TItem, int> itemToIndex) where TItem : notnull
+			{
+				int index;
+				if (!itemToIndex.TryGetValue(item, out index))
+				{
+					index = items.Count;
+					items.Add(item);
+					itemToIndex.Add(item, index);
+				}
+				return index;
+			}
 		}
 
 		static readonly BundleOptions s_defaultOptions = new BundleOptions();
