Horde: Fix OOB error when building directory tree.

diff --git a/Engine/Source/Programs/Shared/EpicGames.Horde/Storage/Nodes/DirectoryNode.cs b/Engine/Source/Programs/Shared/EpicGames.Horde/Storage/Nodes/DirectoryNode.cs
--- a/Engine/Source/Programs/Shared/EpicGames.Horde/Storage/Nodes/DirectoryNode.cs
+++ b/Engine/Source/Programs/Shared/EpicGames.Horde/Storage/Nodes/DirectoryNode.cs
@@ -799,8 +799,15 @@ namespace EpicGames.Horde.Storage.Nodes
 		/// <param name="dataRef">Reference to the file data</param>
 		public FileEntry AddFile(string path, FileEntryFlags flags, long length, NodeRef<ChunkedDataNode> dataRef)
 		{
-			int endIdx = path.LastIndexOf(Path.DirectorySeparatorChar, Path.AltDirectorySeparatorChar);
-			string name = (endIdx == -1) ? path : path.Substring(endIdx + 1);
+			string name = path;
+			for (int idx = path.Length - 1; idx >= 0; idx--)
+			{
+				if (path[idx] == Path.DirectorySeparatorChar || path[idx] == Path.AltDirectorySeparatorChar)
+				{
+					name = path.Substring(idx + 1);
+					break;
+				}
+			}
 
 			FileEntry entry = new FileEntry(name, flags, length, dataRef);
 			AddFile(path, entry);
