Removed C++ UHT

diff --git a/Engine/Source/Programs/Shared/EpicGames.Core/UHTTypes.cs b/Engine/Source/Programs/Shared/EpicGames.Core/UHTTypes.cs
--- a/Engine/Source/Programs/Shared/EpicGames.Core/UHTTypes.cs
+++ b/Engine/Source/Programs/Shared/EpicGames.Core/UHTTypes.cs
@@ -42,8 +42,7 @@ namespace EpicGames.Core
 	};
 
 	/// <summary>
-	/// This MUST be kept in sync with EPackageOverrideType defined in 
-	/// Engine\Source\Programs\UnrealHeaderTool\Public\IScriptGeneratorPluginInterface.h
+	/// Build module override type to add additional PKG flags if necessary, mirrored in ModuleRules.cs, enum PackageOverrideType
 	/// </summary>
 	public enum EPackageOverrideType
 	{
@@ -57,9 +56,6 @@ namespace EpicGames.Core
 
 	/// <summary>
 	/// Type of module. This should be sorted by the order in which we expect modules to be built.
-	/// 
-	/// This MUST be kept in sync with EBuildModuleType defined in 
-	/// Engine\Source\Programs\UnrealHeaderTool\Public\IScriptGeneratorPluginInterface.h
 	/// </summary>
 	public enum UHTModuleType
 	{
diff --git a/Engine/Source/Programs/UnrealHeaderTool/Private/BaseParser.cpp b/Engine/Source/Programs/UnrealHeaderTool/Private/BaseParser.cpp
deleted file mode 100644
diff --git a/Engine/Source/Programs/UnrealHeaderTool/Private/BaseParser.h b/Engine/Source/Programs/UnrealHeaderTool/Private/BaseParser.h
deleted file mode 100644
diff --git a/Engine/Source/Programs/UnrealHeaderTool/Private/ClassMaps.cpp b/Engine/Source/Programs/UnrealHeaderTool/Private/ClassMaps.cpp
deleted file mode 100644
diff --git a/Engine/Source/Programs/UnrealHeaderTool/Private/ClassMaps.h b/Engine/Source/Programs/UnrealHeaderTool/Private/ClassMaps.h
deleted file mode 100644
diff --git a/Engine/Source/Programs/UnrealHeaderTool/Private/CodeGenerator.cpp b/Engine/Source/Programs/UnrealHeaderTool/Private/CodeGenerator.cpp
deleted file mode 100644
diff --git a/Engine/Source/Programs/UnrealHeaderTool/Private/EngineAPI.cpp b/Engine/Source/Programs/UnrealHeaderTool/Private/EngineAPI.cpp
deleted file mode 100644
diff --git a/Engine/Source/Programs/UnrealHeaderTool/Private/EngineAPI.h b/Engine/Source/Programs/UnrealHeaderTool/Private/EngineAPI.h
deleted file mode 100644
diff --git a/Engine/Source/Programs/UnrealHeaderTool/Private/Exceptions.cpp b/Engine/Source/Programs/UnrealHeaderTool/Private/Exceptions.cpp
deleted file mode 100644
diff --git a/Engine/Source/Programs/UnrealHeaderTool/Private/Exceptions.h b/Engine/Source/Programs/UnrealHeaderTool/Private/Exceptions.h
deleted file mode 100644
diff --git a/Engine/Source/Programs/UnrealHeaderTool/Private/GeneratedCodeVersion.h b/Engine/Source/Programs/UnrealHeaderTool/Private/GeneratedCodeVersion.h
deleted file mode 100644
diff --git a/Engine/Source/Programs/UnrealHeaderTool/Private/HeaderParser.cpp b/Engine/Source/Programs/UnrealHeaderTool/Private/HeaderParser.cpp
deleted file mode 100644
diff --git a/Engine/Source/Programs/UnrealHeaderTool/Private/HeaderParser.h b/Engine/Source/Programs/UnrealHeaderTool/Private/HeaderParser.h
deleted file mode 100644
diff --git a/Engine/Source/Programs/UnrealHeaderTool/Private/HeaderProvider.cpp b/Engine/Source/Programs/UnrealHeaderTool/Private/HeaderProvider.cpp
deleted file mode 100644
diff --git a/Engine/Source/Programs/UnrealHeaderTool/Private/HeaderProvider.h b/Engine/Source/Programs/UnrealHeaderTool/Private/HeaderProvider.h
deleted file mode 100644
diff --git a/Engine/Source/Programs/UnrealHeaderTool/Private/IScriptGeneratorPluginInterface.cpp b/Engine/Source/Programs/UnrealHeaderTool/Private/IScriptGeneratorPluginInterface.cpp
deleted file mode 100644
diff --git a/Engine/Source/Programs/UnrealHeaderTool/Private/Manifest.cpp b/Engine/Source/Programs/UnrealHeaderTool/Private/Manifest.cpp
deleted file mode 100644
diff --git a/Engine/Source/Programs/UnrealHeaderTool/Private/Manifest.h b/Engine/Source/Programs/UnrealHeaderTool/Private/Manifest.h
deleted file mode 100644
diff --git a/Engine/Source/Programs/UnrealHeaderTool/Private/NativeClassExporter.h b/Engine/Source/Programs/UnrealHeaderTool/Private/NativeClassExporter.h
deleted file mode 100644
diff --git a/Engine/Source/Programs/UnrealHeaderTool/Private/ParserHelper.cpp b/Engine/Source/Programs/UnrealHeaderTool/Private/ParserHelper.cpp
deleted file mode 100644
diff --git a/Engine/Source/Programs/UnrealHeaderTool/Private/ParserHelper.h b/Engine/Source/Programs/UnrealHeaderTool/Private/ParserHelper.h
deleted file mode 100644
diff --git a/Engine/Source/Programs/UnrealHeaderTool/Private/PropertyTypes.cpp b/Engine/Source/Programs/UnrealHeaderTool/Private/PropertyTypes.cpp
deleted file mode 100644
diff --git a/Engine/Source/Programs/UnrealHeaderTool/Private/PropertyTypes.h b/Engine/Source/Programs/UnrealHeaderTool/Private/PropertyTypes.h
deleted file mode 100644
diff --git a/Engine/Source/Programs/UnrealHeaderTool/Private/Scope.cpp b/Engine/Source/Programs/UnrealHeaderTool/Private/Scope.cpp
deleted file mode 100644
diff --git a/Engine/Source/Programs/UnrealHeaderTool/Private/Scope.h b/Engine/Source/Programs/UnrealHeaderTool/Private/Scope.h
deleted file mode 100644
diff --git a/Engine/Source/Programs/UnrealHeaderTool/Private/Specifiers/CheckedMetadataSpecifiers.def b/Engine/Source/Programs/UnrealHeaderTool/Private/Specifiers/CheckedMetadataSpecifiers.def
deleted file mode 100644
diff --git a/Engine/Source/Programs/UnrealHeaderTool/Private/Specifiers/CheckedMetadataSpecifiers.h b/Engine/Source/Programs/UnrealHeaderTool/Private/Specifiers/CheckedMetadataSpecifiers.h
deleted file mode 100644
diff --git a/Engine/Source/Programs/UnrealHeaderTool/Private/Specifiers/ClassMetadataSpecifiers.def b/Engine/Source/Programs/UnrealHeaderTool/Private/Specifiers/ClassMetadataSpecifiers.def
deleted file mode 100644
diff --git a/Engine/Source/Programs/UnrealHeaderTool/Private/Specifiers/ClassMetadataSpecifiers.h b/Engine/Source/Programs/UnrealHeaderTool/Private/Specifiers/ClassMetadataSpecifiers.h
deleted file mode 100644
diff --git a/Engine/Source/Programs/UnrealHeaderTool/Private/Specifiers/EnumSpecifiers.def b/Engine/Source/Programs/UnrealHeaderTool/Private/Specifiers/EnumSpecifiers.def
deleted file mode 100644
diff --git a/Engine/Source/Programs/UnrealHeaderTool/Private/Specifiers/EnumSpecifiers.h b/Engine/Source/Programs/UnrealHeaderTool/Private/Specifiers/EnumSpecifiers.h
deleted file mode 100644
diff --git a/Engine/Source/Programs/UnrealHeaderTool/Private/Specifiers/FunctionSpecifiers.def b/Engine/Source/Programs/UnrealHeaderTool/Private/Specifiers/FunctionSpecifiers.def
deleted file mode 100644
diff --git a/Engine/Source/Programs/UnrealHeaderTool/Private/Specifiers/FunctionSpecifiers.h b/Engine/Source/Programs/UnrealHeaderTool/Private/Specifiers/FunctionSpecifiers.h
deleted file mode 100644
diff --git a/Engine/Source/Programs/UnrealHeaderTool/Private/Specifiers/InterfaceSpecifiers.def b/Engine/Source/Programs/UnrealHeaderTool/Private/Specifiers/InterfaceSpecifiers.def
deleted file mode 100644
diff --git a/Engine/Source/Programs/UnrealHeaderTool/Private/Specifiers/InterfaceSpecifiers.h b/Engine/Source/Programs/UnrealHeaderTool/Private/Specifiers/InterfaceSpecifiers.h
deleted file mode 100644
diff --git a/Engine/Source/Programs/UnrealHeaderTool/Private/Specifiers/SpecifierArrays.cpp b/Engine/Source/Programs/UnrealHeaderTool/Private/Specifiers/SpecifierArrays.cpp
deleted file mode 100644
diff --git a/Engine/Source/Programs/UnrealHeaderTool/Private/Specifiers/StructSpecifiers.def b/Engine/Source/Programs/UnrealHeaderTool/Private/Specifiers/StructSpecifiers.def
deleted file mode 100644
diff --git a/Engine/Source/Programs/UnrealHeaderTool/Private/Specifiers/StructSpecifiers.h b/Engine/Source/Programs/UnrealHeaderTool/Private/Specifiers/StructSpecifiers.h
deleted file mode 100644
diff --git a/Engine/Source/Programs/UnrealHeaderTool/Private/Specifiers/VariableSpecifiers.def b/Engine/Source/Programs/UnrealHeaderTool/Private/Specifiers/VariableSpecifiers.def
deleted file mode 100644
diff --git a/Engine/Source/Programs/UnrealHeaderTool/Private/Specifiers/VariableSpecifiers.h b/Engine/Source/Programs/UnrealHeaderTool/Private/Specifiers/VariableSpecifiers.h
deleted file mode 100644
diff --git a/Engine/Source/Programs/UnrealHeaderTool/Private/StringUtils.cpp b/Engine/Source/Programs/UnrealHeaderTool/Private/StringUtils.cpp
deleted file mode 100644
diff --git a/Engine/Source/Programs/UnrealHeaderTool/Private/StringUtils.h b/Engine/Source/Programs/UnrealHeaderTool/Private/StringUtils.h
deleted file mode 100644
diff --git a/Engine/Source/Programs/UnrealHeaderTool/Private/UnrealHeaderTool.h b/Engine/Source/Programs/UnrealHeaderTool/Private/UnrealHeaderTool.h
deleted file mode 100644
diff --git a/Engine/Source/Programs/UnrealHeaderTool/Private/UnrealHeaderToolGlobals.h b/Engine/Source/Programs/UnrealHeaderTool/Private/UnrealHeaderToolGlobals.h
deleted file mode 100644
diff --git a/Engine/Source/Programs/UnrealHeaderTool/Private/UnrealHeaderToolMain.cpp b/Engine/Source/Programs/UnrealHeaderTool/Private/UnrealHeaderToolMain.cpp
deleted file mode 100644
diff --git a/Engine/Source/Programs/UnrealHeaderTool/Private/UnrealHeaderToolStats.cpp b/Engine/Source/Programs/UnrealHeaderTool/Private/UnrealHeaderToolStats.cpp
deleted file mode 100644
diff --git a/Engine/Source/Programs/UnrealHeaderTool/Private/UnrealHeaderToolStats.h b/Engine/Source/Programs/UnrealHeaderTool/Private/UnrealHeaderToolStats.h
deleted file mode 100644
diff --git a/Engine/Source/Programs/UnrealHeaderTool/Private/UnrealSourceFile.cpp b/Engine/Source/Programs/UnrealHeaderTool/Private/UnrealSourceFile.cpp
deleted file mode 100644
diff --git a/Engine/Source/Programs/UnrealHeaderTool/Private/UnrealSourceFile.h b/Engine/Source/Programs/UnrealHeaderTool/Private/UnrealSourceFile.h
deleted file mode 100644
diff --git a/Engine/Source/Programs/UnrealHeaderTool/Private/UnrealTypeDefinitionInfo.cpp b/Engine/Source/Programs/UnrealHeaderTool/Private/UnrealTypeDefinitionInfo.cpp
deleted file mode 100644
diff --git a/Engine/Source/Programs/UnrealHeaderTool/Private/UnrealTypeDefinitionInfo.h b/Engine/Source/Programs/UnrealHeaderTool/Private/UnrealTypeDefinitionInfo.h
deleted file mode 100644
diff --git a/Engine/Source/Programs/UnrealHeaderTool/Public/IScriptGeneratorPluginInterface.h b/Engine/Source/Programs/UnrealHeaderTool/Public/IScriptGeneratorPluginInterface.h
deleted file mode 100644
diff --git a/Engine/Source/Programs/UnrealHeaderTool/Resources/UHTDebugging.manifest b/Engine/Source/Programs/UnrealHeaderTool/Resources/UHTDebugging.manifest
deleted file mode 100644
diff --git a/Engine/Source/Programs/UnrealHeaderTool/Resources/UHTDebugging/DocumentationObjects.h b/Engine/Source/Programs/UnrealHeaderTool/Resources/UHTDebugging/DocumentationObjects.h
deleted file mode 100644
diff --git a/Engine/Source/Programs/UnrealHeaderTool/Resources/UHTDebugging/EnumOnlyHeader.h b/Engine/Source/Programs/UnrealHeaderTool/Resources/UHTDebugging/EnumOnlyHeader.h
deleted file mode 100644
diff --git a/Engine/Source/Programs/UnrealHeaderTool/Resources/UHTDebugging/RigVMStructHeader.h b/Engine/Source/Programs/UnrealHeaderTool/Resources/UHTDebugging/RigVMStructHeader.h
deleted file mode 100644
diff --git a/Engine/Source/Programs/UnrealHeaderTool/Resources/UHTDebugging/StructOnlyHeader.h b/Engine/Source/Programs/UnrealHeaderTool/Resources/UHTDebugging/StructOnlyHeader.h
deleted file mode 100644
diff --git a/Engine/Source/Programs/UnrealHeaderTool/Resources/UHTDebugging/TestDeprecatedObject.h b/Engine/Source/Programs/UnrealHeaderTool/Resources/UHTDebugging/TestDeprecatedObject.h
deleted file mode 100644
diff --git a/Engine/Source/Programs/UnrealHeaderTool/Resources/UHTDebugging/TestInterface.h b/Engine/Source/Programs/UnrealHeaderTool/Resources/UHTDebugging/TestInterface.h
deleted file mode 100644
diff --git a/Engine/Source/Programs/UnrealHeaderTool/Resources/UHTDebugging/TestInterfaceObject.h b/Engine/Source/Programs/UnrealHeaderTool/Resources/UHTDebugging/TestInterfaceObject.h
deleted file mode 100644
diff --git a/Engine/Source/Programs/UnrealHeaderTool/Resources/UHTDebugging/TestObject.h b/Engine/Source/Programs/UnrealHeaderTool/Resources/UHTDebugging/TestObject.h
deleted file mode 100644
diff --git a/Engine/Source/Programs/UnrealHeaderTool/Resources/UHTDebugging/TestObjectSerializers.h b/Engine/Source/Programs/UnrealHeaderTool/Resources/UHTDebugging/TestObjectSerializers.h
deleted file mode 100644
diff --git a/Engine/Source/Programs/UnrealHeaderTool/UnrealHeaderTool.Build.cs b/Engine/Source/Programs/UnrealHeaderTool/UnrealHeaderTool.Build.cs
deleted file mode 100644
diff --git a/Engine/Source/Programs/UnrealHeaderTool/UnrealHeaderTool.Target.cs b/Engine/Source/Programs/UnrealHeaderTool/UnrealHeaderTool.Target.cs
deleted file mode 100644
