UBT: Downgrade Horde compute errors to systemic warnings.

diff --git a/Engine/Source/Programs/Shared/EpicGames.Core/KnownLogEvents.cs b/Engine/Source/Programs/Shared/EpicGames.Core/KnownLogEvents.cs
--- a/Engine/Source/Programs/Shared/EpicGames.Core/KnownLogEvents.cs
+++ b/Engine/Source/Programs/Shared/EpicGames.Core/KnownLogEvents.cs
@@ -240,6 +240,11 @@ namespace EpicGames.Core
 		/// </summary>
 		public static EventId Systemic_Horde_ArtifactUpload { get; } = new EventId(731);
 
+		/// <summary>
+		/// Internal Horde error
+		/// </summary>
+		public static EventId Systemic_Horde_Compute { get; } = new EventId(732);
+
 		/// <summary>
 		/// Harmless pdbutil error
 		/// </summary>
