Horde: Start to move API definitions into EpicGames.Horde, starting with the /api/v1/projects endpoint. Exposing these classes to clients makes it much easier to write applications that communicate with the Horde server.

diff --git a/Engine/Source/Programs/Horde/Horde.Server.Tests/BisectTests.cs b/Engine/Source/Programs/Horde/Horde.Server.Tests/BisectTests.cs
--- a/Engine/Source/Programs/Horde/Horde.Server.Tests/BisectTests.cs
+++ b/Engine/Source/Programs/Horde/Horde.Server.Tests/BisectTests.cs
@@ -6,6 +6,7 @@ using System.Linq;
 using System.Threading;
 using System.Threading.Tasks;
 using EpicGames.Core;
+using EpicGames.Horde.Api;
 using Horde.Server.Jobs;
 using Horde.Server.Jobs.Graphs;
 using Horde.Server.Jobs.Bisect;
diff --git a/Engine/Source/Programs/Horde/Horde.Server.Tests/Fixture.cs b/Engine/Source/Programs/Horde/Horde.Server.Tests/Fixture.cs
--- a/Engine/Source/Programs/Horde/Horde.Server.Tests/Fixture.cs
+++ b/Engine/Source/Programs/Horde/Horde.Server.Tests/Fixture.cs
@@ -5,6 +5,7 @@ using System.IO;
 using System.Text;
 using System.Threading.Tasks;
 using EpicGames.Core;
+using EpicGames.Horde.Api;
 using Horde.Server.Agents;
 using Horde.Server.Agents.Pools;
 using Horde.Server.Configuration;
diff --git a/Engine/Source/Programs/Horde/Horde.Server.Tests/Fleet/JobQueueStrategyTest.cs b/Engine/Source/Programs/Horde/Horde.Server.Tests/Fleet/JobQueueStrategyTest.cs
--- a/Engine/Source/Programs/Horde/Horde.Server.Tests/Fleet/JobQueueStrategyTest.cs
+++ b/Engine/Source/Programs/Horde/Horde.Server.Tests/Fleet/JobQueueStrategyTest.cs
@@ -14,6 +14,7 @@ using Horde.Server.Server;
 using Horde.Server.Utilities;
 using Microsoft.VisualStudio.TestTools.UnitTesting;
 using Horde.Server.Agents;
+using EpicGames.Horde.Api;
 
 namespace Horde.Server.Tests.Fleet
 {
diff --git a/Engine/Source/Programs/Horde/Horde.Server.Tests/IssueServiceTests.cs b/Engine/Source/Programs/Horde/Horde.Server.Tests/IssueServiceTests.cs
--- a/Engine/Source/Programs/Horde/Horde.Server.Tests/IssueServiceTests.cs
+++ b/Engine/Source/Programs/Horde/Horde.Server.Tests/IssueServiceTests.cs
@@ -28,6 +28,7 @@ using MongoDB.Driver;
 using Moq;
 using Horde.Server.Issues.Handlers;
 using System.Threading;
+using EpicGames.Horde.Api;
 
 namespace Horde.Server.Tests
 {
diff --git a/Engine/Source/Programs/Horde/Horde.Server.Tests/JobServiceTests.cs b/Engine/Source/Programs/Horde/Horde.Server.Tests/JobServiceTests.cs
--- a/Engine/Source/Programs/Horde/Horde.Server.Tests/JobServiceTests.cs
+++ b/Engine/Source/Programs/Horde/Horde.Server.Tests/JobServiceTests.cs
@@ -23,7 +23,7 @@ using Microsoft.VisualStudio.TestTools.UnitTesting;
 
 namespace Horde.Server.Tests
 {
-	using ProjectId = Projects.ProjectId;
+	using ProjectId = EpicGames.Horde.Api.ProjectId;
 
 	[TestClass]
 	public class JobServiceTests : TestSetup
diff --git a/Engine/Source/Programs/Horde/Horde.Server.Tests/SchedulerTests.cs b/Engine/Source/Programs/Horde/Horde.Server.Tests/SchedulerTests.cs
--- a/Engine/Source/Programs/Horde/Horde.Server.Tests/SchedulerTests.cs
+++ b/Engine/Source/Programs/Horde/Horde.Server.Tests/SchedulerTests.cs
@@ -18,6 +18,7 @@ using Microsoft.Extensions.DependencyInjection;
 using Microsoft.VisualStudio.TestTools.UnitTesting;
 using Horde.Server.Jobs.Graphs;
 using Horde.Server.Perforce;
+using EpicGames.Horde.Api;
 
 namespace Horde.Server.Tests
 {
diff --git a/Engine/Source/Programs/Horde/Horde.Server.Tests/TestDataTests.cs b/Engine/Source/Programs/Horde/Horde.Server.Tests/TestDataTests.cs
--- a/Engine/Source/Programs/Horde/Horde.Server.Tests/TestDataTests.cs
+++ b/Engine/Source/Programs/Horde/Horde.Server.Tests/TestDataTests.cs
@@ -21,6 +21,7 @@ using System.Linq;
 using System.Threading;
 using Horde.Server.Server;
 using EpicGames.Horde;
+using EpicGames.Horde.Api;
 
 namespace Horde.Server.Tests
 {
diff --git a/Engine/Source/Programs/Horde/Horde.Server/Devices/DeviceCollection.cs b/Engine/Source/Programs/Horde/Horde.Server/Devices/DeviceCollection.cs
--- a/Engine/Source/Programs/Horde/Horde.Server/Devices/DeviceCollection.cs
+++ b/Engine/Source/Programs/Horde/Horde.Server/Devices/DeviceCollection.cs
@@ -4,6 +4,7 @@ using System;
 using System.Collections.Generic;
 using System.Linq;
 using System.Threading.Tasks;
+using EpicGames.Horde.Api;
 using Horde.Server.Projects;
 using Horde.Server.Server;
 using Horde.Server.Streams;
diff --git a/Engine/Source/Programs/Horde/Horde.Server/Devices/DeviceService.cs b/Engine/Source/Programs/Horde/Horde.Server/Devices/DeviceService.cs
--- a/Engine/Source/Programs/Horde/Horde.Server/Devices/DeviceService.cs
+++ b/Engine/Source/Programs/Horde/Horde.Server/Devices/DeviceService.cs
@@ -6,6 +6,7 @@ using System.Linq;
 using System.Security.Claims;
 using System.Threading;
 using System.Threading.Tasks;
+using EpicGames.Horde.Api;
 using Horde.Server.Acls;
 using Horde.Server.Jobs;
 using Horde.Server.Jobs.Graphs;
diff --git a/Engine/Source/Programs/Horde/Horde.Server/Devices/DevicesController.cs b/Engine/Source/Programs/Horde/Horde.Server/Devices/DevicesController.cs
--- a/Engine/Source/Programs/Horde/Horde.Server/Devices/DevicesController.cs
+++ b/Engine/Source/Programs/Horde/Horde.Server/Devices/DevicesController.cs
@@ -5,6 +5,7 @@ using System.Collections.Generic;
 using System.Linq;
 using System.Text.Json;
 using System.Threading.Tasks;
+using EpicGames.Horde.Api;
 using Horde.Server.Acls;
 using Horde.Server.Projects;
 using Horde.Server.Server;
diff --git a/Engine/Source/Programs/Horde/Horde.Server/Devices/IDevice.cs b/Engine/Source/Programs/Horde/Horde.Server/Devices/IDevice.cs
--- a/Engine/Source/Programs/Horde/Horde.Server/Devices/IDevice.cs
+++ b/Engine/Source/Programs/Horde/Horde.Server/Devices/IDevice.cs
@@ -2,6 +2,7 @@
 
 using System;
 using System.Collections.Generic;
+using EpicGames.Horde.Api;
 using Horde.Server.Projects;
 using Horde.Server.Streams;
 using Horde.Server.Users;
diff --git a/Engine/Source/Programs/Horde/Horde.Server/Devices/IDeviceCollection.cs b/Engine/Source/Programs/Horde/Horde.Server/Devices/IDeviceCollection.cs
--- a/Engine/Source/Programs/Horde/Horde.Server/Devices/IDeviceCollection.cs
+++ b/Engine/Source/Programs/Horde/Horde.Server/Devices/IDeviceCollection.cs
@@ -3,6 +3,7 @@
 using System;
 using System.Collections.Generic;
 using System.Threading.Tasks;
+using EpicGames.Horde.Api;
 using Horde.Server.Projects;
 using Horde.Server.Users;
 using Horde.Server.Utilities;
diff --git a/Engine/Source/Programs/Horde/Horde.Server/Projects/ProjectConfig.cs b/Engine/Source/Programs/Horde/Horde.Server/Projects/ProjectConfig.cs
--- a/Engine/Source/Programs/Horde/Horde.Server/Projects/ProjectConfig.cs
+++ b/Engine/Source/Programs/Horde/Horde.Server/Projects/ProjectConfig.cs
@@ -16,6 +16,7 @@ using Horde.Server.Streams;
 using Horde.Server.Utilities;
 using Horde.Common;
 using HordeCommon.Rpc.Tasks;
+using EpicGames.Horde.Api;
 
 namespace Horde.Server.Projects
 {
diff --git a/Engine/Source/Programs/Horde/Horde.Server/Projects/ProjectId.cs b/Engine/Source/Programs/Horde/Horde.Server/Projects/ProjectId.cs
deleted file mode 100644
diff --git a/Engine/Source/Programs/Horde/Horde.Server/Projects/ProjectMessages.cs b/Engine/Source/Programs/Horde/Horde.Server/Projects/ProjectMessages.cs
deleted file mode 100644
diff --git a/Engine/Source/Programs/Horde/Horde.Server/Projects/ProjectsController.cs b/Engine/Source/Programs/Horde/Horde.Server/Projects/ProjectsController.cs
--- a/Engine/Source/Programs/Horde/Horde.Server/Projects/ProjectsController.cs
+++ b/Engine/Source/Programs/Horde/Horde.Server/Projects/ProjectsController.cs
@@ -2,7 +2,8 @@
 
 using System.Collections.Generic;
 using System.Linq;
-using System.Threading.Tasks;
+using System.Text.RegularExpressions;
+using EpicGames.Horde.Api;
 using Horde.Server.Acls;
 using Horde.Server.Configuration;
 using Horde.Server.Server;
@@ -59,7 +60,7 @@ namespace Horde.Server.Projects
 					{
 						visibleStreams = projectConfig.Streams.Where(x => x.Authorize(StreamAclAction.ViewStream, User)).ToList();
 					}
-					responses.Add(GetProjectResponse.FromConfig(projectConfig, includeStreams, includeCategories, visibleStreams).ApplyFilter(filter));
+					responses.Add(CreateGetProjectResponse(projectConfig, includeStreams, includeCategories, visibleStreams).ApplyFilter(filter));
 				}
 			}
 			return responses;
@@ -95,7 +96,7 @@ namespace Horde.Server.Projects
 				visibleStreams = projectConfig.Streams.Where(x => x.Authorize(StreamAclAction.ViewStream, User)).ToList();
 			}
 
-			return GetProjectResponse.FromConfig(projectConfig, includeStreams, includeCategories, visibleStreams).ApplyFilter(filter);
+			return CreateGetProjectResponse(projectConfig, includeStreams, includeCategories, visibleStreams).ApplyFilter(filter);
 		}
 
 		/// <summary>
@@ -131,5 +132,74 @@ namespace Horde.Server.Projects
 
 			return new FileContentResult(logoResource.Data.ToArray(), contentType);
 		}
+
+		#region Messages
+
+		internal static GetProjectResponse CreateGetProjectResponse(ProjectConfig projectConfig, bool includeStreams, bool includeCategories, List<StreamConfig>? streamConfigs)
+		{
+			GetProjectResponse response = new GetProjectResponse(projectConfig.Id, projectConfig.Name, projectConfig.Order);
+
+			if (includeStreams)
+			{
+				response.Streams = streamConfigs!.ConvertAll(x => new GetProjectStreamResponse(x.Id.ToString(), x.Name));
+			}
+
+			if (includeCategories)
+			{
+				List<GetProjectCategoryResponse> categoryResponses = projectConfig.Categories.ConvertAll(x => CreateGetProjectCategoryResponse(x));
+				if (streamConfigs != null)
+				{
+					foreach (StreamConfig streamConfig in streamConfigs)
+					{
+						GetProjectCategoryResponse? categoryResponse = categoryResponses.FirstOrDefault(x => MatchCategory(streamConfig.Name, x));
+						if (categoryResponse == null)
+						{
+							int row = (categoryResponses.Count > 0) ? categoryResponses.Max(x => x.Row) : 0;
+							if (categoryResponses.Count(x => x.Row == row) >= 3)
+							{
+								row++;
+							}
+
+							ProjectCategoryConfig otherCategory = new ProjectCategoryConfig();
+							otherCategory.Name = "Other";
+							otherCategory.Row = row;
+							otherCategory.IncludePatterns.Add(".*");
+
+							categoryResponse = CreateGetProjectCategoryResponse(otherCategory);
+							categoryResponses.Add(categoryResponse);
+						}
+						categoryResponse.Streams!.Add(streamConfig.Id.ToString());
+					}
+				}
+				response.Categories = categoryResponses;
+			}
+
+			return response;
+		}
+
+
+		internal static GetProjectCategoryResponse CreateGetProjectCategoryResponse(ProjectCategoryConfig streamCategory)
+		{
+			GetProjectCategoryResponse response = new GetProjectCategoryResponse(streamCategory.Name, streamCategory.Row);
+			response.ShowOnNavMenu = streamCategory.ShowOnNavMenu;
+			response.IncludePatterns.AddRange(streamCategory.IncludePatterns);
+			response.ExcludePatterns.AddRange(streamCategory.ExcludePatterns);
+			return response;
+		}
+
+		// Tests if a category response matches a given stream name
+		static bool MatchCategory(string name, GetProjectCategoryResponse category)
+		{
+			if (category.IncludePatterns.Any(x => Regex.IsMatch(name, x)))
+			{
+				if (!category.ExcludePatterns.Any(x => Regex.IsMatch(name, x)))
+				{
+					return true;
+				}
+			}
+			return false;
+		}
+
+		#endregion
 	}
 }
diff --git a/Engine/Source/Programs/Horde/Horde.Server/Server/GlobalConfig.cs b/Engine/Source/Programs/Horde/Horde.Server/Server/GlobalConfig.cs
--- a/Engine/Source/Programs/Horde/Horde.Server/Server/GlobalConfig.cs
+++ b/Engine/Source/Programs/Horde/Horde.Server/Server/GlobalConfig.cs
@@ -9,6 +9,7 @@ using System.Linq;
 using System.Security.Claims;
 using System.Text.Json.Serialization;
 using EpicGames.Core;
+using EpicGames.Horde.Api;
 using EpicGames.Horde.Common;
 using EpicGames.Horde.Compute;
 using EpicGames.Perforce;
diff --git a/Engine/Source/Programs/Horde/Horde.Server/Startup.cs b/Engine/Source/Programs/Horde/Horde.Server/Startup.cs
--- a/Engine/Source/Programs/Horde/Horde.Server/Startup.cs
+++ b/Engine/Source/Programs/Horde/Horde.Server/Startup.cs
@@ -101,6 +101,8 @@ using Polly;
 using Polly.Extensions.Http;
 using Horde.Server.Jobs.Bisect;
 using Microsoft.AspNetCore.Diagnostics;
+using EpicGames.Horde;
+using EpicGames.Horde.Api;
 
 namespace Horde.Server
 {
@@ -818,15 +820,9 @@ namespace Horde.Server
 
 		public static void ConfigureJsonSerializer(JsonSerializerOptions options)
 		{
-			options.AllowTrailingCommas = true;
-			options.ReadCommentHandling = JsonCommentHandling.Skip;
-			options.DefaultIgnoreCondition = JsonIgnoreCondition.WhenWritingNull;
-			options.PropertyNamingPolicy = JsonNamingPolicy.CamelCase;
-			options.PropertyNameCaseInsensitive = true;
+			HordeHttpClient.ConfigureJsonSerializer(options);
 			options.Converters.Add(new JsonObjectIdConverter());
-			options.Converters.Add(new JsonStringEnumConverter());
 			options.Converters.Add(new JsonKnownTypesConverterFactory());
-			options.Converters.Add(new StringIdJsonConverterFactory());
 			options.Converters.Add(new ObjectIdJsonConverterFactory());
 			options.Converters.Add(new JsonDateTimeConverter());
 			options.Converters.Add(new JsonTimeSpanConverter());
diff --git a/Engine/Source/Programs/Horde/Horde.Server/Streams/StreamMessages.cs b/Engine/Source/Programs/Horde/Horde.Server/Streams/StreamMessages.cs
--- a/Engine/Source/Programs/Horde/Horde.Server/Streams/StreamMessages.cs
+++ b/Engine/Source/Programs/Horde/Horde.Server/Streams/StreamMessages.cs
@@ -15,6 +15,7 @@ using Horde.Server.Projects;
 using Horde.Server.Jobs;
 using Microsoft.AspNetCore.DataProtection;
 using System.Linq;
+using EpicGames.Horde.Api;
 
 namespace Horde.Server.Streams
 {
diff --git a/Engine/Source/Programs/Horde/Horde.Server/Streams/StreamsController.cs b/Engine/Source/Programs/Horde/Horde.Server/Streams/StreamsController.cs
--- a/Engine/Source/Programs/Horde/Horde.Server/Streams/StreamsController.cs
+++ b/Engine/Source/Programs/Horde/Horde.Server/Streams/StreamsController.cs
@@ -5,6 +5,7 @@ using System.Collections.Generic;
 using System.Linq;
 using System.Threading;
 using System.Threading.Tasks;
+using EpicGames.Horde.Api;
 using Horde.Server.Acls;
 using Horde.Server.Jobs;
 using Horde.Server.Jobs.Templates;
diff --git a/Engine/Source/Programs/Horde/Horde.Server/Utilities/HordeControllerBase.cs b/Engine/Source/Programs/Horde/Horde.Server/Utilities/HordeControllerBase.cs
--- a/Engine/Source/Programs/Horde/Horde.Server/Utilities/HordeControllerBase.cs
+++ b/Engine/Source/Programs/Horde/Horde.Server/Utilities/HordeControllerBase.cs
@@ -1,6 +1,7 @@
 // Copyright Epic Games, Inc. All Rights Reserved.
 
 using EpicGames.Core;
+using EpicGames.Horde.Api;
 using EpicGames.Horde.Compute;
 using EpicGames.Horde.Storage;
 using Horde.Server.Acls;
diff --git a/Engine/Source/Programs/Horde/Horde.Server/Utilities/StringIdConverter.cs b/Engine/Source/Programs/Horde/Horde.Server/Utilities/StringIdConverter.cs
--- a/Engine/Source/Programs/Horde/Horde.Server/Utilities/StringIdConverter.cs
+++ b/Engine/Source/Programs/Horde/Horde.Server/Utilities/StringIdConverter.cs
@@ -15,145 +15,6 @@ using ProtoBuf;
 
 namespace Horde.Server.Utilities
 {
-	/// <summary>
-	/// Base class for converting to and from types containing a <see cref="StringId"/>. Useful pattern for reducing boilerplate with strongly typed records.
-	/// </summary>
-	/// <typeparam name="T"></typeparam>
-	abstract class StringIdConverter<T> where T : struct
-	{
-		/// <summary>
-		/// Converts a type to a <see cref="StringId"/>
-		/// </summary>
-		public abstract StringId ToStringId(T value);
-
-		/// <summary>
-		/// Constructs a type from a <see cref="StringId"/>
-		/// </summary>
-		public abstract T FromStringId(StringId id);
-	}
-
-	/// <summary>
-	/// Attribute declaring a <see cref="StringIdConverter{T}"/> for a particular type
-	/// </summary>
-	[AttributeUsage(AttributeTargets.Struct)]
-	sealed class StringIdConverterAttribute : Attribute
-	{
-		/// <summary>
-		/// The converter type
-		/// </summary>
-		public Type ConverterType { get; }
-
-		/// <summary>
-		/// Constructor
-		/// </summary>
-		public StringIdConverterAttribute(Type converterType) => ConverterType = converterType;
-	}
-
-	/// <summary>
-	/// Converter to compact binary objects
-	/// </summary>
-	sealed class StringIdCbConverter<TValue, TConverter> : CbConverterBase<TValue> where TValue : struct where TConverter : StringIdConverter<TValue>, new()
-	{
-		readonly TConverter _converter = new TConverter();
-
-		/// <inheritdoc/>
-		public override TValue Read(CbField field)
-		{
-			return _converter.FromStringId(new StringId(field.AsString()));
-		}
-
-		/// <inheritdoc/>
-		public override void Write(CbWriter writer, TValue value)
-		{
-			writer.WriteStringValue(_converter.ToStringId(value).ToString());
-		}
-
-		/// <inheritdoc/>
-		public override void WriteNamed(CbWriter writer, Utf8String name, TValue value)
-		{
-			writer.WriteString(name, _converter.ToStringId(value).ToString());
-		}
-	}
-
-	/// <summary>
-	/// Class which serializes types with a <see cref="StringIdConverter{T}"/> to Json
-	/// </summary>
-	sealed class StringIdTypeConverter<TValue, TConverter> : TypeConverter where TValue : struct where TConverter : StringIdConverter<TValue>, new()
-	{
-		readonly TConverter _converter = new TConverter();
-
-		/// <inheritdoc/>
-		public override bool CanConvertFrom(ITypeDescriptorContext? context, Type sourceType)
-		{
-			return sourceType == typeof(string) || sourceType == typeof(StringId);
-		}
-
-		/// <inheritdoc/>
-		public override object? ConvertFrom(ITypeDescriptorContext? context, CultureInfo? culture, object value)
-		{
-			if (value is string str)
-			{
-				return _converter.FromStringId(new StringId(str));
-			}
-			if (value is StringId stringId)
-			{
-				return _converter.FromStringId(stringId);
-			}
-			return null;
-		}
-
-		/// <inheritdoc/>
-		public override bool CanConvertTo(ITypeDescriptorContext? context, Type? destinationType) => destinationType == typeof(string) || destinationType == typeof(StringId);
-
-		/// <inheritdoc/>
-		public override object? ConvertTo(ITypeDescriptorContext? context, CultureInfo? culture, object? value, Type destinationType)
-		{
-			if (destinationType == typeof(string))
-			{
-				return _converter.ToStringId((TValue)value!).ToString();
-			}
-			if (destinationType == typeof(StringId))
-			{
-				return _converter.ToStringId((TValue)value!);
-			}
-			return null;
-		}
-	}
-
-	/// <summary>
-	/// Class which serializes types with a <see cref="StringIdConverter{T}"/> to Json
-	/// </summary>
-	sealed class StringIdJsonConverter<TValue, TConverter> : JsonConverter<TValue> where TValue : struct where TConverter : StringIdConverter<TValue>, new()
-	{
-		readonly TConverter _converter = new TConverter();
-
-		/// <inheritdoc/>
-		public override TValue Read(ref Utf8JsonReader reader, Type typeToConvert, JsonSerializerOptions options) => _converter.FromStringId(new StringId(new Utf8String(reader.GetUtf8String().ToArray())));
-		
-		/// <inheritdoc/>
-		public override void Write(Utf8JsonWriter writer, TValue value, JsonSerializerOptions options) => writer.WriteStringValue(_converter.ToStringId(value).Span);
-	}
-
-	/// <summary>
-	/// Creates constructors for types with a <see cref="StringIdConverter{T}"/> to Json
-	/// </summary>
-	sealed class StringIdJsonConverterFactory : JsonConverterFactory
-	{
-		/// <inheritdoc/>
-		public override bool CanConvert(Type typeToConvert) => typeToConvert.GetCustomAttribute<StringIdConverterAttribute>() != null;
-
-		/// <inheritdoc/>
-		public override JsonConverter? CreateConverter(Type type, JsonSerializerOptions options)
-		{
-			StringIdConverterAttribute? attribute = type.GetCustomAttribute<StringIdConverterAttribute>();
-			if (attribute == null)
-			{
-				return null;
-			}
-			return (JsonConverter?)Activator.CreateInstance(typeof(StringIdJsonConverter<,>).MakeGenericType(type, attribute.ConverterType));
-		}
-	}
-
 	/// <summary>
 	/// Class which serializes object id types to BSON
 	/// </summary>
diff --git a/Engine/Source/Programs/Shared/EpicGames.Core/HttpClientExtensions.cs b/Engine/Source/Programs/Shared/EpicGames.Core/HttpClientExtensions.cs
--- a/Engine/Source/Programs/Shared/EpicGames.Core/HttpClientExtensions.cs
+++ b/Engine/Source/Programs/Shared/EpicGames.Core/HttpClientExtensions.cs
@@ -30,7 +30,7 @@ namespace EpicGames.Core
 			using (HttpResponseMessage response = await client.GetAsync(url, cancellationToken))
 			{
 				response.EnsureSuccessStatusCode();
-				return await ParseJsonContent<TResponse>(response);
+				return await ParseJsonContent<TResponse>(response, cancellationToken);
 			}
 		}
 
@@ -40,7 +40,7 @@ namespace EpicGames.Core
 			using (HttpResponseMessage response = await client.GetAsync(url, cancellationToken))
 			{
 				response.EnsureSuccessStatusCode();
-				return await ParseJsonContent<TResponse>(response);
+				return await ParseJsonContent<TResponse>(response, cancellationToken);
 			}
 		}
 
@@ -81,7 +81,7 @@ namespace EpicGames.Core
 			using (HttpResponseMessage response = await PostAsync(client, url, request, cancellationToken))
 			{
 				response.EnsureSuccessStatusCode();
-				return await ParseJsonContent<TResponse>(response);
+				return await ParseJsonContent<TResponse>(response, cancellationToken);
 			}
 		}
 
@@ -91,7 +91,7 @@ namespace EpicGames.Core
 			using (HttpResponseMessage response = await PostAsync(client, url, request, cancellationToken))
 			{
 				response.EnsureSuccessStatusCode();
-				return await ParseJsonContent<TResponse>(response);
+				return await ParseJsonContent<TResponse>(response, cancellationToken);
 			}
 		}
 
@@ -132,7 +132,7 @@ namespace EpicGames.Core
 			using (HttpResponseMessage response = await PutAsync(client, url, request, cancellationToken))
 			{
 				response.EnsureSuccessStatusCode();
-				return await ParseJsonContent<TResponse>(response);
+				return await ParseJsonContent<TResponse>(response, cancellationToken);
 			}
 		}
 
@@ -142,7 +142,7 @@ namespace EpicGames.Core
 			using (HttpResponseMessage response = await PutAsync(client, url, request, cancellationToken))
 			{
 				response.EnsureSuccessStatusCode();
-				return await ParseJsonContent<TResponse>(response);
+				return await ParseJsonContent<TResponse>(response, cancellationToken);
 			}
 		}
 
@@ -162,10 +162,11 @@ namespace EpicGames.Core
 		/// </summary>
 		/// <typeparam name="T">Type of the object to parse</typeparam>
 		/// <param name="message">The message received</param>
+		/// <param name="cancellationToken">Cancellation token for the operation</param>
 		/// <returns>Parsed object instance</returns>
-		private static async Task<T> ParseJsonContent<T>(HttpResponseMessage message)
+		private static async Task<T> ParseJsonContent<T>(HttpResponseMessage message, CancellationToken cancellationToken)
 		{
-			byte[] bytes = await message.Content.ReadAsByteArrayAsync();
+			byte[] bytes = await message.Content.ReadAsByteArrayAsync(cancellationToken);
 			return JsonSerializer.Deserialize<T>(bytes, new JsonSerializerOptions { PropertyNameCaseInsensitive = true })!;
 		}
 	}
diff --git a/Engine/Source/Programs/Shared/EpicGames.Horde/Api/ProjectId.cs b/Engine/Source/Programs/Shared/EpicGames.Horde/Api/ProjectId.cs
new file mode 100644
--- /dev/null
+++ b/Engine/Source/Programs/Shared/EpicGames.Horde/Api/ProjectId.cs
@@ -0,0 +1,44 @@
+﻿// Copyright Epic Games, Inc. All Rights Reserved.
+
+using System.ComponentModel;
+using EpicGames.Core;
+using EpicGames.Serialization;
+
+namespace EpicGames.Horde.Api
+{
+	/// <summary>
+	/// Identifier for a pool
+	/// </summary>
+	/// <param name="Id">Id to construct from</param>
+	[JsonSchemaString]
+	[TypeConverter(typeof(StringIdTypeConverter<ProjectId, ProjectIdConverter>))]
+	[StringIdConverter(typeof(ProjectIdConverter))]
+	[CbConverter(typeof(StringIdCbConverter<ProjectId, ProjectIdConverter>))]
+	public record struct ProjectId(StringId Id)
+	{
+		/// <summary>
+		/// Constructor
+		/// </summary>
+		public ProjectId(string id) : this(new StringId(id))
+		{
+		}
+
+		/// <inheritdoc cref="StringId.IsEmpty"/>
+		public bool IsEmpty => Id.IsEmpty;
+
+		/// <inheritdoc/>
+		public override string ToString() => Id.ToString();
+	}
+
+	/// <summary>
+	/// Converter to and from <see cref="StringId"/> instances.
+	/// </summary>
+	public class ProjectIdConverter : StringIdConverter<ProjectId>
+	{
+		/// <inheritdoc/>
+		public override ProjectId FromStringId(StringId id) => new ProjectId(id);
+
+		/// <inheritdoc/>
+		public override StringId ToStringId(ProjectId value) => value.Id;
+	}
+}
diff --git a/Engine/Source/Programs/Shared/EpicGames.Horde/Api/ProjectMessages.cs b/Engine/Source/Programs/Shared/EpicGames.Horde/Api/ProjectMessages.cs
new file mode 100644
--- /dev/null
+++ b/Engine/Source/Programs/Shared/EpicGames.Horde/Api/ProjectMessages.cs
@@ -0,0 +1,159 @@
+// Copyright Epic Games, Inc. All Rights Reserved.
+
+using System.Collections.Generic;
+using System.Diagnostics;
+using System.Threading;
+using System.Threading.Tasks;
+
+#pragma warning disable CA2227 // Collection properties should be read only
+
+namespace EpicGames.Horde.Api
+{
+	/// <summary>
+	/// Response describing a project
+	/// </summary>
+	[DebuggerDisplay("{Id,nq}")]
+	public class GetProjectResponse
+	{
+		/// <summary>
+		/// Unique id of the project
+		/// </summary>
+		public ProjectId Id { get; set; }
+
+		/// <summary>
+		/// Name of the project
+		/// </summary>
+		public string Name { get; set; }
+
+		/// <summary>
+		/// Order to display this project on the dashboard
+		/// </summary>
+		public int Order { get; set; }
+
+		/// <summary>
+		/// List of streams that are in this project
+		/// </summary>
+		public List<GetProjectStreamResponse>? Streams { get; set; }
+
+		/// <summary>
+		/// List of stream categories to display
+		/// </summary>
+		public List<GetProjectCategoryResponse>? Categories { get; set; }
+
+		/// <summary>
+		/// Constructor
+		/// </summary>
+		/// <param name="id">Unique id of the project</param>
+		/// <param name="name">Name of the project</param>
+		/// <param name="order">Order to show this project on the dashboard</param>
+		public GetProjectResponse(ProjectId id, string name, int order)
+		{
+			Id = id;
+			Name = name;
+			Order = order;
+		}
+	}
+
+	/// <summary>
+	/// Information about a stream within a project
+	/// </summary>
+	[DebuggerDisplay("{Id,nq} ({Name})")]
+	public class GetProjectStreamResponse
+	{
+		/// <summary>
+		/// The stream id
+		/// </summary>
+		public string Id { get; set; }
+
+		/// <summary>
+		/// The stream name
+		/// </summary>
+		public string Name { get; set; }
+
+		/// <summary>
+		/// Constructor
+		/// </summary>
+		public GetProjectStreamResponse(string id, string name)
+		{
+			Id = id;
+			Name = name;
+		}
+	}
+
+	/// <summary>
+	/// Information about a category to display for a stream
+	/// </summary>
+	[DebuggerDisplay("{Name}")]
+	public class GetProjectCategoryResponse
+	{
+		/// <summary>
+		/// Heading for this column
+		/// </summary>
+		public string Name { get; set; }
+
+		/// <summary>
+		/// Index of the row to display this category on
+		/// </summary>
+		public int Row { get; set; }
+
+		/// <summary>
+		/// Whether to show this category on the nav menu
+		/// </summary>
+		public bool ShowOnNavMenu { get; set; }
+
+		/// <summary>
+		/// Patterns for stream names to include
+		/// </summary>
+		public List<string> IncludePatterns { get; init; } = new List<string>();
+
+		/// <summary>
+		/// Patterns for stream names to exclude
+		/// </summary>
+		public List<string> ExcludePatterns { get; init; } = new List<string>();
+
+		/// <summary>
+		/// Streams to include in this category
+		/// </summary>
+		public List<string> Streams { get; init; } = new List<string>();
+
+		/// <summary>
+		/// Constructor
+		/// </summary>
+		public GetProjectCategoryResponse(string name, int row)
+		{
+			Name = name;
+			Row = row;
+		}
+	}
+
+	/// <summary>
+	/// Extension methods for the projects endpoint
+	/// </summary>
+	public static class ProjectExtensions
+	{
+		/// <summary>
+		/// Query all the projects
+		/// </summary>
+		/// <param name="horde">The horde client instance</param>
+		/// <param name="includeStreams">Whether to include streams in the response</param>
+		/// <param name="includeCategories">Whether to include categories in the response</param>
+		/// <param name="cancellationToken">Cancellation token for the operation</param>
+		/// <returns>Information about all the projects</returns>
+		public static Task<List<GetProjectResponse>> GetProjectsAsync(this HordeHttpClient horde, bool includeStreams = false, bool includeCategories = false, CancellationToken cancellationToken = default)
+		{
+			return horde.GetAsync<List<GetProjectResponse>>($"api/v1/projects?includeStreams={includeStreams}&includeCategories={includeCategories}", cancellationToken);
+		}
+
+		/// <summary>
+		/// Retrieve information about a specific project
+		/// </summary>
+		/// <param name="horde">The horde client instance</param>
+		/// <param name="projectId">Id of the project to get information about</param>
+		/// <param name="cancellationToken">Cancellation token for the operation</param>
+		/// <returns>Information about the requested project</returns>
+		public static Task<GetProjectResponse> GetProjectAsync(this HordeHttpClient horde, ProjectId projectId, CancellationToken cancellationToken = default)
+		{
+			return horde.GetAsync<GetProjectResponse>($"api/v1/projects/{projectId}", cancellationToken);
+		}
+	}
+}
diff --git a/Engine/Source/Programs/Shared/EpicGames.Horde/HordeHttpClient.cs b/Engine/Source/Programs/Shared/EpicGames.Horde/HordeHttpClient.cs
new file mode 100644
--- /dev/null
+++ b/Engine/Source/Programs/Shared/EpicGames.Horde/HordeHttpClient.cs
@@ -0,0 +1,88 @@
+// Copyright Epic Games, Inc. All Rights Reserved.
+
+using System;
+using System.Net.Http;
+using System.Net.Http.Headers;
+using System.Text.Json;
+using System.Text.Json.Serialization;
+using System.Threading;
+using System.Threading.Tasks;
+
+#pragma warning disable CA2234 // Pass system uri objects instead of strings
+
+namespace EpicGames.Horde
+{
+	/// <summary>
+	/// Wraps an Http client which communicates with the Horde server
+	/// </summary>
+	public sealed class HordeHttpClient : IDisposable
+	{
+		/// <summary>
+		/// The configured HTTP client
+		/// </summary>
+		public HttpClient HttpClient { get; }
+
+		readonly JsonSerializerOptions _jsonSerializerOptions;
+
+		/// <summary>
+		/// Constructor
+		/// </summary>
+		/// <param name="httpClient">Client to use for communication</param>
+		public HordeHttpClient(HttpClient httpClient)
+		{
+			HttpClient = httpClient;
+
+			_jsonSerializerOptions = new JsonSerializerOptions();
+			ConfigureJsonSerializer(_jsonSerializerOptions);
+		}
+
+		/// <summary>
+		/// Constructor
+		/// </summary>
+		/// <param name="uri">Base URI of the server</param>
+		/// <param name="token">Access token for the connection</param>
+		public HordeHttpClient(Uri uri, string token)
+			: this(new HttpClient { BaseAddress = uri, DefaultRequestHeaders = { Authorization = new AuthenticationHeaderValue("Bearer", token) } })
+		{
+		}
+
+		/// <inheritdoc/>
+		public void Dispose()
+		{
+			HttpClient.Dispose();
+		}
+
+		/// <summary>
+		/// Configures a JSON serializer to read Horde responses
+		/// </summary>
+		/// <param name="options">options for the serializer</param>
+		public static void ConfigureJsonSerializer(JsonSerializerOptions options)
+		{
+			options.AllowTrailingCommas = true;
+			options.ReadCommentHandling = JsonCommentHandling.Skip;
+			options.DefaultIgnoreCondition = JsonIgnoreCondition.WhenWritingNull;
+			options.PropertyNamingPolicy = JsonNamingPolicy.CamelCase;
+			options.PropertyNameCaseInsensitive = true;
+			options.Converters.Add(new JsonStringEnumConverter());
+			options.Converters.Add(new StringIdJsonConverterFactory());
+		}
+
+		/// <summary>
+		/// Gets a resource from an HTTP endpoint and parses it as a JSON object
+		/// </summary>
+		/// <typeparam name="TResponse">The object type to return</typeparam>
+		/// <param name="relativePath">The url to retrieve</param>
+		/// <param name="cancellationToken">Cancels the request</param>
+		/// <returns>New instance of the object</returns>
+		public async Task<TResponse> GetAsync<TResponse>(string relativePath, CancellationToken cancellationToken = default)
+		{
+			using (HttpResponseMessage response = await HttpClient.GetAsync(relativePath, cancellationToken))
+			{
+				response.EnsureSuccessStatusCode();
+
+				byte[] bytes = await response.Content.ReadAsByteArrayAsync(cancellationToken);
+				return JsonSerializer.Deserialize<TResponse>(bytes, _jsonSerializerOptions)!;
+			}
+		}
+	}
+}
diff --git a/Engine/Source/Programs/Shared/EpicGames.Horde/StringIdConverter.cs b/Engine/Source/Programs/Shared/EpicGames.Horde/StringIdConverter.cs
new file mode 100644
--- /dev/null
+++ b/Engine/Source/Programs/Shared/EpicGames.Horde/StringIdConverter.cs
@@ -0,0 +1,152 @@
+// Copyright Epic Games, Inc. All Rights Reserved.
+
+using System;
+using System.ComponentModel;
+using System.Globalization;
+using System.Reflection;
+using System.Text.Json;
+using System.Text.Json.Serialization;
+using EpicGames.Core;
+using EpicGames.Serialization;
+
+namespace EpicGames.Horde
+{
+	/// <summary>
+	/// Base class for converting to and from types containing a <see cref="StringId"/>. Useful pattern for reducing boilerplate with strongly typed records.
+	/// </summary>
+	/// <typeparam name="T"></typeparam>
+	public abstract class StringIdConverter<T> where T : struct
+	{
+		/// <summary>
+		/// Converts a type to a <see cref="StringId"/>
+		/// </summary>
+		public abstract StringId ToStringId(T value);
+
+		/// <summary>
+		/// Constructs a type from a <see cref="StringId"/>
+		/// </summary>
+		public abstract T FromStringId(StringId id);
+	}
+
+	/// <summary>
+	/// Attribute declaring a <see cref="StringIdConverter{T}"/> for a particular type
+	/// </summary>
+	[AttributeUsage(AttributeTargets.Struct)]
+	public sealed class StringIdConverterAttribute : Attribute
+	{
+		/// <summary>
+		/// The converter type
+		/// </summary>
+		public Type ConverterType { get; }
+
+		/// <summary>
+		/// Constructor
+		/// </summary>
+		public StringIdConverterAttribute(Type converterType) => ConverterType = converterType;
+	}
+
+	/// <summary>
+	/// Converter to compact binary objects
+	/// </summary>
+	public sealed class StringIdCbConverter<TValue, TConverter> : CbConverterBase<TValue> where TValue : struct where TConverter : StringIdConverter<TValue>, new()
+	{
+		readonly TConverter _converter = new TConverter();
+
+		/// <inheritdoc/>
+		public override TValue Read(CbField field)
+		{
+			return _converter.FromStringId(new StringId(field.AsString()));
+		}
+
+		/// <inheritdoc/>
+		public override void Write(CbWriter writer, TValue value)
+		{
+			writer.WriteStringValue(_converter.ToStringId(value).ToString());
+		}
+
+		/// <inheritdoc/>
+		public override void WriteNamed(CbWriter writer, Utf8String name, TValue value)
+		{
+			writer.WriteString(name, _converter.ToStringId(value).ToString());
+		}
+	}
+
+	/// <summary>
+	/// Class which serializes types with a <see cref="StringIdConverter{T}"/> to Json
+	/// </summary>
+	public sealed class StringIdTypeConverter<TValue, TConverter> : TypeConverter where TValue : struct where TConverter : StringIdConverter<TValue>, new()
+	{
+		readonly TConverter _converter = new TConverter();
+
+		/// <inheritdoc/>
+		public override bool CanConvertFrom(ITypeDescriptorContext? context, Type sourceType)
+		{
+			return sourceType == typeof(string) || sourceType == typeof(StringId);
+		}
+
+		/// <inheritdoc/>
+		public override object? ConvertFrom(ITypeDescriptorContext? context, CultureInfo? culture, object value)
+		{
+			if (value is string str)
+			{
+				return _converter.FromStringId(new StringId(str));
+			}
+			if (value is StringId stringId)
+			{
+				return _converter.FromStringId(stringId);
+			}
+			return null;
+		}
+
+		/// <inheritdoc/>
+		public override bool CanConvertTo(ITypeDescriptorContext? context, Type? destinationType) => destinationType == typeof(string) || destinationType == typeof(StringId);
+
+		/// <inheritdoc/>
+		public override object? ConvertTo(ITypeDescriptorContext? context, CultureInfo? culture, object? value, Type destinationType)
+		{
+			if (destinationType == typeof(string))
+			{
+				return _converter.ToStringId((TValue)value!).ToString();
+			}
+			if (destinationType == typeof(StringId))
+			{
+				return _converter.ToStringId((TValue)value!);
+			}
+			return null;
+		}
+	}
+
+	/// <summary>
+	/// Class which serializes types with a <see cref="StringIdConverter{T}"/> to Json
+	/// </summary>
+	public sealed class StringIdJsonConverter<TValue, TConverter> : JsonConverter<TValue> where TValue : struct where TConverter : StringIdConverter<TValue>, new()
+	{
+		readonly TConverter _converter = new TConverter();
+
+		/// <inheritdoc/>
+		public override TValue Read(ref Utf8JsonReader reader, Type typeToConvert, JsonSerializerOptions options) => _converter.FromStringId(new StringId(new Utf8String(reader.GetUtf8String().ToArray())));
+
+		/// <inheritdoc/>
+		public override void Write(Utf8JsonWriter writer, TValue value, JsonSerializerOptions options) => writer.WriteStringValue(_converter.ToStringId(value).Span);
+	}
+
+	/// <summary>
+	/// Creates constructors for types with a <see cref="StringIdConverter{T}"/> to Json
+	/// </summary>
+	public sealed class StringIdJsonConverterFactory : JsonConverterFactory
+	{
+		/// <inheritdoc/>
+		public override bool CanConvert(Type typeToConvert) => typeToConvert.GetCustomAttribute<StringIdConverterAttribute>() != null;
+
+		/// <inheritdoc/>
+		public override JsonConverter? CreateConverter(Type type, JsonSerializerOptions options)
+		{
+			StringIdConverterAttribute? attribute = type.GetCustomAttribute<StringIdConverterAttribute>();
+			if (attribute == null)
+			{
+				return null;
+			}
+			return (JsonConverter?)Activator.CreateInstance(typeof(StringIdJsonConverter<,>).MakeGenericType(type, attribute.ConverterType));
+		}
+	}
+}
