Horde: Add a more detailed warning whenever there are no agents found matching the configured requirements during a remote execution lease.

diff --git a/Engine/Source/Programs/Shared/EpicGames.Horde/Compute/Clients/ServerComputeClient.cs b/Engine/Source/Programs/Shared/EpicGames.Horde/Compute/Clients/ServerComputeClient.cs
--- a/Engine/Source/Programs/Shared/EpicGames.Horde/Compute/Clients/ServerComputeClient.cs
+++ b/Engine/Source/Programs/Shared/EpicGames.Horde/Compute/Clients/ServerComputeClient.cs
@@ -28,7 +28,7 @@ namespace EpicGames.Horde.Compute.Clients
 		public const int NonceLength = 64;
 
 		class AssignComputeRequest
-	{
+		{
 			public Requirements? Requirements { get; set; }
 		}
 
@@ -152,6 +152,11 @@ namespace EpicGames.Horde.Compute.Clients
 			AssignComputeResponse? responseMessage;
 			using (HttpResponseMessage response = await client.PostAsync($"api/v2/compute/{clusterId}", request, _cancellationSource.Token))
 			{
+				if (response.StatusCode == HttpStatusCode.NotFound)
+				{
+					throw new NoComputeAgentsFoundException(clusterId, requirements);
+				}
+
 				if (response.StatusCode == HttpStatusCode.ServiceUnavailable)
 				{
 					_logger.LogDebug("No compute resource is available.");
@@ -185,4 +190,30 @@ namespace EpicGames.Horde.Compute.Clients
 			yield return new LeaseInfo(responseMessage.Properties, responseMessage.AssignedResources, computeSocket);
 		}
 	}
+
+	/// <summary>
+	/// Exception indicating that no matching compute agents were found
+	/// </summary>
+	public sealed class NoComputeAgentsFoundException : Exception
+	{
+		/// <summary>
+		/// The compute cluster requested
+		/// </summary>
+		public ClusterId ClusterId { get; }
+
+		/// <summary>
+		/// Requested agent requirements
+		/// </summary>
+		public Requirements? Requirements { get; }
+
+		/// <summary>
+		/// Constructor
+		/// </summary>
+		public NoComputeAgentsFoundException(ClusterId clusterId, Requirements? requirements)
+			: base($"No compute agents found matching '{requirements}' in cluster '{clusterId}'")
+		{
+			ClusterId = clusterId;
+			Requirements = requirements;
+		}
+	}
 }
