Horde: Move response classes for the secrets endpoint into EpicGames.Horde.

diff --git a/Engine/Source/Programs/Horde/Horde.Server/Secrets/SecretConfig.cs b/Engine/Source/Programs/Horde/Horde.Server/Secrets/SecretConfig.cs
--- a/Engine/Source/Programs/Horde/Horde.Server/Secrets/SecretConfig.cs
+++ b/Engine/Source/Programs/Horde/Horde.Server/Secrets/SecretConfig.cs
@@ -3,6 +3,7 @@
 using System.Collections.Generic;
 using System.Security.Claims;
 using System.Text.Json.Serialization;
+using EpicGames.Horde.Api;
 using Horde.Server.Acls;
 using Horde.Server.Server;
 
diff --git a/Engine/Source/Programs/Horde/Horde.Server/Secrets/SecretId.cs b/Engine/Source/Programs/Horde/Horde.Server/Secrets/SecretId.cs
deleted file mode 100644
diff --git a/Engine/Source/Programs/Horde/Horde.Server/Secrets/SecretsController.cs b/Engine/Source/Programs/Horde/Horde.Server/Secrets/SecretsController.cs
--- a/Engine/Source/Programs/Horde/Horde.Server/Secrets/SecretsController.cs
+++ b/Engine/Source/Programs/Horde/Horde.Server/Secrets/SecretsController.cs
@@ -1,6 +1,7 @@
 // Copyright Epic Games, Inc. All Rights Reserved.
 
 using System.Collections.Generic;
+using EpicGames.Horde.Api;
 using Horde.Server.Server;
 using Horde.Server.Utilities;
 using Microsoft.AspNetCore.Authorization;
@@ -9,37 +10,6 @@ using Microsoft.Extensions.Options;
 
 namespace Horde.Server.Secrets
 {
-	/// <summary>
-	/// Response listing all the secrets available to the current user
-	/// </summary>
-	public class GetSecretsResponse
-	{
-		/// <summary>
-		/// List of secret ids
-		/// </summary>
-		public List<SecretId> Ids { get; set; } = new List<SecretId>();
-	}
-
-	/// <summary>
-	/// Gets data for a particular secret
-	/// </summary>
-	public class GetSecretResponse
-	{
-		readonly SecretConfig _config;
-
-		/// <summary>
-		/// Id of the secret
-		/// </summary>
-		public SecretId Id => _config.Id;
-
-		/// <summary>
-		/// Key value pairs for the secret
-		/// </summary>
-		public Dictionary<string, string> Data => _config.Data;
-
-		internal GetSecretResponse(SecretConfig config) => _config = config;
-	}
-
 	/// <summary>
 	/// Controller for the /api/v1/credentials endpoint
 	/// </summary>
@@ -65,15 +35,15 @@ namespace Horde.Server.Secrets
 		[Route("/api/v1/secrets")]
 		public ActionResult<GetSecretsResponse> GetSecrets()
 		{
-			GetSecretsResponse response = new GetSecretsResponse();
+			List<SecretId> secretIds = new List<SecretId>();
 			foreach (SecretConfig secret in _globalConfig.Value.Secrets)
 			{
 				if (secret.Authorize(SecretAclAction.ViewSecret, User))
 				{
-					response.Ids.Add(secret.Id);
+					secretIds.Add(secret.Id);
 				}
 			}
-			return response;
+			return new GetSecretsResponse(secretIds);
 		}
 
 		/// <summary>
@@ -96,7 +66,7 @@ namespace Horde.Server.Secrets
 				return Forbid(SecretAclAction.ViewSecret, secretId);
 			}
 
-			return new GetSecretResponse(secretConfig).ApplyFilter(filter);
+			return new GetSecretResponse(secretConfig.Id, secretConfig.Data).ApplyFilter(filter);
 		}
 	}
 }
diff --git a/Engine/Source/Programs/Shared/EpicGames.Horde/Api/SecretId.cs b/Engine/Source/Programs/Shared/EpicGames.Horde/Api/SecretId.cs
new file mode 100644
--- /dev/null
+++ b/Engine/Source/Programs/Shared/EpicGames.Horde/Api/SecretId.cs
@@ -0,0 +1,44 @@
+// Copyright Epic Games, Inc. All Rights Reserved.
+
+using System.ComponentModel;
+using EpicGames.Core;
+using EpicGames.Serialization;
+
+namespace EpicGames.Horde.Api
+{
+	/// <summary>
+	/// Identifier for a secret
+	/// </summary>
+	/// <param name="Id">Id to construct from</param>
+	[JsonSchemaString]
+	[TypeConverter(typeof(StringIdTypeConverter<SecretId, SecretIdConverter>))]
+	[StringIdConverter(typeof(SecretIdConverter))]
+	[CbConverter(typeof(StringIdCbConverter<SecretId, SecretIdConverter>))]
+	public record struct SecretId(StringId Id)
+	{
+		/// <summary>
+		/// Constructor
+		/// </summary>
+		public SecretId(string id) : this(new StringId(id))
+		{
+		}
+
+		/// <inheritdoc cref="StringId.IsEmpty"/>
+		public bool IsEmpty => Id.IsEmpty;
+
+		/// <inheritdoc/>
+		public override string ToString() => Id.ToString();
+	}
+
+	/// <summary>
+	/// Converter to and from <see cref="StringId"/> instances.
+	/// </summary>
+	class SecretIdConverter : StringIdConverter<SecretId>
+	{
+		/// <inheritdoc/>
+		public override SecretId FromStringId(StringId id) => new SecretId(id);
+
+		/// <inheritdoc/>
+		public override StringId ToStringId(SecretId value) => value.Id;
+	}
+}
diff --git a/Engine/Source/Programs/Shared/EpicGames.Horde/Api/SecretMessages.cs b/Engine/Source/Programs/Shared/EpicGames.Horde/Api/SecretMessages.cs
new file mode 100644
--- /dev/null
+++ b/Engine/Source/Programs/Shared/EpicGames.Horde/Api/SecretMessages.cs
@@ -0,0 +1,83 @@
+// Copyright Epic Games, Inc. All Rights Reserved.
+
+using System;
+using System.Collections.Generic;
+using System.Diagnostics;
+using System.Threading;
+using System.Threading.Tasks;
+
+#pragma warning disable CA2227 // Collection properties should be read only
+
+namespace EpicGames.Horde.Api
+{
+	/// <summary>
+	/// Response listing all the secrets available to the current user
+	/// </summary>
+	public class GetSecretsResponse
+	{
+		/// <summary>
+		/// List of secret ids
+		/// </summary>
+		public List<SecretId> Ids { get; set; }
+
+		/// <summary>
+		/// Constructor
+		/// </summary>
+		public GetSecretsResponse(List<SecretId> ids) => Ids = ids;
+	}
+
+	/// <summary>
+	/// Gets data for a particular secret
+	/// </summary>
+	[DebuggerDisplay("{Id}")]
+	public class GetSecretResponse
+	{
+		/// <summary>
+		/// Id of the secret
+		/// </summary>
+		public SecretId Id { get; }
+
+		/// <summary>
+		/// Key value pairs for the secret
+		/// </summary>
+		public Dictionary<string, string> Data { get; set; }
+
+		/// <summary>
+		/// Constructor
+		/// </summary>
+		public GetSecretResponse(SecretId id, Dictionary<string, string> data)
+		{
+			Id = id;
+			Data = data;
+		}
+	}
+
+	/// <summary>
+	/// Extension methods for the secrets endpoint
+	/// </summary>
+	public static class SecretExtensions
+	{
+		/// <summary>
+		/// Query all the secrets available to the current user
+		/// </summary>
+		/// <param name="horde">The horde client instance</param>
+		/// <param name="cancellationToken">Cancellation token for the operation</param>
+		/// <returns>Information about all the projects</returns>
+		public static Task<GetSecretsResponse> GetSecretsAsync(this HordeHttpClient horde, CancellationToken cancellationToken = default)
+		{
+			return horde.GetAsync<GetSecretsResponse>($"api/v1/secrets", cancellationToken);
+		}
+
+		/// <summary>
+		/// Retrieve information about a specific project
+		/// </summary>
+		/// <param name="horde">The horde client instance</param>
+		/// <param name="secretId">Id of the secret to retrieve</param>
+		/// <param name="cancellationToken">Cancellation token for the operation</param>
+		/// <returns>Information about the requested project</returns>
+		public static Task<GetSecretResponse> GetSecretAsync(this HordeHttpClient horde, SecretId secretId, CancellationToken cancellationToken = default)
+		{
+			return horde.GetAsync<GetSecretResponse>($"api/v1/secrets/{secretId}", cancellationToken);
+		}
+	}
+}
