Horde: Fix un-flushed ref target being written to disk in file storage client.

Also add backoff/retry mechanism to writing refs.

diff --git a/Engine/Source/Programs/Shared/EpicGames.Horde/Storage/Backends/FileStorageClient.cs b/Engine/Source/Programs/Shared/EpicGames.Horde/Storage/Backends/FileStorageClient.cs
--- a/Engine/Source/Programs/Shared/EpicGames.Horde/Storage/Backends/FileStorageClient.cs
+++ b/Engine/Source/Programs/Shared/EpicGames.Horde/Storage/Backends/FileStorageClient.cs
@@ -136,7 +136,21 @@ namespace EpicGames.Horde.Storage.Backends
 			FileReference file = GetRefFile(name);
 			DirectoryReference.CreateDirectory(file.Directory);
 			_logger.LogInformation("Writing {File}", file);
-			await FileReference.WriteAllTextAsync(file, target.ToString());
+
+			NodeLocator locator = await target.FlushAsync(cancellationToken);
+			for (int attempt = 0; ; attempt++)
+			{
+				try
+				{
+					await FileReference.WriteAllTextAsync(file, locator.ToString());
+					break;
+				}
+				catch (IOException ex) when (attempt < 3)
+				{
+					_logger.LogDebug(ex, "Unable to write to {File}; retrying...", file);
+					await Task.Delay(100 * attempt, cancellationToken);
+				}
+			}
 		}
 
 		#endregion
