Horde: Allow spaces in P4 view filters.

diff --git a/Engine/Source/Programs/Shared/EpicGames.Perforce/PerforceViewFilter.cs b/Engine/Source/Programs/Shared/EpicGames.Perforce/PerforceViewFilter.cs
--- a/Engine/Source/Programs/Shared/EpicGames.Perforce/PerforceViewFilter.cs
+++ b/Engine/Source/Programs/Shared/EpicGames.Perforce/PerforceViewFilter.cs
@@ -151,7 +151,7 @@ namespace EpicGames.Perforce
 		/// <returns></returns>
 		public static PerforceViewFilterEntry Parse(string entry)
 		{
-			Match match = Regex.Match(entry, @"^\s*(-?)\s*([^ ]+)\s*$");
+			Match match = Regex.Match(entry, @"^\s*(-?)\s*(.*[^ ])\s*$");
 			if (!match.Success)
 			{
 				throw new PerforceException($"Unable to parse view filter entry: {entry}");
