Honor the editor-only state of sparse class data properties in their accessors

#rnx

diff --git a/Engine/Source/Programs/Shared/EpicGames.UHT/Exporters/CodeGen/UhtHeaderCodeGenerator.cs b/Engine/Source/Programs/Shared/EpicGames.UHT/Exporters/CodeGen/UhtHeaderCodeGenerator.cs
--- a/Engine/Source/Programs/Shared/EpicGames.UHT/Exporters/CodeGen/UhtHeaderCodeGenerator.cs
+++ b/Engine/Source/Programs/Shared/EpicGames.UHT/Exporters/CodeGen/UhtHeaderCodeGenerator.cs
@@ -15,6 +15,8 @@ namespace EpicGames.UHT.Exporters.CodeGen
 		public const string EventParamsMacroSuffix = "EVENT_PARMS";
 		public const string CallbackWrappersMacroSuffix = "CALLBACK_WRAPPERS";
 		public const string SparseDataMacroSuffix = "SPARSE_DATA";
+		public const string SparseDataPropertyAccessorsMacroSuffix = "SPARSE_DATA_PROPERTY_ACCESSORS";
+		public const string EditorOnlySparseDataPropertyAccessorsMacroSuffix = "EDITOR_ONLY_SPARSE_DATA_PROPERTY_ACCESSORS";
 		public const string EditorOnlyRpcWrappersMacroSuffix = "EDITOR_ONLY_RPC_WRAPPERS";
 		public const string RpcWrappersMacroSuffix = "RPC_WRAPPERS";
 		public const string EditorOnlyRpcWrappersNoPureDeclsMacroSuffix = "EDITOR_ONLY_RPC_WRAPPERS_NO_PURE_DECLS";
diff --git a/Engine/Source/Programs/Shared/EpicGames.UHT/Exporters/CodeGen/UhtHeaderCodeGeneratorHFile.cs b/Engine/Source/Programs/Shared/EpicGames.UHT/Exporters/CodeGen/UhtHeaderCodeGeneratorHFile.cs
--- a/Engine/Source/Programs/Shared/EpicGames.UHT/Exporters/CodeGen/UhtHeaderCodeGeneratorHFile.cs
+++ b/Engine/Source/Programs/Shared/EpicGames.UHT/Exporters/CodeGen/UhtHeaderCodeGeneratorHFile.cs
@@ -786,11 +786,21 @@ namespace EpicGames.UHT.Exporters.CodeGen
 		
 		private StringBuilder AppendSparseDeclarations(StringBuilder builder, UhtClass classObj)
 		{
-			// Format the sparse data
+			IEnumerable<UhtScriptStruct> sparseScriptStructs = GetSparseDataStructsToExport(classObj);
+
+			AppendSparseStructDeclarations(builder, classObj, sparseScriptStructs);
+
+			CalculateSparsePropertyDeclarations(sparseScriptStructs, out Dictionary<UhtScriptStruct, List<UhtProperty>> commonProperties, out Dictionary<UhtScriptStruct, List<UhtProperty>> editorOnlyProperties);
+			AppendSparsePropertyDeclarations(builder, classObj, sparseScriptStructs, commonProperties, /*isEditorOnly*/false);
+			AppendSparsePropertyDeclarations(builder, classObj, sparseScriptStructs, editorOnlyProperties, /*isEditorOnly*/true);
+			
+			return builder;
+		}
+
+		private StringBuilder AppendSparseStructDeclarations(StringBuilder builder, UhtClass classObj, IEnumerable<UhtScriptStruct> sparseScriptStructs)
+		{
 			using (UhtMacroCreator macro = new(builder, this, classObj, SparseDataMacroSuffix))
 			{
-				IEnumerable<UhtScriptStruct> sparseScriptStructs = GetSparseDataStructsToExport(classObj);
-
 				if (sparseScriptStructs.Any())
 				{
 					string api = classObj.ClassFlags.HasAnyFlags(EClassFlags.MinimalAPI) ? PackageApi : "";
@@ -804,14 +814,28 @@ namespace EpicGames.UHT.Exporters.CodeGen
 						builder.Append(api).Append("static UScriptStruct* StaticGet").Append(sparseDataType).Append("ScriptStruct(); \\\r\n");
 					}
 				}
+			}
+
+			return builder;
+		}
+
+		private StringBuilder AppendSparsePropertyDeclarations(StringBuilder builder, UhtClass classObj, IEnumerable<UhtScriptStruct> sparseScriptStructs, Dictionary<UhtScriptStruct, List<UhtProperty>> propertiesToDeclare, bool isEditorOnly)
+		{
+			// We can skip writing the macros if there are no properties to declare, as the 'if' and 'else' would be the same
+			if (isEditorOnly && propertiesToDeclare.Count > 0)
+			{
+				builder.Append("#if WITH_EDITORONLY_DATA\r\n");
+			}
 
+			using (UhtMacroCreator macro = new(builder, this, classObj, isEditorOnly ? EditorOnlySparseDataPropertyAccessorsMacroSuffix : SparseDataPropertyAccessorsMacroSuffix))
+			{
 				foreach (UhtScriptStruct sparseScriptStruct in sparseScriptStructs)
 				{
-					string sparseDataType = sparseScriptStruct.EngineName;
-
-					foreach (UhtProperty sparseProperty in sparseScriptStruct.Properties)
+					if (propertiesToDeclare.TryGetValue(sparseScriptStruct, out List<UhtProperty>? sparseProperties))
 					{
-						if (!sparseProperty.MetaData.ContainsKey(UhtNames.NoGetter))
+						string sparseDataType = sparseScriptStruct.EngineName;
+
+						foreach (UhtProperty sparseProperty in sparseProperties)
 						{
 							string propertyName = sparseProperty.SourceName;
 							string cleanPropertyName = propertyName;
@@ -820,26 +844,80 @@ namespace EpicGames.UHT.Exporters.CodeGen
 								cleanPropertyName = propertyName[1..];
 							}
 
-							bool getByRef = sparseProperty.MetaData.ContainsKey(UhtNames.GetByRef);
-
-							if (getByRef)
+							if (sparseProperty.MetaData.ContainsKey(UhtNames.GetByRef))
 							{
-								builder.Append("const ").AppendSparse(sparseProperty).Append("& Get").Append(cleanPropertyName).Append("() const \\\r\n");
+								builder.Append("const ").AppendSparse(sparseProperty).Append("& Get").Append(cleanPropertyName).Append("() const");
 							}
 							else
 							{
-								builder.AppendSparse(sparseProperty).Append(" Get").Append(cleanPropertyName).Append("() const \\\r\n");
+								builder.AppendSparse(sparseProperty).Append(" Get").Append(cleanPropertyName).Append("() const");
 							}
-							builder.Append("{ \\\r\n");
-							builder.Append("\treturn Get").Append(sparseDataType).Append("(EGetSparseClassDataMethod::ArchetypeIfNull)->").Append(propertyName).Append("; \\\r\n");
-							builder.Append("} \\\r\n");
+							builder.Append(" { return Get").Append(sparseDataType).Append("(EGetSparseClassDataMethod::ArchetypeIfNull)->").Append(propertyName).Append("; } \\\r\n");
 						}
 					}
 				}
 			}
+
+			// We can skip writing the macros if there are no properties to declare, as the 'if' and 'else' would be the same
+			if (isEditorOnly && propertiesToDeclare.Count > 0)
+			{
+				// Trim the extra newlines added after the macro generator
+				if (builder.Length > 4 &&
+					builder[builder.Length - 4] == '\r' &&
+					builder[builder.Length - 3] == '\n' &&
+					builder[builder.Length - 2] == '\r' &&
+					builder[builder.Length - 1] == '\n')
+				{
+					builder.Length -= 4;
+				}
+
+				builder.Append("#else //WITH_EDITORONLY_DATA\r\n");
+				using (UhtMacroCreator macro = new(builder, this, classObj, EditorOnlySparseDataPropertyAccessorsMacroSuffix))
+				{
+					// Empty macro when not compiling WITH_EDITORONLY_DATA
+				}
+				builder.Append("#endif //WITH_EDITORONLY_DATA\r\n\r\n\r\n");
+			}
+
 			return builder;
 		}
 
+		private static void CalculateSparsePropertyDeclarations(IEnumerable<UhtScriptStruct> sparseScriptStructs, out Dictionary<UhtScriptStruct, List<UhtProperty>> commonProperties, out Dictionary<UhtScriptStruct, List<UhtProperty>> editorOnlyProperties)
+		{
+			commonProperties = new();
+			editorOnlyProperties = new();
+
+			foreach (UhtScriptStruct sparseScriptStruct in sparseScriptStructs)
+			{
+				List<UhtProperty> commonPropertiesForStruct = new();
+				List<UhtProperty> editorOnlyPropertiesForStruct = new();
+
+				foreach (UhtProperty sparseProperty in sparseScriptStruct.Properties)
+				{
+					if (!sparseProperty.MetaData.ContainsKey(UhtNames.NoGetter))
+					{
+						if (sparseProperty.IsEditorOnlyProperty)
+						{
+							editorOnlyPropertiesForStruct.Add(sparseProperty);
+						}
+						else
+						{
+							commonPropertiesForStruct.Add(sparseProperty);
+						}
+					}
+				}
+
+				if (commonPropertiesForStruct.Count > 0)
+				{
+					commonProperties.Add(sparseScriptStruct, commonPropertiesForStruct);
+				}
+				if (editorOnlyPropertiesForStruct.Count > 0)
+				{
+					editorOnlyProperties.Add(sparseScriptStruct, editorOnlyPropertiesForStruct);
+				}
+			}
+		}
+
 		private StringBuilder AppendRpcFunctions(StringBuilder builder, UhtClass classObj, bool usesLegacy, List<UhtFunction> reversedFunctions, bool editorOnly)
 		{
 			builder.AppendBeginEditorOnlyGuard(editorOnly);
@@ -1536,6 +1614,8 @@ namespace EpicGames.UHT.Exporters.CodeGen
 				builder.Append(DisableDeprecationWarnings).Append(" \\\r\n");
 				builder.Append("public: \\\r\n");
 				builder.Append('\t').AppendMacroName(this, classObj, SparseDataMacroSuffix).Append(" \\\r\n");
+				builder.Append('\t').AppendMacroName(this, classObj, SparseDataPropertyAccessorsMacroSuffix).Append(" \\\r\n");
+				builder.Append('\t').AppendMacroName(this, classObj, EditorOnlySparseDataPropertyAccessorsMacroSuffix).Append(" \\\r\n");
 				if (isLegacy)
 				{
 					builder.Append('\t').AppendMacroName(this, classObj, RpcWrappersMacroSuffix).Append(" \\\r\n");
