LazyObjectPtr and ObjectPtr are allowed as TMap keys in function arguments.

diff --git a/Engine/Source/Programs/Shared/EpicGames.UHT/Types/Properties/UhtLazyObjectPtrProperty.cs b/Engine/Source/Programs/Shared/EpicGames.UHT/Types/Properties/UhtLazyObjectPtrProperty.cs
--- a/Engine/Source/Programs/Shared/EpicGames.UHT/Types/Properties/UhtLazyObjectPtrProperty.cs
+++ b/Engine/Source/Programs/Shared/EpicGames.UHT/Types/Properties/UhtLazyObjectPtrProperty.cs
@@ -99,7 +99,7 @@ namespace EpicGames.UHT.Types
 
 			// UFunctions with a smart pointer as input parameter wont compile anyway, because of missing P_GET_... macro.
 			// UFunctions with a smart pointer as return type will crash when called via blueprint, because they are not supported in VM.
-			if (!options.HasAnyFlags(UhtValidationOptions.IsKey) && PropertyCategory != UhtPropertyCategory.Member)
+			if (PropertyCategory != UhtPropertyCategory.Member)
 			{
 				outerStruct.LogError("UFunctions cannot take a lazy pointer as a parameter.");
 			}
diff --git a/Engine/Source/Programs/Shared/EpicGames.UHT/Types/Properties/UhtObjectPtrProperty.cs b/Engine/Source/Programs/Shared/EpicGames.UHT/Types/Properties/UhtObjectPtrProperty.cs
--- a/Engine/Source/Programs/Shared/EpicGames.UHT/Types/Properties/UhtObjectPtrProperty.cs
+++ b/Engine/Source/Programs/Shared/EpicGames.UHT/Types/Properties/UhtObjectPtrProperty.cs
@@ -1,5 +1,6 @@
 // Copyright Epic Games, Inc. All Rights Reserved.
 
+using System;
 using System.Diagnostics.CodeAnalysis;
 using System.Text;
 using EpicGames.Core;
@@ -94,9 +95,16 @@ namespace EpicGames.UHT.Types
 
 			// UFunctions with a smart pointer as input parameter wont compile anyway, because of missing P_GET_... macro.
 			// UFunctions with a smart pointer as return type will crash when called via blueprint, because they are not supported in VM.
-			if (!options.HasAnyFlags(UhtValidationOptions.IsKey) && PropertyCategory != UhtPropertyCategory.Member)
+			if (PropertyCategory != UhtPropertyCategory.Member)
 			{
-				outerStruct.LogError("UFunctions cannot take a TObjectPtr as a parameter.");
+				// At this point, allow this to appear in TMap keys in the UPlayerMappableInputConfig class
+				if (!options.HasAnyFlags(UhtValidationOptions.IsKey) ||
+					!outerStruct.SourceName.Equals("GetMappingContexts", StringComparison.Ordinal) ||
+					outerStruct.Outer == null ||
+					!outerStruct.Outer.SourceName.Equals("UPlayerMappableInputConfig", StringComparison.Ordinal))
+				{
+					outerStruct.LogError("UFunctions cannot take a TObjectPtr as a parameter.");
+				}
 			}
 		}
 
