Unreal Build Tool:
Update PluginDescriptor to have parity with FPluginDescriptor in terms of handling custom JSON fields. The order in which the original fields are written to disk is preserved.
When saving the PluginDescriptor to disk, the same order is followed and any new entries are added to the end of the plugin file.
Note: PluginDescriptor.Save and PluginDescriptor.Save2 will exist concurrently for now. There are still parts of UBT that use PluginDescriptor.Write and the refactor to create a unified Save function will come later as a lot of testing still needs to be done.
JsonObject:
- Introduced AddOrSetFieldValue() methods overloaded on value types to add/set fields in JsonObject.
- Changed the underlying data structure of JsonObject from Dictionary<string,Object?> to OrderedDictionary to guarantee that order of all .uplugin fields (default and custom) are preserved
- Made the JsonObject constructor that takes an OrderedDictionary private to ensure type safety of values passed to JsonObject
- Introduced ToJsonString() for JsonObject that serializes the JsonObject to match the formatting of .uplugin
IMPORTANT: The ToJsonString() method uses the JavaScriptEncoder.UnsafeRelaxedJsonEscaping encoder for Utf8JsonWriter.
This allows +,<,>,&,` etc characters to be written properly for .uplugins. As per MS docs, with this encoding, none of the contents should be written to HTML or a script.
- Introduced JsonObjectTests that unit tested JsonObject to ensure previous behavior was maintained.

PluginDescriptor
- Added a JsonObject field that acts as the cache of all fields that were read from a file or parsed from a json string.
- Cached JsonObject records the order of all fields that were read and captures custom fields in the .uplugin
- Introduced Save2(), the new save algorithm that writes all fields (default and custom) to disk.
- Deleted the private default constructor because it didn't seem to be used anywhere.
- Introduced PluginDescriptorTests to ensure previous behavior is maintained. In particular when dealing with a .uplugin with only default fields, results of Save() and Save2() are the same.

Descriptor Files:
- Introduced ToJsonObject() methods to replace the previous JsonWriter patterns of writing the descriptor data.
- Introduced the UpdateJsonObject pattern which mimics the old WriteArray pattern of writing data to a JsonWriter.

#test: Developed this in a TDD style with unit tests for all the major functions. Also reviewed output of created .uplugins and diffed them.

diff --git a/Engine/Source/Programs/Shared/EpicGames.Core.Tests/JsonObjectTests.cs b/Engine/Source/Programs/Shared/EpicGames.Core.Tests/JsonObjectTests.cs
new file mode 100644
--- /dev/null
+++ b/Engine/Source/Programs/Shared/EpicGames.Core.Tests/JsonObjectTests.cs
@@ -0,0 +1,676 @@
+// Copyright Epic Games, Inc. All Rights Reserved.
+
+using System;
+using System.Collections.Generic;
+using System.ComponentModel.DataAnnotations;
+using System.IO;
+using System.Linq;
+using System.Text;
+using System.Text.Json;
+using System.Text.Json.Serialization;
+using Microsoft.VisualStudio.TestTools.UnitTesting;
+
+namespace EpicGames.Core.Tests
+{
+	[TestClass]
+	public class JsonObjectTests
+	{
+		private string _testJsonText = null!;
+		private JsonObject _validTestObject = null!;
+		private Enemy _validJsonDataSource = null!;
+		private static DirectoryReference s_tempDirectory = null!;
+
+		[ClassInitialize]
+		public static void Setup(TestContext testContext)
+		{
+			JsonObjectTests.s_tempDirectory = CreateTempDir();
+		}
+
+		[ClassCleanup]
+		public static void TearDown()
+		{
+			if (Directory.Exists(JsonObjectTests.s_tempDirectory.FullName))
+			{
+				Directory.Delete(JsonObjectTests.s_tempDirectory.FullName, true);
+			}
+		}
+
+		[TestInitialize]
+		public void InitializeTest()
+		{
+			Enemy enemy1 = CreateValidJsonDataSource();
+			// This is to make sure that enums are serialized as strings in the case that they were written as.
+			JsonSerializerOptions options = new JsonSerializerOptions();
+			options.Converters.Add(new JsonStringEnumConverter(null));
+			options.WriteIndented = true;
+			_testJsonText = JsonSerializer.Serialize(enemy1, options);
+			_validTestObject = CreateValidJsonObject();
+			_validJsonDataSource = enemy1;
+		}
+
+		[TestCleanup]
+		public void TestCleanup()
+		{
+			
+		}
+
+		[TestMethod]
+		public void ParseValidJsonText()
+		{
+			JsonObject.Parse(_testJsonText);
+		}
+
+		[TestMethod]
+		public void ParseInvalidJsonText()
+		{
+			MakeTestJsonTextInvalid();
+			Assert.ThrowsException<JsonParseException>(() => JsonObject.Parse(_testJsonText));
+		}
+
+		[TestMethod]
+		public void ParseEmptyText()
+		{
+			Assert.ThrowsException<JsonParseException>(() => JsonObject.Parse(""));
+		}
+
+		[TestMethod]
+		public void TryParseValidJsonText()
+		{
+			JsonObject? testObject;
+			Assert.IsTrue(JsonObject.TryParse(_testJsonText, out testObject));
+			Assert.IsNotNull(testObject);
+		}
+
+		[TestMethod]
+		public void TryParseInvalidJsonText()
+		{
+			MakeTestJsonTextInvalid();
+			JsonObject? testObject;
+			Assert.IsFalse(JsonObject.TryParse(_testJsonText, out testObject));
+			Assert.IsNull(testObject);
+		}
+
+		[TestMethod]
+		public void TryParseEmptyJsonText()
+		{
+			JsonObject? testObject;
+			Assert.IsFalse(JsonObject.TryParse("", out testObject));
+			Assert.IsNull(testObject);
+		}
+
+		[TestMethod]
+		public void ReadValidJsonText()
+		{
+			FileReference inputFileReference = CreateTempJsonFile("Valid.json", _testJsonText);
+			JsonObject testObject = JsonObject.Read(inputFileReference);
+			Assert.IsNotNull(testObject);
+		}
+
+		[TestMethod]
+		public void ReadInvalidJsonText()
+		{
+			MakeTestJsonTextInvalid();
+			FileReference inputFileReference = CreateTempJsonFile("Invalid.json", _testJsonText);
+			Assert.ThrowsException<JsonParseException>(() => JsonObject.Read(inputFileReference));
+		}
+
+		[TestMethod]
+		public void ReadEmptyJsonText()
+		{
+			FileReference inputFileReference = CreateTempJsonFile("Empty.json", "");
+			Assert.ThrowsException<JsonParseException>(() => JsonObject.Read(inputFileReference));
+		}
+
+		[TestMethod]
+		public void TryReadValidJsonText()
+		{
+			FileReference inputFileReference = CreateTempJsonFile("Valid.json", _testJsonText);
+			JsonObject? testObject;
+			Assert.IsTrue(JsonObject.TryRead(inputFileReference, out testObject));
+			Assert.IsNotNull(testObject);
+		}
+
+		[TestMethod]
+		public void TryReadInvalidJsonText()
+		{
+			MakeTestJsonTextInvalid();
+			FileReference inputFileReference = CreateTempJsonFile("Invalid.json", _testJsonText);
+			JsonObject? testObject;
+			Assert.IsFalse(JsonObject.TryRead(inputFileReference, out testObject));
+			Assert.IsNull(testObject);
+		}
+
+		[TestMethod]
+		public void TryReadEmptyJsonText()
+		{
+			FileReference inputFileReference = CreateTempJsonFile("Empty.json", "");
+			JsonObject? testObject;
+			Assert.IsFalse(JsonObject.TryRead(inputFileReference, out testObject));
+			Assert.IsNull(testObject);
+		}
+
+		[TestMethod]
+		public void GetIntegerValidField()
+		{
+			string fieldName = "Hp";
+			int correctValue = _validJsonDataSource.Hp;
+			Assert.AreEqual(correctValue, _validTestObject.GetIntegerField(fieldName));
+			int outValue;
+			Assert.IsTrue(_validTestObject.TryGetIntegerField(fieldName, out outValue));
+			Assert.AreEqual(correctValue, outValue);
+		}
+
+		[TestMethod]
+		public void GetIntegerInvalidField()
+		{
+			string fieldName = "Invalid";
+			int correctValue = _validJsonDataSource.Hp;
+			Assert.ThrowsException<JsonParseException>(() => _validTestObject.GetIntegerField(fieldName));
+			int outValue;
+			Assert.IsFalse(_validTestObject.TryGetIntegerField(fieldName, out outValue));
+			Assert.AreNotEqual(correctValue, outValue);
+		}
+
+		[TestMethod]
+		public void GetStringValidField()
+		{
+			string fieldName = "Name";
+			string correctValue = _validJsonDataSource.Name;
+			Assert.AreEqual(correctValue, _validTestObject.GetStringField(fieldName));
+			string? outValue;
+			Assert.IsTrue(_validTestObject.TryGetStringField(fieldName, out outValue));
+			Assert.AreEqual(correctValue, outValue);
+		}
+
+		[TestMethod]
+		public void GetStringInvalidField()
+		{
+			string fieldName = "Invalid";
+			string correctValue = _validJsonDataSource.Name;
+			Assert.ThrowsException<JsonParseException>(() => _validTestObject.GetStringField(fieldName));
+			string? outValue;
+			Assert.IsFalse(_validTestObject.TryGetStringField(fieldName, out outValue));
+			Assert.AreNotEqual(correctValue, outValue);
+		}
+
+		[TestMethod]
+		public void GetObjectValidField()
+		{
+			string fieldName = "StrongestAttack";
+			JsonObject correctValue = _validJsonDataSource.StrongestAttack.ToJsonObject();
+			JsonObject returnValue = _validTestObject.GetObjectField(fieldName);
+			Assert.AreEqual(correctValue, returnValue);
+			JsonObject? outValue;
+			Assert.IsTrue(_validTestObject.TryGetObjectField(fieldName, out outValue));
+			Assert.AreEqual(correctValue, outValue);
+		}
+
+		[TestMethod]
+		public void GetObjectInvalidField()
+		{
+			string fieldName = "Invalid";
+			Assert.ThrowsException<JsonParseException>(() => _validTestObject.GetObjectField(fieldName));
+			JsonObject? outValue;
+			Assert.IsFalse(_validTestObject.TryGetObjectField(fieldName, out outValue));
+			Assert.IsNull(outValue);
+		}
+
+		[TestMethod]
+		public void GetObjectArrayValidField()
+		{
+			string fieldName = "Attacks";
+			JsonObject[] correctValue = _validJsonDataSource.Attacks.Select(x => x.ToJsonObject()).ToArray();
+			JsonObject[] result = _validTestObject.GetObjectArrayField(fieldName);
+			CollectionAssert.AreEqual(correctValue, result);
+			JsonObject[]? outValue;
+			Assert.IsTrue(_validTestObject.TryGetObjectArrayField(fieldName, out outValue));
+			CollectionAssert.AreEqual(correctValue, outValue);
+		}
+
+		[TestMethod]
+		public void GetObjectArrayInvalidField()
+		{
+			string fieldName = "Invalid";
+			Assert.ThrowsException<JsonParseException>(() => _validTestObject.GetObjectArrayField(fieldName));
+			JsonObject[]? outValue;
+			Assert.IsFalse(_validTestObject.TryGetObjectArrayField(fieldName, out outValue));
+			Assert.IsNull(outValue);
+		}
+
+		[TestMethod]
+		public void GetStringArrayValidField()
+		{
+			string fieldName = "Skills";
+			string[] correctValue = _validJsonDataSource.Skills.ToArray();
+			string[] result = _validTestObject.GetStringArrayField(fieldName);
+			CollectionAssert.AreEqual(correctValue, result);
+			string[]? outValue;
+			Assert.IsTrue(_validTestObject.TryGetStringArrayField(fieldName, out outValue));
+			CollectionAssert.AreEqual(correctValue, outValue);
+		}
+
+		[TestMethod]
+		public void GetStringArrayInvalidField()
+		{
+			string fieldName = "Invalid";
+			string[] correctValue = _validJsonDataSource.Skills.ToArray();
+			Assert.ThrowsException<JsonParseException>(() => _validTestObject.GetStringArrayField(fieldName));
+			string[]? outValue;
+			Assert.IsFalse(_validTestObject.TryGetStringArrayField(fieldName, out outValue));
+			Assert.IsNull(outValue);
+		}
+
+		[TestMethod]
+		public void GetEnumValidField()
+		{
+			string fieldName = "CurrentStatusEffect";
+			StatusEffect correctValue = _validJsonDataSource.CurrentStatusEffect;
+			Assert.AreEqual(correctValue, _validTestObject.GetEnumField<StatusEffect>(fieldName));
+			StatusEffect outValue;
+			Assert.IsTrue(_validTestObject.TryGetEnumField<StatusEffect>(fieldName, out outValue));
+			Assert.AreEqual(correctValue, outValue);
+		}
+
+		[TestMethod]
+		public void GetEnumInvalidField()
+		{
+			string fieldName = "Invalid";
+			StatusEffect correctValue = _validJsonDataSource.CurrentStatusEffect;
+			Assert.ThrowsException<JsonParseException>(() => _validTestObject.GetEnumField<StatusEffect>(fieldName));
+			StatusEffect outValue;
+			Assert.IsFalse(_validTestObject.TryGetEnumField<StatusEffect>(fieldName, out outValue));
+			Assert.AreEqual(outValue, default);
+		}
+
+		[TestMethod]
+		public void GetEnumArrayValidField()
+		{
+			string fieldName = "ImmuneStatusEffects";
+			StatusEffect[] correctValue = _validJsonDataSource.ImmuneStatusEffects.ToArray();
+			StatusEffect[]? outValue;
+			Assert.IsTrue(_validTestObject.TryGetEnumArrayField<StatusEffect>(fieldName, out outValue));
+			CollectionAssert.AreEqual(correctValue, outValue);
+		}
+
+		[TestMethod]
+		public void GetEnumArrayInvalidField()
+		{
+			string fieldName = "Invalid";
+			StatusEffect[] correctValu = _validJsonDataSource.ImmuneStatusEffects.ToArray();
+			StatusEffect[]? outValue;
+			Assert.IsFalse(_validTestObject.TryGetEnumArrayField(fieldName, out outValue));
+			Assert.IsNull(outValue);
+			CollectionAssert.AreNotEqual(correctValu, outValue);
+		}
+
+		[TestMethod]
+		public void GetBoolValidField()
+		{
+			string fieldName = "bIsAlive";
+			bool correctValue = _validJsonDataSource.bIsAlive;
+			Assert.AreEqual(correctValue, _validTestObject.GetBoolField(fieldName));
+			bool outValue;
+			Assert.IsTrue(_validTestObject.TryGetBoolField(fieldName, out outValue));
+			Assert.AreEqual(correctValue, outValue);
+		}
+
+		[TestMethod]
+		public void GetBoolInvalidField()
+		{
+			string fieldName = "Invalid";
+			Assert.ThrowsException<JsonParseException>(() => _validTestObject.GetBoolField(fieldName));
+			bool outValue = false;
+			Assert.IsFalse(_validTestObject.TryGetBoolField(fieldName, out outValue));
+			Assert.IsFalse(outValue);
+		}
+
+		[TestMethod]
+		public void ContainsField()
+		{
+			string fieldName1 = "Name";
+			string fieldName2 = "Invalid";
+			Assert.IsTrue(_validTestObject.ContainsField(fieldName1));
+			Assert.IsFalse(_validTestObject.ContainsField(fieldName2));
+		}
+
+		[TestMethod]
+		public void SetIntegerField()
+		{
+			string fieldName = "Hp";
+			Assert.IsTrue(_validTestObject.ContainsField(fieldName));
+			int setValue = 60;
+			_validTestObject.AddOrSetFieldValue(fieldName, setValue);
+			Assert.AreEqual(setValue, _validTestObject.GetIntegerField(fieldName));
+
+			string newFieldName = "New";
+			Assert.IsFalse(_validTestObject.ContainsField(newFieldName));
+			_validTestObject.AddOrSetFieldValue(newFieldName, setValue);
+			Assert.IsTrue(_validTestObject.ContainsField(newFieldName));
+			Assert.AreEqual(setValue, _validTestObject.GetIntegerField(newFieldName));
+		}
+
+		[TestMethod]
+		public void SetDoubleField()
+		{
+			string fieldName = "CriticalHitRate";
+			Assert.IsTrue(_validTestObject.ContainsField(fieldName));
+			double setValue = 1.0;
+			_validTestObject.AddOrSetFieldValue(fieldName, setValue);
+			Assert.AreEqual(setValue, _validTestObject.GetDoubleField(fieldName));
+
+			string newFieldName = "New";
+			Assert.IsFalse(_validTestObject.ContainsField(newFieldName));
+			_validTestObject.AddOrSetFieldValue(newFieldName, setValue);
+			Assert.IsTrue(_validTestObject.ContainsField(newFieldName));
+			Assert.AreEqual(setValue, _validTestObject.GetDoubleField(newFieldName));
+		}
+
+		[TestMethod]
+		public void SetStringField()
+		{
+			string fieldName = "Name";
+			string setValue = "NewValue";
+			Assert.IsTrue(_validTestObject.ContainsField(fieldName));
+			_validTestObject.AddOrSetFieldValue(fieldName, setValue);
+			Assert.AreEqual(setValue, _validTestObject.GetStringField(fieldName));
+
+			string newFieldName = "New";
+			Assert.IsFalse(_validTestObject.ContainsField(newFieldName));
+			_validTestObject.AddOrSetFieldValue(newFieldName, setValue);
+			Assert.IsTrue(_validTestObject.ContainsField(newFieldName));
+			Assert.AreEqual(setValue, _validTestObject.GetStringField(newFieldName));
+
+			string nullFieldName = "Null";
+			string? nullString = null;
+			Assert.IsFalse(_validTestObject.ContainsField(nullFieldName));
+			_validTestObject.AddOrSetFieldValue(nullFieldName, nullString);
+			Assert.AreEqual("", _validTestObject.GetStringField(nullFieldName));
+		}
+
+		[TestMethod]
+		public void SetStringArrayField()
+		{
+			string fieldName = "Skills";
+			string[] setValue = { "SetValue1", "SetValue2", "SetValue3", "SetValue4"};
+			Assert.IsTrue(_validTestObject.ContainsField(fieldName));
+			_validTestObject.AddOrSetFieldValue(fieldName, setValue);
+			CollectionAssert.AreEqual(setValue, _validTestObject.GetStringArrayField(fieldName));
+
+			string newFieldName = "New";
+			Assert.IsFalse(_validTestObject.ContainsField(newFieldName));
+			_validTestObject.AddOrSetFieldValue(newFieldName, setValue);
+			Assert.IsTrue(_validTestObject.ContainsField(newFieldName));
+			CollectionAssert.AreEqual(setValue, _validTestObject.GetStringArrayField(newFieldName));
+
+			// test array of strings that contains null. Null should be set as ""
+			string nullFieldName = "Null";
+			string?[] nullStringArray = { "String1", null, "String2" };
+			string[] correctNullStringArray = { "String1", "", "String2" };
+			Assert.IsFalse(_validTestObject.ContainsField(nullFieldName));
+			_validTestObject.AddOrSetFieldValue(nullFieldName, nullStringArray);
+			CollectionAssert.AreEqual(correctNullStringArray, _validTestObject.GetStringArrayField(nullFieldName));
+
+			// Test for passing in an empty string array. This should still be added and contained within the object 
+			// This test is very important for ModuleDescriptor 
+			string emptyFieldName = "Empty";
+			Assert.IsFalse(_validTestObject.ContainsField(emptyFieldName));
+			string[] emptyStringArray = Array.Empty<string>();
+			_validTestObject.AddOrSetFieldValue(emptyFieldName, emptyStringArray);
+			Assert.IsTrue(_validTestObject.ContainsField(emptyFieldName));
+			string[] returnArray = _validTestObject.GetStringArrayField(emptyFieldName);
+			Assert.IsTrue(returnArray.Length == 0);
+		}
+
+		[TestMethod]
+		public void SetObjectField()
+		{
+			string fieldName = "StrongestAttack";
+			Assert.IsTrue(_validTestObject.ContainsField(fieldName));
+			EnemyAttack setValue = new EnemyAttack("Ultima", 300, StatusEffect.None);
+			JsonObject setValueObject = setValue.ToJsonObject();
+			_validTestObject.AddOrSetFieldValue(fieldName, setValue.ToJsonObject());
+			Assert.AreEqual(setValueObject, _validTestObject.GetObjectField(fieldName));
+
+			string newFieldName = "New";
+			Assert.IsFalse(_validTestObject.ContainsField(newFieldName));
+			_validTestObject.AddOrSetFieldValue(newFieldName, setValue.ToJsonObject());
+			Assert.AreEqual(setValueObject, _validTestObject.GetObjectField(newFieldName));
+		}
+
+		[TestMethod]
+		public void SetObjectArrayField()
+		{
+			string fieldName = "Attacks";
+			EnemyAttack[] attacks = 
+			{ 
+				new EnemyAttack("Attack1", 40, StatusEffect.Sleep), 
+				new EnemyAttack("Attack2", 100, StatusEffect.Burn)
+			};
+			Assert.IsTrue(_validTestObject.ContainsField(fieldName));
+			JsonObject[] setValue = attacks.Select(x => x.ToJsonObject()).ToArray();
+			_validTestObject.AddOrSetFieldValue(fieldName, setValue);
+			CollectionAssert.AreEqual(setValue, _validTestObject.GetObjectArrayField(fieldName));
+
+			string newFieldName = "New";
+			Assert.IsFalse(_validTestObject.ContainsField(newFieldName));
+			_validTestObject.AddOrSetFieldValue(newFieldName, setValue);
+			Assert.IsTrue(_validTestObject.ContainsField(newFieldName));
+			CollectionAssert.AreEqual(setValue, _validTestObject.GetObjectArrayField(newFieldName));
+		}
+
+		[TestMethod]
+		public void SetEnumField()
+		{
+			string fieldName = "CurrentStatusEffect";
+			StatusEffect setValue = StatusEffect.Burn;
+			Assert.IsTrue(_validTestObject.ContainsField(fieldName));
+			_validTestObject.AddOrSetFieldValue(fieldName, setValue);
+			Assert.AreEqual(setValue, _validTestObject.GetEnumField<StatusEffect>(fieldName));
+
+			string newFieldName = "New";
+			Assert.IsFalse(_validTestObject.ContainsField(newFieldName));
+			_validTestObject.AddOrSetFieldValue(newFieldName, setValue);
+			Assert.IsTrue(_validTestObject.ContainsField(newFieldName));
+			Assert.AreEqual(setValue, _validTestObject.GetEnumField<StatusEffect>(newFieldName));
+		}
+
+		[TestMethod]
+		public void SetEnumArrayField()
+		{
+			string fieldName = "ImmuneStatusEffects";
+			StatusEffect[] setValue = { StatusEffect.Sleep, StatusEffect.Poison, StatusEffect.Frozen, StatusEffect.Burn };
+			Assert.IsTrue(_validTestObject.ContainsField(fieldName));
+			_validTestObject.AddOrSetFieldValue(fieldName, setValue);
+			StatusEffect[]? outValue;
+			Assert.IsTrue(_validTestObject.TryGetEnumArrayField(fieldName, out outValue));
+			CollectionAssert.AreEqual(setValue, outValue);
+
+			string newFieldName = "New";
+			Assert.IsFalse(_validTestObject.ContainsField(newFieldName));
+			_validTestObject.AddOrSetFieldValue(newFieldName, setValue);
+			StatusEffect[]? outValue2;
+			Assert.IsTrue(_validTestObject.TryGetEnumArrayField(newFieldName, out outValue2));
+			CollectionAssert.AreEqual(setValue, outValue2);
+
+			// Test passing in an empty array
+			string emptyFieldName = "Empty";
+			StatusEffect[] setValueEmpty = Array.Empty<StatusEffect>();
+			Assert.IsFalse(_validTestObject.ContainsField(emptyFieldName));
+			_validTestObject.AddOrSetFieldValue(emptyFieldName, setValueEmpty);
+			Assert.IsTrue(_validTestObject.ContainsField(emptyFieldName));
+			StatusEffect[]? outValue4;
+			Assert.IsTrue(_validTestObject.TryGetEnumArrayField<StatusEffect>(emptyFieldName, out outValue4));
+			CollectionAssert.AreEqual(setValueEmpty, outValue4);
+		}
+
+		[TestMethod]
+		public void SetEmptyArrayField()
+		{
+			string emptyFieldName = "Empty";
+			Assert.IsFalse(_validTestObject.ContainsField(emptyFieldName));
+			string[] setValue = Array.Empty<string>();
+			_validTestObject.AddOrSetFieldValue(emptyFieldName, setValue);
+			Assert.IsTrue(_validTestObject.ContainsField(emptyFieldName));
+			string[] returnValue = _validTestObject.GetStringArrayField(emptyFieldName);
+			CollectionAssert.AreEqual(setValue, returnValue);
+			Assert.IsTrue(returnValue.Length == 0);
+		}
+
+		[TestMethod]
+		public void SetBoolField()
+		{
+			string fieldName = "bIsAlive";
+			bool setValue = false;
+			Assert.IsTrue(_validTestObject.ContainsField(fieldName));
+			_validTestObject.AddOrSetFieldValue(fieldName, setValue);
+			Assert.AreEqual(setValue, _validTestObject.GetBoolField(fieldName));
+
+			string newFieldName = "New";
+			Assert.IsFalse(_validTestObject.ContainsField(newFieldName));
+			_validTestObject.AddOrSetFieldValue(newFieldName, setValue);
+			Assert.IsTrue(_validTestObject.ContainsField(newFieldName));
+			Assert.AreEqual(setValue, _validTestObject.GetBoolField(newFieldName));
+		}
+
+		[TestMethod]
+		public void Constructor()
+		{
+			JsonObject obj= new JsonObject();
+			Assert.IsTrue(obj.KeyNames.Count() == 0);
+		}
+
+		[TestMethod]
+		public void ToJsonString()
+		{
+			string stringToMatch = new string(_testJsonText.ToCharArray());
+			string[] lines = stringToMatch.Split(new[] { Environment.NewLine }, StringSplitOptions.None);
+			StringBuilder jsonStringBuilder = new StringBuilder();
+			foreach (string line in lines)
+			{
+				// Utf8JsonWriter uses 2 spaces for indents, we replace them with tabs here 
+				int numLeadingSpaces = line.TakeWhile(x => x == ' ').Count();
+				int numLeadingTabs = numLeadingSpaces / 2;
+				jsonStringBuilder.Append('\t', numLeadingTabs);
+				jsonStringBuilder.AppendLine(line.Substring(numLeadingSpaces));
+			}
+			stringToMatch = jsonStringBuilder.ToString();
+
+			string jsonObjectString = _validTestObject.ToJsonString();
+			Assert.AreEqual(stringToMatch, jsonObjectString);
+		}
+
+		[TestMethod]
+		public void ToStringEscapeCharacters()
+		{
+			string escapeCharacterFieldName = "Escape";
+			// .NET Json serialization escapes certian characters like +,<,>,~,& for security reasons. We need to make sure those come through 
+			string escapeString = "++MyTest+Test<>`&";
+			StringBuilder builder = new StringBuilder();
+			builder.AppendLine("{");
+			builder.AppendLine($"\t\"{escapeCharacterFieldName}\": \"{escapeString}\"");
+			builder.AppendLine("}");
+			string correctString = builder.ToString();
+			JsonObject jsonObject = new JsonObject();
+			jsonObject.AddOrSetFieldValue(escapeCharacterFieldName, escapeString);
+			string jsonString = jsonObject.ToJsonString();
+			Assert.AreEqual(correctString , jsonString);
+		}
+
+		private static Enemy CreateValidJsonDataSource()
+		{
+			Enemy enemy1 = new Enemy();
+			enemy1.Name = "Monster";
+			enemy1.Hp = 100;
+			enemy1.CriticalHitRate = 3.0;
+
+			// Add attacks
+			enemy1.Attacks.Add(new EnemyAttack("Slash", 10));
+			enemy1.Attacks.Add(new EnemyAttack("Ice Blade", 20, StatusEffect.Frozen));
+
+			enemy1.StrongestAttack = new EnemyAttack("Ultimate", 100, StatusEffect.Poison);
+			// Add skills 
+			enemy1.Skills = new List<string>();
+			enemy1.Skills.Add("Skill 1");
+			enemy1.Skills.Add("Skill 2");
+
+			// Add status effects
+			enemy1.ImmuneStatusEffects.Add(StatusEffect.Frozen);
+			enemy1.ImmuneStatusEffects.Add(StatusEffect.Sleep);
+
+			enemy1.bIsAlive = true;
+			enemy1.CurrentStatusEffect = StatusEffect.Poison;
+
+			return enemy1;
+		}
+
+		private JsonObject CreateValidJsonObject()
+		{
+			JsonObject validObject = JsonObject.Parse(_testJsonText);
+			return validObject;
+		}
+
+		private static FileReference CreateTempJsonFile(string fileName, string? fileContent)
+		{
+			string inputFile = Path.Join(s_tempDirectory.FullName, fileName);
+			FileReference inputFileReference = new FileReference(inputFile);
+			File.WriteAllText(inputFile, fileContent);
+			return inputFileReference;
+		}
+
+		private void MakeTestJsonTextInvalid()
+		{
+			_testJsonText += @"blah}blah";
+		}
+
+		private static DirectoryReference CreateTempDir()
+		{
+			string tempDir = Path.Join(Path.GetTempPath(), "epicgames-core-tests-" + Guid.NewGuid().ToString()[..8]);
+			Directory.CreateDirectory(tempDir);
+			return new DirectoryReference(tempDir);
+		}
+	}
+
+	enum StatusEffect
+	{
+		Frozen,
+		Poison,
+		Burn,
+		Sleep,
+		None
+	}
+
+	class Enemy
+	{
+		public string Name { get; set; } = "";
+		public int Hp { get; set; } = 0;
+		public double CriticalHitRate { get; set; } = 0.0;
+		public List<EnemyAttack> Attacks { get; set; } = new List<EnemyAttack>();
+		public EnemyAttack StrongestAttack { get; set; } = null!;
+		public List<StatusEffect> ImmuneStatusEffects { get; set; } = new List<StatusEffect>();
+		public StatusEffect CurrentStatusEffect { get; set; } = StatusEffect.None;
+		public List<string> Skills { get; set; } = new List<string>();
+		public bool bIsAlive { get; set; } = false;
+	}
+
+	class EnemyAttack
+	{
+		public EnemyAttack(string name, int damage, StatusEffect effect = StatusEffect.None)
+		{
+			Name = name;
+			Damage = damage;
+			Effect = effect;
+		}
+
+		public JsonObject ToJsonObject()
+		{
+			JsonObject obj = new JsonObject();
+			obj.AddOrSetFieldValue("Name", Name);
+			obj.AddOrSetFieldValue("Damage", Damage);
+			obj.AddOrSetFieldValue<StatusEffect>("Effect", Effect);
+			return obj;
+		}
+		public string Name { get; set; } = "";
+		public int Damage { get; set; } = 0;
+		public StatusEffect Effect { get; set; } = StatusEffect.None;
+	}
+}
diff --git a/Engine/Source/Programs/Shared/EpicGames.Core/JsonObject.cs b/Engine/Source/Programs/Shared/EpicGames.Core/JsonObject.cs
--- a/Engine/Source/Programs/Shared/EpicGames.Core/JsonObject.cs
+++ b/Engine/Source/Programs/Shared/EpicGames.Core/JsonObject.cs
@@ -5,9 +5,74 @@ using System.Collections.Generic;
 using System.Diagnostics.CodeAnalysis;
 using System.Linq;
 using System.Text.Json;
+using System.Collections.Specialized;
+using System.Collections;
+using System.IO;
+using System.Text;
+
+using SystemJsonObject = System.Text.Json.Nodes.JsonObject;
+using SystemJsonNode = System.Text.Json.Nodes.JsonNode;
+using SystemJsonValue = System.Text.Json.Nodes.JsonValue;
+using SystemJsonArray = System.Text.Json.Nodes.JsonArray;
+using System.Text.Encodings.Web;
+using System.Text.Unicode;
 
 namespace EpicGames.Core
 {
+	public static class OrderedDictionaryExtensions
+	{
+		public static bool TryGetValue(this OrderedDictionary dictionary, object key, [NotNullWhen(true)] out object? value)
+		{
+			if (dictionary.Contains(key))
+			{
+				value = dictionary[key]!;
+				return true;
+			}
+			value = null;
+			return false;
+		}
+
+		private static object? DeepCopyHelper(object? original)
+		{
+			if (original is null)
+			{
+				return null;
+			}
+
+			if (original is OrderedDictionary dictionary)
+			{
+				OrderedDictionary copy = new OrderedDictionary();
+				foreach (DictionaryEntry entry in dictionary)
+				{
+					copy.Add(entry.Key, DeepCopyHelper(entry.Value));
+				}
+				return copy;
+			}
+
+			if (original is Array array)
+			{
+				Array copy = Array.CreateInstance(array.GetType().GetElementType()!, array.Length);
+				for (int index = 0; index < array.Length; ++index)
+				{
+					copy.SetValue(DeepCopyHelper(array.GetValue(index)), index);
+				}
+				return copy;
+			}
+			// It's a value type or a string, it's safe to just copy it 
+			return original;
+		}
+
+		public static OrderedDictionary CreateDeepCopy(this OrderedDictionary dictionary)
+		{
+			OrderedDictionary copy = new OrderedDictionary();
+			foreach (DictionaryEntry entry in dictionary)
+			{
+				copy.Add(entry.Key, DeepCopyHelper(entry.Value));
+			}
+			return copy;
+		}
+	}
+
 	/// <summary>
 	/// Exception thrown for errors parsing JSON files
 	/// </summary>
@@ -29,15 +94,20 @@ namespace EpicGames.Core
 	/// </summary>
 	public class JsonObject
 	{
-		readonly Dictionary<string, object?> _rawObject;
+		readonly OrderedDictionary _rawOrderedObject;
+
+		public JsonObject()
+		{
+			_rawOrderedObject = new OrderedDictionary();
+		}
 
 		/// <summary>
-		/// Construct a JSON object from the raw string -> object dictionary
+		/// Construct a JSON object from the OrderedDictionary obtained from reading a file on disk or parsing valid json text.
 		/// </summary>
 		/// <param name="inRawObject">Raw object parsed from disk</param>
-		public JsonObject(Dictionary<string, object?> inRawObject)
+		private JsonObject(OrderedDictionary inRawObject)
 		{
-			_rawObject = new Dictionary<string, object?>(inRawObject, StringComparer.InvariantCultureIgnoreCase);
+			_rawOrderedObject = inRawObject.CreateDeepCopy();
 		}
 
 		/// <summary>
@@ -46,11 +116,43 @@ namespace EpicGames.Core
 		/// <param name="element"></param>
 		public JsonObject(JsonElement element)
 		{
-			_rawObject = new Dictionary<string, object?>(StringComparer.InvariantCultureIgnoreCase);
+			_rawOrderedObject = new OrderedDictionary();
 			foreach (JsonProperty property in element.EnumerateObject())
 			{
-				_rawObject[property.Name] = ParseElement(property.Value);
+				_rawOrderedObject[property.Name] = ParseElement(property.Value);
+			}
+		}
+
+		public override bool Equals(object? obj)
+		{
+			if (obj is null || GetType() != obj.GetType())
+			{
+				return false;
+			}
+			JsonObject jsonObject = (JsonObject) obj;
+			if (_rawOrderedObject.Count != jsonObject._rawOrderedObject.Count)
+			{
+				return false;
+			}
+			
+			foreach (DictionaryEntry entry in _rawOrderedObject)
+			{
+				if (!jsonObject._rawOrderedObject.Contains(entry.Key))
+				{
+					return false;
+				}
+				if (!entry.Value!.Equals(jsonObject._rawOrderedObject[entry.Key!]))
+				{
+					return false;
+				}
 			}
+
+			return true;
+		}
+
+		public override int GetHashCode()
+		{
+			return HashCode.Combine(_rawOrderedObject);
 		}
 
 		/// <summary>
@@ -67,7 +169,12 @@ namespace EpicGames.Core
 				case JsonValueKind.Number:
 					return element.GetDouble();
 				case JsonValueKind.Object:
-					return element.EnumerateObject().ToDictionary(x => x.Name, x => ParseElement(x.Value));
+					OrderedDictionary dictionary = new OrderedDictionary();
+					foreach (JsonProperty property in element.EnumerateObject())
+						{
+						dictionary.Add(property.Name, ParseElement(property.Value));
+					}
+					return dictionary;
 				case JsonValueKind.String:
 					return element.GetString();
 				case JsonValueKind.False:
@@ -81,6 +188,42 @@ namespace EpicGames.Core
 			}
 		}
 
+		private SystemJsonNode ToJsonNode(object? obj)
+		{
+			// All values in the JsonObject are either parsed from a string, read from a file or set/added with the API
+			// All values at this point must be supported and the correct types 
+			switch (obj)
+			{
+				case OrderedDictionary objDictionary:
+					SystemJsonObject dictionaryJson= new SystemJsonObject();
+					foreach (object? key in objDictionary.Keys)
+					{
+						dictionaryJson.Add((string) key, ToJsonNode(objDictionary[key]));
+					}
+					return dictionaryJson;
+				case Array objArray:
+					SystemJsonArray tempJsonArray = new SystemJsonArray();
+					foreach (object? element in objArray)
+					{
+						tempJsonArray.Add(ToJsonNode(element));
+					}
+					return tempJsonArray;
+				default:
+					// We support null as per the json spec 
+					return SystemJsonValue.Create(obj)!;
+			}
+		}
+
+		public SystemJsonObject ToSystemJsonObject()
+		{
+			SystemJsonObject returnObject = new SystemJsonObject();
+			foreach (DictionaryEntry entry in _rawOrderedObject )
+			{
+				returnObject.Add((string) entry.Key, ToJsonNode(entry.Value));
+			}
+			return returnObject;
+		}
+
 		/// <summary>
 		/// Read a JSON file from disk and construct a JsonObject from it
 		/// </summary>
@@ -124,8 +267,15 @@ namespace EpicGames.Core
 		/// <returns>New JsonObject instance</returns>
 		public static JsonObject Parse(string text)
 		{
-			JsonDocument document = JsonDocument.Parse(text, new JsonDocumentOptions { AllowTrailingCommas = true });
-			return new JsonObject(document.RootElement);
+			try
+			{
+				JsonDocument document = JsonDocument.Parse(text, new JsonDocumentOptions { AllowTrailingCommas = true });
+				return new JsonObject(document.RootElement);
+			}
+			catch (Exception ex)
+			{
+				throw new JsonParseException("Failed to parse json text '{0}'. {1}", text, ex.Message);
+			}
 		}
 
 		/// <summary>
@@ -151,7 +301,7 @@ namespace EpicGames.Core
 		/// <summary>
 		/// List of key names in this object
 		/// </summary>
-		public IEnumerable<string> KeyNames => _rawObject.Keys;
+		public IEnumerable<string> KeyNames => _rawOrderedObject.Keys.Cast<string>();
 
 		/// <summary>
 		/// Gets a string field by the given name from the object, throwing an exception if it is not there or cannot be parsed.
@@ -177,7 +327,7 @@ namespace EpicGames.Core
 		public bool TryGetStringField(string fieldName, [NotNullWhen(true)] out string? result)
 		{
 			object? rawValue;
-			if (_rawObject.TryGetValue(fieldName, out rawValue) && (rawValue is string strValue))
+			if (_rawOrderedObject.TryGetValue(fieldName, out rawValue) && (rawValue is string strValue))
 			{
 				result = strValue;
 				return true;
@@ -213,7 +363,7 @@ namespace EpicGames.Core
 		public bool TryGetStringArrayField(string fieldName, [NotNullWhen(true)] out string[]? result)
 		{
 			object? rawValue;
-			if (_rawObject.TryGetValue(fieldName, out rawValue) && (rawValue is IEnumerable<object> enumValue) && enumValue.All(x => x is string))
+			if (_rawOrderedObject.TryGetValue(fieldName, out rawValue) && (rawValue is IEnumerable<object> enumValue) && enumValue.All(x => x is string))
 			{
 				result = enumValue.Select(x => (string)x).ToArray();
 				return true;
@@ -249,7 +399,7 @@ namespace EpicGames.Core
 		public bool TryGetBoolField(string fieldName, out bool result)
 		{
 			object? rawValue;
-			if (_rawObject.TryGetValue(fieldName, out rawValue) && (rawValue is bool boolValue))
+			if (_rawOrderedObject.TryGetValue(fieldName, out rawValue) && (rawValue is bool boolValue))
 			{
 				result = boolValue;
 				return true;
@@ -285,7 +435,7 @@ namespace EpicGames.Core
 		public bool TryGetIntegerField(string fieldName, out int result)
 		{
 			object? rawValue;
-			if (!_rawObject.TryGetValue(fieldName, out rawValue) || !Int32.TryParse(rawValue?.ToString(), out result))
+			if (!_rawOrderedObject.TryGetValue(fieldName, out rawValue) || !Int32.TryParse(rawValue?.ToString(), out result))
 			{
 				result = 0;
 				return false;
@@ -302,7 +452,7 @@ namespace EpicGames.Core
 		public bool TryGetUnsignedIntegerField(string fieldName, out uint result)
 		{
 			object? rawValue;
-			if (!_rawObject.TryGetValue(fieldName, out rawValue) || !UInt32.TryParse(rawValue?.ToString(), out result))
+			if (!_rawOrderedObject.TryGetValue(fieldName, out rawValue) || !UInt32.TryParse(rawValue?.ToString(), out result))
 			{
 				result = 0;
 				return false;
@@ -334,7 +484,7 @@ namespace EpicGames.Core
 		public bool TryGetDoubleField(string fieldName, out double result)
 		{
 			object? rawValue;
-			if (!_rawObject.TryGetValue(fieldName, out rawValue) || !Double.TryParse(rawValue?.ToString(), out result))
+			if (!_rawOrderedObject.TryGetValue(fieldName, out rawValue) || !Double.TryParse(rawValue?.ToString(), out result))
 			{
 				result = 0.0;
 				return false;
@@ -427,9 +577,9 @@ namespace EpicGames.Core
 		public bool TryGetObjectField(string fieldName, [NotNullWhen(true)] out JsonObject? result)
 		{
 			object? rawValue;
-			if (_rawObject.TryGetValue(fieldName, out rawValue) && (rawValue is Dictionary<string, object?> dictValue))
+			if (_rawOrderedObject.TryGetValue(fieldName, out rawValue) && (rawValue is OrderedDictionary orderedDictValue))
 			{
-				result = new JsonObject(dictValue);
+				result = new JsonObject(orderedDictValue);
 				return true;
 			}
 			else
@@ -463,9 +613,9 @@ namespace EpicGames.Core
 		public bool TryGetObjectArrayField(string fieldName, [NotNullWhen(true)] out JsonObject[]? result)
 		{
 			object? rawValue;
-			if (_rawObject.TryGetValue(fieldName, out rawValue) && (rawValue is IEnumerable<object> enumValue) && enumValue.All(x => x is Dictionary<string, object>))
+			if (_rawOrderedObject.TryGetValue(fieldName, out rawValue) && (rawValue is IEnumerable<object> enumValue) && enumValue.All(x => x is OrderedDictionary))
 			{
-				result = enumValue.Select(x => new JsonObject((Dictionary<string, object?>)x)).ToArray();
+				result = enumValue.Select(x => new JsonObject((OrderedDictionary)x)).ToArray();
 				return true;
 			}
 			else
@@ -474,5 +624,170 @@ namespace EpicGames.Core
 				return false;
 			}
 		}
+
+		/// <summary>
+		/// Checks if the provided field exists in this Json object.
+		/// </summary>
+		/// <param name="fieldName">Name of the field to check if it is contained. </param>
+		/// <returns>True if the field exists in the JsonObject, false otherwise</returns>
+		public bool ContainsField(string fieldName)
+		{
+			return _rawOrderedObject.Contains(fieldName);
+		}
+
+		public void AddOrSetFieldValue(string fieldName, int value)
+		{
+			if (String.IsNullOrEmpty(fieldName))
+			{
+				return;
+			}
+			_rawOrderedObject[fieldName] = Convert.ToDouble(value);
+		}
+
+		public void AddOrSetFieldValue(string fieldName, uint value)
+		{
+			if (String.IsNullOrEmpty(fieldName))
+			{
+				return;
+			}
+			_rawOrderedObject[fieldName] = Convert.ToDouble(value);
+		}
+
+		public void AddOrSetFieldValue(string fieldName, double value)
+		{
+			if (String.IsNullOrEmpty(fieldName))
+			{
+				return;
+			}
+			_rawOrderedObject[fieldName] = value;
+		}
+
+		public void AddOrSetFieldValue(string fieldName, bool value)
+		{
+			if (String.IsNullOrEmpty(fieldName))
+			{
+				return;
+			}
+			_rawOrderedObject[fieldName] = value;
+		}
+
+		public void AddOrSetFieldValue(string fieldName, string? value)
+		{
+			if (String.IsNullOrEmpty(fieldName))
+			{
+				return;
+			}
+			// We set null strings to empty strings to follow the behavior of EpicGames.Core.JsonWriter
+			value ??= "";
+			_rawOrderedObject[fieldName] = value;
+		}
+
+		public void AddOrSetFieldValue<T>(string fieldName, T value) where T : Enum
+		{
+			if (String.IsNullOrEmpty(fieldName))
+			{
+				return;
+			}
+			
+			_rawOrderedObject[fieldName] = value.ToString();
+		}
+
+		
+		public void AddOrSetFieldValue(string fieldName, JsonObject? value)
+		{
+			if (String.IsNullOrEmpty(fieldName))
+			{
+				return;
+			}
+
+			_rawOrderedObject[fieldName] = value?._rawOrderedObject.CreateDeepCopy();
+		}
+
+		public void AddOrSetFieldValue(string fieldName, string?[]? value)
+		{
+			if (String.IsNullOrEmpty(fieldName))
+			{
+				return;
+			}
+			if (value is null)
+			{
+				_rawOrderedObject[fieldName] = value;
+				return;
+			}
+			// Contents of the strings array should never be null and should be "" instead to match behavior EpicGames.Core.JsonWriter
+			string[] stringArray = value.Select(x => x ?? "").ToArray();
+			_rawOrderedObject[fieldName] = stringArray;
+		}
+
+		public void AddOrSetFieldValue(string fieldName, JsonObject[]? value)
+		{
+			if (String.IsNullOrEmpty(fieldName))
+			{
+				return;
+			}
+			if (value is null)
+			{
+				_rawOrderedObject[fieldName] = null;
+				return;
+			}
+			List<OrderedDictionary?> objList = new List<OrderedDictionary?>();
+			foreach (JsonObject? obj in value)
+			{
+				objList.Add(obj?._rawOrderedObject.CreateDeepCopy());
+			}
+			_rawOrderedObject[fieldName] = objList.ToArray();
+		}
+
+		public void AddOrSetFieldValue<T>(string fieldName, T[]? value) where T : Enum
+		{
+			if (String.IsNullOrEmpty(fieldName))
+			{
+				return;
+			}
+			if (value is null)
+			{
+				_rawOrderedObject[fieldName] = null;
+				return;
+			}
+
+			string[] stringArray = value.Select(x => x.ToString()!).ToArray();
+			_rawOrderedObject[fieldName] = stringArray;
+		}
+		/// <summary>
+		/// Converts this Json Object to a string representation. This follows formatting of UE .uplugin files with 4 spaces for tabs and indentation enabled.
+		/// IMPORTANT: If this JsonObject contains HTML, the returned string should NOT be used directly for HTML or a script. Read the note below.
+		/// </summary>
+		/// <returns>The formatted, prettified string representation of this Json Object.</returns>
+		public string ToJsonString()
+		{
+			SystemJsonObject jsonObjectToWrite = ToSystemJsonObject();
+			JsonWriterOptions options = new JsonWriterOptions();
+			options.Indented = true;
+			// IMPORTANT: Utf8JsonWriter blocks certain characters like +, &, <, >,` from being escaped in a global block list
+			// Best way around is to use the relaxed JavaScriptEncoder.UnsafeRelaxedJsonEscaping.
+			// However this can have security implications if this string is ever written to an HTML page or script
+			// https://learn.microsoft.com/en-us/dotnet/standard/serialization/system-text-json/character-encoding
+			options.Encoder = JavaScriptEncoder.UnsafeRelaxedJsonEscaping;
+			using MemoryStream jsonMemoryStream = new MemoryStream();
+			using (Utf8JsonWriter writer = new Utf8JsonWriter(jsonMemoryStream, options))
+			{
+				jsonObjectToWrite.WriteTo(writer);
+				writer.Flush();
+			}
+			// The Utf8JsonWriter doesn't format json the same as we want in UE, we massage the string to meet our standards 
+			string jsonString = Encoding.UTF8.GetString(jsonMemoryStream.ToArray());
+			string[] lines = jsonString.Split(new[] { Environment.NewLine }, StringSplitOptions.None);
+			StringBuilder jsonStringBuilder = new StringBuilder();
+			foreach (string line in lines)
+			{
+				// Utf8JsonWriter uses 2 spaces for indents, we replace them with tabs here 
+				int numLeadingSpaces = line.TakeWhile(x => x == ' ').Count();
+				int numLeadingTabs = numLeadingSpaces / 2;
+				jsonStringBuilder.Append('\t', numLeadingTabs);
+				jsonStringBuilder.AppendLine(line.Substring(numLeadingSpaces));
+			}
+
+			return jsonStringBuilder.ToString();
+		}
 	}
 }
diff --git a/Engine/Source/Programs/UnrealBuildTool.Tests/PluginDescriptorTests.cs b/Engine/Source/Programs/UnrealBuildTool.Tests/PluginDescriptorTests.cs
new file mode 100644
--- /dev/null
+++ b/Engine/Source/Programs/UnrealBuildTool.Tests/PluginDescriptorTests.cs
@@ -0,0 +1,286 @@
+// Copyright Epic Games, Inc. All Rights Reserved.
+
+using System;
+using System.Collections.Generic;
+using System.Diagnostics;
+using System.IO;
+using System.Linq;
+using System.Text;
+using EpicGames.Core;
+using UnrealBuildTool;
+using Microsoft.VisualStudio.TestTools.UnitTesting;
+
+namespace UnrealBuildToolTests
+{
+	[TestClass]
+	public class PluginDescriptorTests
+	{
+		private static DirectoryReference s_tempDirectory;
+
+		private static DirectoryReference CreateTempDir()
+		{
+			string tempDir = Path.Join(Path.GetTempPath(), "epicgames-core-tests-" + Guid.NewGuid().ToString()[..8]);
+			Directory.CreateDirectory(tempDir);
+			return new DirectoryReference(tempDir);
+		}
+
+		[ClassInitialize]
+		public static void Setup(TestContext testContext)
+		{
+			PluginDescriptorTests.s_tempDirectory = CreateTempDir();
+		}
+
+		[ClassCleanup]
+		public static void TearDown()
+		{
+			if (Directory.Exists(PluginDescriptorTests.s_tempDirectory.FullName))
+			{
+				Directory.Delete(PluginDescriptorTests.s_tempDirectory.FullName, true);
+			}
+		}
+
+		[TestMethod]
+		public void ReadPluginCustomFields()
+		{
+			string json = CreateUPluginWithCustomFields();
+			string filename = "CustomFields.uplugin";
+			FileReference fileToWrite = CreateTempUPluginFile(filename, json);
+			PluginDescriptor.FromFile(fileToWrite);
+		}
+
+		[TestMethod]
+		public void ReadUpluginMissingFileVersion()
+		{
+			string json = CreateUpluginWithMissingVersion();
+			string filename = "MissingFileVersionField.uplugin";
+			FileReference fileToWrite = CreateTempUPluginFile(filename, json);
+			Assert.ThrowsException<BuildException>(() => PluginDescriptor.FromFile(fileToWrite));
+		}
+
+		[TestMethod]
+		public void ReadUpluginInvalidVersion()
+		{
+			string json = CreateUpluginWithInvalidVersion();
+			string filename = "InvalidFileVersionField.uplugin";
+			FileReference fileToWrite = CreateTempUPluginFile(filename, json);
+			Assert.ThrowsException<BuildException>(() => PluginDescriptor.FromFile(fileToWrite));
+		}
+
+		[TestMethod]
+		public void Save2CustomFieldsMissingDefaults()
+		{
+			string json = CreateUPluginWithCustomFields();
+			string filename = "CustomFields-Save.uplugin";
+			FileReference fileToWrite = CreateTempUPluginFile(filename, json);
+			PluginDescriptor descriptor = PluginDescriptor.FromFile(fileToWrite);
+			string fileToSave = GetAbsolutePathToTempFile("CustomFields-SaveResults.uplugin");
+			descriptor.Save2(fileToSave);
+
+			string saveContents = File.ReadAllText(fileToSave);
+			JsonObject writtenObject = JsonObject.Parse(json);
+			JsonObject savedObject = JsonObject.Parse(saveContents);
+			IEnumerable<string> writtenObjectKeys = writtenObject.KeyNames;
+			// At this point, the written object should contain all the custom fields 
+			// The saved object should have all the custom fields as well as the missing default fields appended
+			// This means we can't do a simple equality check to see if the objects are the same 
+			// We check to make sure that all of the keys in the written object exist in the saved object as a short hand to check that all the custom fields have been saved and carried over to the saved object 
+			foreach (string key in writtenObjectKeys)
+			{
+				Assert.IsTrue(savedObject.ContainsField(key));
+			}
+		}
+
+		[TestMethod]
+		public void Save2DefaultFields()
+		{
+			string json = CreateUpluginWithDefaultFields();
+			string filename = "DefaultFields-Save.uplugin";
+			FileReference fileToWrite = CreateTempUPluginFile(filename, json);
+			PluginDescriptor descriptor = PluginDescriptor.FromFile(fileToWrite);
+			string fileToSave = GetAbsolutePathToTempFile("DefaultFields-SaveResults.uplugin");
+			descriptor.Save2(fileToSave);
+			string saveContents = File.ReadAllText(fileToSave);
+
+			// Nothing should change between the format of the json and what was written to disk 
+			Assert.AreEqual(json, saveContents);
+		}
+
+		[TestMethod]
+		public void Save2CompareWithSave()
+		{
+			string json = CreateUpluginWithDefaultFields();
+			string filename = "DefaultFields-Save2.uplugin";
+			FileReference fileToWrite = CreateTempUPluginFile(filename, json);
+
+			PluginDescriptor descriptor = PluginDescriptor.FromFile(fileToWrite);
+			string fileToSave = GetAbsolutePathToTempFile("DefaultFields-SaveResults.uplugin");
+			string fileToSave2 = GetAbsolutePathToTempFile("DefaultFields-Save2Results.uplugin");
+			descriptor.Save2(fileToSave);
+			descriptor.Save2(fileToSave2);
+
+			string saveContents = File.ReadAllText(fileToSave);
+			string save2Contents = File.ReadAllText(fileToSave2);
+
+			// Nothing should change between the format of the json and what was written to disk 
+			Assert.AreEqual(saveContents, save2Contents);
+		}
+
+		private static string GetAbsolutePathToTempFile(string fileName)
+		{
+			if (String.IsNullOrEmpty(fileName))
+			{
+				return "";
+			}
+			return Path.Join(s_tempDirectory.FullName, fileName);
+		}
+
+		private static FileReference CreateTempUPluginFile(string fileName, string fileContent)
+		{
+			string extension = Path.GetExtension(fileName);
+			Debug.Assert(!String.IsNullOrEmpty(extension));
+			Debug.Assert(extension == ".uplugin");
+			string inputFile = Path.Join(s_tempDirectory.FullName, fileName);
+			FileReference inputFileReference = new FileReference(inputFile);
+			File.WriteAllText(inputFile, fileContent);
+			return inputFileReference;
+		}
+
+		private static string CreateUPluginWithCustomFields()
+		{
+			string json = @"
+				{
+					""FileVersion"": 3,
+					""CanContainContent"": true,
+					""ExplicitlyLoaded"": true,
+					""EnabledByDefault"": false,
+					""EditorCustomVirtualPath"": ""MyPlugin"",
+					""CustomField1"": ""400.0"",
+					""CustomField2"": ""CustomValue2"",
+					""CustomField3"": false,
+					""Modules"": [
+						{
+							""Name"": ""Module1"",
+							""Type"": ""Runtime"",
+							""LoadingPhase"": ""Default""
+						},
+						{
+							""Name"": ""Module2"",
+							""Type"": ""ClientOnly"",
+							""LoadingPhase"": ""Default""
+						}
+					],
+					""Plugins"": [
+						{
+							""Name"": ""Plugin1"",
+							""Enabled"": true
+						},
+						{
+							""Name"": ""Plugin2"",
+							""Enabled"": true
+						},
+						{
+							""Name"": ""Plugin3"",
+							""Enabled"": true
+						},
+						{
+							""Name"": ""Plugin4"",
+							""Enabled"": true
+						},
+						{
+							""Name"": ""Plugin5"",
+							""Enabled"": true
+						}
+					]
+				}";
+
+			return FixVerbatimStringIdentation(json);
+		}
+
+		private static string CreateUpluginWithDefaultFields()
+		{
+			string json = @"
+				{
+					""FileVersion"": 3,
+					""Version"": 1,
+					""VersionName"": ""1.0"",
+					""FriendlyName"": ""Enhanced Input"",
+					""Description"": ""Input handling that allows for contextual and dynamic mappings."",
+					""Category"": ""Input"",
+					""CreatedBy"": ""Epic Games, Inc."",
+					""CreatedByURL"": ""https://epicgames.com"",
+					""DocsURL"": ""https://docs.unrealengine.com/en-US/enhanced-input-in-unreal-engine/"",
+					""MarketplaceURL"": """",
+					""SupportURL"": """",
+					""EnabledByDefault"": true,
+					""CanContainContent"": false,
+					""IsBetaVersion"": false,
+					""Installed"": false,
+					""Modules"": [
+						{
+							""Name"": ""EnhancedInput"",
+							""Type"": ""Runtime"",
+							""LoadingPhase"": ""PreDefault""
+						},
+						{
+							""Name"": ""InputBlueprintNodes"",
+							""Type"": ""UncookedOnly"",
+							""LoadingPhase"": ""Default""
+						},
+						{
+							""Name"": ""InputEditor"",
+							""Type"": ""Editor"",
+							""LoadingPhase"": ""Default""
+						}
+					]
+				}";
+
+			return FixVerbatimStringIdentation(json);
+		}
+
+		private static string CreateUpluginWithMissingVersion()
+		{
+			string json = @"
+				{ 
+					""CustomField1"": ""Value1""
+				}";
+			return FixVerbatimStringIdentation(json);
+		}
+
+		private static string CreateUpluginWithInvalidVersion()
+		{
+			string json = @"
+				{ 
+					""FileVersion"": 4
+				}";
+			return FixVerbatimStringIdentation(json);
+		}
+
+		private static string FixVerbatimStringIdentation(string verbatimString)
+	{
+			if (String.IsNullOrEmpty(verbatimString))
+			{
+				return verbatimString;
+			}
+			string[] lines = verbatimString.Split(new[] { Environment.NewLine }, StringSplitOptions.None);
+			StringBuilder jsonStringBuilder = new StringBuilder();
+			int numberOfTabs = 0;
+			bool bFoundTabLevel = false;
+			foreach (string line in lines)
+			{
+				// Should only apply to the first line. Skip 
+				if (String.IsNullOrEmpty(line))
+				{
+					continue;
+				}
+				// Now we're on the first line with the { we're just going to assume that this is the first line 
+				if (!bFoundTabLevel)
+				{
+					numberOfTabs= line.TakeWhile(x => x == '\t').Count();
+					bFoundTabLevel = true;
+				}
+				jsonStringBuilder.AppendLine(line.Substring(numberOfTabs));
+			}
+			return jsonStringBuilder.ToString();
+	}
+	}
+}
diff --git a/Engine/Source/Programs/UnrealBuildTool/System/CustomBuildSteps.cs b/Engine/Source/Programs/UnrealBuildTool/System/CustomBuildSteps.cs
--- a/Engine/Source/Programs/UnrealBuildTool/System/CustomBuildSteps.cs
+++ b/Engine/Source/Programs/UnrealBuildTool/System/CustomBuildSteps.cs
@@ -72,6 +72,20 @@ namespace UnrealBuildTool
 			Writer.WriteObjectEnd();
 		}
 
+		/// <summary>
+		/// Converts this object to a JsonObject to be saved or manipulated.
+		/// </summary>
+		/// <returns>The JsonObject representation of this object.</returns>
+		public JsonObject ToJsonObject()
+		{
+			JsonObject CustomBuildStepObject = new JsonObject();
+			foreach (KeyValuePair<UnrealTargetPlatform, string[]> Pair in HostPlatformToCommands.OrderBy(x => x.Key.ToString()))
+			{
+				CustomBuildStepObject.AddOrSetFieldValue(Pair.Key.ToString(), Pair.Value);
+			}
+			return CustomBuildStepObject;
+		}
+
 		/// <summary>
 		/// Tries to get the commands for a given host platform
 		/// </summary>
diff --git a/Engine/Source/Programs/UnrealBuildTool/System/LocalizationDescriptor.cs b/Engine/Source/Programs/UnrealBuildTool/System/LocalizationDescriptor.cs
--- a/Engine/Source/Programs/UnrealBuildTool/System/LocalizationDescriptor.cs
+++ b/Engine/Source/Programs/UnrealBuildTool/System/LocalizationDescriptor.cs
@@ -1,5 +1,6 @@
 // Copyright Epic Games, Inc. All Rights Reserved.
 
+using System.Linq;
 using System.Diagnostics;
 using EpicGames.Core;
 
@@ -90,6 +91,15 @@ namespace UnrealBuildTool
 			Writer.WriteObjectEnd();
 		}
 
+		JsonObject ToJsonObject()
+		{
+			JsonObject localizationTargetObject= new JsonObject();
+			localizationTargetObject.AddOrSetFieldValue("Name", Name);
+			localizationTargetObject.AddOrSetFieldValue("LoadingPolicy", LoadingPolicy.ToString());
+			
+			return localizationTargetObject;
+		}
+
 		/// <summary>
 		/// Write an array of target descriptors
 		/// </summary>
@@ -108,5 +118,20 @@ namespace UnrealBuildTool
 				Writer.WriteArrayEnd();
 			}
 		}
+
+		/// <summary>
+		/// Updates a JsonObject with an array of localization target descriptors.
+		/// </summary>
+		/// <param name="InObject">The Json object to update.</param>
+		/// <param name="Name">Name of the array</param>
+		/// <param name="Targets">Array of targets</param>
+		public static void UpdateJson(JsonObject InObject, string Name, LocalizationTargetDescriptor[]? Targets)
+		{
+			if (Targets != null && Targets.Length > 0)
+			{
+				JsonObject[] JsonObjects = Targets.Select(X => X.ToJsonObject()).ToArray();
+				InObject.AddOrSetFieldValue(Name, JsonObjects);
+			}
+		}
 	}
 }
diff --git a/Engine/Source/Programs/UnrealBuildTool/System/ModuleDescriptor.cs b/Engine/Source/Programs/UnrealBuildTool/System/ModuleDescriptor.cs
--- a/Engine/Source/Programs/UnrealBuildTool/System/ModuleDescriptor.cs
+++ b/Engine/Source/Programs/UnrealBuildTool/System/ModuleDescriptor.cs
@@ -433,6 +433,64 @@ namespace UnrealBuildTool
 			Writer.WriteObjectEnd();
 		}
 
+		JsonObject ToJsonObject()
+		{
+			JsonObject ModuleObject = new JsonObject();
+			ModuleObject.AddOrSetFieldValue("Name", Name);
+			ModuleObject.AddOrSetFieldValue("Type", Type.ToString());
+			ModuleObject.AddOrSetFieldValue("LoadingPhase", LoadingPhase.ToString());
+			// important note: we don't check the length of the platform allow list, because if an unknown platform was read in, but was not valid, the 
+			// list will exist but be empty. We don't want to remove the allow list completely, because that would allow this module on all platforms,
+			// which will not be the desired effect
+			if (PlatformAllowList != null)
+			{
+				string[] PlatformAllowListStringArray = PlatformAllowList.Select(X => X.ToString()).ToArray();
+				ModuleObject.AddOrSetFieldValue("PlatformAllowList", PlatformAllowListStringArray);
+			}
+			if (PlatformDenyList != null && PlatformDenyList.Count > 0)
+			{
+				string[] PlatformDenyListStringArray = PlatformDenyList.Select(X => X.ToString()).ToArray();
+				ModuleObject.AddOrSetFieldValue("PlatformDenyList", PlatformDenyListStringArray);
+			}
+			if (TargetAllowList != null && TargetAllowList.Length > 0)
+			{
+				string[] TargetAllowListStringArray = TargetAllowList.Select(X => X.ToString()).ToArray();
+				ModuleObject.AddOrSetFieldValue("TargetAllowList", TargetAllowListStringArray);
+			}
+			if (TargetDenyList != null && TargetDenyList.Length > 0)
+			{
+				string[] TargetDenyListStringArray = TargetDenyList.Select(X => X.ToString()).ToArray();
+				ModuleObject.AddOrSetFieldValue("TargetDenyList", TargetDenyListStringArray);
+			}
+			if (TargetConfigurationAllowList != null && TargetConfigurationAllowList.Length > 0)
+			{
+				string[] TargetConfigurationAllowListStringArray = TargetConfigurationAllowList.Select(X => X.ToString()).ToArray();
+				ModuleObject.AddOrSetFieldValue("TargetConfigurationAllowList", TargetConfigurationAllowListStringArray);
+			}
+			if (TargetConfigurationDenyList != null && TargetConfigurationDenyList.Length > 0)
+			{
+				string[] TargetConfigurationDenyListStringArray = TargetConfigurationDenyList.Select(X => X.ToString()).ToArray();
+				ModuleObject.AddOrSetFieldValue("TargetConfigurationDenyList", TargetConfigurationDenyListStringArray);
+			}
+			if (ProgramAllowList != null && ProgramAllowList.Length > 0)
+			{
+				ModuleObject.AddOrSetFieldValue("ProgramAllowList", ProgramAllowList);
+			}
+			if (ProgramDenyList != null && ProgramDenyList.Length > 0)
+			{
+				ModuleObject.AddOrSetFieldValue("ProgramDenyList", ProgramDenyList);
+			}
+			if (AdditionalDependencies != null && AdditionalDependencies.Length > 0)
+			{
+				ModuleObject.AddOrSetFieldValue("AdditionalDependencies", AdditionalDependencies);
+			}
+			if (bHasExplicitPlatforms)
+			{
+				ModuleObject.AddOrSetFieldValue("HasExplicitPlatforms", bHasExplicitPlatforms);
+			}
+			return ModuleObject;
+		}
+
 		/// <summary>
 		/// Write an array of module descriptors
 		/// </summary>
@@ -452,6 +510,21 @@ namespace UnrealBuildTool
 			}
 		}
 
+		/// <summary>
+		/// Updates a JsonObject with an array of module descriptors 
+		/// </summary>
+		/// <param name="InObject">The JsonObject to update.</param>
+		/// <param name="Name">Name of the array</param>
+		/// <param name="Modules">Array of modules</param>
+		public static void UpdateJson(JsonObject InObject, string Name, ModuleDescriptor[]? Modules)
+		{
+			if (Modules != null && Modules.Length > 0)
+			{
+				JsonObject[] JsonObjects = Modules.Select(X => X.ToJsonObject()).ToArray();
+				InObject.AddOrSetFieldValue(Name, JsonObjects);
+			}
+		}
+
 		/// <summary>
 		/// Produces any warnings and errors for the module settings
 		/// </summary>
diff --git a/Engine/Source/Programs/UnrealBuildTool/System/PluginDescriptor.cs b/Engine/Source/Programs/UnrealBuildTool/System/PluginDescriptor.cs
--- a/Engine/Source/Programs/UnrealBuildTool/System/PluginDescriptor.cs
+++ b/Engine/Source/Programs/UnrealBuildTool/System/PluginDescriptor.cs
@@ -1,9 +1,11 @@
 // Copyright Epic Games, Inc. All Rights Reserved.
 
 using System;
-using System.Collections.Generic;
+using System.IO;
 using System.Linq;
+using System.Collections.Generic;
 using EpicGames.Core;
+using EpicGames.Serialization;
 
 namespace UnrealBuildTool
 {
@@ -222,12 +224,10 @@ namespace UnrealBuildTool
 		public String[]? DisallowedPlugins;
 
 		/// <summary>
-		/// Private constructor. This object should not be created directly; read it from disk using FromFile() instead.
+		/// The JsonObject created from reading a .uplugin on disk or from parsing a json text 
+		/// This preserves the order of all the fields from the source json as well as account for any custom fields.
 		/// </summary>
-		private PluginDescriptor()
-		{
-			FileVersion = (int)PluginDescriptorVersion.Latest;
-		}
+		private readonly JsonObject CachedJson;
 
 		/// <summary>
 		/// Reads a plugin descriptor from a json object
@@ -237,6 +237,7 @@ namespace UnrealBuildTool
 		/// <returns>New plugin descriptor</returns>
 		public PluginDescriptor(JsonObject RawObject, FileReference PluginPath)
 		{
+			CachedJson = RawObject;
 			// Read the version
 			if (!RawObject.TryGetIntegerField("FileVersion", out FileVersion))
 			{
@@ -356,7 +357,7 @@ namespace UnrealBuildTool
 		}
 
 		/// <summary>
-		/// Creates a plugin descriptor from a file on disk
+		/// Creates a plugin descriptor from a file on disk preserving all custom fields in the file.
 		/// </summary>
 		/// <param name="FileName">The filename to read</param>
 		/// <returns>New plugin descriptor</returns>
@@ -376,7 +377,7 @@ namespace UnrealBuildTool
 		}
 
 		/// <summary>
-		/// Saves the descriptor to disk
+		/// Saves the descriptor to disk. This only saves the default fields in a .uplugin and does not account for cusotm fields.
 		/// </summary>
 		/// <param name="FileName">The filename to write to</param>
 		public void Save(string FileName)
@@ -389,6 +390,18 @@ namespace UnrealBuildTool
 			}
 		}
 
+		/// <summary>
+		/// Saves the descriptor to disk preserving all custom fields that were read in.
+		/// </summary>
+		/// <param name="fileName">The filename to write to</param>
+		public void Save2(string fileName)
+		{
+			// @TODO: This should replace all instances of Save() at some point in the future. There's just still a lot of references to test and refactor that needs to be verified. 
+			UpdateJson();
+			string jsonString = CachedJson.ToJsonString();
+			File.WriteAllText(fileName, jsonString);
+		}
+
 		/// <summary>
 		/// Writes the plugin descriptor to an existing Json writer
 		/// </summary>
@@ -507,6 +520,114 @@ namespace UnrealBuildTool
 			}
 		}
 
+		private void UpdateJson()
+		{
+			CachedJson.AddOrSetFieldValue("FileVersion", (int)ProjectDescriptorVersion.Latest);
+			CachedJson.AddOrSetFieldValue("Version", Version);
+			CachedJson.AddOrSetFieldValue("VersionName", VersionName);
+			CachedJson.AddOrSetFieldValue("FriendlyName", FriendlyName);
+			CachedJson.AddOrSetFieldValue("Description", Description);
+			CachedJson.AddOrSetFieldValue("Category", Category);
+			CachedJson.AddOrSetFieldValue("CreatedBy", CreatedBy);
+			CachedJson.AddOrSetFieldValue("CreatedByURL", CreatedByURL);
+			CachedJson.AddOrSetFieldValue("DocsURL", DocsURL);
+			CachedJson.AddOrSetFieldValue("MarketplaceURL", MarketplaceURL);
+			CachedJson.AddOrSetFieldValue("SupportURL", SupportURL);
+			if (!String.IsNullOrEmpty(EngineVersion))
+			{
+				CachedJson.AddOrSetFieldValue("EngineVersion", EngineVersion);
+			}
+			if (!String.IsNullOrEmpty(VersePath))
+			{
+				CachedJson.AddOrSetFieldValue("VersePath", VersePath);
+			}
+			if (VerseScope != VerseScope.User)
+			{
+				CachedJson.AddOrSetFieldValue("VerseScope", VerseScope.ToString());
+			}
+			if (bEnabledByDefault.HasValue)
+			{
+				CachedJson.AddOrSetFieldValue("EnabledByDefault", bEnabledByDefault.Value);
+			}
+			CachedJson.AddOrSetFieldValue("CanContainContent", bCanContainContent);
+			if (bCanContainVerse)
+			{
+				CachedJson.AddOrSetFieldValue("CanContainVerse", bCanContainVerse);
+			}
+			if (bIsBetaVersion)
+			{
+				CachedJson.AddOrSetFieldValue("IsBetaVersion", bIsBetaVersion);
+			}
+			if (bIsExperimentalVersion)
+			{
+				CachedJson.AddOrSetFieldValue("IsExperimentalVersion", bIsExperimentalVersion);
+			}
+			if (bInstalled)
+			{
+				CachedJson.AddOrSetFieldValue("Installed", bInstalled);
+			}
+
+			if (bRequiresBuildPlatform)
+			{
+				CachedJson.AddOrSetFieldValue("RequiresBuildPlatform", bRequiresBuildPlatform);
+			}
+
+			if (bIsSealed)
+			{
+				CachedJson.AddOrSetFieldValue("Sealed", bIsSealed);
+			}
+
+			if (bExplicitlyLoaded)
+			{
+				CachedJson.AddOrSetFieldValue("ExplicitlyLoaded", bExplicitlyLoaded);
+			}
+
+			if (bHasExplicitPlatforms)
+			{
+				CachedJson.AddOrSetFieldValue("HasExplicitPlatforms", bHasExplicitPlatforms);
+			}
+
+			if (SupportedTargetPlatforms != null && SupportedTargetPlatforms.Count > 0)
+			{
+				CachedJson.AddOrSetFieldValue("SupportedTargetPlatforms", SupportedTargetPlatforms.Select<UnrealTargetPlatform, string>(x => x.ToString()).ToArray());
+			}
+
+			if (SupportedPrograms != null && SupportedPrograms.Length > 0)
+			{
+				CachedJson.AddOrSetFieldValue("SupportedPrograms", SupportedPrograms);
+			}
+			if (bIsPluginExtension)
+			{
+				CachedJson.AddOrSetFieldValue("bIsPluginExtension", bIsPluginExtension);
+			}
+
+			if (Modules != null && Modules.Count > 0)
+			{
+				ModuleDescriptor.UpdateJson(CachedJson, "Modules", Modules.ToArray());
+			}
+
+			LocalizationTargetDescriptor.UpdateJson(CachedJson, "LocalizationTargets", LocalizationTargets);
+
+			if (PreBuildSteps != null)
+			{
+				CachedJson.AddOrSetFieldValue("PreBuildSteps", PreBuildSteps.ToJsonObject());
+			}
+
+			if (PostBuildSteps != null)
+			{
+				CachedJson.AddOrSetFieldValue("PostBuildSteps", PostBuildSteps.ToJsonObject());
+			}
+
+			if (Plugins != null && Plugins.Count > 0)
+			{
+				PluginReferenceDescriptor.UpdateJson(CachedJson, "Plugins", Plugins.ToArray());
+			}
+			if (DisallowedPlugins != null && DisallowedPlugins.Length > 0)
+			{
+				CachedJson.AddOrSetFieldValue("DisallowedPlugins", DisallowedPlugins);
+			}
+		}
+
 		/// <summary>
 		/// Produces any warnings and errors for the plugin descriptor
 		/// </summary>
diff --git a/Engine/Source/Programs/UnrealBuildTool/System/PluginReferenceDescriptor.cs b/Engine/Source/Programs/UnrealBuildTool/System/PluginReferenceDescriptor.cs
--- a/Engine/Source/Programs/UnrealBuildTool/System/PluginReferenceDescriptor.cs
+++ b/Engine/Source/Programs/UnrealBuildTool/System/PluginReferenceDescriptor.cs
@@ -157,6 +157,62 @@ namespace UnrealBuildTool
 			Writer.WriteObjectEnd();
 		}
 
+		private JsonObject ToJsonObject()
+		{
+			JsonObject PluginReferenceObject = new JsonObject();
+			PluginReferenceObject.AddOrSetFieldValue("Name", Name);
+			PluginReferenceObject.AddOrSetFieldValue("Enabled", bEnabled);
+			if (bEnabled && bOptional)
+			{
+				PluginReferenceObject.AddOrSetFieldValue("Optional", bOptional);
+			}
+			if (!String.IsNullOrEmpty(Description))
+			{
+				PluginReferenceObject.AddOrSetFieldValue("Description", Description);
+			}
+			if (!String.IsNullOrEmpty(MarketplaceURL))
+			{
+				PluginReferenceObject.AddOrSetFieldValue("MarketplaceURL", MarketplaceURL);
+			}
+			if (PlatformAllowList != null && PlatformAllowList.Length > 0)
+			{
+				PluginReferenceObject.AddOrSetFieldValue("PlatformAllowList", PlatformAllowList.Select(x => x.ToString()).ToArray());
+			}
+			if (PlatformDenyList != null && PlatformDenyList.Length > 0)
+			{
+				PluginReferenceObject.AddOrSetFieldValue("PlatformDenyList", PlatformDenyList.Select(x => x.ToString()).ToArray());
+			}
+			if (TargetConfigurationAllowList != null && TargetConfigurationAllowList.Length > 0)
+			{
+				PluginReferenceObject.AddOrSetFieldValue("TargetConfigurationAllowList", TargetConfigurationAllowList);
+			}
+			if (TargetConfigurationDenyList != null && TargetConfigurationDenyList.Length > 0)
+			{
+				PluginReferenceObject.AddOrSetFieldValue("TargetConfigurationDenyList", TargetConfigurationDenyList);
+			}
+			if (TargetAllowList != null && TargetAllowList.Length > 0)
+			{
+				PluginReferenceObject.AddOrSetFieldValue("TargetAllowList", TargetAllowList);
+			}
+			if (TargetDenyList != null && TargetDenyList.Length > 0)
+			{
+				PluginReferenceObject.AddOrSetFieldValue("TargetDenyList", TargetDenyList);
+			}
+			if (SupportedTargetPlatforms != null && SupportedTargetPlatforms.Length > 0)
+			{
+				PluginReferenceObject.AddOrSetFieldValue("SupportedTargetPlatforms", SupportedTargetPlatforms.Select(x => x.ToString()).ToArray());
+			}
+			if (bHasExplicitPlatforms)
+			{
+				PluginReferenceObject.AddOrSetFieldValue("HasExplicitPlatforms", bHasExplicitPlatforms);
+			}
+			if (bEnabled && RequestedVersion != null)
+			{
+				PluginReferenceObject.AddOrSetFieldValue("Version", RequestedVersion.Value);
+			}
+			return PluginReferenceObject;
+		}
+
 		/// <summary>
 		/// Write an array of module descriptors
 		/// </summary>
@@ -176,6 +232,21 @@ namespace UnrealBuildTool
 			}
 		}
 
+		/// <summary>
+		/// Updates the json object with an array of plugin descriptors.
+		/// </summary>
+		/// <param name="InObject">The Json object to update.</param>
+		/// <param name="Name">Name of the array</param>
+		/// <param name="Plugins">Array of plugins</param>
+		public static void UpdateJson(JsonObject InObject, string Name, PluginReferenceDescriptor[]? Plugins)
+		{
+			if (Plugins != null && Plugins.Length > 0)
+			{
+				JsonObject[] JsonObjects = Plugins.Select(X => X.ToJsonObject()).ToArray();
+				InObject.AddOrSetFieldValue(Name, JsonObjects);
+			}
+		}
+
 		/// <summary>
 		/// Checks the given platform names array and logs a message about any that are not known. This may simply be due to not having the platform synced in Engine/Platforms/[PlatformName].
 		/// </summary>
