[Backout] - CL26105383
Original CL Desc
-----------------------------------------------------------------
Unreal Build Tool:
Update PluginDescriptor to have parity with FPluginDescriptor in terms of handling custom JSON fields. The order in which the original fields are written to disk is preserved.
When saving the PluginDescriptor to disk, the same order is followed and any new entries are added to the end of the plugin file.
Note: PluginDescriptor.Save and PluginDescriptor.Save2 will exist concurrently for now. There are still parts of UBT that use PluginDescriptor.Write and the refactor to create a unified Save function will come later as a lot of testing still needs to be done.
JsonObject:
- Introduced AddOrSetFieldValue() methods overloaded on value types to add/set fields in JsonObject.
- Changed the underlying data structure of JsonObject from Dictionary<string,Object?> to OrderedDictionary to guarantee that order of all .uplugin fields (default and custom) are preserved
- Made the JsonObject constructor that takes an OrderedDictionary private to ensure type safety of values passed to JsonObject
- Introduced ToJsonString() for JsonObject that serializes the JsonObject to match the formatting of .uplugin
IMPORTANT: The ToJsonString() method uses the JavaScriptEncoder.UnsafeRelaxedJsonEscaping encoder for Utf8JsonWriter.
This allows +,<,>,&,` etc characters to be written properly for .uplugins. As per MS docs, with this encoding, none of the contents should be written to HTML or a script.
- Introduced JsonObjectTests that unit tested JsonObject to ensure previous behavior was maintained.

PluginDescriptor
- Added a JsonObject field that acts as the cache of all fields that were read from a file or parsed from a json string.
- Cached JsonObject records the order of all fields that were read and captures custom fields in the .uplugin
- Introduced Save2(), the new save algorithm that writes all fields (default and custom) to disk.
- Deleted the private default constructor because it didn't seem to be used anywhere.
- Introduced PluginDescriptorTests to ensure previous behavior is maintained. In particular when dealing with a .uplugin with only default fields, results of Save() and Save2() are the same.

Descriptor Files:
- Introduced ToJsonObject() methods to replace the previous JsonWriter patterns of writing the descriptor data.
- Introduced the UpdateJsonObject pattern which mimics the old WriteArray pattern of writing data to a JsonWriter.

#test: Developed this in a TDD style with unit tests for all the major functions. Also reviewed output of created .uplugins and diffed them.

diff --git a/Engine/Source/Programs/Shared/EpicGames.Core.Tests/JsonObjectTests.cs b/Engine/Source/Programs/Shared/EpicGames.Core.Tests/JsonObjectTests.cs
deleted file mode 100644
diff --git a/Engine/Source/Programs/Shared/EpicGames.Core/JsonObject.cs b/Engine/Source/Programs/Shared/EpicGames.Core/JsonObject.cs
--- a/Engine/Source/Programs/Shared/EpicGames.Core/JsonObject.cs
+++ b/Engine/Source/Programs/Shared/EpicGames.Core/JsonObject.cs
@@ -5,74 +5,9 @@ using System.Collections.Generic;
 using System.Diagnostics.CodeAnalysis;
 using System.Linq;
 using System.Text.Json;
-using System.Collections.Specialized;
-using System.Collections;
-using System.IO;
-using System.Text;
-
-using SystemJsonObject = System.Text.Json.Nodes.JsonObject;
-using SystemJsonNode = System.Text.Json.Nodes.JsonNode;
-using SystemJsonValue = System.Text.Json.Nodes.JsonValue;
-using SystemJsonArray = System.Text.Json.Nodes.JsonArray;
-using System.Text.Encodings.Web;
-using System.Text.Unicode;
 
 namespace EpicGames.Core
 {
-	public static class OrderedDictionaryExtensions
-	{
-		public static bool TryGetValue(this OrderedDictionary dictionary, object key, [NotNullWhen(true)] out object? value)
-		{
-			if (dictionary.Contains(key))
-			{
-				value = dictionary[key]!;
-				return true;
-			}
-			value = null;
-			return false;
-		}
-
-		private static object? DeepCopyHelper(object? original)
-		{
-			if (original is null)
-			{
-				return null;
-			}
-
-			if (original is OrderedDictionary dictionary)
-			{
-				OrderedDictionary copy = new OrderedDictionary();
-				foreach (DictionaryEntry entry in dictionary)
-				{
-					copy.Add(entry.Key, DeepCopyHelper(entry.Value));
-				}
-				return copy;
-			}
-
-			if (original is Array array)
-			{
-				Array copy = Array.CreateInstance(array.GetType().GetElementType()!, array.Length);
-				for (int index = 0; index < array.Length; ++index)
-				{
-					copy.SetValue(DeepCopyHelper(array.GetValue(index)), index);
-				}
-				return copy;
-			}
-			// It's a value type or a string, it's safe to just copy it 
-			return original;
-		}
-
-		public static OrderedDictionary CreateDeepCopy(this OrderedDictionary dictionary)
-		{
-			OrderedDictionary copy = new OrderedDictionary();
-			foreach (DictionaryEntry entry in dictionary)
-			{
-				copy.Add(entry.Key, DeepCopyHelper(entry.Value));
-			}
-			return copy;
-		}
-	}
-
 	/// <summary>
 	/// Exception thrown for errors parsing JSON files
 	/// </summary>
@@ -94,20 +29,15 @@ namespace EpicGames.Core
 	/// </summary>
 	public class JsonObject
 	{
-		readonly OrderedDictionary _rawOrderedObject;
-
-		public JsonObject()
-		{
-			_rawOrderedObject = new OrderedDictionary();
-		}
+		readonly Dictionary<string, object?> _rawObject;
 
 		/// <summary>
-		/// Construct a JSON object from the OrderedDictionary obtained from reading a file on disk or parsing valid json text.
+		/// Construct a JSON object from the raw string -> object dictionary
 		/// </summary>
 		/// <param name="inRawObject">Raw object parsed from disk</param>
-		private JsonObject(OrderedDictionary inRawObject)
+		public JsonObject(Dictionary<string, object?> inRawObject)
 		{
-			_rawOrderedObject = inRawObject.CreateDeepCopy();
+			_rawObject = new Dictionary<string, object?>(inRawObject, StringComparer.InvariantCultureIgnoreCase);
 		}
 
 		/// <summary>
@@ -116,43 +46,11 @@ namespace EpicGames.Core
 		/// <param name="element"></param>
 		public JsonObject(JsonElement element)
 		{
-			_rawOrderedObject = new OrderedDictionary();
+			_rawObject = new Dictionary<string, object?>(StringComparer.InvariantCultureIgnoreCase);
 			foreach (JsonProperty property in element.EnumerateObject())
 			{
-				_rawOrderedObject[property.Name] = ParseElement(property.Value);
-			}
-		}
-
-		public override bool Equals(object? obj)
-		{
-			if (obj is null || GetType() != obj.GetType())
-			{
-				return false;
-			}
-			JsonObject jsonObject = (JsonObject) obj;
-			if (_rawOrderedObject.Count != jsonObject._rawOrderedObject.Count)
-			{
-				return false;
-			}
-			
-			foreach (DictionaryEntry entry in _rawOrderedObject)
-			{
-				if (!jsonObject._rawOrderedObject.Contains(entry.Key))
-				{
-					return false;
-				}
-				if (!entry.Value!.Equals(jsonObject._rawOrderedObject[entry.Key!]))
-				{
-					return false;
-				}
+				_rawObject[property.Name] = ParseElement(property.Value);
 			}
-
-			return true;
-		}
-
-		public override int GetHashCode()
-		{
-			return HashCode.Combine(_rawOrderedObject);
 		}
 
 		/// <summary>
@@ -169,12 +67,7 @@ namespace EpicGames.Core
 				case JsonValueKind.Number:
 					return element.GetDouble();
 				case JsonValueKind.Object:
-					OrderedDictionary dictionary = new OrderedDictionary();
-					foreach (JsonProperty property in element.EnumerateObject())
-						{
-						dictionary.Add(property.Name, ParseElement(property.Value));
-					}
-					return dictionary;
+					return element.EnumerateObject().ToDictionary(x => x.Name, x => ParseElement(x.Value));
 				case JsonValueKind.String:
 					return element.GetString();
 				case JsonValueKind.False:
@@ -188,42 +81,6 @@ namespace EpicGames.Core
 			}
 		}
 
-		private SystemJsonNode ToJsonNode(object? obj)
-		{
-			// All values in the JsonObject are either parsed from a string, read from a file or set/added with the API
-			// All values at this point must be supported and the correct types 
-			switch (obj)
-			{
-				case OrderedDictionary objDictionary:
-					SystemJsonObject dictionaryJson= new SystemJsonObject();
-					foreach (object? key in objDictionary.Keys)
-					{
-						dictionaryJson.Add((string) key, ToJsonNode(objDictionary[key]));
-					}
-					return dictionaryJson;
-				case Array objArray:
-					SystemJsonArray tempJsonArray = new SystemJsonArray();
-					foreach (object? element in objArray)
-					{
-						tempJsonArray.Add(ToJsonNode(element));
-					}
-					return tempJsonArray;
-				default:
-					// We support null as per the json spec 
-					return SystemJsonValue.Create(obj)!;
-			}
-		}
-
-		public SystemJsonObject ToSystemJsonObject()
-		{
-			SystemJsonObject returnObject = new SystemJsonObject();
-			foreach (DictionaryEntry entry in _rawOrderedObject )
-			{
-				returnObject.Add((string) entry.Key, ToJsonNode(entry.Value));
-			}
-			return returnObject;
-		}
-
 		/// <summary>
 		/// Read a JSON file from disk and construct a JsonObject from it
 		/// </summary>
@@ -267,15 +124,8 @@ namespace EpicGames.Core
 		/// <returns>New JsonObject instance</returns>
 		public static JsonObject Parse(string text)
 		{
-			try
-			{
-				JsonDocument document = JsonDocument.Parse(text, new JsonDocumentOptions { AllowTrailingCommas = true });
-				return new JsonObject(document.RootElement);
-			}
-			catch (Exception ex)
-			{
-				throw new JsonParseException("Failed to parse json text '{0}'. {1}", text, ex.Message);
-			}
+			JsonDocument document = JsonDocument.Parse(text, new JsonDocumentOptions { AllowTrailingCommas = true });
+			return new JsonObject(document.RootElement);
 		}
 
 		/// <summary>
@@ -301,7 +151,7 @@ namespace EpicGames.Core
 		/// <summary>
 		/// List of key names in this object
 		/// </summary>
-		public IEnumerable<string> KeyNames => _rawOrderedObject.Keys.Cast<string>();
+		public IEnumerable<string> KeyNames => _rawObject.Keys;
 
 		/// <summary>
 		/// Gets a string field by the given name from the object, throwing an exception if it is not there or cannot be parsed.
@@ -327,7 +177,7 @@ namespace EpicGames.Core
 		public bool TryGetStringField(string fieldName, [NotNullWhen(true)] out string? result)
 		{
 			object? rawValue;
-			if (_rawOrderedObject.TryGetValue(fieldName, out rawValue) && (rawValue is string strValue))
+			if (_rawObject.TryGetValue(fieldName, out rawValue) && (rawValue is string strValue))
 			{
 				result = strValue;
 				return true;
@@ -363,7 +213,7 @@ namespace EpicGames.Core
 		public bool TryGetStringArrayField(string fieldName, [NotNullWhen(true)] out string[]? result)
 		{
 			object? rawValue;
-			if (_rawOrderedObject.TryGetValue(fieldName, out rawValue) && (rawValue is IEnumerable<object> enumValue) && enumValue.All(x => x is string))
+			if (_rawObject.TryGetValue(fieldName, out rawValue) && (rawValue is IEnumerable<object> enumValue) && enumValue.All(x => x is string))
 			{
 				result = enumValue.Select(x => (string)x).ToArray();
 				return true;
@@ -399,7 +249,7 @@ namespace EpicGames.Core
 		public bool TryGetBoolField(string fieldName, out bool result)
 		{
 			object? rawValue;
-			if (_rawOrderedObject.TryGetValue(fieldName, out rawValue) && (rawValue is bool boolValue))
+			if (_rawObject.TryGetValue(fieldName, out rawValue) && (rawValue is bool boolValue))
 			{
 				result = boolValue;
 				return true;
@@ -435,7 +285,7 @@ namespace EpicGames.Core
 		public bool TryGetIntegerField(string fieldName, out int result)
 		{
 			object? rawValue;
-			if (!_rawOrderedObject.TryGetValue(fieldName, out rawValue) || !Int32.TryParse(rawValue?.ToString(), out result))
+			if (!_rawObject.TryGetValue(fieldName, out rawValue) || !Int32.TryParse(rawValue?.ToString(), out result))
 			{
 				result = 0;
 				return false;
@@ -452,7 +302,7 @@ namespace EpicGames.Core
 		public bool TryGetUnsignedIntegerField(string fieldName, out uint result)
 		{
 			object? rawValue;
-			if (!_rawOrderedObject.TryGetValue(fieldName, out rawValue) || !UInt32.TryParse(rawValue?.ToString(), out result))
+			if (!_rawObject.TryGetValue(fieldName, out rawValue) || !UInt32.TryParse(rawValue?.ToString(), out result))
 			{
 				result = 0;
 				return false;
@@ -484,7 +334,7 @@ namespace EpicGames.Core
 		public bool TryGetDoubleField(string fieldName, out double result)
 		{
 			object? rawValue;
-			if (!_rawOrderedObject.TryGetValue(fieldName, out rawValue) || !Double.TryParse(rawValue?.ToString(), out result))
+			if (!_rawObject.TryGetValue(fieldName, out rawValue) || !Double.TryParse(rawValue?.ToString(), out result))
 			{
 				result = 0.0;
 				return false;
@@ -577,9 +427,9 @@ namespace EpicGames.Core
 		public bool TryGetObjectField(string fieldName, [NotNullWhen(true)] out JsonObject? result)
 		{
 			object? rawValue;
-			if (_rawOrderedObject.TryGetValue(fieldName, out rawValue) && (rawValue is OrderedDictionary orderedDictValue))
+			if (_rawObject.TryGetValue(fieldName, out rawValue) && (rawValue is Dictionary<string, object?> dictValue))
 			{
-				result = new JsonObject(orderedDictValue);
+				result = new JsonObject(dictValue);
 				return true;
 			}
 			else
@@ -613,9 +463,9 @@ namespace EpicGames.Core
 		public bool TryGetObjectArrayField(string fieldName, [NotNullWhen(true)] out JsonObject[]? result)
 		{
 			object? rawValue;
-			if (_rawOrderedObject.TryGetValue(fieldName, out rawValue) && (rawValue is IEnumerable<object> enumValue) && enumValue.All(x => x is OrderedDictionary))
+			if (_rawObject.TryGetValue(fieldName, out rawValue) && (rawValue is IEnumerable<object> enumValue) && enumValue.All(x => x is Dictionary<string, object>))
 			{
-				result = enumValue.Select(x => new JsonObject((OrderedDictionary)x)).ToArray();
+				result = enumValue.Select(x => new JsonObject((Dictionary<string, object?>)x)).ToArray();
 				return true;
 			}
 			else
@@ -624,170 +474,5 @@ namespace EpicGames.Core
 				return false;
 			}
 		}
-
-		/// <summary>
-		/// Checks if the provided field exists in this Json object.
-		/// </summary>
-		/// <param name="fieldName">Name of the field to check if it is contained. </param>
-		/// <returns>True if the field exists in the JsonObject, false otherwise</returns>
-		public bool ContainsField(string fieldName)
-		{
-			return _rawOrderedObject.Contains(fieldName);
-		}
-
-		public void AddOrSetFieldValue(string fieldName, int value)
-		{
-			if (String.IsNullOrEmpty(fieldName))
-			{
-				return;
-			}
-			_rawOrderedObject[fieldName] = Convert.ToDouble(value);
-		}
-
-		public void AddOrSetFieldValue(string fieldName, uint value)
-		{
-			if (String.IsNullOrEmpty(fieldName))
-			{
-				return;
-			}
-			_rawOrderedObject[fieldName] = Convert.ToDouble(value);
-		}
-
-		public void AddOrSetFieldValue(string fieldName, double value)
-		{
-			if (String.IsNullOrEmpty(fieldName))
-			{
-				return;
-			}
-			_rawOrderedObject[fieldName] = value;
-		}
-
-		public void AddOrSetFieldValue(string fieldName, bool value)
-		{
-			if (String.IsNullOrEmpty(fieldName))
-			{
-				return;
-			}
-			_rawOrderedObject[fieldName] = value;
-		}
-
-		public void AddOrSetFieldValue(string fieldName, string? value)
-		{
-			if (String.IsNullOrEmpty(fieldName))
-			{
-				return;
-			}
-			// We set null strings to empty strings to follow the behavior of EpicGames.Core.JsonWriter
-			value ??= "";
-			_rawOrderedObject[fieldName] = value;
-		}
-
-		public void AddOrSetFieldValue<T>(string fieldName, T value) where T : Enum
-		{
-			if (String.IsNullOrEmpty(fieldName))
-			{
-				return;
-			}
-			
-			_rawOrderedObject[fieldName] = value.ToString();
-		}
-
-		
-		public void AddOrSetFieldValue(string fieldName, JsonObject? value)
-		{
-			if (String.IsNullOrEmpty(fieldName))
-			{
-				return;
-			}
-
-			_rawOrderedObject[fieldName] = value?._rawOrderedObject.CreateDeepCopy();
-		}
-
-		public void AddOrSetFieldValue(string fieldName, string?[]? value)
-		{
-			if (String.IsNullOrEmpty(fieldName))
-			{
-				return;
-			}
-			if (value is null)
-			{
-				_rawOrderedObject[fieldName] = value;
-				return;
-			}
-			// Contents of the strings array should never be null and should be "" instead to match behavior EpicGames.Core.JsonWriter
-			string[] stringArray = value.Select(x => x ?? "").ToArray();
-			_rawOrderedObject[fieldName] = stringArray;
-		}
-
-		public void AddOrSetFieldValue(string fieldName, JsonObject[]? value)
-		{
-			if (String.IsNullOrEmpty(fieldName))
-			{
-				return;
-			}
-			if (value is null)
-			{
-				_rawOrderedObject[fieldName] = null;
-				return;
-			}
-			List<OrderedDictionary?> objList = new List<OrderedDictionary?>();
-			foreach (JsonObject? obj in value)
-			{
-				objList.Add(obj?._rawOrderedObject.CreateDeepCopy());
-			}
-			_rawOrderedObject[fieldName] = objList.ToArray();
-		}
-
-		public void AddOrSetFieldValue<T>(string fieldName, T[]? value) where T : Enum
-		{
-			if (String.IsNullOrEmpty(fieldName))
-			{
-				return;
-			}
-			if (value is null)
-			{
-				_rawOrderedObject[fieldName] = null;
-				return;
-			}
-
-			string[] stringArray = value.Select(x => x.ToString()!).ToArray();
-			_rawOrderedObject[fieldName] = stringArray;
-		}
-		/// <summary>
-		/// Converts this Json Object to a string representation. This follows formatting of UE .uplugin files with 4 spaces for tabs and indentation enabled.
-		/// IMPORTANT: If this JsonObject contains HTML, the returned string should NOT be used directly for HTML or a script. Read the note below.
-		/// </summary>
-		/// <returns>The formatted, prettified string representation of this Json Object.</returns>
-		public string ToJsonString()
-		{
-			SystemJsonObject jsonObjectToWrite = ToSystemJsonObject();
-			JsonWriterOptions options = new JsonWriterOptions();
-			options.Indented = true;
-			// IMPORTANT: Utf8JsonWriter blocks certain characters like +, &, <, >,` from being escaped in a global block list
-			// Best way around is to use the relaxed JavaScriptEncoder.UnsafeRelaxedJsonEscaping.
-			// However this can have security implications if this string is ever written to an HTML page or script
-			// https://learn.microsoft.com/en-us/dotnet/standard/serialization/system-text-json/character-encoding
-			options.Encoder = JavaScriptEncoder.UnsafeRelaxedJsonEscaping;
-			using MemoryStream jsonMemoryStream = new MemoryStream();
-			using (Utf8JsonWriter writer = new Utf8JsonWriter(jsonMemoryStream, options))
-			{
-				jsonObjectToWrite.WriteTo(writer);
-				writer.Flush();
-			}
-			// The Utf8JsonWriter doesn't format json the same as we want in UE, we massage the string to meet our standards 
-			string jsonString = Encoding.UTF8.GetString(jsonMemoryStream.ToArray());
-			string[] lines = jsonString.Split(new[] { Environment.NewLine }, StringSplitOptions.None);
-			StringBuilder jsonStringBuilder = new StringBuilder();
-			foreach (string line in lines)
-			{
-				// Utf8JsonWriter uses 2 spaces for indents, we replace them with tabs here 
-				int numLeadingSpaces = line.TakeWhile(x => x == ' ').Count();
-				int numLeadingTabs = numLeadingSpaces / 2;
-				jsonStringBuilder.Append('\t', numLeadingTabs);
-				jsonStringBuilder.AppendLine(line.Substring(numLeadingSpaces));
-			}
-
-			return jsonStringBuilder.ToString();
-		}
 	}
 }
diff --git a/Engine/Source/Programs/UnrealBuildTool.Tests/PluginDescriptorTests.cs b/Engine/Source/Programs/UnrealBuildTool.Tests/PluginDescriptorTests.cs
deleted file mode 100644
diff --git a/Engine/Source/Programs/UnrealBuildTool/System/CustomBuildSteps.cs b/Engine/Source/Programs/UnrealBuildTool/System/CustomBuildSteps.cs
--- a/Engine/Source/Programs/UnrealBuildTool/System/CustomBuildSteps.cs
+++ b/Engine/Source/Programs/UnrealBuildTool/System/CustomBuildSteps.cs
@@ -72,20 +72,6 @@ namespace UnrealBuildTool
 			Writer.WriteObjectEnd();
 		}
 
-		/// <summary>
-		/// Converts this object to a JsonObject to be saved or manipulated.
-		/// </summary>
-		/// <returns>The JsonObject representation of this object.</returns>
-		public JsonObject ToJsonObject()
-		{
-			JsonObject CustomBuildStepObject = new JsonObject();
-			foreach (KeyValuePair<UnrealTargetPlatform, string[]> Pair in HostPlatformToCommands.OrderBy(x => x.Key.ToString()))
-			{
-				CustomBuildStepObject.AddOrSetFieldValue(Pair.Key.ToString(), Pair.Value);
-			}
-			return CustomBuildStepObject;
-		}
-
 		/// <summary>
 		/// Tries to get the commands for a given host platform
 		/// </summary>
diff --git a/Engine/Source/Programs/UnrealBuildTool/System/LocalizationDescriptor.cs b/Engine/Source/Programs/UnrealBuildTool/System/LocalizationDescriptor.cs
--- a/Engine/Source/Programs/UnrealBuildTool/System/LocalizationDescriptor.cs
+++ b/Engine/Source/Programs/UnrealBuildTool/System/LocalizationDescriptor.cs
@@ -1,6 +1,5 @@
 // Copyright Epic Games, Inc. All Rights Reserved.
 
-using System.Linq;
 using System.Diagnostics;
 using EpicGames.Core;
 
@@ -91,15 +90,6 @@ namespace UnrealBuildTool
 			Writer.WriteObjectEnd();
 		}
 
-		JsonObject ToJsonObject()
-		{
-			JsonObject localizationTargetObject= new JsonObject();
-			localizationTargetObject.AddOrSetFieldValue("Name", Name);
-			localizationTargetObject.AddOrSetFieldValue("LoadingPolicy", LoadingPolicy.ToString());
-			
-			return localizationTargetObject;
-		}
-
 		/// <summary>
 		/// Write an array of target descriptors
 		/// </summary>
@@ -118,20 +108,5 @@ namespace UnrealBuildTool
 				Writer.WriteArrayEnd();
 			}
 		}
-
-		/// <summary>
-		/// Updates a JsonObject with an array of localization target descriptors.
-		/// </summary>
-		/// <param name="InObject">The Json object to update.</param>
-		/// <param name="Name">Name of the array</param>
-		/// <param name="Targets">Array of targets</param>
-		public static void UpdateJson(JsonObject InObject, string Name, LocalizationTargetDescriptor[]? Targets)
-		{
-			if (Targets != null && Targets.Length > 0)
-			{
-				JsonObject[] JsonObjects = Targets.Select(X => X.ToJsonObject()).ToArray();
-				InObject.AddOrSetFieldValue(Name, JsonObjects);
-			}
-		}
 	}
 }
diff --git a/Engine/Source/Programs/UnrealBuildTool/System/ModuleDescriptor.cs b/Engine/Source/Programs/UnrealBuildTool/System/ModuleDescriptor.cs
--- a/Engine/Source/Programs/UnrealBuildTool/System/ModuleDescriptor.cs
+++ b/Engine/Source/Programs/UnrealBuildTool/System/ModuleDescriptor.cs
@@ -433,64 +433,6 @@ namespace UnrealBuildTool
 			Writer.WriteObjectEnd();
 		}
 
-		JsonObject ToJsonObject()
-		{
-			JsonObject ModuleObject = new JsonObject();
-			ModuleObject.AddOrSetFieldValue("Name", Name);
-			ModuleObject.AddOrSetFieldValue("Type", Type.ToString());
-			ModuleObject.AddOrSetFieldValue("LoadingPhase", LoadingPhase.ToString());
-			// important note: we don't check the length of the platform allow list, because if an unknown platform was read in, but was not valid, the 
-			// list will exist but be empty. We don't want to remove the allow list completely, because that would allow this module on all platforms,
-			// which will not be the desired effect
-			if (PlatformAllowList != null)
-			{
-				string[] PlatformAllowListStringArray = PlatformAllowList.Select(X => X.ToString()).ToArray();
-				ModuleObject.AddOrSetFieldValue("PlatformAllowList", PlatformAllowListStringArray);
-			}
-			if (PlatformDenyList != null && PlatformDenyList.Count > 0)
-			{
-				string[] PlatformDenyListStringArray = PlatformDenyList.Select(X => X.ToString()).ToArray();
-				ModuleObject.AddOrSetFieldValue("PlatformDenyList", PlatformDenyListStringArray);
-			}
-			if (TargetAllowList != null && TargetAllowList.Length > 0)
-			{
-				string[] TargetAllowListStringArray = TargetAllowList.Select(X => X.ToString()).ToArray();
-				ModuleObject.AddOrSetFieldValue("TargetAllowList", TargetAllowListStringArray);
-			}
-			if (TargetDenyList != null && TargetDenyList.Length > 0)
-			{
-				string[] TargetDenyListStringArray = TargetDenyList.Select(X => X.ToString()).ToArray();
-				ModuleObject.AddOrSetFieldValue("TargetDenyList", TargetDenyListStringArray);
-			}
-			if (TargetConfigurationAllowList != null && TargetConfigurationAllowList.Length > 0)
-			{
-				string[] TargetConfigurationAllowListStringArray = TargetConfigurationAllowList.Select(X => X.ToString()).ToArray();
-				ModuleObject.AddOrSetFieldValue("TargetConfigurationAllowList", TargetConfigurationAllowListStringArray);
-			}
-			if (TargetConfigurationDenyList != null && TargetConfigurationDenyList.Length > 0)
-			{
-				string[] TargetConfigurationDenyListStringArray = TargetConfigurationDenyList.Select(X => X.ToString()).ToArray();
-				ModuleObject.AddOrSetFieldValue("TargetConfigurationDenyList", TargetConfigurationDenyListStringArray);
-			}
-			if (ProgramAllowList != null && ProgramAllowList.Length > 0)
-			{
-				ModuleObject.AddOrSetFieldValue("ProgramAllowList", ProgramAllowList);
-			}
-			if (ProgramDenyList != null && ProgramDenyList.Length > 0)
-			{
-				ModuleObject.AddOrSetFieldValue("ProgramDenyList", ProgramDenyList);
-			}
-			if (AdditionalDependencies != null && AdditionalDependencies.Length > 0)
-			{
-				ModuleObject.AddOrSetFieldValue("AdditionalDependencies", AdditionalDependencies);
-			}
-			if (bHasExplicitPlatforms)
-			{
-				ModuleObject.AddOrSetFieldValue("HasExplicitPlatforms", bHasExplicitPlatforms);
-			}
-			return ModuleObject;
-		}
-
 		/// <summary>
 		/// Write an array of module descriptors
 		/// </summary>
@@ -510,21 +452,6 @@ namespace UnrealBuildTool
 			}
 		}
 
-		/// <summary>
-		/// Updates a JsonObject with an array of module descriptors 
-		/// </summary>
-		/// <param name="InObject">The JsonObject to update.</param>
-		/// <param name="Name">Name of the array</param>
-		/// <param name="Modules">Array of modules</param>
-		public static void UpdateJson(JsonObject InObject, string Name, ModuleDescriptor[]? Modules)
-		{
-			if (Modules != null && Modules.Length > 0)
-			{
-				JsonObject[] JsonObjects = Modules.Select(X => X.ToJsonObject()).ToArray();
-				InObject.AddOrSetFieldValue(Name, JsonObjects);
-			}
-		}
-
 		/// <summary>
 		/// Produces any warnings and errors for the module settings
 		/// </summary>
diff --git a/Engine/Source/Programs/UnrealBuildTool/System/PluginDescriptor.cs b/Engine/Source/Programs/UnrealBuildTool/System/PluginDescriptor.cs
--- a/Engine/Source/Programs/UnrealBuildTool/System/PluginDescriptor.cs
+++ b/Engine/Source/Programs/UnrealBuildTool/System/PluginDescriptor.cs
@@ -1,11 +1,9 @@
 // Copyright Epic Games, Inc. All Rights Reserved.
 
 using System;
-using System.IO;
-using System.Linq;
 using System.Collections.Generic;
+using System.Linq;
 using EpicGames.Core;
-using EpicGames.Serialization;
 
 namespace UnrealBuildTool
 {
@@ -224,10 +222,12 @@ namespace UnrealBuildTool
 		public String[]? DisallowedPlugins;
 
 		/// <summary>
-		/// The JsonObject created from reading a .uplugin on disk or from parsing a json text 
-		/// This preserves the order of all the fields from the source json as well as account for any custom fields.
+		/// Private constructor. This object should not be created directly; read it from disk using FromFile() instead.
 		/// </summary>
-		private readonly JsonObject CachedJson;
+		private PluginDescriptor()
+		{
+			FileVersion = (int)PluginDescriptorVersion.Latest;
+		}
 
 		/// <summary>
 		/// Reads a plugin descriptor from a json object
@@ -237,7 +237,6 @@ namespace UnrealBuildTool
 		/// <returns>New plugin descriptor</returns>
 		public PluginDescriptor(JsonObject RawObject, FileReference PluginPath)
 		{
-			CachedJson = RawObject;
 			// Read the version
 			if (!RawObject.TryGetIntegerField("FileVersion", out FileVersion))
 			{
@@ -357,7 +356,7 @@ namespace UnrealBuildTool
 		}
 
 		/// <summary>
-		/// Creates a plugin descriptor from a file on disk preserving all custom fields in the file.
+		/// Creates a plugin descriptor from a file on disk
 		/// </summary>
 		/// <param name="FileName">The filename to read</param>
 		/// <returns>New plugin descriptor</returns>
@@ -377,7 +376,7 @@ namespace UnrealBuildTool
 		}
 
 		/// <summary>
-		/// Saves the descriptor to disk. This only saves the default fields in a .uplugin and does not account for cusotm fields.
+		/// Saves the descriptor to disk
 		/// </summary>
 		/// <param name="FileName">The filename to write to</param>
 		public void Save(string FileName)
@@ -390,18 +389,6 @@ namespace UnrealBuildTool
 			}
 		}
 
-		/// <summary>
-		/// Saves the descriptor to disk preserving all custom fields that were read in.
-		/// </summary>
-		/// <param name="fileName">The filename to write to</param>
-		public void Save2(string fileName)
-		{
-			// @TODO: This should replace all instances of Save() at some point in the future. There's just still a lot of references to test and refactor that needs to be verified. 
-			UpdateJson();
-			string jsonString = CachedJson.ToJsonString();
-			File.WriteAllText(fileName, jsonString);
-		}
-
 		/// <summary>
 		/// Writes the plugin descriptor to an existing Json writer
 		/// </summary>
@@ -520,114 +507,6 @@ namespace UnrealBuildTool
 			}
 		}
 
-		private void UpdateJson()
-		{
-			CachedJson.AddOrSetFieldValue("FileVersion", (int)ProjectDescriptorVersion.Latest);
-			CachedJson.AddOrSetFieldValue("Version", Version);
-			CachedJson.AddOrSetFieldValue("VersionName", VersionName);
-			CachedJson.AddOrSetFieldValue("FriendlyName", FriendlyName);
-			CachedJson.AddOrSetFieldValue("Description", Description);
-			CachedJson.AddOrSetFieldValue("Category", Category);
-			CachedJson.AddOrSetFieldValue("CreatedBy", CreatedBy);
-			CachedJson.AddOrSetFieldValue("CreatedByURL", CreatedByURL);
-			CachedJson.AddOrSetFieldValue("DocsURL", DocsURL);
-			CachedJson.AddOrSetFieldValue("MarketplaceURL", MarketplaceURL);
-			CachedJson.AddOrSetFieldValue("SupportURL", SupportURL);
-			if (!String.IsNullOrEmpty(EngineVersion))
-			{
-				CachedJson.AddOrSetFieldValue("EngineVersion", EngineVersion);
-			}
-			if (!String.IsNullOrEmpty(VersePath))
-			{
-				CachedJson.AddOrSetFieldValue("VersePath", VersePath);
-			}
-			if (VerseScope != VerseScope.User)
-			{
-				CachedJson.AddOrSetFieldValue("VerseScope", VerseScope.ToString());
-			}
-			if (bEnabledByDefault.HasValue)
-			{
-				CachedJson.AddOrSetFieldValue("EnabledByDefault", bEnabledByDefault.Value);
-			}
-			CachedJson.AddOrSetFieldValue("CanContainContent", bCanContainContent);
-			if (bCanContainVerse)
-			{
-				CachedJson.AddOrSetFieldValue("CanContainVerse", bCanContainVerse);
-			}
-			if (bIsBetaVersion)
-			{
-				CachedJson.AddOrSetFieldValue("IsBetaVersion", bIsBetaVersion);
-			}
-			if (bIsExperimentalVersion)
-			{
-				CachedJson.AddOrSetFieldValue("IsExperimentalVersion", bIsExperimentalVersion);
-			}
-			if (bInstalled)
-			{
-				CachedJson.AddOrSetFieldValue("Installed", bInstalled);
-			}
-
-			if (bRequiresBuildPlatform)
-			{
-				CachedJson.AddOrSetFieldValue("RequiresBuildPlatform", bRequiresBuildPlatform);
-			}
-
-			if (bIsSealed)
-			{
-				CachedJson.AddOrSetFieldValue("Sealed", bIsSealed);
-			}
-
-			if (bExplicitlyLoaded)
-			{
-				CachedJson.AddOrSetFieldValue("ExplicitlyLoaded", bExplicitlyLoaded);
-			}
-
-			if (bHasExplicitPlatforms)
-			{
-				CachedJson.AddOrSetFieldValue("HasExplicitPlatforms", bHasExplicitPlatforms);
-			}
-
-			if (SupportedTargetPlatforms != null && SupportedTargetPlatforms.Count > 0)
-			{
-				CachedJson.AddOrSetFieldValue("SupportedTargetPlatforms", SupportedTargetPlatforms.Select<UnrealTargetPlatform, string>(x => x.ToString()).ToArray());
-			}
-
-			if (SupportedPrograms != null && SupportedPrograms.Length > 0)
-			{
-				CachedJson.AddOrSetFieldValue("SupportedPrograms", SupportedPrograms);
-			}
-			if (bIsPluginExtension)
-			{
-				CachedJson.AddOrSetFieldValue("bIsPluginExtension", bIsPluginExtension);
-			}
-
-			if (Modules != null && Modules.Count > 0)
-			{
-				ModuleDescriptor.UpdateJson(CachedJson, "Modules", Modules.ToArray());
-			}
-
-			LocalizationTargetDescriptor.UpdateJson(CachedJson, "LocalizationTargets", LocalizationTargets);
-
-			if (PreBuildSteps != null)
-			{
-				CachedJson.AddOrSetFieldValue("PreBuildSteps", PreBuildSteps.ToJsonObject());
-			}
-
-			if (PostBuildSteps != null)
-			{
-				CachedJson.AddOrSetFieldValue("PostBuildSteps", PostBuildSteps.ToJsonObject());
-			}
-
-			if (Plugins != null && Plugins.Count > 0)
-			{
-				PluginReferenceDescriptor.UpdateJson(CachedJson, "Plugins", Plugins.ToArray());
-			}
-			if (DisallowedPlugins != null && DisallowedPlugins.Length > 0)
-			{
-				CachedJson.AddOrSetFieldValue("DisallowedPlugins", DisallowedPlugins);
-			}
-		}
-
 		/// <summary>
 		/// Produces any warnings and errors for the plugin descriptor
 		/// </summary>
diff --git a/Engine/Source/Programs/UnrealBuildTool/System/PluginReferenceDescriptor.cs b/Engine/Source/Programs/UnrealBuildTool/System/PluginReferenceDescriptor.cs
--- a/Engine/Source/Programs/UnrealBuildTool/System/PluginReferenceDescriptor.cs
+++ b/Engine/Source/Programs/UnrealBuildTool/System/PluginReferenceDescriptor.cs
@@ -157,62 +157,6 @@ namespace UnrealBuildTool
 			Writer.WriteObjectEnd();
 		}
 
-		private JsonObject ToJsonObject()
-		{
-			JsonObject PluginReferenceObject = new JsonObject();
-			PluginReferenceObject.AddOrSetFieldValue("Name", Name);
-			PluginReferenceObject.AddOrSetFieldValue("Enabled", bEnabled);
-			if (bEnabled && bOptional)
-			{
-				PluginReferenceObject.AddOrSetFieldValue("Optional", bOptional);
-			}
-			if (!String.IsNullOrEmpty(Description))
-			{
-				PluginReferenceObject.AddOrSetFieldValue("Description", Description);
-			}
-			if (!String.IsNullOrEmpty(MarketplaceURL))
-			{
-				PluginReferenceObject.AddOrSetFieldValue("MarketplaceURL", MarketplaceURL);
-			}
-			if (PlatformAllowList != null && PlatformAllowList.Length > 0)
-			{
-				PluginReferenceObject.AddOrSetFieldValue("PlatformAllowList", PlatformAllowList.Select(x => x.ToString()).ToArray());
-			}
-			if (PlatformDenyList != null && PlatformDenyList.Length > 0)
-			{
-				PluginReferenceObject.AddOrSetFieldValue("PlatformDenyList", PlatformDenyList.Select(x => x.ToString()).ToArray());
-			}
-			if (TargetConfigurationAllowList != null && TargetConfigurationAllowList.Length > 0)
-			{
-				PluginReferenceObject.AddOrSetFieldValue("TargetConfigurationAllowList", TargetConfigurationAllowList);
-			}
-			if (TargetConfigurationDenyList != null && TargetConfigurationDenyList.Length > 0)
-			{
-				PluginReferenceObject.AddOrSetFieldValue("TargetConfigurationDenyList", TargetConfigurationDenyList);
-			}
-			if (TargetAllowList != null && TargetAllowList.Length > 0)
-			{
-				PluginReferenceObject.AddOrSetFieldValue("TargetAllowList", TargetAllowList);
-			}
-			if (TargetDenyList != null && TargetDenyList.Length > 0)
-			{
-				PluginReferenceObject.AddOrSetFieldValue("TargetDenyList", TargetDenyList);
-			}
-			if (SupportedTargetPlatforms != null && SupportedTargetPlatforms.Length > 0)
-			{
-				PluginReferenceObject.AddOrSetFieldValue("SupportedTargetPlatforms", SupportedTargetPlatforms.Select(x => x.ToString()).ToArray());
-			}
-			if (bHasExplicitPlatforms)
-			{
-				PluginReferenceObject.AddOrSetFieldValue("HasExplicitPlatforms", bHasExplicitPlatforms);
-			}
-			if (bEnabled && RequestedVersion != null)
-			{
-				PluginReferenceObject.AddOrSetFieldValue("Version", RequestedVersion.Value);
-			}
-			return PluginReferenceObject;
-		}
-
 		/// <summary>
 		/// Write an array of module descriptors
 		/// </summary>
@@ -232,21 +176,6 @@ namespace UnrealBuildTool
 			}
 		}
 
-		/// <summary>
-		/// Updates the json object with an array of plugin descriptors.
-		/// </summary>
-		/// <param name="InObject">The Json object to update.</param>
-		/// <param name="Name">Name of the array</param>
-		/// <param name="Plugins">Array of plugins</param>
-		public static void UpdateJson(JsonObject InObject, string Name, PluginReferenceDescriptor[]? Plugins)
-		{
-			if (Plugins != null && Plugins.Length > 0)
-			{
-				JsonObject[] JsonObjects = Plugins.Select(X => X.ToJsonObject()).ToArray();
-				InObject.AddOrSetFieldValue(Name, JsonObjects);
-			}
-		}
-
 		/// <summary>
 		/// Checks the given platform names array and logs a message about any that are not known. This may simply be due to not having the platform synced in Engine/Platforms/[PlatformName].
 		/// </summary>
