UHT: RigVM Make sure to rely on slice index for lazy values

diff --git a/Engine/Source/Programs/Shared/EpicGames.UHT/Exporters/CodeGen/UhtHeaderCodeGeneratorHFile.cs b/Engine/Source/Programs/Shared/EpicGames.UHT/Exporters/CodeGen/UhtHeaderCodeGeneratorHFile.cs
--- a/Engine/Source/Programs/Shared/EpicGames.UHT/Exporters/CodeGen/UhtHeaderCodeGeneratorHFile.cs
+++ b/Engine/Source/Programs/Shared/EpicGames.UHT/Exporters/CodeGen/UhtHeaderCodeGeneratorHFile.cs
@@ -282,7 +282,7 @@ namespace EpicGames.UHT.Exporters.CodeGen
 									string paramTypeOriginal = parameter.TypeOriginal(true);
 									string paramNameOriginal = parameter.NameOriginal(false);
 									string additionalParameters = String.Empty;
-									if (!parameter.Input && !parameter.Output && !parameter.Singleton)
+									if (parameter.IsLazy || (!parameter.Input && !parameter.Output && !parameter.Singleton))
 									{
 										additionalParameters = ", RigVMExecuteContext.GetSlice().GetIndex()";
 									}
