- Moved SDK versions (MainVersion and Min/Max SDK version) from C# strings to SDK.json files in the Platform's Config folder
- Allow for per-project override in a project's platform's config folder (this is merging with my other work on per-project SDKs and validation of multi-target builds)
- More versions will move over after this

diff --git a/Engine/Config/Android/Android_SDK.json b/Engine/Config/Android/Android_SDK.json
new file mode 100644
--- /dev/null
+++ b/Engine/Config/Android/Android_SDK.json
@@ -0,0 +1,7 @@
+{
+"MainVersion" : "r25b",
+"AutoSDKDirectory" : "-26",
+
+"MinVersion" : "r25b",
+"MaxVersion" : "r25b",
+}
diff --git a/Engine/Config/Apple/Apple_SDK.json b/Engine/Config/Apple/Apple_SDK.json
new file mode 100644
--- /dev/null
+++ b/Engine/Config/Apple/Apple_SDK.json
@@ -0,0 +1,8 @@
+{
+"MainVersion" : "14.1",
+"MinVersion" : "14.1.0",
+"MaxVersion" : "15.9.0",
+// The versions on Windows are iTunes versions
+"MinVersion_Win64" : "1100.0.0.0",
+"MaxVersion_Win64" : "8999.0",
+}
diff --git a/Engine/Config/IOS/IOS_SDK.json b/Engine/Config/IOS/IOS_SDK.json
new file mode 100644
--- /dev/null
+++ b/Engine/Config/IOS/IOS_SDK.json
@@ -0,0 +1,3 @@
+{
+"ParentSDKFile" : "../Apple/Apple_SDK.json"
+}
diff --git a/Engine/Config/Linux/Linux_SDK.json b/Engine/Config/Linux/Linux_SDK.json
new file mode 100644
--- /dev/null
+++ b/Engine/Config/Linux/Linux_SDK.json
@@ -0,0 +1,7 @@
+{
+"MainVersion" : "v21_clang-15.0.1-centos7",
+"MinVersion" : "v21_clang-15.0.1-centos7",
+"MaxVersion" : "v21_clang-15.0.1-centos7",
+
+"AutoSDKPlatform" : "Linux_x64",
+}
diff --git a/Engine/Config/LinuxArm64/LinuxArm64_SDK.json b/Engine/Config/LinuxArm64/LinuxArm64_SDK.json
new file mode 100644
--- /dev/null
+++ b/Engine/Config/LinuxArm64/LinuxArm64_SDK.json
@@ -0,0 +1,3 @@
+{
+"ParentSDKFile" : "../Linux/Linux_SDK.json"
+}
diff --git a/Engine/Config/Mac/Mac_SDK.json b/Engine/Config/Mac/Mac_SDK.json
new file mode 100644
--- /dev/null
+++ b/Engine/Config/Mac/Mac_SDK.json
@@ -0,0 +1,3 @@
+{
+"ParentSDKFile" : "../Apple/Apple_SDK.json"
+}
diff --git a/Engine/Config/TVOS/TVOS_SDK.json b/Engine/Config/TVOS/TVOS_SDK.json
new file mode 100644
--- /dev/null
+++ b/Engine/Config/TVOS/TVOS_SDK.json
@@ -0,0 +1,3 @@
+{
+"ParentSDKFile" : "../Apple/Apple_SDK.json"
+}
diff --git a/Engine/Config/Windows/Windows_SDK.json b/Engine/Config/Windows/Windows_SDK.json
new file mode 100644
--- /dev/null
+++ b/Engine/Config/Windows/Windows_SDK.json
@@ -0,0 +1,5 @@
+{
+"MainVersion" : "10.0.18362.0",
+"MinVersion" : "10.0.00000.0",
+"MaxVersion" : "10.9.99999.0",
+}
diff --git a/Engine/Source/Programs/AutomationTool/CookedEditor/CookedEditor.Automation.cs b/Engine/Source/Programs/AutomationTool/CookedEditor/CookedEditor.Automation.cs
--- a/Engine/Source/Programs/AutomationTool/CookedEditor/CookedEditor.Automation.cs
+++ b/Engine/Source/Programs/AutomationTool/CookedEditor/CookedEditor.Automation.cs
@@ -291,8 +291,8 @@ public class ModifyStageContext
 	private void AddUFSFilesToList(List<FileReference> FileList, string Extension, DeploymentContext SC)
 	{
 		// look in SC and UFSFiles
-		FileList.AddRange(SC.FilesToStage.UFSFiles.Keys.Where(x => x.HasExtension(Extension)).Select(y => DeploymentContext.UnmakeRelativeStagedReference(SC, y)));
-		FileList.AddRange(UFSFilesToStage.Where(x => x.GetExtension().Equals(Extension, StringComparison.InvariantCultureIgnoreCase)));
+		FileList.AddRange(SC.FilesToStage.UFSFiles.Keys.Where(x => x.Name.EndsWith(Extension, StringComparison.OrdinalIgnoreCase)).Select(y => DeploymentContext.UnmakeRelativeStagedReference(SC, y)));
+		FileList.AddRange(UFSFilesToStage.Where(x => x.FullName.EndsWith(Extension, StringComparison.InvariantCultureIgnoreCase)));
 	}
 
 	private void UncookMaps(DeploymentContext SC)
@@ -331,6 +331,7 @@ public class ModifyStageContext
 			// UAT needs uplugin and ini files, so make sure they are not in the .pak
 			AddUFSFilesToList(NonUFSFilesToStage, ".uplugin", SC);
 			AddUFSFilesToList(NonUFSFilesToStage, ".ini", SC);
+			AddUFSFilesToList(NonUFSFilesToStage, "SDK.json", SC);
 			NonUFSFilesToStage = NonUFSFilesToStage.Where(x => x.GetFileName() != "BinaryConfig.ini").ToList();
 		}
 	}
@@ -767,7 +768,7 @@ public class MakeCookedEditor : BuildCommand
 		// plugins are already handled in the Plugins staging code
 		List<string> RootFoldersToStrip = new List<string> { "source", "plugins" };//, "binaries" };
 		List<string> SubFoldersToStrip = new List<string> { "source", "intermediate", "tests", "binaries" + Path.DirectorySeparatorChar + HostPlatform.Current.HostEditorPlatform.ToString().ToLower() };
-		List<string> RootNonUFSFolders = new List<string> { "shaders", "binaries", "build", "extras" };
+		List<string> RootNonUFSFolders = new List<string> { "shaders", "binaries", "build", "extras", "config" };
 
 
 		if (!Context.bStageShaderDirs)
diff --git a/Engine/Source/Programs/Shared/EpicGames.Build/UEBuildPlatformSDK.cs b/Engine/Source/Programs/Shared/EpicGames.Build/UEBuildPlatformSDK.cs
--- a/Engine/Source/Programs/Shared/EpicGames.Build/UEBuildPlatformSDK.cs
+++ b/Engine/Source/Programs/Shared/EpicGames.Build/UEBuildPlatformSDK.cs
@@ -9,6 +9,9 @@ using System.Linq;
 using System.Text;
 using System.Text.Json;
 using Microsoft.Extensions.Logging;
+using UnrealBuildBase;
+using EpicGames.Core;
+using System.Runtime.InteropServices;
 
 namespace EpicGames.Core
 {
@@ -310,21 +313,7 @@ namespace EpicGames.Core
 			{
 				throw new Exception(string.Format("Re-registering SDK for {0}. All Platforms must have a unique SDK object", PlatformName));
 			}
-
-			// check that none of our override versions conflict with default version, by using MainVersion unless overridden and we aren't _also_
-			// requiring default version
-			if (!PlatformSDKMainVersions.ContainsKey(PlatformName))
-			{
-				PlatformSDKMainVersions.Add(PlatformName, SDK.GetMainVersionInternal());
-			}
-			else if (PlatformsThatNeedDefaultSDK != null && PlatformsThatNeedDefaultSDK.Contains(PlatformName))
-			{
-				if (!PlatformSDKMainVersions[PlatformName].Equals(SDK.GetMainVersionInternal(), StringComparison.OrdinalIgnoreCase))
-				{
-					throw new Exception($"A target overrode SDK for {PlatformName} to version {PlatformSDKMainVersions[PlatformName]}, but another target uses the default SDK version {SDK.GetMainVersionInternal()}");
-				}
-			}
-
+			
 			SDKRegistry.Add(PlatformName, SDK);
 
 			SDK.Init(PlatformName, bIsSdkAllowedOnHost);
@@ -335,6 +324,9 @@ namespace EpicGames.Core
 			PlatformName = InPlatformName;
 			bIsSdkAllowedOnHost = bInIsSdkAllowedOnHost;
 
+			// load the SDK config file
+			LoadJsonFile(PlatformName);
+
 			// if the parent set up autosdk, the env vars will be wrong, but we can still get the manual SDK version from before it was setup
 			string? ParentManualSDKVersions = Environment.GetEnvironmentVariable(GetPlatformManualSDKSetupEnvVar());
 			if (!string.IsNullOrEmpty(ParentManualSDKVersions))
@@ -390,19 +382,21 @@ namespace EpicGames.Core
 			get { return SDKRegistry.Values.ToArray(); }
 		}
 
+		// True if at least one platform has had its version changed from the standard SDK version, which means it may not be compatible with other projects
+		public static bool bHasAnySDKOverride = false;
+
+		// True if this SDK has had its version changed from the standard SDK version, which means it may not be compatible with other projects
+		public bool bHasSDKOverride = false;
+
+		// Contains a list of projects that had per-project SDK overrides, used for validating conflicting SDKs
+		public List<FileReference> ProjectsThatOverrodeSDK = new();
+	
 		// String name of the platform (will match an UnrealTargetPlatform)
 		public string? PlatformName;
 
 		// True if this Sdk is allowed to be used by this host - if not, we can skip a lot 
 		public bool bIsSdkAllowedOnHost;
 
-		// Retrieves the version of the default SDK that is preferred choice, as well as the single AutoSDK version that will be used
-		public string GetMainVersion()
-		{
-			// at this point, the PlatformSDKMainVersions map is fully set up with default version, or with overrides that have been verified to not conflict
-			return PlatformSDKMainVersions[PlatformName!];
-		}
-
 		public SDKCollection GetAllSDKInfo()
 		{
 			SDKCollection AllSdks = new SDKCollection(this);
@@ -542,22 +536,28 @@ namespace EpicGames.Core
 		/// Return the SDK version that the platform wants to use (AutoSDK dir must match this, full SDKs can be in a valid range)
 		/// </summary>
 		/// <returns></returns>
-		protected abstract string GetMainVersionInternal();
+		public virtual string GetMainVersion()
+		{
+			// by default, look up in SDK config value
+			return GetRequiredVersionFromConfig("MainVersion");
+		}
 
 		/// <summary>
 		/// Gets the valid string range of Sdk versions. TryConvertVersionToInt() will need to succeed to make this usable for range checks
 		/// </summary>
 		/// <param name="MinVersion">Smallest version allowed</param>
 		/// <param name="MaxVersion">Largest version allowed (inclusive)</param>
-		protected virtual void GetValidVersionRange(out string MinVersion, out string MaxVersion)
+		protected virtual void GetValidVersionRange(out string? MinVersion, out string? MaxVersion)
 		{
-			throw new Exception($"This platform's Sdk class ({GetType().Name}) must implement either GetValidVersionRange() or GetValidVersions(). If this triggers, and GetValidVersions() is implemented, that means GetValidVersionRange was called directly");
+			// by default, use SDK config file values
+			MinVersion = GetVersionFromConfig("MinVersion");
+			MaxVersion = GetVersionFromConfig("MaxVersion");
 		}
 
 		protected virtual SDKCollection GetValidVersions()
 		{
 			// if the platform doesn't override this, then it only has one sdk, so put it into the collection as the single sdk (this is very much the standard behavior)
-			string MinVersion, MaxVersion;
+			string? MinVersion, MaxVersion;
 			GetValidVersionRange(out MinVersion, out MaxVersion);
 			return new SDKCollection(MinVersion, MaxVersion, this);
 		}
@@ -646,7 +646,8 @@ namespace EpicGames.Core
 		/// <returns></returns>
 		public virtual string GetAutoSDKDirectoryForMainVersion()
 		{
-			return GetMainVersion();
+			// if there is an AutoSDKDirectory property, use it, otherwise, use the MainVersion string
+			return GetVersionFromConfig("AutoSDKDirectory") ?? GetMainVersion();
 		}
 
 
@@ -818,29 +819,19 @@ namespace EpicGames.Core
 		/// Allow the platform SDK to override the name it will use in AutoSDK, but default to the platform name
 		/// </summary>
 		/// <returns>The name of the directory to use inside the AutoSDK system</returns>
-		public virtual string? GetAutoSDKPlatformName()
+		public virtual string GetAutoSDKPlatformName()
 		{
-			return PlatformName;
+			return GetVersionFromConfig("AutoSDKPlatform") ?? PlatformName!;
 		}
 
 		#endregion
 
 		#region Per-project SDK support
 
-		private static Dictionary<string, string> PlatformSDKMainVersions = new();
-		private static HashSet<string>? PlatformsThatNeedDefaultSDK = new();
-		public static void InitializePerProjectSDKVersions(Dictionary<string, string> PerPlatformSDKOverrides, HashSet<string> InPlatformsThatNeedDefaultSDK)
+		private static List<FileReference> ProjectsToCheckForSDKOverrides = new();
+		public static void InitializePerProjectSDKVersions(IEnumerable<FileReference> ProjectsToCheckForOverrides)
 		{
-			if (PlatformSDKMainVersions.Count > 0)
-			{
-				throw new Exception("InitializePerProjectSDKVersions() was called multiple times, which is not allowed");
-			}
-
-			// cache off the overrides - Register will fill in the rest
-			PlatformSDKMainVersions = PerPlatformSDKOverrides;
-
-			// remember the platforms that need the default version, which we will use later to detect conflicting versions in Register			
-			PlatformsThatNeedDefaultSDK = InPlatformsThatNeedDefaultSDK;
+			ProjectsToCheckForSDKOverrides = ProjectsToCheckForOverrides.ToList();
 		}
 
 		#endregion
@@ -959,6 +950,155 @@ namespace EpicGames.Core
 			return null;
 		}
 
+		// cached SDK info from the SDK.json file
+		private Dictionary<string, string> ConfigSDKVersions = new(StringComparer.OrdinalIgnoreCase);
+		private Dictionary<string, string> DefaultConfigSDKVersions = new(StringComparer.OrdinalIgnoreCase);
+
+		private void LoadJsonFile(string Platform)
+		{
+			// fixup for Windows
+			if (Platform == "Win64")
+			{
+				Platform = "Windows";
+			}
+
+			Func<DirectoryReference, bool, FileReference?> MakeConfigFilename = (RootDir, bIsRequired) =>
+			{
+				FileReference PlatformExtensionLocation = FileReference.Combine(RootDir, "Platforms", Platform, "Config", $"{Platform}_SDK.json");
+				if (FileReference.Exists(PlatformExtensionLocation))
+				{
+					return PlatformExtensionLocation;
+				}
+				FileReference StandardLocation = FileReference.Combine(RootDir, "Config", Platform, $"{Platform}_SDK.json");
+				if (FileReference.Exists(StandardLocation))
+				{
+					return StandardLocation;
+				}
+				if (bIsRequired)
+				{
+					throw new Exception($"Failed to find required SDK.json for {Platform}. Looked in '{StandardLocation}' and '{PlatformExtensionLocation}'.");
+				}
+				return null;
+			};
+
+			FileReference EngineSDKConfigFile = MakeConfigFilename(Unreal.EngineDirectory, true)!;
+
+			// load the file, along with any chained group file
+			ProcessJsonFile(EngineSDKConfigFile, ConfigSDKVersions);
+
+			// copy off the versions to the defaults, so we can later get the 
+			foreach (string Key in ConfigSDKVersions.Keys)
+			{
+				DefaultConfigSDKVersions[Key] = ConfigSDKVersions[Key];
+			}
+
+			// now read overrides into the array
+			foreach (FileReference ProjectFile in ProjectsToCheckForSDKOverrides)
+			{
+				FileReference? ProjectSDKConfigFile = MakeConfigFilename(ProjectFile.Directory, false);
+				if (ProjectSDKConfigFile != null)
+				{
+					// load a project's SDK file if thre is one, into a temp dictionary
+					Dictionary<string, string> OverrideConfigSDKVersions = new(StringComparer.OrdinalIgnoreCase);
+					ProcessJsonFile(ProjectSDKConfigFile, OverrideConfigSDKVersions);
+
+					// currently only care about MainVersion in the overrides
+					string? OverrideMainVersion;
+					const string MainVersionKey = "MainVersion";
+					if (OverrideConfigSDKVersions.TryGetValue(MainVersionKey, out OverrideMainVersion))
+					{
+						// if different from default, then remmber it, and mark that it's been overridden
+						if (!OverrideMainVersion.Equals(DefaultConfigSDKVersions[MainVersionKey], StringComparison.OrdinalIgnoreCase))
+						{
+							// now check if it was already overridden, in which case we have a conflict we can't resolve, so error
+							if (ProjectsThatOverrodeSDK.Count > 0 && !OverrideMainVersion.Equals(ConfigSDKVersions[MainVersionKey], StringComparison.OrdinalIgnoreCase))
+							{
+								throw new Exception($"Project {ProjectFile.GetFileNameWithoutAnyExtensions()} wants to override {Platform} SDK to {OverrideMainVersion}, but it was already overridden to version {ConfigSDKVersions[MainVersionKey]}");							
+							}
+
+							Logger.LogWarning("Project {Project} is overriding {Platform} SDK to {OverrideMainVersion}", ProjectFile.GetFileNameWithoutAnyExtensions(), Platform, OverrideMainVersion);
+
+							ConfigSDKVersions[MainVersionKey] = OverrideMainVersion;
+							UEBuildPlatformSDK.bHasAnySDKOverride = true;
+							bHasSDKOverride = true;
+							ProjectsThatOverrodeSDK.Add(ProjectFile);
+						}
+					}
+				}
+			}
+		}
+
+		private void ProcessJsonFile(FileReference SDKConfigFile, Dictionary<string, string> VersionMap)
+		{
+			string Contents = FileReference.ReadAllText(SDKConfigFile);
+
+			// load the json into a dictionary
+			JsonSerializerOptions Options = new()
+			{
+				AllowTrailingCommas = true,
+				ReadCommentHandling = JsonCommentHandling.Skip,
+			};
+			Dictionary<string, string>? LoadedDictionary = JsonSerializer.Deserialize<Dictionary<string, string>>(Contents, Options);
+			if (LoadedDictionary == null)
+			{
+				throw new Exception($"Failed to parse SDK version file '{SDKConfigFile}'");
+			}
+
+			// look for special key to chain to other - after processing everything else (we only add from the "parent" settings not already in the map)
+			List<string> Parents = new();
+			foreach (string Key in LoadedDictionary.Keys)
+			{
+				if (Key.Equals("ParentSDKFile", StringComparison.OrdinalIgnoreCase))
+				{
+					Parents.Add(LoadedDictionary[Key]);
+				}
+				// add this key if it's not already in the (case-insensitive) dictionary 
+				else if (!VersionMap.ContainsKey(Key))
+				{
+					VersionMap.Add(Key, LoadedDictionary[Key]);
+				}
+			}
+
+			// now load parents, filling in unset properties
+			foreach (string Parent in Parents)
+			{
+				FileReference ParentConfigFile = FileReference.Combine(SDKConfigFile.Directory, Parent);
+				ProcessJsonFile(ParentConfigFile, VersionMap);
+			}
+		}
+
+		protected string GetRequiredVersionFromConfig(string VersionName)
+		{
+			// when bIsRequired is true, then we know it will return non-null
+			return GetVersionFromConfig(VersionName, bIsRequired:true)!;
+		}
+
+		protected string? GetVersionFromConfig(string VersionName, bool bIsRequired=false)
+		{
+			string? Version;
+			string HostPlatform = "Win64";
+			if (OperatingSystem.IsMacOS())
+			{
+				HostPlatform = "Mac";
+			}
+			else if (OperatingSystem.IsLinux())
+			{
+				HostPlatform = "Linux";
+			}
+			// look up both Version_Host and Version (Host specific version wins)
+			string HostSpecificVersionName = $"{VersionName}_{HostPlatform}";
+			if (ConfigSDKVersions!.TryGetValue(HostSpecificVersionName, out Version) || ConfigSDKVersions.TryGetValue(VersionName, out Version))
+			{
+				return Version;
+			}
+
+			if (bIsRequired)
+			{
+				throw new Exception($"Unable to find required SDK version '{VersionName}' for platform {PlatformName}. Check your SDK.json files");
+			}
+			return null;
+		}
+
 		#endregion
 
 		// AutoSDKs handling portion
@@ -1071,7 +1211,7 @@ namespace EpicGames.Core
 			{
 				if (SDKRoot != "")
 				{
-					SDKPath = Path.Combine(SDKRoot, "Host" + GetAutoSDKHostPlatform(), GetAutoSDKPlatformName()!);
+					SDKPath = Path.Combine(SDKRoot, "Host" + GetAutoSDKHostPlatform(), GetAutoSDKPlatformName());
 				}
 			}
 			return SDKPath;
diff --git a/Engine/Source/Programs/UnrealBuildTool/Configuration/TargetDescriptor.cs b/Engine/Source/Programs/UnrealBuildTool/Configuration/TargetDescriptor.cs
--- a/Engine/Source/Programs/UnrealBuildTool/Configuration/TargetDescriptor.cs
+++ b/Engine/Source/Programs/UnrealBuildTool/Configuration/TargetDescriptor.cs
@@ -238,6 +238,13 @@ namespace UnrealBuildTool
 			return TargetDescriptors;
 		}
 
+		public static IEnumerable<FileReference?> ParseCommandLineForProjects(CommandLineArguments Arguments, ILogger Logger)
+		{
+			List<TargetDescriptor> SimplifiedTargetDescriptors = new List<TargetDescriptor>();
+			ParseCommandLine(Arguments, false, true, false, bCreateSimplifiedDescriptors: true, SimplifiedTargetDescriptors, Logger);
+			return SimplifiedTargetDescriptors.Select(x => x.ProjectFile);
+		}
+
 		/// <summary>
 		/// Parse a list of target descriptors from the command line
 		/// </summary>
@@ -248,6 +255,12 @@ namespace UnrealBuildTool
 		/// <param name="TargetDescriptors">Receives the list of parsed target descriptors</param>
 		/// <param name="Logger">Logger for output</param>
 		public static void ParseCommandLine(CommandLineArguments Arguments, bool bUsePrecompiled, bool bSkipRulesCompile, bool bForceRulesCompile, List<TargetDescriptor> TargetDescriptors, ILogger Logger)
+		{
+			// call the internal one
+			ParseCommandLine(Arguments, bUsePrecompiled, bSkipRulesCompile, bForceRulesCompile, bCreateSimplifiedDescriptors: false, TargetDescriptors, Logger);
+		}
+
+		private static void ParseCommandLine(CommandLineArguments Arguments, bool bUsePrecompiled, bool bSkipRulesCompile, bool bForceRulesCompile, bool bCreateSimplifiedDescriptors, List<TargetDescriptor> TargetDescriptors, ILogger Logger)
 		{
 			List<string> TargetLists;
 			Arguments = Arguments.Remove("-TargetList=", out TargetLists);
@@ -267,7 +280,7 @@ namespace UnrealBuildTool
 						if (TrimLine.Length > 0 && TrimLine[0] != ';')
 						{
 							CommandLineArguments NewArguments = Arguments.Append(CommandLineArguments.Split(TrimLine));
-							ParseCommandLine(NewArguments, bUsePrecompiled, bSkipRulesCompile, bForceRulesCompile, TargetDescriptors, Logger);
+							ParseCommandLine(NewArguments, bUsePrecompiled, bSkipRulesCompile, bForceRulesCompile, bCreateSimplifiedDescriptors, TargetDescriptors, Logger);
 						}
 					}
 				}
@@ -275,13 +288,13 @@ namespace UnrealBuildTool
 				foreach (string Target in Targets)
 				{
 					CommandLineArguments NewArguments = Arguments.Append(CommandLineArguments.Split(Target));
-					ParseCommandLine(NewArguments, bUsePrecompiled, bSkipRulesCompile, bForceRulesCompile, TargetDescriptors, Logger);
+					ParseCommandLine(NewArguments, bUsePrecompiled, bSkipRulesCompile, bForceRulesCompile, bCreateSimplifiedDescriptors, TargetDescriptors, Logger);
 				}
 			}
 			else
 			{
 				// Otherwise just process the whole command line together
-				ParseSingleCommandLine(Arguments, bUsePrecompiled, bSkipRulesCompile, bForceRulesCompile, TargetDescriptors, Logger);
+				ParseSingleCommandLine(Arguments, bUsePrecompiled, bSkipRulesCompile, bForceRulesCompile, bCreateSimplifiedDescriptors, TargetDescriptors, Logger);
 			}
 		}
 
@@ -295,6 +308,12 @@ namespace UnrealBuildTool
 		/// <param name="TargetDescriptors">List of target descriptors</param>
 		/// <param name="Logger">Logger for output</param>
 		public static void ParseSingleCommandLine(CommandLineArguments Arguments, bool bUsePrecompiled, bool bSkipRulesCompile, bool bForceRulesCompile, List<TargetDescriptor> TargetDescriptors, ILogger Logger)
+		{
+			// call the internal one
+			ParseSingleCommandLine(Arguments, bUsePrecompiled, bSkipRulesCompile, bForceRulesCompile, bCreateSimplifiedDescriptors: false, TargetDescriptors, Logger);
+		}
+
+		private static void ParseSingleCommandLine(CommandLineArguments Arguments, bool bUsePrecompiled, bool bSkipRulesCompile, bool bForceRulesCompile, bool bCreateSimplifiedDescriptors, List<TargetDescriptor> TargetDescriptors, ILogger Logger)
 		{
 			List<UnrealTargetPlatform> Platforms = new List<UnrealTargetPlatform>();
 			List<UnrealTargetConfiguration> Configurations = new List<UnrealTargetConfiguration>();
@@ -368,6 +387,32 @@ namespace UnrealBuildTool
 				throw new BuildException("No configurations specified for target");
 			}
 
+			// make a single simple descriptor with out a lot of processing or rules assembly creation
+			if (bCreateSimplifiedDescriptors)
+			{
+				// we expect either a targetname from the above loop, or -targetype combined with -project, so we should have -project or a TargetName already
+				string? TargetName = TargetNames.FirstOrDefault();
+				if (TargetName == null)
+				{
+					if (ProjectFile == null)
+					{
+						Logger.LogInformation("When looking for per-project SDK overrides, we got a commandline that could not be used to find a target project ({CmdLine})", Arguments.ToString());
+						return;
+					}
+					// just assume Game target type, it doesn't matter
+					TargetName = ProjectFile.GetFileNameWithoutAnyExtensions() + "Game";
+				}
+				if (ProjectFile == null)
+				{
+					NativeProjects.TryGetProjectForTarget(TargetName, Logger, out ProjectFile);
+				}
+
+				TargetDescriptors.Add(new TargetDescriptor(ProjectFile, TargetName, Platforms.First(), Configurations.First(), new UnrealArchitectures(UnrealArch.X64), Arguments));
+
+				// we are done now with simple descriptors
+				return;
+			}
+
 			// Make sure the project file exists, and make sure we're using the correct case.
 			if (ProjectFile != null)
 			{
@@ -395,10 +440,7 @@ namespace UnrealBuildTool
 					}
 				}
 
-				// there are times that this can be run before we have setup build platforms (ie Per-Project AutoSDK version), in which case we
-				// cannot get Architecture info, but it's not needed yet
-				UEBuildPlatform? BuildPlatform;
-				UEBuildPlatform.TryGetBuildPlatform(Platform, out BuildPlatform);
+				UEBuildPlatform BuildPlatform = UEBuildPlatform.GetBuildPlatform(Platform);
 
 				// Parse the architecture parameter, or use null to look up platform defaults later
 				string ParamArchitectureList = Arguments.GetStringOrDefault("-Architecture=", "") + Arguments.GetStringOrDefault("-Architectures=", "");
@@ -461,11 +503,6 @@ namespace UnrealBuildTool
 						{
 							Architectures = ParamArchitectures;
 						}
-						else if (BuildPlatform == null)
-						{
-							// if we can't use BuildPlatform yet, we just use any random architecture, since we need _something_
-							Architectures = new UnrealArchitectures(UnrealArch.X64);
-						}
 						else
 						{
 							// ask the platform what achitectures it wants for this project
@@ -473,7 +510,7 @@ namespace UnrealBuildTool
 						}
 
 						// If the platform wants a target for each architecture, make a target descriptor for each architecture, otherwise one target for all architectures
-						if (BuildPlatform != null && BuildPlatform.ArchitectureConfig.Mode == UnrealArchitectureMode.OneTargetPerArchitecture)
+						if (BuildPlatform.ArchitectureConfig.Mode == UnrealArchitectureMode.OneTargetPerArchitecture)
 						{
 							foreach (UnrealArch Architecture in Architectures.Architectures)
 							{
diff --git a/Engine/Source/Programs/UnrealBuildTool/Configuration/TargetRules.cs b/Engine/Source/Programs/UnrealBuildTool/Configuration/TargetRules.cs
--- a/Engine/Source/Programs/UnrealBuildTool/Configuration/TargetRules.cs
+++ b/Engine/Source/Programs/UnrealBuildTool/Configuration/TargetRules.cs
@@ -1007,6 +1007,15 @@ namespace UnrealBuildTool
 			get => bNeedsExtraShaderFormatsOverride ?? (bForceBuildShaderFormats || bBuildDeveloperTools) && (Type == TargetType.Editor || Type == TargetType.Program);
 		}
 
+		/// <summary>
+		/// Checks the above flags to see if target SDKs are likely to be compiled in
+		/// </summary>
+		private bool bUsesTargetSDKs
+		{
+			get => bNeedsExtraShaderFormats || bBuildDeveloperTools || bBuildTargetDeveloperTools || bForceBuildShaderFormats || bForceBuildTargetPlatforms;
+		}
+
+
 		/// <summary>
 		/// Whether we should compile SQLite using the custom "Unreal" platform (true), or using the native platform (false).
 		/// </summary>
@@ -2388,6 +2397,51 @@ namespace UnrealBuildTool
 		[ConfigSubObject]
 		public WindowsTargetRules WindowsPlatform; // Requires 'this' parameter; initialized in constructor
 
+		/// <summary>
+		/// Checks if the target needs to care about SDK version - if not, we can share the target between projects that are on differnet SDK versions.
+		/// For instance, ShaderCompileWorker is very tied to, say, console SDK versions, but UnrealPak doesn't have SDK things compiled in
+		/// </summary>
+		/// <returns></returns>
+		public bool IsSDKVersionRelevant(UnrealTargetPlatform SDKPlatform)
+		{
+			// we always care bout the SDK of the platform we are compiling for (Platform is the platform this target is being built for)
+			if (this.Platform == SDKPlatform)
+			{
+				return true;
+			}
+
+			// at this point, we are checking against an SDK that we are not compiling for, ie a target platform sdk, so check if the Target has overriden it
+			if (bAreTargetSDKVersionsRelevantOverride != null)
+			{
+				return bAreTargetSDKVersionsRelevantOverride.Value;
+			}
+
+			// developer tools are what pull in SDK versions for other platforms, so Programs that don't use dev tools don't care about versions
+			if (Type == global::UnrealBuildTool.TargetType.Program || Type == global::UnrealBuildTool.TargetType.Editor)
+			{
+				return this.bUsesTargetSDKs;
+			}
+
+			// games/client/servers are not expected to contain other platform SDK bits in them, so we don't need to care
+			return false;
+		}
+
+		/// <summary>
+		/// Allow a target to override whether or not AreSDKVersionsRelevant uses default logic
+		/// </summary>
+		protected bool? bAreTargetSDKVersionsRelevantOverride = null;
+
+		/// <summary>
+		/// Checks if the target can use non-standard SDK versions. For instance, the editor is usually compiled as a modular build shared between projects,
+		/// so if one projet uses SDK version 1, and another uses 2, then the generated engine modules will use the most recently built SDK, which could break
+		/// </summary>
+		/// <returns></returns>
+		public bool AllowsPerProjectSDKVersion()
+		{
+			// modular target with TargetBuildEnvironment.Shared build type cannot allow per-project SDKs
+			return LinkType == TargetLinkType.Monolithic || BuildEnvironment == TargetBuildEnvironment.Unique;
+		}
+
 		/// <summary>
 		/// Create a TargetRules instance
 		/// </summary>
diff --git a/Engine/Source/Programs/UnrealBuildTool/Configuration/UEBuildPlatform.cs b/Engine/Source/Programs/UnrealBuildTool/Configuration/UEBuildPlatform.cs
--- a/Engine/Source/Programs/UnrealBuildTool/Configuration/UEBuildPlatform.cs
+++ b/Engine/Source/Programs/UnrealBuildTool/Configuration/UEBuildPlatform.cs
@@ -10,6 +10,7 @@ using System.Text;
 using EpicGames.Core;
 using Microsoft.Extensions.Logging;
 using OpenTracing.Util;
+using static UnrealBuildTool.DataDrivenPlatformInfo;
 
 namespace UnrealBuildTool
 {
@@ -239,7 +240,12 @@ namespace UnrealBuildTool
 
 			// check DDPI to see if the platform is enabled on this host platform
 			string IniPlatformName = ConfigHierarchy.GetIniPlatformName(Platform);
-			bool bIsEnabled = DataDrivenPlatformInfo.GetDataDrivenInfoForPlatform(IniPlatformName)?.bIsEnabled ?? true;
+			bool bIsEnabled = true; 
+			ConfigDataDrivenPlatformInfo? DDPI = DataDrivenPlatformInfo.GetDataDrivenInfoForPlatform(IniPlatformName);
+			if (DDPI != null)
+			{
+				bIsEnabled = DDPI.bIsEnabled;
+			}
 
 			// set up the SDK if the platform is enabled
 			UEBuildPlatformSDK.RegisterSDKForPlatform(SDK, Platform.ToString(), bIsEnabled);
@@ -255,7 +261,6 @@ namespace UnrealBuildTool
 		private static void InitializePerPlatformSDKs(string[] Args, bool bArgumentsAreForUBT, ILogger Logger)
 		{
 			Dictionary<string, string> PlatformToVersionMap = new();
-			HashSet<string> PlatformsThatNeedDefaultSDK = new();
 
 			IEnumerable<FileReference?> ProjectFiles;
 
@@ -278,46 +283,10 @@ namespace UnrealBuildTool
 
 				// get the project files for all targets - we allow null uproject files which means to use the defaults
 				// (same as a uproject that has no override SDK versions set)
-				List<TargetDescriptor> AllTargets = TargetDescriptor.ParseCommandLine(CommandLine, BuildConfiguration.bUsePrecompiled, bSkipRulesCompile: true, bForceRulesCompile: false, Logger);
-				ProjectFiles = AllTargets.Select(x => x.ProjectFile);
+				ProjectFiles = TargetDescriptor.ParseCommandLineForProjects(CommandLine, Logger);
 			}
 
-			foreach (FileReference? ProjectFile in ProjectFiles)
-			{
-				// because we can't safely tell, at this point, if any platform SDKs will be used for a given target, we have to check all of them
-				foreach (UnrealTargetPlatform Platform in UnrealTargetPlatform.GetValidPlatforms())
-				{
-					string PlatformName = Platform.ToString();
-					string? OverrideSDKVersion = null;
-
-					if (ProjectFile != null)
-					{
-						ConfigHierarchy Ini = ConfigCache.ReadHierarchy(ConfigHierarchyType.Engine, ProjectFile.Directory, Platform);
-						if (Ini.TryGetValue("OverrideSDK", "SDKVersion", out OverrideSDKVersion))
-						{
-							Logger.LogWarning("Project {Project} is overriding the active main SDK for {Platform} to {Override}", ProjectFile.GetFileNameWithoutAnyExtensions(), Platform, OverrideSDKVersion);
-
-							if (PlatformToVersionMap.ContainsKey(PlatformName) && PlatformToVersionMap[PlatformName] != OverrideSDKVersion)
-							{
-								throw new BuildException("Attempting to build with mismatching SDK versions! [MORE INFO HERE, BUT ALSO CHECK THAT ALL PROJECTS USE SAME VERSION]");
-							}
-						}
-					}
-					
-					// if we have a null version here, that means to use the default, later when the PlatformSDK is created
-					if (OverrideSDKVersion == null)
-					{
-						// mark _something_ in the map so we can detect mismatches
-						PlatformsThatNeedDefaultSDK.Add(PlatformName);
-					}
-					else
-					{
-						PlatformToVersionMap[PlatformName] = OverrideSDKVersion;
-					}
-				}
-			}
-
-			UEBuildPlatformSDK.InitializePerProjectSDKVersions(PlatformToVersionMap, PlatformsThatNeedDefaultSDK);
+			UEBuildPlatformSDK.InitializePerProjectSDKVersions(ProjectFiles.OfType<FileReference>());
 		}
 
 		/// <summary>
@@ -359,11 +328,10 @@ namespace UnrealBuildTool
 				InstalledPlatformInfo.Initialize();
 			}
 
-// disabling _for now_ as it is causing issues when using -TargetType param, causing too much stuff to be created early
-//			using (GlobalTracer.Instance.BuildSpan("Initializing PerPlatformSDKs").StartActive())
-//			{
-//				InitializePerPlatformSDKs(ArgumentsForPerPlatform, bArgumentsAreForUBT, Logger);
-//			}
+			using (GlobalTracer.Instance.BuildSpan("Initializing PerPlatformSDKs").StartActive())
+			{
+				InitializePerPlatformSDKs(ArgumentsForPerPlatform, bArgumentsAreForUBT, Logger);
+			}
 
 			// Find and register all tool chains and build platforms that are present
 			Type[] AllTypes;
diff --git a/Engine/Source/Programs/UnrealBuildTool/Configuration/UEBuildTarget.cs b/Engine/Source/Programs/UnrealBuildTool/Configuration/UEBuildTarget.cs
--- a/Engine/Source/Programs/UnrealBuildTool/Configuration/UEBuildTarget.cs
+++ b/Engine/Source/Programs/UnrealBuildTool/Configuration/UEBuildTarget.cs
@@ -1553,6 +1553,15 @@ namespace UnrealBuildTool
 			return bCompileMonolithic;  // @todo ubtmake: We need to make sure this function and similar things aren't called in assembler mode
 		}
 
+		/// <summary>
+		/// Whether this target was compiled with an overridden sdk (usually from a platform .ini specifying an alternate SDK version)
+		/// </summary>
+		/// <returns></returns>
+		public bool BuiltWithOverriddenSDKs()
+		{
+			return UEBuildPlatformSDK.bHasAnySDKOverride;
+		}
+
 		/// <summary>
 		/// Constructor.
 		/// </summary>
diff --git a/Engine/Source/Programs/UnrealBuildTool/Platform/Android/AndroidPlatformSDK.Versions.cs b/Engine/Source/Programs/UnrealBuildTool/Platform/Android/AndroidPlatformSDK.Versions.cs
--- a/Engine/Source/Programs/UnrealBuildTool/Platform/Android/AndroidPlatformSDK.Versions.cs
+++ b/Engine/Source/Programs/UnrealBuildTool/Platform/Android/AndroidPlatformSDK.Versions.cs
@@ -4,24 +4,12 @@ using EpicGames.Core;
 
 namespace UnrealBuildTool
 {
+	/////////////////////////////////////////////////////////////////////////////////////
+	// If you are looking for any version numbers not listed here, see GDK_SDK.json
+	/////////////////////////////////////////////////////////////////////////////////////
+
 	partial class AndroidPlatformSDK : UEBuildPlatformSDK
 	{
-		protected override string GetMainVersionInternal()
-		{
-			return "r25b";
-		}
-
-		public override string GetAutoSDKDirectoryForMainVersion()
-		{
-			return "-26";
-		}
-
-		protected override void GetValidVersionRange(out string MinVersion, out string MaxVersion)
-		{
-			MinVersion = "r25b";
-			MaxVersion = "r25b";
-		}
-
 		protected override void GetValidSoftwareVersionRange(out string? MinVersion, out string? MaxVersion)
 		{
 			MinVersion = MaxVersion = null;
diff --git a/Engine/Source/Programs/UnrealBuildTool/Platform/Linux/LinuxPlatformSDK.Versions.cs b/Engine/Source/Programs/UnrealBuildTool/Platform/Linux/LinuxPlatformSDK.Versions.cs
--- a/Engine/Source/Programs/UnrealBuildTool/Platform/Linux/LinuxPlatformSDK.Versions.cs
+++ b/Engine/Source/Programs/UnrealBuildTool/Platform/Linux/LinuxPlatformSDK.Versions.cs
@@ -4,20 +4,12 @@ using EpicGames.Core;
 
 namespace UnrealBuildTool
 {
+	/////////////////////////////////////////////////////////////////////////////////////
+	// If you are looking for any version numbers not listed here, see Linux_SDK.json
+	/////////////////////////////////////////////////////////////////////////////////////
+
 	partial class LinuxPlatformSDK : UEBuildPlatformSDK
 	{
-		protected override string GetMainVersionInternal()
-		{
-			return "v21_clang-15.0.1-centos7";
-		}
-
-		protected override void GetValidVersionRange(out string MinVersion, out string MaxVersion)
-		{
-			// all that matters is the number after the v, according to TryConvertVersionToInt()
-			MinVersion = "v21_clang-15.0.1-centos7";
-			MaxVersion = "v21_clang-15.0.1-centos7";
-		}
-
 		protected override void GetValidSoftwareVersionRange(out string? MinVersion, out string? MaxVersion)
 		{
 			MinVersion = MaxVersion = null;
diff --git a/Engine/Source/Programs/UnrealBuildTool/Platform/Linux/LinuxPlatformSDK.cs b/Engine/Source/Programs/UnrealBuildTool/Platform/Linux/LinuxPlatformSDK.cs
--- a/Engine/Source/Programs/UnrealBuildTool/Platform/Linux/LinuxPlatformSDK.cs
+++ b/Engine/Source/Programs/UnrealBuildTool/Platform/Linux/LinuxPlatformSDK.cs
@@ -21,11 +21,6 @@ namespace UnrealBuildTool
 		{
 		}
 
-		public override string GetAutoSDKPlatformName()
-		{
-			return TargetPlatformName;
-		}
-
 		protected override string? GetInstalledSDKVersion()
 		{
 			// @todo turnkey: ForceUseSystemCompiler() returns true, we should probably run system clang -V or similar to get version
@@ -121,7 +116,7 @@ namespace UnrealBuildTool
 			string? SDKRoot = Environment.GetEnvironmentVariable(SDKRootEnvVar);
 			if (!String.IsNullOrEmpty(SDKRoot))
 			{
-				string AutoSDKPath = Path.Combine(SDKRoot, "Host" + BuildHostPlatform.Current.Platform, TargetPlatformName, GetAutoSDKDirectoryForMainVersion(), LinuxPlatform.DefaultHostArchitecture.LinuxName);
+				string AutoSDKPath = Path.Combine(SDKRoot, "Host" + BuildHostPlatform.Current.Platform, GetAutoSDKPlatformName(), GetAutoSDKDirectoryForMainVersion(), LinuxPlatform.DefaultHostArchitecture.LinuxName);
 				if (DirectoryReference.Exists(new DirectoryReference(AutoSDKPath)))
 				{
 					return AutoSDKPath;
diff --git a/Engine/Source/Programs/UnrealBuildTool/Platform/Mac/ApplePlatformSDK.Versions.cs b/Engine/Source/Programs/UnrealBuildTool/Platform/Mac/ApplePlatformSDK.Versions.cs
--- a/Engine/Source/Programs/UnrealBuildTool/Platform/Mac/ApplePlatformSDK.Versions.cs
+++ b/Engine/Source/Programs/UnrealBuildTool/Platform/Mac/ApplePlatformSDK.Versions.cs
@@ -3,31 +3,17 @@
 using System;
 using EpicGames.Core;
 
+
+using UnrealBuildBase;
+
 namespace UnrealBuildTool
 {
+	/////////////////////////////////////////////////////////////////////////////////////
+	// If you are looking for any version numbers not listed here, see Apple_SDK.json
+	/////////////////////////////////////////////////////////////////////////////////////
+
 	internal partial class ApplePlatformSDK : UEBuildPlatformSDK
 	{
-		protected override string GetMainVersionInternal()
-		{
-			// Xcode prefered version?
-			return "14.1";
-		}
-
-		protected override void GetValidVersionRange(out string MinVersion, out string MaxVersion)
-		{
-			if (RuntimePlatform.IsMac)
-			{
-				MinVersion = "14.1.0";
-				MaxVersion = "15.9.9";
-			}
-			else
-			{
-				// @todo turnkey: these are MobileDevice .dll versions in Windows - to get the iTunes app version (12.3.4.1 etc) would need to hunt down the .exe
-				MinVersion = "1100.0.0.0";
-				MaxVersion = "8999.0";
-			}
-		}
-
 		protected override void GetValidSoftwareVersionRange(out string? MinVersion, out string? MaxVersion)
 		{
 			MinVersion = MaxVersion = null;
@@ -37,6 +23,5 @@ namespace UnrealBuildTool
 		/// The minimum macOS SDK version that a dynamic library can be built with
 		/// </summary>
 		public virtual Version MinimumDynamicLibSDKVersion => new Version("12.1");      // SDK used in Xcode13.2.1
-
 	}
 }
diff --git a/Engine/Source/Programs/UnrealBuildTool/Platform/Windows/MicrosoftPlatformSDK.Versions.cs b/Engine/Source/Programs/UnrealBuildTool/Platform/Windows/MicrosoftPlatformSDK.Versions.cs
--- a/Engine/Source/Programs/UnrealBuildTool/Platform/Windows/MicrosoftPlatformSDK.Versions.cs
+++ b/Engine/Source/Programs/UnrealBuildTool/Platform/Windows/MicrosoftPlatformSDK.Versions.cs
@@ -5,16 +5,12 @@ using EpicGames.Core;
 
 namespace UnrealBuildTool
 {
+	/////////////////////////////////////////////////////////////////////////////////////
+	// If you are looking for any version numbers not listed here, see Windows_SDK.json
+	/////////////////////////////////////////////////////////////////////////////////////
+
 	partial class MicrosoftPlatformSDK : UEBuildPlatformSDK
 	{
-		/// <summary>
-		/// The default Windows SDK version to be used, if installed.
-		/// </summary>
-		static readonly VersionNumber[] PreferredWindowsSdkVersions = new VersionNumber[]
-		{
-			VersionNumber.Parse("10.0.18362.0")
-		};
-
 		/// <summary>
 		/// The minimum Windows SDK version to be used. If this is null then it means there is no minimum version
 		/// </summary>
@@ -70,20 +66,6 @@ namespace UnrealBuildTool
 
 		static readonly VersionNumber MinimumIntelOneApiVersion = new VersionNumber(2023, 0, 0);
 
-		/// <inheritdoc/>
-		protected override string GetMainVersionInternal()
-		{
-			// preferred/main version is the top of the Preferred list - 
-			return PreferredWindowsSdkVersions.First().ToString();
-		}
-
-		/// <inheritdoc/>
-		protected override void GetValidVersionRange(out string MinVersion, out string MaxVersion)
-		{
-			MinVersion = "10.0.00000.0";
-			MaxVersion = "10.9.99999.0";
-		}
-
 		/// <inheritdoc/>
 		protected override void GetValidSoftwareVersionRange(out string? MinVersion, out string? MaxVersion)
 		{
diff --git a/Engine/Source/Programs/UnrealBuildTool/Platform/Windows/MicrosoftPlatformSDK.cs b/Engine/Source/Programs/UnrealBuildTool/Platform/Windows/MicrosoftPlatformSDK.cs
--- a/Engine/Source/Programs/UnrealBuildTool/Platform/Windows/MicrosoftPlatformSDK.cs
+++ b/Engine/Source/Programs/UnrealBuildTool/Platform/Windows/MicrosoftPlatformSDK.cs
@@ -203,8 +203,10 @@ namespace UnrealBuildTool
 			}
 			else if (MinVersion == null && MaxVersion == null)
 			{
-				WindowsSdkVersion = PreferredWindowsSdkVersions.FirstOrDefault(x => CachedWindowsSdkDirs!.ContainsKey(x));
-				if (WindowsSdkVersion == null && CachedWindowsSdkDirs!.Count > 0)
+				// convert the desired version into a VersionNumber
+				WindowsSdkVersion = VersionNumber.Parse(UEBuildPlatformSDK.GetSDKForPlatform("Win64")!.GetMainVersion());
+				// if it's not a known version, use the highest known
+				if (!CachedWindowsSdkDirs!.ContainsKey(WindowsSdkVersion))
 				{
 					WindowsSdkVersion = CachedWindowsSdkDirs.OrderBy(x => x.Key).Last().Key;
 				}
diff --git a/Engine/Source/Programs/UnrealBuildTool/System/RulesAssembly.cs b/Engine/Source/Programs/UnrealBuildTool/System/RulesAssembly.cs
--- a/Engine/Source/Programs/UnrealBuildTool/System/RulesAssembly.cs
+++ b/Engine/Source/Programs/UnrealBuildTool/System/RulesAssembly.cs
@@ -699,6 +699,8 @@ namespace UnrealBuildTool
 			{
 				UEBuildPlatform Platform = UEBuildPlatform.GetBuildPlatform(Rules.Platform);
 				Platform.ValidateTarget(Rules);
+
+				ValidateSDKs(Rules);
 			}
 
 			// Some platforms may *require* monolithic compilation...
@@ -715,6 +717,42 @@ namespace UnrealBuildTool
 			return Rules;
 		}
 
+		private void ValidateSDKs(TargetRules Rules)
+		{
+			// ask if each plaform SDK matters to this target
+
+			foreach (UnrealTargetPlatform Platform in UnrealTargetPlatform.GetValidPlatforms())
+			{
+				if (!Rules.IsSDKVersionRelevant(Platform))
+				{
+					continue;
+				}
+
+				UEBuildPlatformSDK? SDK = UEBuildPlatformSDK.GetSDKForPlatform(Platform.ToString());
+				if (SDK == null)
+				{
+					continue;
+				}
+
+				if (SDK.bHasSDKOverride)
+				{
+					// if the target doesn't allow for an override at all, error
+					if (!Rules.AllowsPerProjectSDKVersion())
+					{
+						throw new BuildException($"Target {Rules.Name} is being built with a overridden {Platform} SDK version to '{SDK.GetMainVersion()}', but this target is not allowed - likely due to a modular build using a Shared BuildEnvironment");
+					}
+
+					// check to see if the project _doesn't_ override the SDK version (if it did, this was already handled during early processing)
+					if (Rules.ProjectFile == null || !SDK.ProjectsThatOverrodeSDK.Contains(Rules.ProjectFile))
+					{
+						string OverrideProject = SDK.ProjectsThatOverrodeSDK[0].GetFileNameWithoutAnyExtensions();
+						throw new BuildException($"Target {Rules.Name} is using default {Platform} SDK version, but another target (probably {OverrideProject}) has overridden the SDK version to '{SDK.GetMainVersion()}'. If this target doesn't care about SDK versions, set 'bAreTargetSDKVersionsRelevantOverride = false' in your Target.cs file");
+					}
+				}
+			}
+		}
+
+
 		/// <summary>
 		/// Creates a target rules object for the specified target name.
 		/// </summary>
