Updated UHT to use the real C++ type for USTRUCT properties using a template wrapper struct

#rnx

diff --git a/Engine/Source/Programs/Shared/EpicGames.UHT/Types/Properties/UhtStructProperty.cs b/Engine/Source/Programs/Shared/EpicGames.UHT/Types/Properties/UhtStructProperty.cs
--- a/Engine/Source/Programs/Shared/EpicGames.UHT/Types/Properties/UhtStructProperty.cs
+++ b/Engine/Source/Programs/Shared/EpicGames.UHT/Types/Properties/UhtStructProperty.cs
@@ -1,5 +1,6 @@
 // Copyright Epic Games, Inc. All Rights Reserved.
 
+using System;
 using System.Collections.Generic;
 using System.Diagnostics.CodeAnalysis;
 using System.Globalization;
@@ -12,6 +13,48 @@ using EpicGames.UHT.Utils;
 
 namespace EpicGames.UHT.Types
 {
+	/// <summary>
+	/// Extra information needed for a FStructProperty that is generated from a template wrapper struct
+	/// </summary>
+	/// <param name="TemplateSourceName">Source name of the template wrapper struct</param>
+	/// <param name="InnerScriptStruct">USTRUCTs referenced by the args of the template instance</param>
+	record class UhtStructTemplateWrapper(string TemplateSourceName, UhtScriptStruct InnerScriptStruct)
+	{
+		public StringBuilder AppendSourceName(StringBuilder builder)
+		{
+			builder.Append(TemplateSourceName);
+			builder.Append('<');
+			builder.Append(InnerScriptStruct.SourceName);
+			builder.Append('>');
+
+			return builder;
+		}
+
+		public StringBuilder AppendForwardDeclarations(StringBuilder builder)
+		{
+			builder.Append("struct ");
+			builder.Append(InnerScriptStruct.SourceName);
+			builder.Append(';');
+			builder.Append("\r\n");
+
+			builder.Append("template struct ");
+			AppendSourceName(builder);
+			builder.Append(';');
+
+			return builder;
+		}
+
+		public void ParseTemplateType(IUhtTokenReader tokenReader)
+		{
+			tokenReader
+				.Require(TemplateSourceName)
+				.Require('<')
+				.Optional("struct")
+				.Require(InnerScriptStruct.SourceName)
+				.Require('>');
+		}
+	};
+
 	/// <summary>
 	/// FStructProperty
 	/// </summary>
@@ -37,6 +80,11 @@ namespace EpicGames.UHT.Types
 		[JsonConverter(typeof(UhtTypeSourceNameJsonConverter<UhtScriptStruct>))]
 		public UhtScriptStruct ScriptStruct { get; set; }
 
+		/// <summary>
+		/// Extra information for struct properties generated from a template wrapper struct
+		/// </summary>
+		private UhtStructTemplateWrapper? TemplateWrapper { get; set; }
+
 		/// <summary>
 		/// Construct property
 		/// </summary>
@@ -106,7 +154,19 @@ namespace EpicGames.UHT.Types
 		/// <inheritdoc/>
 		public override string? GetForwardDeclarations()
 		{
-			return !ScriptStruct.IsCoreType ? $"struct {ScriptStruct.SourceName};" : null;
+			if (ScriptStruct.IsCoreType)
+			{
+				return null;
+			}
+
+			if (TemplateWrapper != null)
+			{
+				StringBuilder builder = new();
+				TemplateWrapper.AppendForwardDeclarations(builder);
+				return builder.ToString();
+			}
+
+			return $"struct {ScriptStruct.SourceName};";
 		}
 
 		/// <inheritdoc/>
@@ -121,7 +181,14 @@ namespace EpicGames.UHT.Types
 			switch (textType)
 			{
 				default:
-					builder.Append(ScriptStruct.SourceName);
+					if (TemplateWrapper != null)
+					{
+						TemplateWrapper.AppendSourceName(builder);
+					}
+					else
+					{
+						builder.Append(ScriptStruct.SourceName);
+					}
 					break;
 			}
 			return builder;
@@ -449,8 +516,16 @@ namespace EpicGames.UHT.Types
 		[SuppressMessage("CodeQuality", "IDE0051:Remove unused private members", Justification = "Attribute accessed method")]
 		private static bool DefaultStructDefaultValue(UhtStructProperty property, IUhtTokenReader defaultValueReader, StringBuilder innerDefaultValue)
 		{
+			if (property.TemplateWrapper != null)
+			{
+				property.TemplateWrapper.ParseTemplateType(defaultValueReader);
+			}
+			else
+			{
+				defaultValueReader
+					.Require(property.ScriptStruct.SourceName);
+			}
 			defaultValueReader
-				.Require(property.ScriptStruct.SourceName)
 				.Require('(')
 				.Require(')');
 			innerDefaultValue.Append("()");
@@ -514,6 +589,7 @@ namespace EpicGames.UHT.Types
 
 			// With TInstancedStruct, BaseStruct is used as a type limiter.
 			UhtStructProperty instancedStructProperty = new UhtStructProperty(propertySettings, baseScriptStruct.Session.FInstancedStruct);
+			instancedStructProperty.TemplateWrapper = new UhtStructTemplateWrapper("TInstancedStruct", baseScriptStruct);
 			instancedStructProperty.MetaData.Add("BaseStruct", baseScriptStruct.PathName);
 			return instancedStructProperty;
 		}
