Horde: Add a 'ping' message that can be used to keep the remote execution channel alive.

diff --git a/Engine/Source/Programs/Shared/EpicGames.Horde/Compute/AgentMessage.cs b/Engine/Source/Programs/Shared/EpicGames.Horde/Compute/AgentMessage.cs
--- a/Engine/Source/Programs/Shared/EpicGames.Horde/Compute/AgentMessage.cs
+++ b/Engine/Source/Programs/Shared/EpicGames.Horde/Compute/AgentMessage.cs
@@ -23,6 +23,11 @@ namespace EpicGames.Horde.Compute
 		/// </summary>
 		None = 0x00,
 
+		/// <summary>
+		/// No-op message sent to keep the connection alive. Remote should reply with the same message.
+		/// </summary>
+		Ping = 0x01,
+
 		/// <summary>
 		/// Sent in place of a regular response if an error occurs on the remote
 		/// </summary>
@@ -235,6 +240,15 @@ namespace EpicGames.Horde.Compute
 			message.Send();
 		}
 
+		/// <summary>
+		/// Sends a ping message to the remote
+		/// </summary>
+		public static async ValueTask PingAsync(this AgentMessageChannel channel, CancellationToken cancellationToken = default)
+		{
+			using IAgentMessageBuilder message = await channel.CreateMessageAsync(AgentMessageType.Ping, cancellationToken);
+			message.Send();
+		}
+
 		/// <summary>
 		/// Sends an exception response to the remote
 		/// </summary>
diff --git a/Engine/Source/Programs/Shared/EpicGames.Horde/Compute/AgentMessageHandler.cs b/Engine/Source/Programs/Shared/EpicGames.Horde/Compute/AgentMessageHandler.cs
--- a/Engine/Source/Programs/Shared/EpicGames.Horde/Compute/AgentMessageHandler.cs
+++ b/Engine/Source/Programs/Shared/EpicGames.Horde/Compute/AgentMessageHandler.cs
@@ -64,6 +64,9 @@ namespace EpicGames.Horde.Compute
 						case AgentMessageType.None:
 							await Task.WhenAll(childTasks);
 							return;
+						case AgentMessageType.Ping:
+							await channel.PingAsync(cancellationToken);
+							break;
 						case AgentMessageType.Fork:
 							{
 								ForkMessage fork = message.ParseForkMessage();
