EpicGames.Core: Add AsyncThreadPoolWorkQueue

diff --git a/Engine/Source/Programs/Shared/EpicGames.Core.Tests/AsyncThreadPoolWorkQueueTests.cs b/Engine/Source/Programs/Shared/EpicGames.Core.Tests/AsyncThreadPoolWorkQueueTests.cs
new file mode 100644
--- /dev/null
+++ b/Engine/Source/Programs/Shared/EpicGames.Core.Tests/AsyncThreadPoolWorkQueueTests.cs
@@ -0,0 +1,95 @@
+// Copyright Epic Games, Inc. All Rights Reserved.
+
+using System;
+using System.Collections.Generic;
+using System.Threading;
+using System.Threading.Tasks;
+using Microsoft.VisualStudio.TestTools.UnitTesting;
+
+namespace EpicGames.Core.Tests;
+
+[TestClass]
+public sealed class AsyncThreadPoolWorkQueueTests : IDisposable
+{
+	private readonly AsyncThreadPoolWorkQueue _queue = new (4);
+	
+	public void Dispose()
+	{
+		_queue.Dispose();
+	}
+
+	[TestMethod]
+	public async Task ExecuteEnqueuedTasks()
+	{
+		int c = 0;
+		await _queue.EnqueueAsync((_) => { c++; return Task.CompletedTask; });
+		await _queue.EnqueueAsync((_) => { c++; return Task.CompletedTask; });
+		await _queue.ExecuteAsync();
+		Assert.AreEqual(2, c);
+	}
+	
+	[TestMethod]
+	public async Task ExecuteManyEnqueuedTasks()
+	{
+		int c = 0;
+		int numTasks = 100000;
+		
+		for (int i = 0; i < numTasks; i++)
+		{
+			await _queue.EnqueueAsync((_) => { Interlocked.Increment(ref c); return Task.CompletedTask; });	
+		}
+
+		await _queue.ExecuteAsync();
+		Assert.AreEqual(numTasks, c);
+	}
+	
+	[TestMethod]
+	public async Task ExecuteNestedEnqueuedTasks()
+	{
+		bool task1 = false;
+		bool task2 = false;
+		bool task3 = false;
+		
+		await _queue.EnqueueAsync(async (ct1) =>
+		{
+			await Task.Delay(10, ct1);
+			await _queue.EnqueueAsync(async (ct2) =>
+			{
+				await Task.Delay(10, ct2);
+				await _queue.EnqueueAsync(async (ct3) =>
+				{
+					await Task.Delay(10, ct3);
+					task3 = true;
+				});
+				task2 = true;
+			});
+			task1 = true;
+		});
+		
+		await _queue.ExecuteAsync();
+		Assert.IsTrue(task1);
+		Assert.IsTrue(task2);
+		Assert.IsTrue(task3);
+	}
+	
+	[TestMethod]
+	public async Task VerifyMaxConcurrency()
+	{
+		const int NumTasks = 10000;
+		const int NumWorkers = 3;
+		int c = 0;
+		using AsyncThreadPoolWorkQueue queue = new (NumWorkers);
+		
+		HashSet<int> uniqueTaskIds = new(NumTasks);
+		for (int i = 0; i < NumTasks; i++)
+		{
+			await queue.EnqueueAsync((_) => { uniqueTaskIds.Add(Task.CurrentId!.Value); Interlocked.Increment(ref c); return Task.CompletedTask; });	
+		}
+		
+		await queue.ExecuteAsync();
+		
+		Assert.AreEqual(NumTasks, c);
+		Assert.AreEqual(NumWorkers, uniqueTaskIds.Count);
+	}
+}
+
diff --git a/Engine/Source/Programs/Shared/EpicGames.Core/AsyncThreadPoolWorkQueue.cs b/Engine/Source/Programs/Shared/EpicGames.Core/AsyncThreadPoolWorkQueue.cs
new file mode 100644
--- /dev/null
+++ b/Engine/Source/Programs/Shared/EpicGames.Core/AsyncThreadPoolWorkQueue.cs
@@ -0,0 +1,141 @@
+// Copyright Epic Games, Inc. All Rights Reserved.
+
+using System;
+using System.Collections.Generic;
+using System.Threading;
+using System.Threading.Channels;
+using System.Threading.Tasks;
+
+namespace EpicGames.Core;
+
+/// <summary>
+/// Allows queuing a large number of async tasks to the default thread pool
+/// In contrast to <see cref="ThreadPoolWorkQueue" />, the number of concurrent workers can be set 
+/// </summary>
+public sealed class AsyncThreadPoolWorkQueue : IDisposable
+{
+	/// <summary>
+	/// Object used for controlling access to updating _lockObject
+	/// </summary>
+	private readonly object _lockObject = new();
+	
+	/// <summary>
+	/// Event which indicates whether the queue is empty.
+	/// </summary>
+	private readonly ManualResetEvent _emptyEvent = new(true);
+	
+	/// <summary>
+	/// Number of concurrent workers that will process tasks
+	/// </summary>
+	private readonly int _numWorkers;
+	
+	/// <summary>
+	/// Number of tasks remaining in the queue. This is updated in an atomic way.
+	/// </summary>
+	private int _numOutstandingTasks;
+
+	/// <summary>
+	/// Channel used as the queue for tasks
+	/// </summary>
+	readonly Channel<Func<CancellationToken, Task>> _tasks = Channel.CreateUnbounded<Func<CancellationToken, Task>>();
+	
+	/// <summary>
+	/// Constructor
+	/// </summary>
+	/// <param name="numWorkers">Number of concurrent workers executing tasks</param>
+	public AsyncThreadPoolWorkQueue(int numWorkers)
+	{
+		_numWorkers = numWorkers;
+	}
+	
+	/// <summary>
+	/// Dispose
+	/// </summary>
+	public void Dispose()
+	{
+		_emptyEvent.Dispose();
+	}
+
+	/// <summary>
+	/// Enqueue a task to be executed
+	/// </summary>
+	/// <param name="task"></param>
+	public async Task EnqueueAsync(Func<CancellationToken, Task> task)
+	{
+		if (Interlocked.Increment(ref _numOutstandingTasks) == 1)
+		{
+			SetEventState();
+		}
+		await _tasks.Writer.WriteAsync(task);
+	}
+
+	/// <summary>
+	/// Execute all enqueued tasks
+	/// In contrast to Parallel.ForEachAsync, new tasks can be enqueued during execution.
+	/// </summary>
+	/// <param name="cancellationToken">Cancellation token</param>
+	public async Task ExecuteAsync(CancellationToken cancellationToken = default)
+	{
+		if (_numOutstandingTasks < 1)
+		{
+			return;
+		}
+		
+		List<Task> tasks = new();
+		for (int i = 0; i < _numWorkers; i++)
+		{
+			tasks.Add(Task.Run(() => Worker(cancellationToken), cancellationToken));
+		}
+
+		Task queueCompletionTask = Task.Run(async () =>
+		{
+			await _emptyEvent.WaitOneAsync(cancellationToken: cancellationToken);
+			_tasks.Writer.Complete();
+		}, cancellationToken);
+
+		tasks.Add(queueCompletionTask);
+		await Task.WhenAll(tasks);
+	}
+
+	/// <summary>
+	/// Worker executing tasks. Any exception encountered will be bubbled up as all calls are awaited.
+	/// </summary>
+	/// <param name="cancellationToken">Cancellation token</param>
+	private async Task Worker(CancellationToken cancellationToken)
+	{
+		await foreach (Func<CancellationToken, Task> task in _tasks.Reader.ReadAllAsync(cancellationToken))
+		{
+			try
+			{
+				await task(cancellationToken);
+			}
+			finally
+			{
+				if (Interlocked.Decrement(ref _numOutstandingTasks) == 0)
+				{
+					SetEventState();
+				}
+			}
+		}
+	}
+	
+	/// <summary>
+	/// When the number of outstanding tasks transitions between 0 and 1 and vice-versa, we need to set the event 
+	/// based on the actual counter. While the adjustment of the counter happens with interlocked instructions,
+	/// the setting of the event must be done under a lock to serialize the set/reset
+	/// </summary>
+	private void SetEventState()
+	{
+		lock (_lockObject)
+		{
+			if (_numOutstandingTasks > 0)
+			{
+				_emptyEvent?.Reset();
+			}
+			else
+			{
+				_emptyEvent?.Set();
+			}
+		}
+	}
+}
