- Disallow SDK.json reading if the platform isn't enabled on the platform
- If no DDPI.ini exists, mark the SDK as not usable on the platform (default the UEBuildPlaformSDK to disabled until the DDPI.ini file is read)

diff --git a/Engine/Source/Programs/AutomationTool/Turnkey/SdkInfo.cs b/Engine/Source/Programs/AutomationTool/Turnkey/SdkInfo.cs
--- a/Engine/Source/Programs/AutomationTool/Turnkey/SdkInfo.cs
+++ b/Engine/Source/Programs/AutomationTool/Turnkey/SdkInfo.cs
@@ -41,6 +41,8 @@ namespace Turnkey
 
 			Device_AutoSoftwareUpdates_Disabled		= (1 << 15),
 			Device_AutoSoftwareUpdates_Enabled		= (1 << 16),
+
+			Host_Unsupported						= (1 << 17),
 		}
 
 		static public LocalAvailability GetLocalAvailability(AutomationTool.Platform AutomationPlatform, bool bAllowUpdatingPrerequisites, TurnkeyContextImpl TurnkeyContext)
@@ -54,6 +56,11 @@ namespace Turnkey
 				return Result;
 			}
 
+			if (!SDK.bIsSdkAllowedOnHost)
+			{
+				return LocalAvailability.Host_Unsupported;
+			}
+
 			if (AutomationPlatform.UpdateHostPrerequisites(TurnkeyUtils.CommandUtilHelper, TurnkeyContext, !bAllowUpdatingPrerequisites))
 			{
 				Result |= LocalAvailability.Platform_ValidHostPrerequisites;
diff --git a/Engine/Source/Programs/Shared/EpicGames.Build/UEBuildPlatformSDK.cs b/Engine/Source/Programs/Shared/EpicGames.Build/UEBuildPlatformSDK.cs
--- a/Engine/Source/Programs/Shared/EpicGames.Build/UEBuildPlatformSDK.cs
+++ b/Engine/Source/Programs/Shared/EpicGames.Build/UEBuildPlatformSDK.cs
@@ -325,7 +325,10 @@ namespace EpicGames.Core
 			bIsSdkAllowedOnHost = bInIsSdkAllowedOnHost;
 
 			// load the SDK config file
-			LoadJsonFile(PlatformName);
+			if (bIsSdkAllowedOnHost)
+			{
+				LoadJsonFile(PlatformName);
+			}
 
 			// if the parent set up autosdk, the env vars will be wrong, but we can still get the manual SDK version from before it was setup
 			string? ParentManualSDKVersions = Environment.GetEnvironmentVariable(GetPlatformManualSDKSetupEnvVar());
@@ -401,25 +404,28 @@ namespace EpicGames.Core
 		{
 			SDKCollection AllSdks = new SDKCollection(this);
 
-			// walk over each one version the platform supports, and get it's current installed info
-			foreach (SDKDescriptor Desc in GetValidVersions().Sdks)
+			if (bIsSdkAllowedOnHost)
 			{
-				AllSdks.SetupSDK(Desc.Name, Desc.Min, Desc.Max, CachedManualSDKVersions.GetValueOrDefault(Desc.Name), Desc.GroupName);
-			}
+				// walk over each one version the platform supports, and get it's current installed info
+				foreach (SDKDescriptor Desc in GetValidVersions().Sdks)
+				{
+					AllSdks.SetupSDK(Desc.Name, Desc.Min, Desc.Max, CachedManualSDKVersions.GetValueOrDefault(Desc.Name), Desc.GroupName);
+				}
 
-			// now get current AutoSDK version (if autosdk is set up, then the GetInstalledSDKVersion will return AutoSDK version)
-			bool bIsAutoSDK = false;
-			string? CurrentAutoSDKVersion = (PlatformSupportsAutoSDKs() && HasRequiredAutoSDKInstalled() == SDKStatus.Valid && HasSetupAutoSDK()) ? GetInstalledVersion(out bIsAutoSDK) : null;
-			AllSdks.SetupSDK("AutoSdk", GetMainVersion(), GetMainVersion(), CurrentAutoSDKVersion, null);
+				// now get current AutoSDK version (if autosdk is set up, then the GetInstalledSDKVersion will return AutoSDK version)
+				bool bIsAutoSDK = false;
+				string? CurrentAutoSDKVersion = (PlatformSupportsAutoSDKs() && HasRequiredAutoSDKInstalled() == SDKStatus.Valid && HasSetupAutoSDK()) ? GetInstalledVersion(out bIsAutoSDK) : null;
+				AllSdks.SetupSDK("AutoSdk", GetMainVersion(), GetMainVersion(), CurrentAutoSDKVersion, null);
 
-			// verify some assumptions
-			if (CurrentAutoSDKVersion != null && bIsAutoSDK == false)
-			{
-				throw new Exception($"AutoSDK was indicated to be setup ({CurrentAutoSDKVersion}), but GetInstalledSDKVersion returned false for bIsAutoSDK");
-			}
-			if (CurrentAutoSDKVersion != null && CurrentAutoSDKVersion != GetMainVersion())
-			{
-				throw new Exception($"AutoSDK was indicated to be setup, but the version if returned ({CurrentAutoSDKVersion}) doesn't equal the MainVersion ({GetMainVersion()}");
+				// verify some assumptions
+				if (CurrentAutoSDKVersion != null && bIsAutoSDK == false)
+				{
+					throw new Exception($"AutoSDK was indicated to be setup ({CurrentAutoSDKVersion}), but GetInstalledSDKVersion returned false for bIsAutoSDK");
+				}
+				if (CurrentAutoSDKVersion != null && CurrentAutoSDKVersion != GetMainVersion())
+				{
+					throw new Exception($"AutoSDK was indicated to be setup, but the version if returned ({CurrentAutoSDKVersion}) doesn't equal the MainVersion ({GetMainVersion()}");
+				}
 			}
 
 			return AllSdks;
@@ -596,11 +602,14 @@ namespace EpicGames.Core
 
 		protected virtual SDKCollection GetInstalledSDKVersions()
 		{
-			// if the platform doesn't override this, then it only has one sdk, so put it into the collection as the single sdk (this is very much the standard behavior)
-			string? Version = GetInstalledSDKVersion();
-			if (Version != null)
+			if (bIsSdkAllowedOnHost)
 			{
-				return new SDKCollection(Version, this);
+				// if the platform doesn't override this, then it only has one sdk, so put it into the collection as the single sdk (this is very much the standard behavior)
+				string? Version = GetInstalledSDKVersion();
+				if (Version != null)
+				{
+					return new SDKCollection(Version, this);
+				}
 			}
 
 			// if no manual version, then return an empty list of current sdks
diff --git a/Engine/Source/Programs/UnrealBuildTool/Configuration/UEBuildPlatform.cs b/Engine/Source/Programs/UnrealBuildTool/Configuration/UEBuildPlatform.cs
--- a/Engine/Source/Programs/UnrealBuildTool/Configuration/UEBuildPlatform.cs
+++ b/Engine/Source/Programs/UnrealBuildTool/Configuration/UEBuildPlatform.cs
@@ -240,7 +240,7 @@ namespace UnrealBuildTool
 
 			// check DDPI to see if the platform is enabled on this host platform
 			string IniPlatformName = ConfigHierarchy.GetIniPlatformName(Platform);
-			bool bIsEnabled = true; 
+			bool bIsEnabled = false; 
 			ConfigDataDrivenPlatformInfo? DDPI = DataDrivenPlatformInfo.GetDataDrivenInfoForPlatform(IniPlatformName);
 			if (DDPI != null)
 			{
@@ -427,7 +427,12 @@ namespace UnrealBuildTool
 		/// </summary>
 		public SDKStatus HasRequiredSDKsInstalled()
 		{
-			return UEBuildPlatform.GetSDK(Platform)!.HasRequiredSDKsInstalled();
+			UEBuildPlatformSDK? SDK = UEBuildPlatform.GetSDK(Platform);
+			if (SDK == null || !SDK.bIsSdkAllowedOnHost)
+			{
+				return SDKStatus.Invalid;
+			}
+			return SDK.HasRequiredSDKsInstalled();
 		}
 
 		/// <summary>
diff --git a/Engine/Source/Programs/UnrealBuildTool/Configuration/UEBuildTarget.cs b/Engine/Source/Programs/UnrealBuildTool/Configuration/UEBuildTarget.cs
--- a/Engine/Source/Programs/UnrealBuildTool/Configuration/UEBuildTarget.cs
+++ b/Engine/Source/Programs/UnrealBuildTool/Configuration/UEBuildTarget.cs
@@ -1165,7 +1165,7 @@ namespace UnrealBuildTool
 			// make sure we are allowed to build this platform
 			if (!UEBuildPlatform.IsPlatformAvailable(Descriptor.Platform))
 			{
-				throw new BuildException("Platform {0} is not a valid platform to build. Check that the SDK is installed properly.", Descriptor.Platform);
+				throw new BuildException("Platform {0} is not a valid platform to build. Check that the SDK is installed properly and that you have the necessary platorm support files (DataDrivenPlatformInfo.ini, SDK.json, etc).", Descriptor.Platform);
 			}
 
 			RulesAssembly RulesAssembly;
