EpicGames.Perforce: Replace thread pool queue with new async-based one (with concurrency limits)

This should address I/O issues seen during job batch init, where ManagedWorkspaces are set up and synced. Prior to this CL, ThreadPoolWorkQueue could monopolize the shared thread pool which would starve other I/O operations, such as network traffic. Now these operations have a max concurrency of 4.

diff --git a/Engine/Source/Programs/Horde/Horde/Commands/Workspace/WorkspaceStatus.cs b/Engine/Source/Programs/Horde/Horde/Commands/Workspace/WorkspaceStatus.cs
--- a/Engine/Source/Programs/Horde/Horde/Commands/Workspace/WorkspaceStatus.cs
+++ b/Engine/Source/Programs/Horde/Horde/Commands/Workspace/WorkspaceStatus.cs
@@ -1,6 +1,5 @@
 // Copyright Epic Games, Inc. All Rights Reserved.
 
-using System.Threading.Tasks;
 using EpicGames.Core;
 using EpicGames.Perforce;
 using EpicGames.Perforce.Managed;
@@ -11,10 +10,9 @@ namespace Horde.Commands.Workspace
 	[Command("workspace", "status", "Prints information about the state of the cache and workspace")]
 	class WorkspaceStatus : WorkspaceBase
 	{
-		protected override Task ExecuteAsync(IPerforceConnection perforce, ManagedWorkspace repo, ILogger logger)
+		protected override async Task ExecuteAsync(IPerforceConnection perforce, ManagedWorkspace repo, ILogger logger)
 		{
-			repo.Status();
-			return Task.CompletedTask;
+			await repo.Status();
 		}
 	}
 }
diff --git a/Engine/Source/Programs/Shared/EpicGames.Perforce.Managed/ManagedWorkspace.cs b/Engine/Source/Programs/Shared/EpicGames.Perforce.Managed/ManagedWorkspace.cs
--- a/Engine/Source/Programs/Shared/EpicGames.Perforce.Managed/ManagedWorkspace.cs
+++ b/Engine/Source/Programs/Shared/EpicGames.Perforce.Managed/ManagedWorkspace.cs
@@ -455,7 +455,9 @@ namespace EpicGames.Perforce.Managed
 			{
 				Stopwatch timer = Stopwatch.StartNew();
 
-				_workspace.Refresh(removeUntracked, out filesToDelete, out directoriesToDelete);
+				(FileInfo[] filesToDelete, DirectoryInfo[] directoriesToDelete) refreshResult = await _workspace.Refresh(removeUntracked, MaxFileConcurrency);
+				filesToDelete = refreshResult.filesToDelete;
+				directoriesToDelete = refreshResult.directoriesToDelete;
 
 				status.Progress = $"({timer.Elapsed.TotalSeconds:0.0}s)";
 			}
@@ -756,14 +758,14 @@ namespace EpicGames.Perforce.Managed
 		/// <summary>
 		/// Prints information about the repository state
 		/// </summary>
-		public void Status()
+		public async Task Status()
 		{
 			// Print size stats
 			_logger.LogInformation("Cache contains {NumFiles:n0} files, {TotalSize:n1}mb", _contentIdToTrackedFile.Count, _contentIdToTrackedFile.Values.Sum(x => x.Length) / (1024.0 * 1024.0));
 			_logger.LogInformation("Stage contains {NumFiles:n0} files, {TotalSize:n1}mb", _workspace.GetFiles().Count, _workspace.GetFiles().Sum(x => x._length) / (1024.0 * 1024.0));
 
 			// Print the contents of the workspace
-			string[] differences = _workspace.FindDifferences();
+			string[] differences = await _workspace.FindDifferencesAsync(MaxFileConcurrency);
 			if (differences.Length > 0)
 			{
 				_logger.LogInformation("Local changes:");
@@ -1719,7 +1721,7 @@ namespace EpicGames.Perforce.Managed
 			{
 				Stopwatch timer = Stopwatch.StartNew();
 
-				transaction = new RemoveTransaction(_workspace, contents, _contentIdToTrackedFile);
+				transaction = await RemoveTransaction.CreateAsync(_workspace, contents, _contentIdToTrackedFile, MaxFileConcurrency);
 
 				scope.Progress = $"({timer.Elapsed.TotalSeconds:0.0}s)";
 			}
@@ -1855,7 +1857,7 @@ namespace EpicGames.Perforce.Managed
 			{
 				Stopwatch timer = Stopwatch.StartNew();
 
-				transaction = new AddTransaction(_workspace, stream, _contentIdToTrackedFile);
+				transaction = await AddTransaction.CreateAsync(_workspace, stream, _contentIdToTrackedFile, MaxFileConcurrency);
 				_workspace = transaction._newWorkspaceRootDir;
 				await SaveAsync(TransactionState.Dirty, cancellationToken);
 
diff --git a/Engine/Source/Programs/Shared/EpicGames.Perforce.Managed/Transactions/AddTransaction.cs b/Engine/Source/Programs/Shared/EpicGames.Perforce.Managed/Transactions/AddTransaction.cs
--- a/Engine/Source/Programs/Shared/EpicGames.Perforce.Managed/Transactions/AddTransaction.cs
+++ b/Engine/Source/Programs/Shared/EpicGames.Perforce.Managed/Transactions/AddTransaction.cs
@@ -3,6 +3,7 @@
 using System.Collections.Concurrent;
 using System.Collections.Generic;
 using System.IO;
+using System.Threading.Tasks;
 using EpicGames.Core;
 
 namespace EpicGames.Perforce.Managed
@@ -57,7 +58,7 @@ namespace EpicGames.Perforce.Managed
 		readonly Dictionary<FileContentId, CachedFileInfo> _contentIdToTrackedFile;
 		readonly Dictionary<FileContentId, WorkspaceFileInfo> _contentIdToWorkspaceFile = new Dictionary<FileContentId, WorkspaceFileInfo>();
 
-		public AddTransaction(WorkspaceDirectoryInfo workspaceRootDir, StreamSnapshot streamSnapshot, Dictionary<FileContentId, CachedFileInfo> contentIdToTrackedFile)
+		private AddTransaction(WorkspaceDirectoryInfo workspaceRootDir, StreamSnapshot streamSnapshot, Dictionary<FileContentId, CachedFileInfo> contentIdToTrackedFile)
 		{
 			_streamSnapshot = streamSnapshot;
 			_newWorkspaceRootDir = new WorkspaceDirectoryInfo(workspaceRootDir.GetLocation());
@@ -69,14 +70,18 @@ namespace EpicGames.Perforce.Managed
 			{
 				_contentIdToWorkspaceFile[workspaceFile.ContentId] = workspaceFile;
 			}
+		}
 
-			using (ThreadPoolWorkQueue queue = new ThreadPoolWorkQueue())
-			{
-				queue.Enqueue(() => MergeDirectory(workspaceRootDir, _newWorkspaceRootDir, streamSnapshot.Root, queue));
-			}
+		public static async Task<AddTransaction> CreateAsync(WorkspaceDirectoryInfo workspaceRootDir, StreamSnapshot streamSnapshot, Dictionary<FileContentId, CachedFileInfo> contentIdToTrackedFile, int numWorkers)
+		{
+			AddTransaction transaction = new(workspaceRootDir, streamSnapshot, contentIdToTrackedFile);
+			using AsyncThreadPoolWorkQueue queue = new(numWorkers);
+			await queue.EnqueueAsync(_ => transaction.MergeDirectoryAsync(workspaceRootDir, transaction._newWorkspaceRootDir, streamSnapshot.Root, queue));
+			await queue.ExecuteAsync();
+			return transaction;
 		}
 
-		void MergeDirectory(WorkspaceDirectoryInfo workspaceDir, WorkspaceDirectoryInfo newWorkspaceDir, StreamTreeRef streamTreeRef, ThreadPoolWorkQueue queue)
+		async Task MergeDirectoryAsync(WorkspaceDirectoryInfo workspaceDir, WorkspaceDirectoryInfo newWorkspaceDir, StreamTreeRef streamTreeRef, AsyncThreadPoolWorkQueue queue)
 		{
 			// Make sure the directory exists
 			Directory.CreateDirectory(workspaceDir.GetFullName());
@@ -88,11 +93,11 @@ namespace EpicGames.Perforce.Managed
 				WorkspaceDirectoryInfo? workspaceSubDir;
 				if (workspaceDir.NameToSubDirectory.TryGetValue(subDirName, out workspaceSubDir))
 				{
-					MergeSubDirectory(subDirName, workspaceSubDir, subDirRef, newWorkspaceDir, queue);
+					await MergeSubDirectoryAsync(subDirName, workspaceSubDir, subDirRef, newWorkspaceDir, queue);
 				}
 				else
 				{
-					AddSubDirectory(subDirName, newWorkspaceDir, subDirRef, queue);
+					await AddSubDirectoryAsync(subDirName, newWorkspaceDir, subDirRef, queue);
 				}
 			}
 
@@ -111,7 +116,7 @@ namespace EpicGames.Perforce.Managed
 			}
 		}
 
-		void AddDirectory(WorkspaceDirectoryInfo newWorkspaceDir, StreamTreeRef streamTreeRef, ThreadPoolWorkQueue queue)
+		async Task AddDirectoryAsync(WorkspaceDirectoryInfo newWorkspaceDir, StreamTreeRef streamTreeRef, AsyncThreadPoolWorkQueue queue)
 		{
 			StreamTree streamDir = _streamSnapshot.Lookup(streamTreeRef);
 
@@ -121,7 +126,7 @@ namespace EpicGames.Perforce.Managed
 			// Add all the sub directories and files
 			foreach ((Utf8String subDirName, StreamTreeRef subDirRef) in streamDir.NameToTree)
 			{
-				AddSubDirectory(subDirName, newWorkspaceDir, subDirRef, queue);
+				await AddSubDirectoryAsync(subDirName, newWorkspaceDir, subDirRef, queue);
 			}
 			foreach ((Utf8String name, StreamFile streamFile) in streamDir.NameToFile)
 			{
@@ -129,18 +134,18 @@ namespace EpicGames.Perforce.Managed
 			}
 		}
 
-		void MergeSubDirectory(Utf8String name, WorkspaceDirectoryInfo workspaceSubDir, StreamTreeRef streamSubTreeRef, WorkspaceDirectoryInfo newWorkspaceDir, ThreadPoolWorkQueue queue)
+		async Task MergeSubDirectoryAsync(Utf8String name, WorkspaceDirectoryInfo workspaceSubDir, StreamTreeRef streamSubTreeRef, WorkspaceDirectoryInfo newWorkspaceDir, AsyncThreadPoolWorkQueue queue)
 		{
 			WorkspaceDirectoryInfo newWorkspaceSubDir = new WorkspaceDirectoryInfo(newWorkspaceDir, name, streamSubTreeRef);
 			newWorkspaceDir.NameToSubDirectory.Add(name, newWorkspaceSubDir);
-			queue.Enqueue(() => MergeDirectory(workspaceSubDir, newWorkspaceSubDir, streamSubTreeRef, queue));
+			await queue.EnqueueAsync(_ => MergeDirectoryAsync(workspaceSubDir, newWorkspaceSubDir, streamSubTreeRef, queue));
 		}
 
-		void AddSubDirectory(Utf8String name, WorkspaceDirectoryInfo newWorkspaceDir, StreamTreeRef streamSubTreeRef, ThreadPoolWorkQueue queue)
+		async Task AddSubDirectoryAsync(Utf8String name, WorkspaceDirectoryInfo newWorkspaceDir, StreamTreeRef streamSubTreeRef, AsyncThreadPoolWorkQueue queue)
 		{
 			WorkspaceDirectoryInfo newWorkspaceSubDir = new WorkspaceDirectoryInfo(newWorkspaceDir, name, streamSubTreeRef);
 			newWorkspaceDir.NameToSubDirectory.Add(name, newWorkspaceSubDir);
-			queue.Enqueue(() => AddDirectory(newWorkspaceSubDir, streamSubTreeRef, queue));
+			await queue.EnqueueAsync(_ => AddDirectoryAsync(newWorkspaceSubDir, streamSubTreeRef, queue));
 		}
 
 		void AddFile(WorkspaceDirectoryInfo newWorkspaceDir, Utf8String name, StreamFile streamFile)
diff --git a/Engine/Source/Programs/Shared/EpicGames.Perforce.Managed/Transactions/RemoveTransaction.cs b/Engine/Source/Programs/Shared/EpicGames.Perforce.Managed/Transactions/RemoveTransaction.cs
--- a/Engine/Source/Programs/Shared/EpicGames.Perforce.Managed/Transactions/RemoveTransaction.cs
+++ b/Engine/Source/Programs/Shared/EpicGames.Perforce.Managed/Transactions/RemoveTransaction.cs
@@ -2,6 +2,7 @@
 
 using System.Collections.Concurrent;
 using System.Collections.Generic;
+using System.Threading.Tasks;
 using EpicGames.Core;
 
 namespace EpicGames.Perforce.Managed
@@ -21,14 +22,19 @@ namespace EpicGames.Perforce.Managed
 			_newWorkspaceRootDir = new WorkspaceDirectoryInfo(workspaceRootDir.GetLocation());
 			_streamSnapshot = streamSnapshot;
 			_contentIdToTrackedFile = contentIdToTrackedFile;
+		}
 
-			using (ThreadPoolWorkQueue queue = new ThreadPoolWorkQueue())
-			{
-				queue.Enqueue(() => Merge(workspaceRootDir, _newWorkspaceRootDir, streamSnapshot.Root, queue));
-			}
+		public static async Task<RemoveTransaction> CreateAsync(WorkspaceDirectoryInfo workspaceRootDir, StreamSnapshot streamSnapshot, Dictionary<FileContentId, CachedFileInfo> contentIdToTrackedFile, int numWorkers)
+		{
+			RemoveTransaction transaction = new(workspaceRootDir, streamSnapshot, contentIdToTrackedFile);
+			
+			using AsyncThreadPoolWorkQueue queue = new (numWorkers);
+			await queue.EnqueueAsync(_ => transaction.MergeAsync(workspaceRootDir, transaction._newWorkspaceRootDir, streamSnapshot.Root, queue));
+			await queue.ExecuteAsync();
+			return transaction;
 		}
 
-		void Merge(WorkspaceDirectoryInfo workspaceDir, WorkspaceDirectoryInfo newWorkspaceDir, StreamTreeRef streamTreeRef, ThreadPoolWorkQueue queue)
+		async Task MergeAsync(WorkspaceDirectoryInfo workspaceDir, WorkspaceDirectoryInfo newWorkspaceDir, StreamTreeRef streamTreeRef, AsyncThreadPoolWorkQueue queue)
 		{
 			StreamTree streamDir = _streamSnapshot.Lookup(streamTreeRef);
 
@@ -38,11 +44,11 @@ namespace EpicGames.Perforce.Managed
 				StreamTreeRef? streamSubTreeRef;
 				if (streamDir.NameToTree.TryGetValue(workspaceSubDir.Name, out streamSubTreeRef))
 				{
-					MergeSubDirectory(workspaceSubDir.Name, newWorkspaceDir, workspaceSubDir, streamSubTreeRef, queue);
+					await MergeSubDirectoryAsync(workspaceSubDir.Name, newWorkspaceDir, workspaceSubDir, streamSubTreeRef, queue);
 				}
 				else
 				{
-					RemoveDirectory(workspaceSubDir, queue);
+					await RemoveDirectoryAsync(workspaceSubDir, queue);
 				}
 			}
 
@@ -61,20 +67,20 @@ namespace EpicGames.Perforce.Managed
 			}
 		}
 
-		void MergeSubDirectory(Utf8String name, WorkspaceDirectoryInfo newWorkspaceDir, WorkspaceDirectoryInfo workspaceSubDir, StreamTreeRef streamSubTreeRef, ThreadPoolWorkQueue queue)
+		async Task MergeSubDirectoryAsync(Utf8String name, WorkspaceDirectoryInfo newWorkspaceDir, WorkspaceDirectoryInfo workspaceSubDir, StreamTreeRef streamSubTreeRef, AsyncThreadPoolWorkQueue queue)
 		{
 			WorkspaceDirectoryInfo newWorkspaceSubDir = new WorkspaceDirectoryInfo(newWorkspaceDir, name, streamSubTreeRef);
 			newWorkspaceDir.NameToSubDirectory.Add(newWorkspaceSubDir.Name, newWorkspaceSubDir);
-			queue.Enqueue(() => Merge(workspaceSubDir, newWorkspaceSubDir, streamSubTreeRef, queue));
+			await queue.EnqueueAsync(_ => MergeAsync(workspaceSubDir, newWorkspaceSubDir, streamSubTreeRef, queue));
 		}
 
-		void RemoveDirectory(WorkspaceDirectoryInfo workspaceDir, ThreadPoolWorkQueue queue)
+		async Task RemoveDirectoryAsync(WorkspaceDirectoryInfo workspaceDir, AsyncThreadPoolWorkQueue queue)
 		{
 			_directoriesToDelete.Enqueue(workspaceDir);
 
 			foreach (WorkspaceDirectoryInfo workspaceSubDir in workspaceDir.NameToSubDirectory.Values)
 			{
-				queue.Enqueue(() => RemoveDirectory(workspaceSubDir, queue));
+				await queue.EnqueueAsync(_ => RemoveDirectoryAsync(workspaceSubDir, queue));
 			}
 			foreach (WorkspaceFileInfo workspaceFile in workspaceDir.NameToFile.Values)
 			{
diff --git a/Engine/Source/Programs/Shared/EpicGames.Perforce.Managed/WorkspaceDirectoryInfo.cs b/Engine/Source/Programs/Shared/EpicGames.Perforce.Managed/WorkspaceDirectoryInfo.cs
--- a/Engine/Source/Programs/Shared/EpicGames.Perforce.Managed/WorkspaceDirectoryInfo.cs
+++ b/Engine/Source/Programs/Shared/EpicGames.Perforce.Managed/WorkspaceDirectoryInfo.cs
@@ -6,6 +6,7 @@ using System.Collections.Generic;
 using System.IO;
 using System.Linq;
 using System.Text;
+using System.Threading.Tasks;
 using EpicGames.Core;
 
 namespace EpicGames.Perforce.Managed
@@ -125,18 +126,15 @@ namespace EpicGames.Perforce.Managed
 		/// Refresh the state of the workspace on disk
 		/// </summary>
 		/// <param name="removeUntracked">Whether to remove files that are not part of the stream</param>
-		/// <param name="filesToDelete">Receives an array of files to delete</param>
-		/// <param name="directoriesToDelete">Recevies an array of directories to delete</param>
-		public void Refresh(bool removeUntracked, out FileInfo[] filesToDelete, out DirectoryInfo[] directoriesToDelete)
+		/// <param name="numWorkers">Number of concurrent workers to use when refreshing</param>
+		public async Task<(FileInfo[] filesToDelete, DirectoryInfo[] directoriesToDelete)> Refresh(bool removeUntracked, int numWorkers)
 		{
 			ConcurrentQueue<FileInfo> concurrentFilesToDelete = new ConcurrentQueue<FileInfo>();
 			ConcurrentQueue<DirectoryInfo> concurrentDirectoriesToDelete = new ConcurrentQueue<DirectoryInfo>();
-			using (ThreadPoolWorkQueue queue = new ThreadPoolWorkQueue())
-			{
-				queue.Enqueue(() => Refresh(new DirectoryInfo(GetFullName()), removeUntracked, concurrentFilesToDelete, concurrentDirectoriesToDelete, queue));
-			}
-			directoriesToDelete = concurrentDirectoriesToDelete.ToArray();
-			filesToDelete = concurrentFilesToDelete.ToArray();
+			using AsyncThreadPoolWorkQueue queue = new (numWorkers);
+			await queue.EnqueueAsync(_ => Refresh(new DirectoryInfo(GetFullName()), removeUntracked, concurrentFilesToDelete, concurrentDirectoriesToDelete, queue));
+			await queue.ExecuteAsync();
+			return (concurrentFilesToDelete.ToArray(), concurrentDirectoriesToDelete.ToArray());
 		}
 
 		/// <summary>
@@ -147,7 +145,7 @@ namespace EpicGames.Perforce.Managed
 		/// <param name="filesToDelete"></param>
 		/// <param name="directoriesToDelete"></param>
 		/// <param name="queue"></param>
-		void Refresh(DirectoryInfo info, bool removeUntracked, ConcurrentQueue<FileInfo> filesToDelete, ConcurrentQueue<DirectoryInfo> directoriesToDelete, ThreadPoolWorkQueue queue)
+		async Task Refresh(DirectoryInfo info, bool removeUntracked, ConcurrentQueue<FileInfo> filesToDelete, ConcurrentQueue<DirectoryInfo> directoriesToDelete, AsyncThreadPoolWorkQueue queue)
 		{
 			// Recurse through subdirectories
 			Dictionary<Utf8String, WorkspaceDirectoryInfo> newNameToSubDirectory = new Dictionary<Utf8String, WorkspaceDirectoryInfo>(NameToSubDirectory.Count, NameToSubDirectory.Comparer);
@@ -157,7 +155,7 @@ namespace EpicGames.Perforce.Managed
 				if (NameToSubDirectory.TryGetValue(subDirectoryInfo.Name, out subDirectory))
 				{
 					newNameToSubDirectory.Add(subDirectory.Name, subDirectory);
-					queue.Enqueue(() => subDirectory.Refresh(subDirectoryInfo, removeUntracked, filesToDelete, directoriesToDelete, queue));
+					await queue.EnqueueAsync(_ => subDirectory.Refresh(subDirectoryInfo, removeUntracked, filesToDelete, directoriesToDelete, queue));
 				}
 				else if (removeUntracked)
 				{
@@ -208,13 +206,12 @@ namespace EpicGames.Perforce.Managed
 		/// Builds a list of differences from the working directory
 		/// </summary>
 		/// <returns></returns>
-		public string[] FindDifferences()
+		public async Task<string[]> FindDifferencesAsync(int numWorkers)
 		{
-			ConcurrentQueue<string> paths = new ConcurrentQueue<string>();
-			using (ThreadPoolWorkQueue queue = new ThreadPoolWorkQueue())
-			{
-				queue.Enqueue(() => FindDifferences(new DirectoryInfo(GetFullName()), "/", paths, queue));
-			}
+			ConcurrentQueue<string> paths = new ();
+			using AsyncThreadPoolWorkQueue queue = new (numWorkers);
+			await queue.EnqueueAsync(_ => FindDifferencesAsync(new DirectoryInfo(GetFullName()), "/", paths, queue));
+			await queue.ExecuteAsync();
 			return paths.OrderBy(x => x).ToArray();
 		}
 
@@ -225,7 +222,7 @@ namespace EpicGames.Perforce.Managed
 		/// <param name="path"></param>
 		/// <param name="paths"></param>
 		/// <param name="queue"></param>
-		void FindDifferences(DirectoryInfo directory, string path, ConcurrentQueue<string> paths, ThreadPoolWorkQueue queue)
+		async Task FindDifferencesAsync(DirectoryInfo directory, string path, ConcurrentQueue<string> paths, AsyncThreadPoolWorkQueue queue)
 		{
 			// Recurse through subdirectories
 			HashSet<Utf8String> remainingSubDirectoryNames = new HashSet<Utf8String>(NameToSubDirectory.Keys);
@@ -235,7 +232,7 @@ namespace EpicGames.Perforce.Managed
 				if (NameToSubDirectory.TryGetValue(subDirectory.Name, out stagedSubDirectory))
 				{
 					remainingSubDirectoryNames.Remove(subDirectory.Name);
-					queue.Enqueue(() => stagedSubDirectory.FindDifferences(subDirectory, String.Format("{0}{1}/", path, subDirectory.Name), paths, queue));
+					await queue.EnqueueAsync(_ => stagedSubDirectory.FindDifferencesAsync(subDirectory, String.Format("{0}{1}/", path, subDirectory.Name), paths, queue));
 					continue;
 				}
 				paths.Enqueue(String.Format("+{0}{1}/...", path, subDirectory.Name));
