Horde: Use multi-threaded Blake3 hash computation for large objects.

diff --git a/Engine/Source/Programs/Shared/EpicGames.IoHash/IoHash.cs b/Engine/Source/Programs/Shared/EpicGames.IoHash/IoHash.cs
--- a/Engine/Source/Programs/Shared/EpicGames.IoHash/IoHash.cs
+++ b/Engine/Source/Programs/Shared/EpicGames.IoHash/IoHash.cs
@@ -30,6 +30,11 @@ namespace EpicGames.Core
 		/// </summary>
 		public const int NumBits = NumBytes * 8;
 
+		/// <summary>
+		/// Threshold size at which to use multiple threads for hashing
+		/// </summary>
+		const int MultiThreadedSize = 1_000_000;
+
 		readonly ulong _a;
 		readonly ulong _b;
 		readonly uint _c;
@@ -77,7 +82,18 @@ namespace EpicGames.Core
 		public static IoHash Compute(ReadOnlySpan<byte> data)
 		{
 			Span<byte> output = stackalloc byte[32];
-			Blake3.Hasher.Hash(data, output);
+			using (Blake3.Hasher hasher = Blake3.Hasher.New())
+			{
+				if (data.Length < MultiThreadedSize)
+				{
+					hasher.Update(data);
+				}
+				else
+				{
+					hasher.UpdateWithJoin(data);
+				}
+				hasher.Finalize(output);
+			}
 			return new IoHash(output);
 		}
 
@@ -97,7 +113,14 @@ namespace EpicGames.Core
 			{
 				foreach (ReadOnlyMemory<byte> segment in sequence)
 				{
-					hasher.Update(segment.Span);
+					if (segment.Length < MultiThreadedSize)
+					{
+						hasher.Update(segment.Span);
+					}
+					else
+					{
+						hasher.UpdateWithJoin(segment.Span);
+					}
 				}
 				return FromBlake3(hasher);
 			}
@@ -113,11 +136,13 @@ namespace EpicGames.Core
 			using (Blake3.Hasher hasher = Blake3.Hasher.New())
 			{
 				Span<byte> buffer = stackalloc byte[16384];
+
 				int length;
 				while ((length = stream.Read(buffer)) > 0)
 				{
 					hasher.Update(buffer.Slice(0, length));
 				}
+
 				return FromBlake3(hasher);
 			}
 		}
