Allow _DEPRECATED properties in sparse class data to not be marked as VisibleDefaultsOnly or EditDefaultsOnly

diff --git a/Engine/Source/Programs/Shared/EpicGames.UHT/Types/UhtStruct.cs b/Engine/Source/Programs/Shared/EpicGames.UHT/Types/UhtStruct.cs
--- a/Engine/Source/Programs/Shared/EpicGames.UHT/Types/UhtStruct.cs
+++ b/Engine/Source/Programs/Shared/EpicGames.UHT/Types/UhtStruct.cs
@@ -426,7 +426,7 @@ namespace EpicGames.UHT.Types
 					}
 
 					// all sparse properties should have EditDefaultsOnly
-					if (!property.PropertyFlags.HasAllFlags(EPropertyFlags.Edit | EPropertyFlags.DisableEditOnInstance))
+					if (!property.PropertyFlags.HasFlag(EPropertyFlags.Deprecated) && !property.PropertyFlags.HasAllFlags(EPropertyFlags.Edit | EPropertyFlags.DisableEditOnInstance))
 					{
 						property.LogError($"Sparse class data types must be VisibleDefaultsOnly or EditDefaultsOnly. Type '{sparseScriptStruct.EngineName}' Property '{property.SourceName}'");
 					}
