Fixed issue where if export code threw an exception, the build will not display the error message or fail the build.

diff --git a/Engine/Source/Programs/Shared/EpicGames.UHT/Utils/UhtSession.cs b/Engine/Source/Programs/Shared/EpicGames.UHT/Utils/UhtSession.cs
--- a/Engine/Source/Programs/Shared/EpicGames.UHT/Utils/UhtSession.cs
+++ b/Engine/Source/Programs/Shared/EpicGames.UHT/Utils/UhtSession.cs
@@ -172,16 +172,16 @@ namespace EpicGames.UHT.Utils
 				Task[]? prereqTasks = prereqs?.Where(x => x != null).Cast<Task>().ToArray();
 				if (prereqTasks != null && prereqTasks.Length > 0)
 				{
-					return Task.Factory.ContinueWhenAll(prereqTasks, (Task[] tasks) => { action(this); });
+					return Task.Factory.ContinueWhenAll(prereqTasks, (Task[] tasks) => { Session.TryNoErrorCheck(null, () => action(this)); });
 				}
 				else
 				{
-					return Task.Factory.StartNew(() => { action(this); }, CancellationToken.None, TaskCreationOptions.None, TaskScheduler.Default);
+					return Task.Factory.StartNew(() => { Session.TryNoErrorCheck(null, () => action(this)); }, CancellationToken.None, TaskCreationOptions.None, TaskScheduler.Default);
 				}
 			}
 			else
 			{
-				action(this);
+				Session.TryNoErrorCheck(null, () => action(this));
 				return null;
 			}
 		}
@@ -1091,7 +1091,8 @@ namespace EpicGames.UHT.Utils
 
 		/// <summary>
 		/// Try the given action.  If an exception occurs that doesn't have the required
-		/// context, use the supplied context to generate the message.
+		/// context, use the supplied context to generate the message.  If a previous error
+		/// has occured, the action will not be executed
 		/// </summary>
 		/// <param name="messageSource">Message context for when the exception doesn't contain a context.</param>
 		/// <param name="action">The lambda to be invoked</param>
@@ -1099,14 +1100,25 @@ namespace EpicGames.UHT.Utils
 		{
 			if (!HasErrors)
 			{
-				try
-				{
-					action();
-				}
-				catch (Exception e)
-				{
-					HandleException(messageSource, e);
-				}
+				TryNoErrorCheck(messageSource, action);
+			}
+		}
+
+		/// <summary>
+		/// Try the given action.  If an exception occurs that doesn't have the required
+		/// context, use the supplied context to generate the message.
+		/// </summary>
+		/// <param name="messageSource">Message context for when the exception doesn't contain a context.</param>
+		/// <param name="action">The lambda to be invoked</param>
+		public void TryNoErrorCheck(IUhtMessageSource? messageSource, Action action)
+		{
+			try
+			{
+				action();
+			}
+			catch (Exception e)
+			{
+				HandleException(messageSource, e);
 			}
 		}
 
