Fixed issue where invalid code would be generated if all FieldNotify properties in class were marked as editor only data.

diff --git a/Engine/Source/Programs/Shared/EpicGames.UHT/Exporters/CodeGen/UhtHeaderCodeGeneratorHFile.cs b/Engine/Source/Programs/Shared/EpicGames.UHT/Exporters/CodeGen/UhtHeaderCodeGeneratorHFile.cs
--- a/Engine/Source/Programs/Shared/EpicGames.UHT/Exporters/CodeGen/UhtHeaderCodeGeneratorHFile.cs
+++ b/Engine/Source/Programs/Shared/EpicGames.UHT/Exporters/CodeGen/UhtHeaderCodeGeneratorHFile.cs
@@ -664,7 +664,7 @@ namespace EpicGames.UHT.Exporters.CodeGen
 			// Scan the children to see what we have
 			GetFieldNotifyStats(classObj, out bool hasProperties, out bool hasFunctions, out bool hasEditorFields, out bool allEditorFields);
 
-			// If we only have editor fields or no editor fields, then we only emit one block
+			// If we have editor fields, then we must place them in #if blocks
 			if (hasEditorFields)
 			{
 				builder.Append("#if WITH_EDITORONLY_DATA\r\n");
@@ -672,10 +672,13 @@ namespace EpicGames.UHT.Exporters.CodeGen
 				{
 					AppendFieldNotify(builder, classObj, hasProperties, hasFunctions, hasEditorFields, allEditorFields, true);
 				}
-				builder.Append("#else //WITH_EDITORONLY_DATA\r\n");
-				using (UhtMacroCreator macro = new(builder, this, classObj, FieldNotifyMacroSuffix))
+				if (!allEditorFields)
 				{
-					AppendFieldNotify(builder, classObj, hasProperties, hasFunctions, hasEditorFields, allEditorFields, false);
+					builder.Append("#else //WITH_EDITORONLY_DATA\r\n");
+					using (UhtMacroCreator macro = new(builder, this, classObj, FieldNotifyMacroSuffix))
+					{
+						AppendFieldNotify(builder, classObj, hasProperties, hasFunctions, hasEditorFields, allEditorFields, false);
+					}
 				}
 				builder.Append("#endif // WITH_EDITORONLY_DATA\r\n");
 			}
