Fix new C# warning messages in UHT

#rnx

diff --git a/Engine/Source/Programs/Shared/EpicGames.UHT/Exporters/CodeGen/UhtHeaderCodeGeneratorCppFile.cs b/Engine/Source/Programs/Shared/EpicGames.UHT/Exporters/CodeGen/UhtHeaderCodeGeneratorCppFile.cs
--- a/Engine/Source/Programs/Shared/EpicGames.UHT/Exporters/CodeGen/UhtHeaderCodeGeneratorCppFile.cs
+++ b/Engine/Source/Programs/Shared/EpicGames.UHT/Exporters/CodeGen/UhtHeaderCodeGeneratorCppFile.cs
@@ -816,7 +816,7 @@ namespace EpicGames.UHT.Exporters.CodeGen
 					{
 						if (predicateInfo.IsPredicate)
 						{
-							builder.Append("\t").Append(predicateInfo.Name).Append("Struct ").Append(predicateInfo.Name).Append("Predicate; \r\n");
+							builder.Append('\t').Append(predicateInfo.Name).Append("Struct ").Append(predicateInfo.Name).Append("Predicate; \r\n");
 						}
 					}
 
diff --git a/Engine/Source/Programs/Shared/EpicGames.UHT/Exporters/CodeGen/UhtHeaderCodeGeneratorHFile.cs b/Engine/Source/Programs/Shared/EpicGames.UHT/Exporters/CodeGen/UhtHeaderCodeGeneratorHFile.cs
--- a/Engine/Source/Programs/Shared/EpicGames.UHT/Exporters/CodeGen/UhtHeaderCodeGeneratorHFile.cs
+++ b/Engine/Source/Programs/Shared/EpicGames.UHT/Exporters/CodeGen/UhtHeaderCodeGeneratorHFile.cs
@@ -8,7 +8,6 @@ using System.Text;
 using EpicGames.Core;
 using EpicGames.UHT.Types;
 using EpicGames.UHT.Utils;
-using Microsoft.Extensions.Primitives;
 
 namespace EpicGames.UHT.Exporters.CodeGen
 {
@@ -210,14 +209,14 @@ namespace EpicGames.UHT.Exporters.CodeGen
 								foreach (UhtRigVMParameter parameter in methodInfo.Parameters)
 								{
 									string baseType = parameter.TypeOriginal().ToString();
-									if (baseType.StartsWith("const"))
+									if (baseType.StartsWith("const", StringComparison.Ordinal))
 									{
-										baseType = baseType.Substring(5).Trim();
+										baseType = baseType[5..].Trim();
 									}
 
-									if (baseType.EndsWith("&"))
+									if (baseType.EndsWith("&", StringComparison.Ordinal))
 									{
-										baseType = baseType.Substring(0, baseType.Length - 1).Trim();
+										baseType = baseType[..^1].Trim();
 									}
 									builder.Append("\t\t\t\t*(").Append(baseType).Append("*) Branch.MemoryHandles[")
 										.Append(parameterIndex++).Append("].GetData(false) = ")
@@ -227,7 +226,7 @@ namespace EpicGames.UHT.Exporters.CodeGen
 								builder.Append("\t\t\t\tBranch.Execute(*Context);  \\\r\n");
 								builder.Append("\t\t\t\treturn *(").Append(methodInfo.ReturnType).Append("*)Branch.MemoryHandles[").Append(methodInfo.Parameters.Count).Append("].GetData(false);  \\\r\n");
 								builder.Append("\t\t\t}  \\\r\n");
-								builder.Append("\t\t\treturn ").Append(methodInfo.Name).Append("(");
+								builder.Append("\t\t\treturn ").Append(methodInfo.Name).Append('(');
 								builder.AppendParameterNames(methodInfo.Parameters, false, ", ", false, false);
 								builder.Append("); \\\r\n");
 								builder.Append("\t\t}  \\\r\n");
@@ -675,10 +674,8 @@ namespace EpicGames.UHT.Exporters.CodeGen
 				if (!allEditorFields)
 				{
 					builder.Append("#else //WITH_EDITORONLY_DATA\r\n");
-					using (UhtMacroCreator macro = new(builder, this, classObj, FieldNotifyMacroSuffix))
-					{
-						AppendFieldNotify(builder, classObj, hasProperties, hasFunctions, hasEditorFields, allEditorFields, false);
-					}
+					using UhtMacroCreator macro = new(builder, this, classObj, FieldNotifyMacroSuffix);
+					AppendFieldNotify(builder, classObj, hasProperties, hasFunctions, hasEditorFields, allEditorFields, false);
 				}
 				builder.Append("#endif // WITH_EDITORONLY_DATA\r\n");
 			}
@@ -757,7 +754,7 @@ namespace EpicGames.UHT.Exporters.CodeGen
 			return builder;
 		}
 		
-		private StringBuilder AppendAutoGetterSetter(StringBuilder builder, UhtClass classObj,
+		private static StringBuilder AppendAutoGetterSetter(StringBuilder builder, UhtClass classObj,
 			bool hasProperties, bool hasEditorFields, bool allEditorFields,
 			bool includeEditorOnlyFields)
 		{
@@ -866,10 +863,10 @@ namespace EpicGames.UHT.Exporters.CodeGen
 			{
 				// Trim the extra newlines added after the macro generator
 				if (builder.Length > 4 &&
-					builder[builder.Length - 4] == '\r' &&
-					builder[builder.Length - 3] == '\n' &&
-					builder[builder.Length - 2] == '\r' &&
-					builder[builder.Length - 1] == '\n')
+					builder[^4] == '\r' &&
+					builder[^3] == '\n' &&
+					builder[^2] == '\r' &&
+					builder[^1] == '\n')
 				{
 					builder.Length -= 4;
 				}
diff --git a/Engine/Source/Programs/Shared/EpicGames.UHT/Exporters/CodeGen/UhtPackageCodeGeneratorHFile.cs b/Engine/Source/Programs/Shared/EpicGames.UHT/Exporters/CodeGen/UhtPackageCodeGeneratorHFile.cs
--- a/Engine/Source/Programs/Shared/EpicGames.UHT/Exporters/CodeGen/UhtPackageCodeGeneratorHFile.cs
+++ b/Engine/Source/Programs/Shared/EpicGames.UHT/Exporters/CodeGen/UhtPackageCodeGeneratorHFile.cs
@@ -40,11 +40,14 @@ namespace EpicGames.UHT.Exporters.CodeGen
 				List<UhtHeaderFile> headerFiles = new(Package.Children.Count * 2);
 				headerFiles.AddRange(packageSortedHeaders);
 
-				foreach (UhtHeaderFile headerFile in Package.Children)
+				foreach (UhtType type in Package.Children)
 				{
-					if (headerFile.HeaderFileType == UhtHeaderFileType.Classes)
+					if (type is UhtHeaderFile headerFile)
 					{
-						headerFiles.Add(headerFile);
+						if (headerFile.HeaderFileType == UhtHeaderFileType.Classes)
+						{
+							headerFiles.Add(headerFile);
+						}
 					}
 				}
 
diff --git a/Engine/Source/Programs/Shared/EpicGames.UHT/Parsers/UhtNativeInterfaceClassParser.cs b/Engine/Source/Programs/Shared/EpicGames.UHT/Parsers/UhtNativeInterfaceClassParser.cs
--- a/Engine/Source/Programs/Shared/EpicGames.UHT/Parsers/UhtNativeInterfaceClassParser.cs
+++ b/Engine/Source/Programs/Shared/EpicGames.UHT/Parsers/UhtNativeInterfaceClassParser.cs
@@ -120,11 +120,7 @@ namespace EpicGames.UHT.Parsers
 				tokenReader.LogError($"When compiling class definition for '{classObj.SourceName}', attempting to strip prefix results in an empty name. Did you leave off a prefix?");
 			}
 
-			if (classObj.Outer != null)
-			{
-				classObj.Outer.AddChild(classObj);
-			}
-
+			classObj.Outer?.AddChild(classObj);
 			classObj.SuperIdentifier = superName;
 
 			//TODO - C++ UHT compatibility - When we know for sure that we have a native interface, then we should error out.  Due to the lack of a symbol table,
diff --git a/Engine/Source/Programs/Shared/EpicGames.UHT/Tokenizer/UhtTokenReplayReader.cs b/Engine/Source/Programs/Shared/EpicGames.UHT/Tokenizer/UhtTokenReplayReader.cs
--- a/Engine/Source/Programs/Shared/EpicGames.UHT/Tokenizer/UhtTokenReplayReader.cs
+++ b/Engine/Source/Programs/Shared/EpicGames.UHT/Tokenizer/UhtTokenReplayReader.cs
@@ -91,11 +91,7 @@ namespace EpicGames.UHT.Tokenizer
 		/// <exception cref="UhtIceException">Thrown if the TLS value can not be retrieved.</exception>
 		public static UhtTokenReplayReader GetThreadInstance(IUhtMessageSite messageSite, ReadOnlyMemory<char> data, ReadOnlyMemory<UhtToken> tokens, UhtTokenType endTokenType)
 		{
-			UhtTokenReplayReader? reader = s_tls.Value;
-			if (reader == null)
-			{
-				throw new UhtIceException("Unable to acquire a UhtTokenReplayReader");
-			}
+			UhtTokenReplayReader? reader = s_tls.Value ?? throw new UhtIceException("Unable to acquire a UhtTokenReplayReader");
 			reader.Reset(messageSite, data, tokens, endTokenType);
 			return reader;
 		}
diff --git a/Engine/Source/Programs/Shared/EpicGames.UHT/Types/UhtProperty.cs b/Engine/Source/Programs/Shared/EpicGames.UHT/Types/UhtProperty.cs
--- a/Engine/Source/Programs/Shared/EpicGames.UHT/Types/UhtProperty.cs
+++ b/Engine/Source/Programs/Shared/EpicGames.UHT/Types/UhtProperty.cs
@@ -1540,10 +1540,7 @@ namespace EpicGames.UHT.Types
 						if (dim.Length > 0 && !UhtFCString.IsDigit(dim[0]))
 						{
 							UhtEnum? enumObj = Session.FindRegularEnumValue(dim.ToString());
-							if (enumObj == null)
-							{
-								enumObj = Session.FindType(Outer, UhtFindOptions.Enum | UhtFindOptions.SourceName, dim.ToString()) as UhtEnum;
-							}
+							enumObj ??= Session.FindType(Outer, UhtFindOptions.Enum | UhtFindOptions.SourceName, dim.ToString()) as UhtEnum;
 							if (enumObj != null)
 							{
 								MetaData.Add(UhtNames.ArraySizeEnum, enumObj.PathName);
diff --git a/Engine/Source/Programs/Shared/EpicGames.UHT/Utils/UhtSession.cs b/Engine/Source/Programs/Shared/EpicGames.UHT/Utils/UhtSession.cs
--- a/Engine/Source/Programs/Shared/EpicGames.UHT/Utils/UhtSession.cs
+++ b/Engine/Source/Programs/Shared/EpicGames.UHT/Utils/UhtSession.cs
@@ -305,6 +305,7 @@ namespace EpicGames.UHT.Utils
 					if (!String.IsNullOrEmpty(message))
 					{
 						Session.Logger.LogInformation("{Message}", message);
+						Session.SignalError();
 						lock (ReferenceErrorMessages)
 						{
 							ReferenceErrorMessages.Add(message, true);
@@ -996,21 +997,21 @@ namespace EpicGames.UHT.Utils
 		{
 			if (FileManager == null)
 			{
-				Interlocked.Increment(ref _errorCount);
+				SignalError();
 				Logger.LogError("No file manager supplied, aborting.");
 				return;
 			}
 
 			if (Config == null)
 			{
-				Interlocked.Increment(ref _errorCount);
+				SignalError();
 				Logger.LogError("No configuration supplied, aborting.");
 				return;
 			}
 
 			if (Tables == null)
 			{
-				Interlocked.Increment(ref _errorCount);
+				SignalError();
 				Logger.LogError("No parsing tables supplied, aborting.");
 				return;
 			}
@@ -1530,7 +1531,7 @@ namespace EpicGames.UHT.Utils
 			{
 				case UhtMessageType.Error:
 				case UhtMessageType.Ice:
-					Interlocked.Increment(ref _errorCount);
+					SignalError();
 					break;
 
 				case UhtMessageType.Warning:
@@ -1545,6 +1546,14 @@ namespace EpicGames.UHT.Utils
 		}
 		#endregion
 
+		/// <summary>
+		/// Register than an error has happened
+		/// </summary>
+		public void SignalError()
+		{
+			Interlocked.Increment(ref _errorCount);
+		}
+
 		/// <summary>
 		/// Log all the collected messages to the log/console.  If messages aren't being
 		/// cached, then this just waits until the flush task has completed.  If messages
@@ -1558,10 +1567,7 @@ namespace EpicGames.UHT.Utils
 			{
 				messageTask = _messageTask;
 			}
-			if (messageTask != null)
-			{
-				messageTask.Wait();
-			}
+			messageTask?.Wait();
 
 			foreach (UhtMessage message in FetchOrderedMessages())
 			{
@@ -1781,10 +1787,7 @@ namespace EpicGames.UHT.Utils
 			{
 				case UhtException uhtException:
 					UhtMessage message = uhtException.UhtMessage;
-					if (message.MessageSource == null)
-					{
-						message.MessageSource = messageSource;
-					}
+					message.MessageSource ??= messageSource;
 					AddMessage(message);
 					break;
 
@@ -1985,11 +1988,14 @@ namespace EpicGames.UHT.Utils
 				{
 					foreach (UhtPackage package in _packages)
 					{
-						foreach (UhtHeaderFile headerFile in package.Children)
+						foreach (UhtType type in package.Children)
 						{
-							if (headerFile != noExportTypes)
+							if (type is UhtHeaderFile headerFile)
 							{
-								headerFile.AddReferencedHeader(noExportTypes);
+								if (headerFile != noExportTypes)
+								{
+									headerFile.AddReferencedHeader(noExportTypes);
+								}
 							}
 						}
 					}
