Horde: Change IComputeTransport into a concrete base class, to align with C++ implementation.

diff --git a/Engine/Source/Programs/Horde/Horde.Agent/Leases/Handlers/ComputeHandler.cs b/Engine/Source/Programs/Horde/Horde.Agent/Leases/Handlers/ComputeHandler.cs
--- a/Engine/Source/Programs/Horde/Horde.Agent/Leases/Handlers/ComputeHandler.cs
+++ b/Engine/Source/Programs/Horde/Horde.Agent/Leases/Handlers/ComputeHandler.cs
@@ -21,15 +21,13 @@ namespace Horde.Agent.Leases.Handlers
 	/// </summary>
 	class ComputeHandler : LeaseHandler<ComputeTask>
 	{
-		class TcpTransportWithTimeout : IComputeTransport
+		class TcpTransportWithTimeout : ComputeTransport
 		{
 			readonly TcpTransport _inner;
 			long _lastPingTicks;
 
 			static readonly double s_ticksToSystemTicks = (double)TimeSpan.TicksPerSecond / Stopwatch.Frequency;
 
-			public long Position => _inner.Position;
-
 			public TcpTransportWithTimeout(Socket socket)
 			{
 				_inner = new TcpTransport(socket);
@@ -38,9 +36,9 @@ namespace Horde.Agent.Leases.Handlers
 
 			public TimeSpan TimeSinceActivity => TimeSpan.FromTicks((long)((Stopwatch.GetTimestamp() - Interlocked.CompareExchange(ref _lastPingTicks, 0, 0)) * s_ticksToSystemTicks));
 
-			public ValueTask MarkCompleteAsync(CancellationToken cancellationToken) => _inner.MarkCompleteAsync(cancellationToken);
+			public override ValueTask MarkCompleteAsync(CancellationToken cancellationToken) => _inner.MarkCompleteAsync(cancellationToken);
 
-			public async ValueTask<int> ReadPartialAsync(Memory<byte> buffer, CancellationToken cancellationToken)
+			public override async ValueTask<int> ReadPartialAsync(Memory<byte> buffer, CancellationToken cancellationToken)
 			{
 				int result = await _inner.ReadPartialAsync(buffer, cancellationToken);
 				if (result > 0)
@@ -50,7 +48,7 @@ namespace Horde.Agent.Leases.Handlers
 				return result;
 			}
 
-			public async ValueTask WriteAsync(ReadOnlySequence<byte> buffer, CancellationToken cancellationToken)
+			public override async ValueTask WriteAsync(ReadOnlySequence<byte> buffer, CancellationToken cancellationToken)
 			{
 				await _inner.WriteAsync(buffer, cancellationToken);
 				Interlocked.Exchange(ref _lastPingTicks, Stopwatch.GetTimestamp());
diff --git a/Engine/Source/Programs/Shared/EpicGames.Horde.Tests/TransportTests.cs b/Engine/Source/Programs/Shared/EpicGames.Horde.Tests/TransportTests.cs
--- a/Engine/Source/Programs/Shared/EpicGames.Horde.Tests/TransportTests.cs
+++ b/Engine/Source/Programs/Shared/EpicGames.Horde.Tests/TransportTests.cs
@@ -85,7 +85,7 @@ namespace EpicGames.Horde.Tests
 			Assert.IsTrue(input.SequenceEqual(output));
 		}
 
-		public static async Task CopyToAsync(IComputeTransport transport, Memory<byte> buffer)
+		public static async Task CopyToAsync(ComputeTransport transport, Memory<byte> buffer)
 		{
 			for (int offset = 0; offset < buffer.Length;)
 			{
diff --git a/Engine/Source/Programs/Shared/EpicGames.Horde/Compute/ComputeSocket.cs b/Engine/Source/Programs/Shared/EpicGames.Horde/Compute/ComputeSocket.cs
--- a/Engine/Source/Programs/Shared/EpicGames.Horde/Compute/ComputeSocket.cs
+++ b/Engine/Source/Programs/Shared/EpicGames.Horde/Compute/ComputeSocket.cs
@@ -42,7 +42,7 @@ namespace EpicGames.Horde.Compute
 
 		bool _complete;
 
-		readonly IComputeTransport _transport;
+		readonly ComputeTransport _transport;
 		readonly ComputeSocketEndpoint _endpoint;
 		readonly ILogger _logger;
 
@@ -63,7 +63,7 @@ namespace EpicGames.Horde.Compute
 		/// <param name="transport">Transport to communicate with the remote</param>
 		/// <param name="endpoint">Tag for log messages</param>
 		/// <param name="logger">Logger for trace output</param>
-		public ComputeSocket(IComputeTransport transport, ComputeSocketEndpoint endpoint, ILogger logger)
+		public ComputeSocket(ComputeTransport transport, ComputeSocketEndpoint endpoint, ILogger logger)
 		{
 			_transport = transport;
 			_endpoint = endpoint;
@@ -115,7 +115,7 @@ namespace EpicGames.Horde.Compute
 			GC.SuppressFinalize(this);
 		}
 
-		async Task RunRecvTaskAsync(IComputeTransport transport, CancellationToken cancellationToken)
+		async Task RunRecvTaskAsync(ComputeTransport transport, CancellationToken cancellationToken)
 		{
 			_logger.LogTrace("[{Tag}] Started socket reader", Tag);
 
@@ -175,7 +175,7 @@ namespace EpicGames.Horde.Compute
 			_logger.LogTrace("[{Tag}] Closing reader", Tag);
 		}
 
-		async Task ReadPacketAsync(IComputeTransport transport, int id, int size, IComputeBufferWriter writer, CancellationToken cancellationToken)
+		async Task ReadPacketAsync(ComputeTransport transport, int id, int size, IComputeBufferWriter writer, CancellationToken cancellationToken)
 		{
 			Memory<byte> memory = writer.GetWriteBuffer();
 			while (memory.Length < size)
diff --git a/Engine/Source/Programs/Shared/EpicGames.Horde/Compute/Transports/AesTransport.cs b/Engine/Source/Programs/Shared/EpicGames.Horde/Compute/Transports/AesTransport.cs
--- a/Engine/Source/Programs/Shared/EpicGames.Horde/Compute/Transports/AesTransport.cs
+++ b/Engine/Source/Programs/Shared/EpicGames.Horde/Compute/Transports/AesTransport.cs
@@ -15,7 +15,7 @@ namespace EpicGames.Horde.Compute.Transports
 	/// Transport layer that adds AES encryption on top of an underlying transport implementation. Key must be exchanged separately
 	/// (eg. via the HTTPS request to negotiate a lease with the server).
 	/// </summary>
-	public sealed class AesTransport : IComputeTransport, IAsyncDisposable
+	public sealed class AesTransport : ComputeTransport, IAsyncDisposable
 	{
 		/// <summary>
 		/// Length of the required encrption key. 
@@ -30,7 +30,7 @@ namespace EpicGames.Horde.Compute.Transports
 		const int HeaderLength = sizeof(int); // Unencrypted size of the message
 		const int FooterLength = 16; // 16-byte auth tag for encryption
 
-		readonly IComputeTransport _inner;
+		readonly ComputeTransport _inner;
 		readonly Pipe _readPipe;
 		readonly AesGcm _aesGcm;
 
@@ -55,7 +55,7 @@ namespace EpicGames.Horde.Compute.Transports
 		/// <param name="inner">The underlying transport implementation</param>
 		/// <param name="key">AES encryption key (256 bits / 32 bytes)</param>
 		/// <param name="nonce">Cryptographic nonce to identify the connection. Must be longer than <see cref="NonceLength"/>.</param>
-		public AesTransport(IComputeTransport inner, ReadOnlySpan<byte> key, ReadOnlySpan<byte> nonce)
+		public AesTransport(ComputeTransport inner, ReadOnlySpan<byte> key, ReadOnlySpan<byte> nonce)
 		{
 			if (key.Length != KeyLength)
 			{
@@ -143,7 +143,7 @@ namespace EpicGames.Horde.Compute.Transports
 		}
 
 		/// <inheritdoc/>
-		public async ValueTask<int> ReadPartialAsync(Memory<byte> buffer, CancellationToken cancellationToken)
+		public override async ValueTask<int> ReadPartialAsync(Memory<byte> buffer, CancellationToken cancellationToken)
 		{
 			int sizeRead = 0;
 			while (sizeRead == 0)
@@ -167,7 +167,7 @@ namespace EpicGames.Horde.Compute.Transports
 		}
 
 		/// <inheritdoc/>
-		public async ValueTask WriteAsync(ReadOnlySequence<byte> buffer, CancellationToken cancellationToken)
+		public override async ValueTask WriteAsync(ReadOnlySequence<byte> buffer, CancellationToken cancellationToken)
 		{
 			while (buffer.Length > 0)
 			{
@@ -212,6 +212,6 @@ namespace EpicGames.Horde.Compute.Transports
 		}
 
 		/// <inheritdoc/>
-		public ValueTask MarkCompleteAsync(CancellationToken cancellationToken) => _inner.MarkCompleteAsync(cancellationToken);
+		public override ValueTask MarkCompleteAsync(CancellationToken cancellationToken) => _inner.MarkCompleteAsync(cancellationToken);
 	}
 }
diff --git a/Engine/Source/Programs/Shared/EpicGames.Horde/Compute/Transports/PipeTransport.cs b/Engine/Source/Programs/Shared/EpicGames.Horde/Compute/Transports/PipeTransport.cs
--- a/Engine/Source/Programs/Shared/EpicGames.Horde/Compute/Transports/PipeTransport.cs
+++ b/Engine/Source/Programs/Shared/EpicGames.Horde/Compute/Transports/PipeTransport.cs
@@ -9,10 +9,10 @@ using System.Threading.Tasks;
 namespace EpicGames.Horde.Compute.Transports
 {
 	/// <summary>
-	/// Implementation of <see cref="IComputeTransport"/> for communicating over a <see cref="Pipe"/>. 
+	/// Implementation of <see cref="ComputeTransport"/> for communicating over a <see cref="Pipe"/>. 
 	/// (Note: this uses a .NET in-process pipe, not an IPC pipe).
 	/// </summary>
-	public class PipeTransport : IComputeTransport
+	public class PipeTransport : ComputeTransport
 	{
 		readonly PipeReader _reader;
 		readonly PipeWriter _writer;
@@ -35,7 +35,7 @@ namespace EpicGames.Horde.Compute.Transports
 		}
 
 		/// <inheritdoc/>
-		public async ValueTask<int> ReadPartialAsync(Memory<byte> buffer, CancellationToken cancellationToken)
+		public override async ValueTask<int> ReadPartialAsync(Memory<byte> buffer, CancellationToken cancellationToken)
 		{
 			int sizeRead = 0;
 			while (sizeRead < buffer.Length)
@@ -78,7 +78,7 @@ namespace EpicGames.Horde.Compute.Transports
 		}
 
 		/// <inheritdoc/>
-		public async ValueTask WriteAsync(ReadOnlySequence<byte> buffer, CancellationToken cancellationToken)
+		public override async ValueTask WriteAsync(ReadOnlySequence<byte> buffer, CancellationToken cancellationToken)
 		{
 			foreach (ReadOnlyMemory<byte> memory in buffer)
 			{
@@ -88,6 +88,6 @@ namespace EpicGames.Horde.Compute.Transports
 		}
 
 		/// <inheritdoc/>
-		public ValueTask MarkCompleteAsync(CancellationToken cancellationToken) => _writer.CompleteAsync();
+		public override ValueTask MarkCompleteAsync(CancellationToken cancellationToken) => _writer.CompleteAsync();
 	}
 }
diff --git a/Engine/Source/Programs/Shared/EpicGames.Horde/Compute/Transports/StreamTransport.cs b/Engine/Source/Programs/Shared/EpicGames.Horde/Compute/Transports/StreamTransport.cs
--- a/Engine/Source/Programs/Shared/EpicGames.Horde/Compute/Transports/StreamTransport.cs
+++ b/Engine/Source/Programs/Shared/EpicGames.Horde/Compute/Transports/StreamTransport.cs
@@ -11,7 +11,7 @@ namespace EpicGames.Horde.Compute.Transports
 	/// <summary>
 	/// Compute transport which wraps an underlying stream
 	/// </summary>
-	class StreamTransport : IComputeTransport
+	class StreamTransport : ComputeTransport
 	{
 		readonly Stream _stream;
 
@@ -25,7 +25,7 @@ namespace EpicGames.Horde.Compute.Transports
 		public StreamTransport(Stream stream) => _stream = stream;
 
 		/// <inheritdoc/>
-		public async ValueTask<int> ReadPartialAsync(Memory<byte> buffer, CancellationToken cancellationToken)
+		public override async ValueTask<int> ReadPartialAsync(Memory<byte> buffer, CancellationToken cancellationToken)
 		{
 			int length = await _stream.ReadAsync(buffer, cancellationToken);
 			Position += length;
@@ -33,7 +33,7 @@ namespace EpicGames.Horde.Compute.Transports
 		}
 
 		/// <inheritdoc/>
-		public async ValueTask WriteAsync(ReadOnlySequence<byte> buffer, CancellationToken cancellationToken)
+		public override async ValueTask WriteAsync(ReadOnlySequence<byte> buffer, CancellationToken cancellationToken)
 		{
 			foreach (ReadOnlyMemory<byte> memory in buffer)
 			{
@@ -43,6 +43,6 @@ namespace EpicGames.Horde.Compute.Transports
 		}
 
 		/// <inheritdoc/>
-		public ValueTask MarkCompleteAsync(CancellationToken cancellationToken) => new ValueTask();
+		public override ValueTask MarkCompleteAsync(CancellationToken cancellationToken) => new ValueTask();
 	}
 }
diff --git a/Engine/Source/Programs/Shared/EpicGames.Horde/Compute/Transports/TcpTransport.cs b/Engine/Source/Programs/Shared/EpicGames.Horde/Compute/Transports/TcpTransport.cs
--- a/Engine/Source/Programs/Shared/EpicGames.Horde/Compute/Transports/TcpTransport.cs
+++ b/Engine/Source/Programs/Shared/EpicGames.Horde/Compute/Transports/TcpTransport.cs
@@ -10,9 +10,9 @@ using EpicGames.Core;
 namespace EpicGames.Horde.Compute.Transports
 {
 	/// <summary>
-	/// Implementation of <see cref="IComputeTransport"/> for communicating over a socket
+	/// Implementation of <see cref="ComputeTransport"/> for communicating over a socket
 	/// </summary>
-	public sealed class TcpTransport : IComputeTransport
+	public sealed class TcpTransport : ComputeTransport
 	{
 		readonly Socket _socket;
 
@@ -26,7 +26,7 @@ namespace EpicGames.Horde.Compute.Transports
 		public TcpTransport(Socket socket) => _socket = socket;
 
 		/// <inheritdoc/>
-		public async ValueTask<int> ReadPartialAsync(Memory<byte> buffer, CancellationToken cancellationToken)
+		public override async ValueTask<int> ReadPartialAsync(Memory<byte> buffer, CancellationToken cancellationToken)
 		{
 			int read = await _socket.ReceiveAsync(buffer, SocketFlags.None, cancellationToken);
 			Position += read;
@@ -34,7 +34,7 @@ namespace EpicGames.Horde.Compute.Transports
 		}
 
 		/// <inheritdoc/>
-		public async ValueTask WriteAsync(ReadOnlySequence<byte> buffer, CancellationToken cancellationToken)
+		public override async ValueTask WriteAsync(ReadOnlySequence<byte> buffer, CancellationToken cancellationToken)
 		{
 			foreach (ReadOnlyMemory<byte> memory in buffer)
 			{
@@ -44,7 +44,7 @@ namespace EpicGames.Horde.Compute.Transports
 		}
 
 		/// <inheritdoc/>
-		public ValueTask MarkCompleteAsync(CancellationToken cancelationToken)
+		public override ValueTask MarkCompleteAsync(CancellationToken cancelationToken)
 		{
 			_socket.Shutdown(SocketShutdown.Send);
 			return new ValueTask();
