Horde: Missing file.

diff --git a/Engine/Source/Programs/Shared/EpicGames.Horde/Compute/ComputeTransport.cs b/Engine/Source/Programs/Shared/EpicGames.Horde/Compute/ComputeTransport.cs
new file mode 100644
--- /dev/null
+++ b/Engine/Source/Programs/Shared/EpicGames.Horde/Compute/ComputeTransport.cs
@@ -0,0 +1,83 @@
+// Copyright Epic Games, Inc. All Rights Reserved.
+
+using System;
+using System.Buffers;
+using System.IO;
+using System.Threading;
+using System.Threading.Tasks;
+
+namespace EpicGames.Horde.Compute
+{
+	/// <summary>
+	/// Low-level interface for transferring data
+	/// </summary>
+	public abstract class ComputeTransport
+	{
+		/// <summary>
+		/// Reads data from the underlying transport into an output buffer
+		/// </summary>
+		/// <param name="buffer">Buffer to read into</param>
+		/// <param name="cancellationToken">Cancellation token for the operation</param>
+		public abstract ValueTask<int> ReadPartialAsync(Memory<byte> buffer, CancellationToken cancellationToken);
+
+		/// <summary>
+		/// Writes data to the underlying transport
+		/// </summary>
+		/// <param name="buffer">Buffer to be written</param>
+		/// <param name="cancellationToken">Cancellation token for the operation</param>
+		public abstract ValueTask WriteAsync(ReadOnlySequence<byte> buffer, CancellationToken cancellationToken);
+
+		/// <summary>
+		/// Indicate that all data has been written to the transport layer, and that there will be no more calls to <see cref="WriteAsync(ReadOnlySequence{Byte}, CancellationToken)"/>
+		/// </summary>
+		/// <param name="cancellationToken">Cancellation token for the operation</param>
+		public abstract ValueTask MarkCompleteAsync(CancellationToken cancellationToken);
+
+		/// <summary>
+		/// Fill the given buffer with data
+		/// </summary>
+		/// <param name="buffer">Buffer to read into</param>
+		/// <param name="cancellationToken">Cancellation token for the operation</param>
+		public async ValueTask ReadFullAsync(Memory<byte> buffer, CancellationToken cancellationToken)
+		{
+			int read = 0;
+			while (read < buffer.Length)
+			{
+				int partialRead = await ReadPartialAsync(buffer.Slice(read, buffer.Length - read), cancellationToken);
+				if (partialRead == 0)
+				{
+					throw new EndOfStreamException();
+				}
+				read += partialRead;
+			}
+		}
+
+		/// <summary>
+		/// Fill the given buffer with data
+		/// </summary>
+		/// <param name="buffer">Buffer to read into</param>
+		/// <param name="cancellationToken">Cancellation token for the operation</param>
+		public async ValueTask<bool> ReadOptionalAsync(Memory<byte> buffer, CancellationToken cancellationToken)
+		{
+			int read = await ReadPartialAsync(buffer, cancellationToken);
+			if (read == 0)
+			{
+				return false;
+			}
+			if (read < buffer.Length)
+			{
+				await ReadFullAsync(buffer.Slice(read), cancellationToken);
+			}
+			return true;
+		}
+
+		/// <summary>
+		/// Writes data to the underlying transport
+		/// </summary>
+		/// <param name="buffer">Buffer to be written</param>
+		/// <param name="cancellationToken">Cancellation token for the operation</param>
+		public ValueTask WriteAsync(ReadOnlyMemory<byte> buffer, CancellationToken cancellationToken)
+			=> WriteAsync(new ReadOnlySequence<byte>(buffer), cancellationToken);
+	}
+}
+
diff --git a/Engine/Source/Programs/Shared/EpicGames.Horde/Compute/IComputeTransport.cs b/Engine/Source/Programs/Shared/EpicGames.Horde/Compute/IComputeTransport.cs
deleted file mode 100644
