Horde: Change IComputeChannel to a concrete implementation.

diff --git a/Engine/Source/Programs/Horde/Samples/RemoteWorker/RemoteWorkerApp.cs b/Engine/Source/Programs/Horde/Samples/RemoteWorker/RemoteWorkerApp.cs
--- a/Engine/Source/Programs/Horde/Samples/RemoteWorker/RemoteWorkerApp.cs
+++ b/Engine/Source/Programs/Horde/Samples/RemoteWorker/RemoteWorkerApp.cs
@@ -13,7 +13,7 @@ namespace RemoteWorker
 
 			using WorkerComputeSocket socket = WorkerComputeSocket.Open();
 
-			using IComputeChannel channel = socket.CreateChannel(ChannelId);
+			using ComputeChannel channel = socket.CreateChannel(ChannelId);
 			Console.WriteLine("Connected to initiator");
 
 			byte[] data = new byte[4];
diff --git a/Engine/Source/Programs/Shared/EpicGames.Horde/Compute/ComputeChannel.cs b/Engine/Source/Programs/Shared/EpicGames.Horde/Compute/ComputeChannel.cs
new file mode 100644
--- /dev/null
+++ b/Engine/Source/Programs/Shared/EpicGames.Horde/Compute/ComputeChannel.cs
@@ -0,0 +1,129 @@
+// Copyright Epic Games, Inc. All Rights Reserved.
+
+using System;
+using System.IO;
+using System.Threading;
+using System.Threading.Tasks;
+using EpicGames.Horde.Compute.Buffers;
+
+namespace EpicGames.Horde.Compute
+{
+	/// <summary>
+	/// Conventional TCP-like interface for writing data to a socket. Sends are "push", receives are "pull".
+	/// </summary>
+	public sealed class ComputeChannel : IDisposable
+	{
+		/// <summary>
+		/// Reader for the channel
+		/// </summary>
+		public IComputeBufferReader Reader { get; }
+
+		/// <summary>
+		/// Writer for the channel
+		/// </summary>
+		public IComputeBufferWriter Writer { get; }
+
+		/// <summary>
+		/// Constructor
+		/// </summary>
+		/// <param name="recvBufferReader"></param>
+		/// <param name="sendBufferWriter"></param>
+		internal ComputeChannel(IComputeBufferReader recvBufferReader, IComputeBufferWriter sendBufferWriter)
+		{
+			Reader = recvBufferReader.AddRef();
+			Writer = sendBufferWriter.AddRef();
+		}
+
+		/// <inheritdoc/>
+		public void Dispose()
+		{
+			Reader.Dispose();
+
+			Writer.MarkComplete();
+			Writer.Dispose();
+		}
+
+		/// <summary>
+		/// Sends data to a remote channel
+		/// </summary>
+		/// <param name="memory">Memory to write</param>
+		/// <param name="cancellationToken">Cancellation token for the operation</param>
+		public ValueTask SendAsync(ReadOnlyMemory<byte> memory, CancellationToken cancellationToken = default) => Writer.WriteAsync(memory, cancellationToken);
+
+		/// <summary>
+		/// Marks a channel as complete
+		/// </summary>
+		/// <param name="buffer">Buffer to receive the data</param>
+		/// <param name="cancellationToken">Cancellation token for the operation</param>
+		public ValueTask<int> ReceiveAsync(Memory<byte> buffer, CancellationToken cancellationToken = default) => Reader.ReadAsync(buffer, cancellationToken);
+
+		/// <summary>
+		/// Reads a complete message from the given socket, retrying reads until the buffer is full.
+		/// </summary>
+		/// <param name="buffer">Buffer to store the data</param>
+		/// <param name="cancellationToken">Cancellation token for the operation</param>
+		public async ValueTask ReceiveMessageAsync(Memory<byte> buffer, CancellationToken cancellationToken = default)
+		{
+			if (!await TryReceiveMessageAsync(buffer, cancellationToken))
+			{
+				throw new EndOfStreamException();
+			}
+		}
+
+		/// <summary>
+		/// Reads either a full message or end of stream from the channel
+		/// </summary>
+		/// <param name="buffer">Buffer to store the data</param>
+		/// <param name="cancellationToken">Cancellation token for the operation</param>
+		public async ValueTask<bool> TryReceiveMessageAsync(Memory<byte> buffer, CancellationToken cancellationToken = default)
+		{
+			for (int offset = 0; offset < buffer.Length;)
+			{
+				int read = await ReceiveAsync(buffer.Slice(offset), cancellationToken);
+				if (read == 0)
+				{
+					return false;
+				}
+				offset += read;
+			}
+			return true;
+		}
+
+		/// <summary>
+		/// Mark the channel as complete (ie. that no more data will be sent)
+		/// </summary>
+		public void MarkComplete() => Writer.MarkComplete();
+	}
+
+	/// <summary>
+	/// Opens a channel for compute workers
+	/// </summary>
+	public static class ComputeChannelExtensions
+	{
+		/// <summary>
+		/// Creates a channel using a socket and receive buffer
+		/// </summary>
+		/// <param name="socket">Socket to use for sending data</param>
+		/// <param name="channelId">Channel id to send and receive data</param>
+		public static ComputeChannel CreateChannel(this IComputeSocket socket, int channelId)
+		{
+			using SharedMemoryBuffer recvBuffer = SharedMemoryBuffer.CreateNew(null, 65536);
+			using SharedMemoryBuffer sendBuffer = SharedMemoryBuffer.CreateNew(null, 65536);
+			return CreateChannel(socket, channelId, recvBuffer, sendBuffer);
+		}
+
+		/// <summary>
+		/// Creates a channel using a socket and receive buffer
+		/// </summary>
+		/// <param name="socket">Socket to use for sending data</param>
+		/// <param name="channelId">Channel id to send and receive data</param>
+		/// <param name="recvBuffer">Buffer for receiving data</param>
+		/// <param name="sendBuffer">Buffer for sending data</param>
+		public static ComputeChannel CreateChannel(this IComputeSocket socket, int channelId, IComputeBuffer recvBuffer, IComputeBuffer sendBuffer)
+		{
+			socket.AttachRecvBuffer(channelId, recvBuffer.Writer);
+			socket.AttachSendBuffer(channelId, sendBuffer.Reader);
+			return new ComputeChannel(recvBuffer.Reader, sendBuffer.Writer);
+		}
+	}
+}
diff --git a/Engine/Source/Programs/Shared/EpicGames.Horde/Compute/IComputeChannel.cs b/Engine/Source/Programs/Shared/EpicGames.Horde/Compute/IComputeChannel.cs
deleted file mode 100644
