Horde: Change IComputeSocket to a concrete base class.

diff --git a/Engine/Source/Programs/Horde/Horde.Agent/Commands/Compute/ComputeCommand.cs b/Engine/Source/Programs/Horde/Horde.Agent/Commands/Compute/ComputeCommand.cs
--- a/Engine/Source/Programs/Horde/Horde.Agent/Commands/Compute/ComputeCommand.cs
+++ b/Engine/Source/Programs/Horde/Horde.Agent/Commands/Compute/ComputeCommand.cs
@@ -127,7 +127,7 @@ namespace Horde.Agent.Commands.Compute
 			using Socket tcpSocket = new Socket(SocketType.Stream, ProtocolType.IP);
 			await tcpSocket.ConnectAsync(IPAddress.Loopback, Port);
 
-			await using (ComputeSocket socket = new ComputeSocket(new TcpTransport(tcpSocket), ComputeSocketEndpoint.Remote, logger))
+			await using (RemoteComputeSocket socket = new RemoteComputeSocket(new TcpTransport(tcpSocket), ComputeSocketEndpoint.Remote, logger))
 			{
 				logger.LogInformation("Running worker...");
 				await RunWorkerAsync(socket, _memoryCache, logger, CancellationToken.None);
@@ -139,7 +139,7 @@ namespace Horde.Agent.Commands.Compute
 			return 0;
 		}
 
-		public static async Task RunWorkerAsync(IComputeSocket socket, IMemoryCache memoryCache, ILogger logger, CancellationToken cancellationToken)
+		public static async Task RunWorkerAsync(ComputeSocket socket, IMemoryCache memoryCache, ILogger logger, CancellationToken cancellationToken)
 		{
 			DirectoryReference sandboxDir = DirectoryReference.Combine(Program.DataDir, "Sandbox");
 
diff --git a/Engine/Source/Programs/Horde/Horde.Agent/Commands/Compute/DdcComputeCommand.cs b/Engine/Source/Programs/Horde/Horde.Agent/Commands/Compute/DdcComputeCommand.cs
--- a/Engine/Source/Programs/Horde/Horde.Agent/Commands/Compute/DdcComputeCommand.cs
+++ b/Engine/Source/Programs/Horde/Horde.Agent/Commands/Compute/DdcComputeCommand.cs
@@ -30,7 +30,7 @@ namespace Horde.Agent.Commands.Compute
 		/// <inheritdoc/>
 		protected override async Task<bool> HandleRequestAsync(IComputeLease lease, CancellationToken cancellationToken)
 		{
-			IComputeSocket socket = lease.Socket;
+			ComputeSocket socket = lease.Socket;
 			using (AgentMessageChannel channel = socket.CreateAgentMessageChannel(0, _logger))
 			{
 				_logger.LogInformation("Sending XOR request");
diff --git a/Engine/Source/Programs/Horde/Horde.Agent/Commands/Compute/XorComputeCommand.cs b/Engine/Source/Programs/Horde/Horde.Agent/Commands/Compute/XorComputeCommand.cs
--- a/Engine/Source/Programs/Horde/Horde.Agent/Commands/Compute/XorComputeCommand.cs
+++ b/Engine/Source/Programs/Horde/Horde.Agent/Commands/Compute/XorComputeCommand.cs
@@ -30,7 +30,7 @@ namespace Horde.Agent.Commands.Compute
 		/// <inheritdoc/>
 		protected override async Task<bool> HandleRequestAsync(IComputeLease lease, CancellationToken cancellationToken)
 		{
-			IComputeSocket socket = lease.Socket;
+			ComputeSocket socket = lease.Socket;
 			using (AgentMessageChannel mainChannel = socket.CreateAgentMessageChannel(0, _logger))
 			{
 				_logger.LogInformation("Forking compute channel...");
diff --git a/Engine/Source/Programs/Horde/Horde.Agent/Leases/Handlers/ComputeHandler.cs b/Engine/Source/Programs/Horde/Horde.Agent/Leases/Handlers/ComputeHandler.cs
--- a/Engine/Source/Programs/Horde/Horde.Agent/Leases/Handlers/ComputeHandler.cs
+++ b/Engine/Source/Programs/Horde/Horde.Agent/Leases/Handlers/ComputeHandler.cs
@@ -96,7 +96,7 @@ namespace Horde.Agent.Leases.Handlers
 					await using BackgroundTask timeoutTask = BackgroundTask.StartNew(ctx => TickTimeoutAsync(transport, cts, ctx));
 					try
 					{
-						await using (ComputeSocket socket = new ComputeSocket(transport, ComputeSocketEndpoint.Local, _logger))
+						await using (RemoteComputeSocket socket = new RemoteComputeSocket(transport, ComputeSocketEndpoint.Local, _logger))
 						{
 							DirectoryReference sandboxDir = DirectoryReference.Combine(session.WorkingDir, "Sandbox", leaseId);
 							try
diff --git a/Engine/Source/Programs/Horde/Samples/RemoteClient/RemoteClientApp.cs b/Engine/Source/Programs/Horde/Samples/RemoteClient/RemoteClientApp.cs
--- a/Engine/Source/Programs/Horde/Samples/RemoteClient/RemoteClientApp.cs
+++ b/Engine/Source/Programs/Horde/Samples/RemoteClient/RemoteClientApp.cs
@@ -119,7 +119,7 @@ namespace RemoteClient
 			await lease.CloseAsync();
 		}
 		
-		static async Task WriteNumbersAsync(ComputeSocket socket, ILogger logger, CancellationToken cancellationToken)
+		static async Task WriteNumbersAsync(RemoteComputeSocket socket, ILogger logger, CancellationToken cancellationToken)
 		{
 			// Wait until the remote sends a message indicating that it's ready
 			using (PooledBuffer recvBuffer = new PooledBuffer(1, 20))
diff --git a/Engine/Source/Programs/Shared/EpicGames.Horde.Tests/BufferTests.cs b/Engine/Source/Programs/Shared/EpicGames.Horde.Tests/BufferTests.cs
--- a/Engine/Source/Programs/Shared/EpicGames.Horde.Tests/BufferTests.cs
+++ b/Engine/Source/Programs/Shared/EpicGames.Horde.Tests/BufferTests.cs
@@ -97,8 +97,8 @@ namespace EpicGames.Horde.Tests
 			Pipe sourceToTargetPipe = new Pipe();
 			Pipe targetToSourcePipe = new Pipe();
 
-			await using ComputeSocket producerSocket = new ComputeSocket(new PipeTransport(targetToSourcePipe.Reader, sourceToTargetPipe.Writer), ComputeSocketEndpoint.Local, NullLogger.Instance);
-			await using ComputeSocket consumerSocket = new ComputeSocket(new PipeTransport(sourceToTargetPipe.Reader, targetToSourcePipe.Writer), ComputeSocketEndpoint.Remote, NullLogger.Instance);
+			await using RemoteComputeSocket producerSocket = new RemoteComputeSocket(new PipeTransport(targetToSourcePipe.Reader, sourceToTargetPipe.Writer), ComputeSocketEndpoint.Local, NullLogger.Instance);
+			await using RemoteComputeSocket consumerSocket = new RemoteComputeSocket(new PipeTransport(sourceToTargetPipe.Reader, targetToSourcePipe.Writer), ComputeSocketEndpoint.Remote, NullLogger.Instance);
 
 			using IComputeBuffer consumerBuffer = createBuffer(Length);
 			consumerSocket.AttachRecvBuffer(ChannelId, consumerBuffer.Writer);
@@ -113,7 +113,7 @@ namespace EpicGames.Horde.Tests
 			Assert.IsTrue(input.SequenceEqual(output));
 		}
 
-		static async Task RunProducerAsync(ComputeSocket socket, ReadOnlyMemory<byte> input)
+		static async Task RunProducerAsync(RemoteComputeSocket socket, ReadOnlyMemory<byte> input)
 		{
 			int offset = 0;
 			while (offset < input.Length)
diff --git a/Engine/Source/Programs/Shared/EpicGames.Horde/Compute/AgentMessageChannel.cs b/Engine/Source/Programs/Shared/EpicGames.Horde/Compute/AgentMessageChannel.cs
--- a/Engine/Source/Programs/Shared/EpicGames.Horde/Compute/AgentMessageChannel.cs
+++ b/Engine/Source/Programs/Shared/EpicGames.Horde/Compute/AgentMessageChannel.cs
@@ -74,7 +74,7 @@ namespace EpicGames.Horde.Compute
 			public Span<byte> GetSpan(int sizeHint = 0) => GetMemory(sizeHint).Span;
 		}
 
-		readonly IComputeSocket _socket;
+		readonly ComputeSocket _socket;
 		readonly IComputeBufferReader _recvBufferReader;
 		readonly IComputeBufferWriter _sendBufferWriter;
 
@@ -107,7 +107,7 @@ namespace EpicGames.Horde.Compute
 		/// <param name="recvBuffer"></param>
 		/// <param name="sendBuffer"></param>
 		/// <param name="logger">Logger for diagnostic output</param>
-		public AgentMessageChannel(IComputeSocket socket, int channelId, IComputeBuffer recvBuffer, IComputeBuffer sendBuffer, ILogger logger)
+		public AgentMessageChannel(ComputeSocket socket, int channelId, IComputeBuffer recvBuffer, IComputeBuffer sendBuffer, ILogger logger)
 		{
 			_socket = socket;
 			_socket.AttachRecvBuffer(channelId, recvBuffer.Writer);
@@ -210,7 +210,7 @@ namespace EpicGames.Horde.Compute
 		/// <param name="socket">Socket to create a channel for</param>
 		/// <param name="channelId">Identifier for the channel</param>
 		/// <param name="logger">Logger for the channel</param>
-		public static AgentMessageChannel CreateAgentMessageChannel(this IComputeSocket socket, int channelId, ILogger logger)
+		public static AgentMessageChannel CreateAgentMessageChannel(this ComputeSocket socket, int channelId, ILogger logger)
 			=> socket.CreateAgentMessageChannel(channelId, 65536, logger);
 
 		/// <summary>
@@ -220,7 +220,7 @@ namespace EpicGames.Horde.Compute
 		/// <param name="channelId">Identifier for the channel</param>
 		/// <param name="bufferSize">Size of the send and receive buffer</param>
 		/// <param name="logger">Logger for the channel</param>
-		public static AgentMessageChannel CreateAgentMessageChannel(this IComputeSocket socket, int channelId, int bufferSize, ILogger logger)
+		public static AgentMessageChannel CreateAgentMessageChannel(this ComputeSocket socket, int channelId, int bufferSize, ILogger logger)
 			=> socket.CreateAgentMessageChannel(channelId, bufferSize, bufferSize, logger);
 
 		/// <summary>
@@ -231,7 +231,7 @@ namespace EpicGames.Horde.Compute
 		/// <param name="sendBufferSize">Size of the send buffer</param>
 		/// <param name="recvBufferSize">Size of the recieve buffer</param>
 		/// <param name="logger">Logger for the channel</param>
-		public static AgentMessageChannel CreateAgentMessageChannel(this IComputeSocket socket, int channelId, int sendBufferSize, int recvBufferSize, ILogger logger)
+		public static AgentMessageChannel CreateAgentMessageChannel(this ComputeSocket socket, int channelId, int sendBufferSize, int recvBufferSize, ILogger logger)
 		{
 			using IComputeBuffer sendBuffer = new PooledBuffer(sendBufferSize);
 			using IComputeBuffer recvBuffer = new PooledBuffer(recvBufferSize);
diff --git a/Engine/Source/Programs/Shared/EpicGames.Horde/Compute/AgentMessageHandler.cs b/Engine/Source/Programs/Shared/EpicGames.Horde/Compute/AgentMessageHandler.cs
--- a/Engine/Source/Programs/Shared/EpicGames.Horde/Compute/AgentMessageHandler.cs
+++ b/Engine/Source/Programs/Shared/EpicGames.Horde/Compute/AgentMessageHandler.cs
@@ -43,12 +43,12 @@ namespace EpicGames.Horde.Compute
 		/// </summary>
 		/// <param name="socket">Socket to read from</param>
 		/// <param name="cancellationToken">Cancellation token for the operation</param>
-		public async Task RunAsync(IComputeSocket socket, CancellationToken cancellationToken)
+		public async Task RunAsync(ComputeSocket socket, CancellationToken cancellationToken)
 		{
 			await RunAsync(socket, 0, 4 * 1024 * 1024, cancellationToken);
 		}
 
-		async Task RunAsync(IComputeSocket socket, int channelId, int bufferSize, CancellationToken cancellationToken)
+		async Task RunAsync(ComputeSocket socket, int channelId, int bufferSize, CancellationToken cancellationToken)
 		{
 			using (AgentMessageChannel channel = socket.CreateAgentMessageChannel(channelId, bufferSize, _logger))
 			{
@@ -152,7 +152,7 @@ namespace EpicGames.Horde.Compute
 			}
 		}
 
-		async Task ExecuteProcessAsync(IComputeSocket socket, AgentMessageChannel channel, string executable, IReadOnlyList<string> arguments, string? workingDir, IReadOnlyDictionary<string, string?>? envVars, CancellationToken cancellationToken)
+		async Task ExecuteProcessAsync(ComputeSocket socket, AgentMessageChannel channel, string executable, IReadOnlyList<string> arguments, string? workingDir, IReadOnlyDictionary<string, string?>? envVars, CancellationToken cancellationToken)
 		{
 			try
 			{
@@ -171,7 +171,7 @@ namespace EpicGames.Horde.Compute
 			}
 		}
 
-		async Task ExecuteProcessWindowsAsync(IComputeSocket socket, AgentMessageChannel channel, string executable, IReadOnlyList<string> arguments, string? workingDir, IReadOnlyDictionary<string, string?>? envVars, CancellationToken cancellationToken)
+		async Task ExecuteProcessWindowsAsync(ComputeSocket socket, AgentMessageChannel channel, string executable, IReadOnlyList<string> arguments, string? workingDir, IReadOnlyDictionary<string, string?>? envVars, CancellationToken cancellationToken)
 		{
 			Dictionary<string, string?> newEnvVars = new Dictionary<string, string?>();
 			if (envVars != null)
@@ -198,7 +198,7 @@ namespace EpicGames.Horde.Compute
 			_logger.LogInformation("Child process has shut down");
 		}
 
-		async Task ProcessIpcMessagesAsync(IComputeSocket socket, IComputeBufferReader ipcReader, CancellationToken cancellationToken)
+		async Task ProcessIpcMessagesAsync(ComputeSocket socket, IComputeBufferReader ipcReader, CancellationToken cancellationToken)
 		{
 			List<(IpcMessage, int, SharedMemoryBuffer)> buffers = new();
 			try
diff --git a/Engine/Source/Programs/Shared/EpicGames.Horde/Compute/Clients/AgentComputeClient.cs b/Engine/Source/Programs/Shared/EpicGames.Horde/Compute/Clients/AgentComputeClient.cs
--- a/Engine/Source/Programs/Shared/EpicGames.Horde/Compute/Clients/AgentComputeClient.cs
+++ b/Engine/Source/Programs/Shared/EpicGames.Horde/Compute/Clients/AgentComputeClient.cs
@@ -22,7 +22,7 @@ namespace EpicGames.Horde.Compute.Clients
 	{
 		class LeaseImpl : IComputeLease
 		{
-			readonly IAsyncEnumerator<ComputeSocket> _source;
+			readonly IAsyncEnumerator<RemoteComputeSocket> _source;
 
 			/// <inheritdoc/>
 			public IReadOnlyList<string> Properties { get; } = new List<string>();
@@ -31,9 +31,9 @@ namespace EpicGames.Horde.Compute.Clients
 			public IReadOnlyDictionary<string, int> AssignedResources => new Dictionary<string, int>();
 
 			/// <inheritdoc/>
-			public ComputeSocket Socket => _source.Current;
+			public RemoteComputeSocket Socket => _source.Current;
 
-			public LeaseImpl(IAsyncEnumerator<ComputeSocket> source) => _source = source;
+			public LeaseImpl(IAsyncEnumerator<RemoteComputeSocket> source) => _source = source;
 
 			/// <inheritdoc/>
 			public async ValueTask DisposeAsync()
@@ -73,7 +73,7 @@ namespace EpicGames.Horde.Compute.Clients
 			_logger.LogInformation("Launching {Path} to handle remote", _hordeAgentAssembly);
 
 			// The connection logic is an async enumerator that returns the socket, then shuts down.
-			IAsyncEnumerator<ComputeSocket> source = ConnectAsync(cancellationToken).GetAsyncEnumerator(cancellationToken);
+			IAsyncEnumerator<RemoteComputeSocket> source = ConnectAsync(cancellationToken).GetAsyncEnumerator(cancellationToken);
 			if (!await source.MoveNextAsync())
 			{
 				await source.DisposeAsync();
@@ -83,7 +83,7 @@ namespace EpicGames.Horde.Compute.Clients
 			return new LeaseImpl(source);
 		}
 
-		async IAsyncEnumerable<ComputeSocket> ConnectAsync([EnumeratorCancellation] CancellationToken cancellationToken)
+		async IAsyncEnumerable<RemoteComputeSocket> ConnectAsync([EnumeratorCancellation] CancellationToken cancellationToken)
 		{
 			using Socket listener = new Socket(SocketType.Stream, ProtocolType.IP);
 			listener.Bind(new IPEndPoint(IPAddress.Loopback, _port));
@@ -92,7 +92,7 @@ namespace EpicGames.Horde.Compute.Clients
 			using BackgroundTask agentTask = BackgroundTask.StartNew(ctx => RunAgentAsync(_hordeAgentAssembly, _port, _logger, ctx));
 			using Socket tcpSocket = await listener.AcceptAsync(cancellationToken);
 
-			await using ComputeSocket socket = new ComputeSocket(new TcpTransport(tcpSocket), ComputeSocketEndpoint.Local, _logger);
+			await using RemoteComputeSocket socket = new RemoteComputeSocket(new TcpTransport(tcpSocket), ComputeSocketEndpoint.Local, _logger);
 			yield return socket;
 
 			await socket.CloseAsync(cancellationToken);
diff --git a/Engine/Source/Programs/Shared/EpicGames.Horde/Compute/Clients/LocalComputeClient.cs b/Engine/Source/Programs/Shared/EpicGames.Horde/Compute/Clients/LocalComputeClient.cs
--- a/Engine/Source/Programs/Shared/EpicGames.Horde/Compute/Clients/LocalComputeClient.cs
+++ b/Engine/Source/Programs/Shared/EpicGames.Horde/Compute/Clients/LocalComputeClient.cs
@@ -22,11 +22,11 @@ namespace EpicGames.Horde.Compute.Clients
 		{
 			public IReadOnlyList<string> Properties { get; } = new List<string>();
 			public IReadOnlyDictionary<string, int> AssignedResources => new Dictionary<string, int>();
-			public ComputeSocket Socket => _socket;
+			public RemoteComputeSocket Socket => _socket;
 
-			readonly ComputeSocket _socket;
+			readonly RemoteComputeSocket _socket;
 
-			public LeaseImpl(ComputeSocket socket) => _socket = socket;
+			public LeaseImpl(RemoteComputeSocket socket) => _socket = socket;
 
 			/// <inheritdoc/>
 			public ValueTask DisposeAsync() => _socket.DisposeAsync();
@@ -77,7 +77,7 @@ namespace EpicGames.Horde.Compute.Clients
 
 			using MemoryCache memoryCache = new MemoryCache(new MemoryCacheOptions { SizeLimit = 10 * 1024 * 1024 });
 
-			await using (ComputeSocket socket = new ComputeSocket(new TcpTransport(tcpSocket), ComputeSocketEndpoint.Remote, logger))
+			await using (RemoteComputeSocket socket = new RemoteComputeSocket(new TcpTransport(tcpSocket), ComputeSocketEndpoint.Remote, logger))
 			{
 				AgentMessageHandler worker = new AgentMessageHandler(sandboxDir, memoryCache, logger);
 				await worker.RunAsync(socket, cancellationToken);
@@ -89,7 +89,7 @@ namespace EpicGames.Horde.Compute.Clients
 		public Task<IComputeLease?> TryAssignWorkerAsync(ClusterId clusterId, Requirements? requirements, CancellationToken cancellationToken)
 		{
 #pragma warning disable CA2000 // Dispose objects before losing scope
-			ComputeSocket socket = new ComputeSocket(new TcpTransport(_socket), ComputeSocketEndpoint.Local, _logger);
+			RemoteComputeSocket socket = new RemoteComputeSocket(new TcpTransport(_socket), ComputeSocketEndpoint.Local, _logger);
 			return Task.FromResult<IComputeLease?>(new LeaseImpl(socket));
 #pragma warning restore CA2000 // Dispose objects before losing scope
 		}
diff --git a/Engine/Source/Programs/Shared/EpicGames.Horde/Compute/Clients/ServerComputeClient.cs b/Engine/Source/Programs/Shared/EpicGames.Horde/Compute/Clients/ServerComputeClient.cs
--- a/Engine/Source/Programs/Shared/EpicGames.Horde/Compute/Clients/ServerComputeClient.cs
+++ b/Engine/Source/Programs/Shared/EpicGames.Horde/Compute/Clients/ServerComputeClient.cs
@@ -42,7 +42,7 @@ namespace EpicGames.Horde.Compute.Clients
 			public List<string> Properties { get; set; } = new List<string>();
 		}
 
-		record class LeaseInfo(IReadOnlyList<string> Properties, IReadOnlyDictionary<string, int> AssignedResources, ComputeSocket Socket);
+		record class LeaseInfo(IReadOnlyList<string> Properties, IReadOnlyDictionary<string, int> AssignedResources, RemoteComputeSocket Socket);
 
 		class LeaseImpl : IComputeLease
 		{
@@ -50,7 +50,7 @@ namespace EpicGames.Horde.Compute.Clients
 
 			public IReadOnlyList<string> Properties => _source.Current.Properties;
 			public IReadOnlyDictionary<string, int> AssignedResources => _source.Current.AssignedResources;
-			public ComputeSocket Socket => _source.Current.Socket;
+			public RemoteComputeSocket Socket => _source.Current.Socket;
 
 			public LeaseImpl(IAsyncEnumerator<LeaseInfo> source)
 			{
@@ -186,7 +186,7 @@ namespace EpicGames.Horde.Compute.Clients
 			// Pass the rest of the call over to the handler
 			byte[] key = StringUtils.ParseHexString(responseMessage.Key);
 
-			await using ComputeSocket computeSocket = new ComputeSocket(new TcpTransport(socket), ComputeSocketEndpoint.Local, _logger);
+			await using RemoteComputeSocket computeSocket = new RemoteComputeSocket(new TcpTransport(socket), ComputeSocketEndpoint.Local, _logger);
 			yield return new LeaseInfo(responseMessage.Properties, responseMessage.AssignedResources, computeSocket);
 		}
 	}
diff --git a/Engine/Source/Programs/Shared/EpicGames.Horde/Compute/ComputeChannel.cs b/Engine/Source/Programs/Shared/EpicGames.Horde/Compute/ComputeChannel.cs
--- a/Engine/Source/Programs/Shared/EpicGames.Horde/Compute/ComputeChannel.cs
+++ b/Engine/Source/Programs/Shared/EpicGames.Horde/Compute/ComputeChannel.cs
@@ -105,7 +105,7 @@ namespace EpicGames.Horde.Compute
 		/// </summary>
 		/// <param name="socket">Socket to use for sending data</param>
 		/// <param name="channelId">Channel id to send and receive data</param>
-		public static ComputeChannel CreateChannel(this IComputeSocket socket, int channelId)
+		public static ComputeChannel CreateChannel(this ComputeSocket socket, int channelId)
 		{
 			using SharedMemoryBuffer recvBuffer = SharedMemoryBuffer.CreateNew(null, 65536);
 			using SharedMemoryBuffer sendBuffer = SharedMemoryBuffer.CreateNew(null, 65536);
@@ -119,7 +119,7 @@ namespace EpicGames.Horde.Compute
 		/// <param name="channelId">Channel id to send and receive data</param>
 		/// <param name="recvBuffer">Buffer for receiving data</param>
 		/// <param name="sendBuffer">Buffer for sending data</param>
-		public static ComputeChannel CreateChannel(this IComputeSocket socket, int channelId, IComputeBuffer recvBuffer, IComputeBuffer sendBuffer)
+		public static ComputeChannel CreateChannel(this ComputeSocket socket, int channelId, IComputeBuffer recvBuffer, IComputeBuffer sendBuffer)
 		{
 			socket.AttachRecvBuffer(channelId, recvBuffer.Writer);
 			socket.AttachSendBuffer(channelId, sendBuffer.Reader);
diff --git a/Engine/Source/Programs/Shared/EpicGames.Horde/Compute/ComputeSocket.cs b/Engine/Source/Programs/Shared/EpicGames.Horde/Compute/ComputeSocket.cs
--- a/Engine/Source/Programs/Shared/EpicGames.Horde/Compute/ComputeSocket.cs
+++ b/Engine/Source/Programs/Shared/EpicGames.Horde/Compute/ComputeSocket.cs
@@ -8,10 +8,102 @@ using System.Linq;
 using System.Threading;
 using System.Threading.Tasks;
 using EpicGames.Core;
+using EpicGames.Horde.Compute.Buffers;
 using Microsoft.Extensions.Logging;
 
 namespace EpicGames.Horde.Compute
 {
+	/// <summary>
+	/// Socket for sending and reciving data using a "push" model. The application can attach multiple writers to accept received data.
+	/// </summary>
+	public abstract class ComputeSocket
+	{
+		/// <summary>
+		/// Attaches a buffer to receive data.
+		/// </summary>
+		/// <param name="channelId">Channel to receive data on</param>
+		/// <param name="recvBufferWriter">Writer for the buffer to store received data</param>
+		public abstract void AttachRecvBuffer(int channelId, IComputeBufferWriter recvBufferWriter);
+
+		/// <summary>
+		/// Attaches a buffer to send data.
+		/// </summary>
+		/// <param name="channelId">Channel to receive data on</param>
+		/// <param name="sendBufferReader">Reader for the buffer to send data from</param>
+		public abstract void AttachSendBuffer(int channelId, IComputeBufferReader sendBufferReader);
+	}
+
+	internal enum IpcMessage
+	{
+		AttachRecvBuffer = 0,
+		AttachSendBuffer = 1,
+	}
+
+	/// <summary>
+	/// Provides functionality for attaching buffers for compute workers 
+	/// </summary>
+	public sealed class WorkerComputeSocket : ComputeSocket, IDisposable
+	{
+		/// <summary>
+		/// Name of the environment variable for passing the name of the compute channel
+		/// </summary>
+		public const string IpcEnvVar = "UE_HORDE_COMPUTE_IPC";
+
+		readonly SharedMemoryBuffer _commandBuffer;
+
+		/// <summary>
+		/// Creates a socket for a worker
+		/// </summary>
+		private WorkerComputeSocket(SharedMemoryBuffer commandBuffer)
+		{
+			_commandBuffer = commandBuffer;
+		}
+
+		/// <summary>
+		/// Opens a socket which allows a worker to communicate with the Horde Agent
+		/// </summary>
+		public static WorkerComputeSocket Open()
+		{
+			string? baseName = Environment.GetEnvironmentVariable(IpcEnvVar);
+			if (baseName == null)
+			{
+				throw new InvalidOperationException($"Environment variable {IpcEnvVar} is not defined; cannot connect as worker.");
+			}
+
+			SharedMemoryBuffer commandBuffer = SharedMemoryBuffer.OpenExisting(baseName);
+			return new WorkerComputeSocket(commandBuffer);
+		}
+
+		/// <inheritdoc/>
+		public void Dispose()
+		{
+			_commandBuffer.Dispose();
+		}
+
+		/// <inheritdoc/>
+		public override void AttachRecvBuffer(int channelId, IComputeBufferWriter writer)
+		{
+			string bufferName = SharedMemoryBuffer.GetName(writer);
+			AttachBuffer(IpcMessage.AttachRecvBuffer, channelId, bufferName);
+		}
+
+		/// <inheritdoc/>
+		public override void AttachSendBuffer(int channelId, IComputeBufferReader reader)
+		{
+			string bufferName = SharedMemoryBuffer.GetName(reader);
+			AttachBuffer(IpcMessage.AttachSendBuffer, channelId, bufferName);
+		}
+
+		void AttachBuffer(IpcMessage message, int channelId, string bufferName)
+		{
+			MemoryWriter writer = new MemoryWriter(_commandBuffer.Writer.GetWriteBuffer());
+			writer.WriteUnsignedVarInt((int)message);
+			writer.WriteUnsignedVarInt(channelId);
+			writer.WriteString(bufferName);
+			_commandBuffer.Writer.AdvanceWritePosition(writer.Length);
+		}
+	}
+
 	/// <summary>
 	/// Enum identifying which end of the socket a particular machine is
 	/// </summary>
@@ -31,7 +123,7 @@ namespace EpicGames.Horde.Compute
 	/// <summary>
 	/// Manages a set of readers and writers to buffers across a transport layer
 	/// </summary>
-	public class ComputeSocket : IComputeSocket, IAsyncDisposable
+	public class RemoteComputeSocket : ComputeSocket, IAsyncDisposable
 	{
 		enum ControlMessageType
 		{
@@ -63,7 +155,7 @@ namespace EpicGames.Horde.Compute
 		/// <param name="transport">Transport to communicate with the remote</param>
 		/// <param name="endpoint">Tag for log messages</param>
 		/// <param name="logger">Logger for trace output</param>
-		public ComputeSocket(ComputeTransport transport, ComputeSocketEndpoint endpoint, ILogger logger)
+		public RemoteComputeSocket(ComputeTransport transport, ComputeSocketEndpoint endpoint, ILogger logger)
 		{
 			_transport = transport;
 			_endpoint = endpoint;
@@ -262,7 +354,7 @@ namespace EpicGames.Horde.Compute
 		}
 
 		/// <inheritdoc/>
-		public void AttachRecvBuffer(int channelId, IComputeBufferWriter recvBufferWriter)
+		public override void AttachRecvBuffer(int channelId, IComputeBufferWriter recvBufferWriter)
 		{
 			bool complete;
 			lock (_lockObject)
@@ -300,7 +392,7 @@ namespace EpicGames.Horde.Compute
 		}
 
 		/// <inheritdoc/>
-		public void AttachSendBuffer(int channelId, IComputeBufferReader sendBufferReader)
+		public override void AttachSendBuffer(int channelId, IComputeBufferReader sendBufferReader)
 		{
 			sendBufferReader = sendBufferReader.AddRef();
 			lock (_lockObject)
diff --git a/Engine/Source/Programs/Shared/EpicGames.Horde/Compute/IComputeLease.cs b/Engine/Source/Programs/Shared/EpicGames.Horde/Compute/IComputeLease.cs
--- a/Engine/Source/Programs/Shared/EpicGames.Horde/Compute/IComputeLease.cs
+++ b/Engine/Source/Programs/Shared/EpicGames.Horde/Compute/IComputeLease.cs
@@ -25,7 +25,7 @@ namespace EpicGames.Horde.Compute
 		/// <summary>
 		/// Socket to communicate with the remote
 		/// </summary>
-		ComputeSocket Socket { get; }
+		RemoteComputeSocket Socket { get; }
 
 		/// <summary>
 		/// Relinquish the lease gracefully
diff --git a/Engine/Source/Programs/Shared/EpicGames.Horde/Compute/IComputeSocket.cs b/Engine/Source/Programs/Shared/EpicGames.Horde/Compute/IComputeSocket.cs
deleted file mode 100644
diff --git a/Engine/Source/Programs/Shared/EpicGames.Horde/Compute/WorkerComputeSocket.cs b/Engine/Source/Programs/Shared/EpicGames.Horde/Compute/WorkerComputeSocket.cs
deleted file mode 100644
