Horde: Change IComputeBuffer to a concrete base class.

diff --git a/Engine/Source/Programs/Shared/EpicGames.Horde.Tests/BufferTests.cs b/Engine/Source/Programs/Shared/EpicGames.Horde.Tests/BufferTests.cs
--- a/Engine/Source/Programs/Shared/EpicGames.Horde.Tests/BufferTests.cs
+++ b/Engine/Source/Programs/Shared/EpicGames.Horde.Tests/BufferTests.cs
@@ -90,7 +90,7 @@ namespace EpicGames.Horde.Tests
 			}
 		}
 
-		static async Task TestProducerConsumerAsync(Func<int, IComputeBuffer> createBuffer, CancellationToken cancellationToken)
+		static async Task TestProducerConsumerAsync(Func<int, ComputeBuffer> createBuffer, CancellationToken cancellationToken)
 		{
 			const int Length = 8000;
 
@@ -100,7 +100,7 @@ namespace EpicGames.Horde.Tests
 			await using RemoteComputeSocket producerSocket = new RemoteComputeSocket(new PipeTransport(targetToSourcePipe.Reader, sourceToTargetPipe.Writer), ComputeSocketEndpoint.Local, NullLogger.Instance);
 			await using RemoteComputeSocket consumerSocket = new RemoteComputeSocket(new PipeTransport(sourceToTargetPipe.Reader, targetToSourcePipe.Writer), ComputeSocketEndpoint.Remote, NullLogger.Instance);
 
-			using IComputeBuffer consumerBuffer = createBuffer(Length);
+			using ComputeBuffer consumerBuffer = createBuffer(Length);
 			consumerSocket.AttachRecvBuffer(ChannelId, consumerBuffer.Writer);
 
 			byte[] input = RandomNumberGenerator.GetBytes(Length);
@@ -126,7 +126,7 @@ namespace EpicGames.Horde.Tests
 			await socket.MarkCompleteAsync(ChannelId);
 		}
 
-		static async Task RunConsumerAsync(IComputeBufferReader reader, Memory<byte> output)
+		static async Task RunConsumerAsync(ComputeBufferReader reader, Memory<byte> output)
 		{
 			int offset = 0;
 			while (!reader.IsComplete)
diff --git a/Engine/Source/Programs/Shared/EpicGames.Horde/Compute/AgentMessageChannel.cs b/Engine/Source/Programs/Shared/EpicGames.Horde/Compute/AgentMessageChannel.cs
--- a/Engine/Source/Programs/Shared/EpicGames.Horde/Compute/AgentMessageChannel.cs
+++ b/Engine/Source/Programs/Shared/EpicGames.Horde/Compute/AgentMessageChannel.cs
@@ -26,14 +26,14 @@ namespace EpicGames.Horde.Compute
 		class MessageBuilder : IAgentMessageBuilder
 		{
 			readonly AgentMessageChannel _channel;
-			readonly IComputeBufferWriter _sendBufferWriter;
+			readonly ComputeBufferWriter _sendBufferWriter;
 			readonly AgentMessageType _type;
 			int _length;
 
 			/// <inheritdoc/>
 			public int Length => _length;
 
-			public MessageBuilder(AgentMessageChannel channel, IComputeBufferWriter sendBufferWriter, AgentMessageType type)
+			public MessageBuilder(AgentMessageChannel channel, ComputeBufferWriter sendBufferWriter, AgentMessageType type)
 			{
 				_channel = channel;
 				_sendBufferWriter = sendBufferWriter;
@@ -75,8 +75,8 @@ namespace EpicGames.Horde.Compute
 		}
 
 		readonly ComputeSocket _socket;
-		readonly IComputeBufferReader _recvBufferReader;
-		readonly IComputeBufferWriter _sendBufferWriter;
+		readonly ComputeBufferReader _recvBufferReader;
+		readonly ComputeBufferWriter _sendBufferWriter;
 
 		// Can lock chunked memory writer to acuqire pointer
 		readonly ILogger _logger;
@@ -91,7 +91,7 @@ namespace EpicGames.Horde.Compute
 		/// <param name="recvBufferReader"></param>
 		/// <param name="sendBufferWriter"></param>
 		/// <param name="logger">Logger for diagnostic output</param>
-		public AgentMessageChannel(IComputeBufferReader recvBufferReader, IComputeBufferWriter sendBufferWriter, ILogger logger)
+		public AgentMessageChannel(ComputeBufferReader recvBufferReader, ComputeBufferWriter sendBufferWriter, ILogger logger)
 		{
 			_socket = null!;
 			_recvBufferReader = recvBufferReader.AddRef();
@@ -107,7 +107,7 @@ namespace EpicGames.Horde.Compute
 		/// <param name="recvBuffer"></param>
 		/// <param name="sendBuffer"></param>
 		/// <param name="logger">Logger for diagnostic output</param>
-		public AgentMessageChannel(ComputeSocket socket, int channelId, IComputeBuffer recvBuffer, IComputeBuffer sendBuffer, ILogger logger)
+		public AgentMessageChannel(ComputeSocket socket, int channelId, ComputeBuffer recvBuffer, ComputeBuffer sendBuffer, ILogger logger)
 		{
 			_socket = socket;
 			_socket.AttachRecvBuffer(channelId, recvBuffer.Writer);
@@ -233,8 +233,8 @@ namespace EpicGames.Horde.Compute
 		/// <param name="logger">Logger for the channel</param>
 		public static AgentMessageChannel CreateAgentMessageChannel(this ComputeSocket socket, int channelId, int sendBufferSize, int recvBufferSize, ILogger logger)
 		{
-			using IComputeBuffer sendBuffer = new PooledBuffer(sendBufferSize);
-			using IComputeBuffer recvBuffer = new PooledBuffer(recvBufferSize);
+			using ComputeBuffer sendBuffer = new PooledBuffer(sendBufferSize);
+			using ComputeBuffer recvBuffer = new PooledBuffer(recvBufferSize);
 			return new AgentMessageChannel(socket, channelId, sendBuffer, recvBuffer, logger);
 		}
 
diff --git a/Engine/Source/Programs/Shared/EpicGames.Horde/Compute/AgentMessageHandler.cs b/Engine/Source/Programs/Shared/EpicGames.Horde/Compute/AgentMessageHandler.cs
--- a/Engine/Source/Programs/Shared/EpicGames.Horde/Compute/AgentMessageHandler.cs
+++ b/Engine/Source/Programs/Shared/EpicGames.Horde/Compute/AgentMessageHandler.cs
@@ -198,12 +198,12 @@ namespace EpicGames.Horde.Compute
 			_logger.LogInformation("Child process has shut down");
 		}
 
-		async Task ProcessIpcMessagesAsync(ComputeSocket socket, IComputeBufferReader ipcReader, CancellationToken cancellationToken)
+		async Task ProcessIpcMessagesAsync(ComputeSocket socket, ComputeBufferReader ipcReader, CancellationToken cancellationToken)
 		{
 			List<(IpcMessage, int, SharedMemoryBuffer)> buffers = new();
 			try
 			{
-				List<(int, IComputeBufferWriter)> writers = new List<(int, IComputeBufferWriter)>();
+				List<(int, ComputeBufferWriter)> writers = new List<(int, ComputeBufferWriter)>();
 				while (await ipcReader.WaitToReadAsync(1, cancellationToken))
 				{
 					ReadOnlyMemory<byte> memory = ipcReader.GetReadBuffer();
diff --git a/Engine/Source/Programs/Shared/EpicGames.Horde/Compute/Buffers/PooledBuffer.cs b/Engine/Source/Programs/Shared/EpicGames.Horde/Compute/Buffers/PooledBuffer.cs
--- a/Engine/Source/Programs/Shared/EpicGames.Horde/Compute/Buffers/PooledBuffer.cs
+++ b/Engine/Source/Programs/Shared/EpicGames.Horde/Compute/Buffers/PooledBuffer.cs
@@ -14,67 +14,6 @@ namespace EpicGames.Horde.Compute.Buffers
 	/// </summary>
 	public sealed class PooledBuffer : ComputeBuffer
 	{
-		class Resources : ResourcesBase
-		{
-			readonly GCHandle _headerHandle;
-			readonly IMemoryOwner<byte>[] _chunkOwners;
-
-			readonly AsyncEvent _writerEvent = new AsyncEvent();
-			readonly AsyncEvent _readerEvent = new AsyncEvent();
-
-			public Resources(HeaderPtr headerPtr, GCHandle headerHandle, Memory<byte>[] chunks, IMemoryOwner<byte>[] chunkOwners)
-				: base(headerPtr, chunks)
-			{
-				_headerHandle = headerHandle;
-				_chunkOwners = chunkOwners;
-			}
-
-			public static unsafe Resources Create(int numChunks, int chunkLength, int numReaders)
-			{
-				byte[] header = new byte[HeaderSize];
-				GCHandle headerHandle = GCHandle.Alloc(header, GCHandleType.Pinned);
-				HeaderPtr headerPtr = new HeaderPtr((ulong*)headerHandle.AddrOfPinnedObject().ToPointer(), numReaders, numChunks, chunkLength);
-
-				Memory<byte>[] chunks = new Memory<byte>[numChunks];
-				IMemoryOwner<byte>[] chunkOwners = new IMemoryOwner<byte>[numChunks];
-
-				for (int idx = 0; idx < numChunks; idx++)
-				{
-					chunkOwners[idx] = MemoryPool<byte>.Shared.Rent(chunkLength);
-					chunks[idx] = chunkOwners[idx].Memory.Slice(0, chunkLength);
-				}
-
-				return new Resources(headerPtr, headerHandle, chunks, chunkOwners);
-			}
-
-			public override void Dispose()
-			{
-				for (int idx = 0; idx < _chunkOwners.Length; idx++)
-				{
-					_chunkOwners[idx].Dispose();
-				}
-				_headerHandle.Free();
-			}
-
-			/// <inheritdoc/>
-			public override void SetReadEvent(int readerIdx) => _readerEvent.Set();
-
-			/// <inheritdoc/>
-			public override void ResetReadEvent(int readerIdx) => _readerEvent.Reset();
-
-			/// <inheritdoc/>
-			public override Task WaitForReadEvent(int readerIdx, CancellationToken cancellationToken) => _readerEvent.Task.WaitAsync(cancellationToken);
-
-			/// <inheritdoc/>
-			public override void SetWriteEvent() => _writerEvent.Set();
-
-			/// <inheritdoc/>
-			public override void ResetWriteEvent() => _writerEvent.Reset();
-
-			/// <inheritdoc/>
-			public override Task WaitForWriteEvent(CancellationToken cancellationToken) => _writerEvent.Task.WaitAsync(cancellationToken);
-		}
-
 		/// <summary>
 		/// Constructor
 		/// </summary>
@@ -91,20 +30,88 @@ namespace EpicGames.Horde.Compute.Buffers
 		/// <param name="chunkLength">Length of each chunk</param>
 		/// <param name="numReaders">Number of readers for this buffer</param>
 		public PooledBuffer(int numChunks, int chunkLength, int numReaders = 1)
-			: base(Resources.Create(numChunks, chunkLength, numReaders))
+			: base(PooledBufferDetail.Create(numChunks, chunkLength, numReaders))
 		{
 		}
 
-		private PooledBuffer(ResourcesBase resources)
+		private PooledBuffer(ComputeBufferDetail resources)
 			: base(resources)
 		{
 		}
 
 		/// <inheritdoc/>
-		public override IComputeBuffer AddRef()
+		public override ComputeBuffer AddRef()
 		{
-			_resources.AddRef();
-			return new PooledBuffer(_resources);
+			_detail.AddRef();
+			return new PooledBuffer(_detail);
 		}
 	}
+
+	/// <summary>
+	/// Core implementation of <see cref="PooledBuffer"/>
+	/// </summary>
+	class PooledBufferDetail : ComputeBufferDetail
+	{
+		readonly GCHandle _headerHandle;
+		readonly IMemoryOwner<byte>[] _chunkOwners;
+
+		readonly AsyncEvent _writerEvent = new AsyncEvent();
+		readonly AsyncEvent _readerEvent = new AsyncEvent();
+
+		internal PooledBufferDetail(HeaderPtr headerPtr, GCHandle headerHandle, Memory<byte>[] chunks, IMemoryOwner<byte>[] chunkOwners)
+			: base(headerPtr, chunks)
+		{
+			_headerHandle = headerHandle;
+			_chunkOwners = chunkOwners;
+		}
+
+		public static unsafe PooledBufferDetail Create(int numChunks, int chunkLength, int numReaders)
+		{
+			byte[] header = new byte[HeaderSize];
+			GCHandle headerHandle = GCHandle.Alloc(header, GCHandleType.Pinned);
+			HeaderPtr headerPtr = new HeaderPtr((ulong*)headerHandle.AddrOfPinnedObject().ToPointer(), numReaders, numChunks, chunkLength);
+
+			Memory<byte>[] chunks = new Memory<byte>[numChunks];
+			IMemoryOwner<byte>[] chunkOwners = new IMemoryOwner<byte>[numChunks];
+
+			for (int idx = 0; idx < numChunks; idx++)
+			{
+				chunkOwners[idx] = MemoryPool<byte>.Shared.Rent(chunkLength);
+				chunks[idx] = chunkOwners[idx].Memory.Slice(0, chunkLength);
+			}
+
+			return new PooledBufferDetail(headerPtr, headerHandle, chunks, chunkOwners);
+		}
+
+		/// <inheritdoc/>
+		protected override void Dispose(bool disposing)
+		{
+			if (disposing)
+			{
+				for (int idx = 0; idx < _chunkOwners.Length; idx++)
+				{
+					_chunkOwners[idx].Dispose();
+				}
+				_headerHandle.Free();
+			}
+		}
+
+		/// <inheritdoc/>
+		public override void SetReadEvent(int readerIdx) => _readerEvent.Set();
+
+		/// <inheritdoc/>
+		public override void ResetReadEvent(int readerIdx) => _readerEvent.Reset();
+
+		/// <inheritdoc/>
+		public override Task WaitForReadEvent(int readerIdx, CancellationToken cancellationToken) => _readerEvent.Task.WaitAsync(cancellationToken);
+
+		/// <inheritdoc/>
+		public override void SetWriteEvent() => _writerEvent.Set();
+
+		/// <inheritdoc/>
+		public override void ResetWriteEvent() => _writerEvent.Reset();
+
+		/// <inheritdoc/>
+		public override Task WaitForWriteEvent(CancellationToken cancellationToken) => _writerEvent.Task.WaitAsync(cancellationToken);
+	}
 }
diff --git a/Engine/Source/Programs/Shared/EpicGames.Horde/Compute/Buffers/SharedMemoryBuffer.cs b/Engine/Source/Programs/Shared/EpicGames.Horde/Compute/Buffers/SharedMemoryBuffer.cs
--- a/Engine/Source/Programs/Shared/EpicGames.Horde/Compute/Buffers/SharedMemoryBuffer.cs
+++ b/Engine/Source/Programs/Shared/EpicGames.Horde/Compute/Buffers/SharedMemoryBuffer.cs
@@ -15,75 +15,22 @@ namespace EpicGames.Horde.Compute.Buffers
 	/// </summary>
 	public sealed class SharedMemoryBuffer : ComputeBuffer
 	{
-		class Resources : ResourcesBase
-		{
-			readonly MemoryMappedFile _memoryMappedFile;
-			readonly MemoryMappedViewAccessor _memoryMappedViewAccessor;
-			readonly MemoryMappedView _memoryMappedView;
-
-			readonly Native.EventHandle _readerEvent;
-			readonly Native.EventHandle _writerEvent;
-
-			public string Name { get; }
-
-			public Resources(string name, HeaderPtr headerPtr, Memory<byte>[] chunks, MemoryMappedFile memoryMappedFile, MemoryMappedViewAccessor memoryMappedViewAccessor, MemoryMappedView memoryMappedView, Native.EventHandle readerEvent, Native.EventHandle writerEvent)
-				: base(headerPtr, chunks)
-			{
-				Name = name;
-
-				_memoryMappedFile = memoryMappedFile;
-				_memoryMappedViewAccessor = memoryMappedViewAccessor;
-				_memoryMappedView = memoryMappedView;
-
-				_readerEvent = readerEvent;
-				_writerEvent = writerEvent;
-			}
-
-			public override void Dispose()
-			{
-				_readerEvent.Dispose();
-				_writerEvent.Dispose();
-
-				_memoryMappedView.Dispose();
-				_memoryMappedViewAccessor.Dispose();
-				_memoryMappedFile.Dispose();
-			}
-
-			/// <inheritdoc/>
-			public override void SetReadEvent(int readerIdx) => _readerEvent.Set();
-
-			/// <inheritdoc/>
-			public override void ResetReadEvent(int readerIdx) => _readerEvent.Reset();
-
-			/// <inheritdoc/>
-			public override Task WaitForReadEvent(int readerIdx, CancellationToken cancellationToken) => _readerEvent.WaitOneAsync(cancellationToken);
-
-			/// <inheritdoc/>
-			public override void SetWriteEvent() => _writerEvent.Set();
-
-			/// <inheritdoc/>
-			public override void ResetWriteEvent() => _writerEvent.Reset();
-
-			/// <inheritdoc/>
-			public override Task WaitForWriteEvent(CancellationToken cancellationToken) => _writerEvent.WaitOneAsync(cancellationToken);
-		}
-
 		/// <inheritdoc/>
-		public string Name => ((Resources)_resources).Name;
+		public string Name => ((SharedMemoryBufferDetail)_detail).Name;
 
 		/// <summary>
 		/// Constructor
 		/// </summary>
-		private SharedMemoryBuffer(ResourcesBase resources)
+		private SharedMemoryBuffer(ComputeBufferDetail resources)
 			: base(resources)
 		{
 		}
 
 		/// <inheritdoc/>
-		public override IComputeBuffer AddRef()
+		public override ComputeBuffer AddRef()
 		{
-			_resources.AddRef();
-			return new SharedMemoryBuffer(_resources);
+			_detail.AddRef();
+			return new SharedMemoryBuffer(_detail);
 		}
 
 		/// <summary>
@@ -103,7 +50,49 @@ namespace EpicGames.Horde.Compute.Buffers
 		/// <param name="name">Name of the buffer</param>
 		/// <param name="numChunks">Number of chunks in the buffer</param>
 		/// <param name="chunkLength">Length of each chunk</param>
-		public static unsafe SharedMemoryBuffer CreateNew(string? name, int numChunks, int chunkLength)
+		public static unsafe SharedMemoryBuffer CreateNew(string? name, int numChunks, int chunkLength) => new SharedMemoryBuffer(SharedMemoryBufferDetail.CreateNew(name, numChunks, chunkLength));
+
+		/// <summary>
+		/// Open an existing buffer by name
+		/// </summary>
+		/// <param name="name">Name of the buffer to open</param>
+		public static unsafe SharedMemoryBuffer OpenExisting(string name) => new SharedMemoryBuffer(SharedMemoryBufferDetail.OpenExisting(name));
+	}
+
+	/// <summary>
+	/// Core implementation of <see cref="SharedMemoryBuffer"/>
+	/// </summary>
+	class SharedMemoryBufferDetail : ComputeBufferDetail
+	{
+		readonly MemoryMappedFile _memoryMappedFile;
+		readonly MemoryMappedViewAccessor _memoryMappedViewAccessor;
+		readonly MemoryMappedView _memoryMappedView;
+
+		readonly Native.EventHandle _readerEvent;
+		readonly Native.EventHandle _writerEvent;
+
+		public string Name { get; }
+
+		internal SharedMemoryBufferDetail(string name, HeaderPtr headerPtr, Memory<byte>[] chunks, MemoryMappedFile memoryMappedFile, MemoryMappedViewAccessor memoryMappedViewAccessor, MemoryMappedView memoryMappedView, Native.EventHandle readerEvent, Native.EventHandle writerEvent)
+			: base(headerPtr, chunks)
+		{
+			Name = name;
+
+			_memoryMappedFile = memoryMappedFile;
+			_memoryMappedViewAccessor = memoryMappedViewAccessor;
+			_memoryMappedView = memoryMappedView;
+
+			_readerEvent = readerEvent;
+			_writerEvent = writerEvent;
+		}
+
+		/// <summary>
+		/// Create a new shared memory buffer
+		/// </summary>
+		/// <param name="name">Name of the buffer</param>
+		/// <param name="numChunks">Number of chunks in the buffer</param>
+		/// <param name="chunkLength">Length of each chunk</param>
+		public static unsafe SharedMemoryBufferDetail CreateNew(string? name, int numChunks, int chunkLength)
 		{
 			long capacity = HeaderSize + (numChunks * sizeof(ulong)) + (numChunks * chunkLength);
 
@@ -119,14 +108,14 @@ namespace EpicGames.Horde.Compute.Buffers
 			HeaderPtr headerPtr = new HeaderPtr((ulong*)memoryMappedView.GetPointer(), 1, numChunks, chunkLength);
 			Memory<byte>[] chunks = CreateChunks(headerPtr.NumChunks, headerPtr.ChunkLength, memoryMappedView);
 
-			return new SharedMemoryBuffer(new Resources(name, headerPtr, chunks, memoryMappedFile, memoryMappedViewAccessor, memoryMappedView, readerEvent, writerEvent));
+			return new SharedMemoryBufferDetail(name, headerPtr, chunks, memoryMappedFile, memoryMappedViewAccessor, memoryMappedView, readerEvent, writerEvent);
 		}
 
 		/// <summary>
 		/// Open an existing buffer by name
 		/// </summary>
 		/// <param name="name">Name of the buffer to open</param>
-		public static unsafe SharedMemoryBuffer OpenExisting(string name)
+		public static unsafe SharedMemoryBufferDetail OpenExisting(string name)
 		{
 			if (!RuntimeInformation.IsOSPlatform(OSPlatform.Windows))
 			{
@@ -143,7 +132,7 @@ namespace EpicGames.Horde.Compute.Buffers
 			HeaderPtr headerPtr = new HeaderPtr((ulong*)memoryMappedView.GetPointer());
 			Memory<byte>[] chunks = CreateChunks(headerPtr.NumChunks, headerPtr.ChunkLength, memoryMappedView);
 
-			return new SharedMemoryBuffer(new Resources(name, headerPtr, chunks, memoryMappedFile, memoryMappedViewAccessor, memoryMappedView, readerEvent, writerEvent));
+			return new SharedMemoryBufferDetail(name, headerPtr, chunks, memoryMappedFile, memoryMappedViewAccessor, memoryMappedView, readerEvent, writerEvent);
 		}
 
 		static Memory<byte>[] CreateChunks(int numChunks, int chunkLength, MemoryMappedView memoryMappedView)
@@ -157,16 +146,36 @@ namespace EpicGames.Horde.Compute.Buffers
 			return chunks;
 		}
 
-		internal static string GetName(IComputeBufferReader reader)
+		/// <inheritdoc/>
+		protected override void Dispose(bool disposing)
 		{
-			Resources resources = (Resources)((ReaderImpl)reader).GetResources();
-			return resources.Name;
-		}
+			if (disposing)
+			{
+				_readerEvent.Dispose();
+				_writerEvent.Dispose();
 
-		internal static string GetName(IComputeBufferWriter writer)
-		{
-			Resources resources = (Resources)((WriterImpl)writer).GetResources();
-			return resources.Name;
+				_memoryMappedView.Dispose();
+				_memoryMappedViewAccessor.Dispose();
+				_memoryMappedFile.Dispose();
+			}
 		}
+
+		/// <inheritdoc/>
+		public override void SetReadEvent(int readerIdx) => _readerEvent.Set();
+
+		/// <inheritdoc/>
+		public override void ResetReadEvent(int readerIdx) => _readerEvent.Reset();
+
+		/// <inheritdoc/>
+		public override Task WaitForReadEvent(int readerIdx, CancellationToken cancellationToken) => _readerEvent.WaitOneAsync(cancellationToken);
+
+		/// <inheritdoc/>
+		public override void SetWriteEvent() => _writerEvent.Set();
+
+		/// <inheritdoc/>
+		public override void ResetWriteEvent() => _writerEvent.Reset();
+
+		/// <inheritdoc/>
+		public override Task WaitForWriteEvent(CancellationToken cancellationToken) => _writerEvent.WaitOneAsync(cancellationToken);
 	}
 }
diff --git a/Engine/Source/Programs/Shared/EpicGames.Horde/Compute/ComputeBuffer.cs b/Engine/Source/Programs/Shared/EpicGames.Horde/Compute/ComputeBuffer.cs
--- a/Engine/Source/Programs/Shared/EpicGames.Horde/Compute/ComputeBuffer.cs
+++ b/Engine/Source/Programs/Shared/EpicGames.Horde/Compute/ComputeBuffer.cs
@@ -11,17 +11,237 @@ namespace EpicGames.Horde.Compute
 	/// <summary>
 	/// In-process buffer used to store compute messages
 	/// </summary>
-	public abstract class ComputeBuffer : IComputeBuffer
+	public abstract class ComputeBuffer : IDisposable
 	{
-		const int MaxChunks = 16;
-		const int MaxReaders = 16;
+		/// <summary>
+		/// Maximum number of chunks in a buffer
+		/// </summary>
+		public const int MaxChunks = 16;
+
+		/// <summary>
+		/// Maximum number of readers
+		/// </summary>
+		public const int MaxReaders = 16;
+
+		internal ComputeBufferDetail _detail;
+
+		/// <summary>
+		/// Reader for this buffer
+		/// </summary>
+		public ComputeBufferReader Reader { get; }
+
+		/// <summary>
+		/// Writer for this buffer
+		/// </summary>
+		public ComputeBufferWriter Writer { get; }
+
+		/// <summary>
+		/// Constructor
+		/// </summary>
+		/// <param name="detail">Resources shared between instances of the buffer</param>
+		internal ComputeBuffer(ComputeBufferDetail detail)
+		{
+			_detail = detail;
+			Reader = new ComputeBufferReader(detail, 0);
+			Writer = new ComputeBufferWriter(detail);
+		}
+
+		/// <inheritdoc/>
+		public void Dispose()
+		{
+			Dispose(true);
+			GC.SuppressFinalize(this);
+		}
+
+		/// <summary>
+		/// Overriable dispose method
+		/// </summary>
+		protected virtual void Dispose(bool disposing)
+		{
+			if (_detail != null)
+			{
+				_detail.Release();
+				_detail = null!;
+			}
+		}
+
+		/// <summary>
+		/// Creates a new reference to the underlying buffer. The underlying resources will only be destroyed once all instances are disposed of.
+		/// </summary>
+		public abstract ComputeBuffer AddRef();
+	}
+
+	/// <summary>
+	/// Read interface for a compute buffer
+	/// </summary>
+	public sealed class ComputeBufferReader : IDisposable
+	{
+		ComputeBufferDetail _detail;
+		readonly int _readerIdx;
+
+		internal ComputeBufferDetail Detail => _detail;
+
+		internal ComputeBufferReader(ComputeBufferDetail detail, int readerIdx)
+		{
+			_detail = detail;
+			_readerIdx = readerIdx;
+		}
+
+		/// <summary>
+		/// Create a new reader instance using the same underlying buffer
+		/// </summary>
+		public ComputeBufferReader AddRef()
+		{
+			_detail.AddRef();
+			return new ComputeBufferReader(_detail, _readerIdx);
+		}
+
+		/// <inheritdoc/>
+		public void Dispose()
+		{
+			if (_detail != null)
+			{
+				_detail.Release();
+				_detail = null!;
+			}
+		}
+
+		/// <summary>
+		/// Whether this buffer is complete (no more data will be added)
+		/// </summary>
+		public bool IsComplete => _detail.IsReadingComplete(_readerIdx);
+
+		/// <summary>
+		/// Updates the read position
+		/// </summary>
+		/// <param name="length">Size of data that was read</param>
+		public void AdvanceReadPosition(int length) => _detail.AdvanceReadPosition(_readerIdx, length);
+
+		/// <summary>
+		/// Gets the next data to read
+		/// </summary>
+		/// <returns>Memory to read from</returns>
+		public ReadOnlyMemory<byte> GetReadBuffer() => _detail.GetReadBuffer(_readerIdx);
+
+		/// <summary>
+		/// Read from a buffer into another buffer
+		/// </summary>
+		/// <param name="buffer">Memory to receive the read data</param>
+		/// <param name="cancellationToken">Cancellation token for the operation</param>
+		/// <returns>Number of bytes read</returns>
+		public async ValueTask<int> ReadAsync(Memory<byte> buffer, CancellationToken cancellationToken = default)
+		{
+			for (; ; )
+			{
+				ReadOnlyMemory<byte> readMemory = GetReadBuffer();
+				if (IsComplete || readMemory.Length > 0)
+				{
+					int length = Math.Min(readMemory.Length, buffer.Length);
+					readMemory.Slice(0, length).CopyTo(buffer);
+					AdvanceReadPosition(length);
+					return length;
+				}
+				await WaitToReadAsync(1, cancellationToken);
+			}
+		}
+
+		/// <summary>
+		/// Wait for data to be available, or for the buffer to be marked as complete
+		/// </summary>
+		/// <param name="minLength">Minimum amount of data to read</param>
+		/// <param name="cancellationToken">Cancellation token for the operation</param>
+		/// <returns>True if new data is available, false if the buffer is complete</returns>
+		public ValueTask<bool> WaitToReadAsync(int minLength, CancellationToken cancellationToken = default) => _detail.WaitToReadAsync(_readerIdx, minLength, cancellationToken);
+	}
+
+	/// <summary>
+	/// Buffer that can receive data from a remote machine.
+	/// </summary>
+	public sealed class ComputeBufferWriter : IDisposable
+	{
+		ComputeBufferDetail _detail;
+
+		internal ComputeBufferDetail Detail => _detail;
+
+		internal ComputeBufferWriter(ComputeBufferDetail detail)
+		{
+			_detail = detail;
+		}
+
+		/// <summary>
+		/// Create a new writer instance using the same underlying buffer
+		/// </summary>
+		public ComputeBufferWriter AddRef()
+		{
+			_detail.AddRef();
+			return new ComputeBufferWriter(_detail);
+		}
+
+		/// <inheritdoc/>
+		public void Dispose()
+		{
+			if (_detail != null)
+			{
+				_detail.Release();
+				_detail = null!;
+			}
+		}
+
+		/// <inheritdoc/>
+		public void AdvanceWritePosition(int size) => _detail.AdvanceWritePosition(size);
+
+		/// <summary>
+		/// Gets memory to write to
+		/// </summary>
+		/// <returns>Memory to be written to</returns>
+		public Memory<byte> GetWriteBuffer() => _detail.GetWriteBuffer();
 
-		internal const int HeaderSize = (2 + MaxChunks + MaxReaders) * sizeof(ulong);
+		/// <summary>
+		/// Mark the output to this buffer as complete
+		/// </summary>
+		/// <returns>Whether the writer was marked as complete. False if the writer has already been marked as complete.</returns>
+		public bool MarkComplete() => _detail.MarkComplete();
+
+		/// <summary>
+		/// Writes data into a buffer from a memory block
+		/// </summary>
+		/// <param name="buffer">The data to write</param>
+		/// <param name="cancellationToken">Cancellation token for the operation</param>
+		public async ValueTask WriteAsync(ReadOnlyMemory<byte> buffer, CancellationToken cancellationToken = default)
+		{
+			while (buffer.Length > 0)
+			{
+				Memory<byte> writeMemory = GetWriteBuffer();
+				if (writeMemory.Length >= buffer.Length)
+				{
+					buffer.CopyTo(writeMemory);
+					AdvanceWritePosition(buffer.Length);
+					break;
+				}
+				await WaitToWriteAsync(writeMemory.Length, cancellationToken);
+			}
+		}
+
+		/// <summary>
+		/// Gets memory to write to
+		/// </summary>
+		/// <param name="minLength">Minimum size of the desired write buffer</param>
+		/// <param name="cancellationToken">Cancellation token for the operation</param>
+		/// <returns>Memory to be written to</returns>
+		public ValueTask WaitToWriteAsync(int minLength, CancellationToken cancellationToken = default) => _detail.WaitToWriteAsync(minLength, cancellationToken);
+	}
+
+	/// <summary>
+	/// State shared between buffer instances
+	/// </summary>
+	internal abstract class ComputeBufferDetail : IDisposable
+	{
+		internal const int HeaderSize = (2 + ComputeBuffer.MaxChunks + ComputeBuffer.MaxReaders) * sizeof(ulong);
 
 		/// <summary>
 		/// Tracked state of the buffer
 		/// </summary>
-		internal readonly unsafe struct HeaderPtr
+		protected internal readonly unsafe struct HeaderPtr
 		{
 			readonly ulong* _data;
 
@@ -48,13 +268,13 @@ namespace EpicGames.Horde.Compute
 
 			public ChunkStatePtr GetChunkStatePtr(int chunkIdx) => new ChunkStatePtr(_data + 2 + chunkIdx);
 
-			public ReaderStatePtr GetReaderStatePtr(int readerIdx) => new ReaderStatePtr(_data + 2 + MaxChunks + readerIdx);
+			public ReaderStatePtr GetReaderStatePtr(int readerIdx) => new ReaderStatePtr(_data + 2 + ComputeBuffer.MaxChunks + readerIdx);
 		}
 
 		/// <summary>
 		/// Write state for a chunk
 		/// </summary>
-		internal enum WriteState
+		protected internal enum WriteState
 		{
 			/// <summary>
 			/// Writer has moved to the next chunk
@@ -75,7 +295,7 @@ namespace EpicGames.Horde.Compute
 		/// <summary>
 		/// Stores the state of a chunk in a 64-bit value, which can be updated atomically
 		/// </summary>
-		internal record struct ChunkState(ulong Value)
+		protected internal record struct ChunkState(ulong Value)
 		{
 			// Written length of this chunk
 			public int Length => (int)(Value & 0x7fffffff);
@@ -99,7 +319,7 @@ namespace EpicGames.Horde.Compute
 		/// <summary>
 		/// Wraps a pointer to the state of a chunk
 		/// </summary>
-		internal readonly unsafe struct ChunkStatePtr
+		protected internal readonly unsafe struct ChunkStatePtr
 		{
 			readonly ulong* _data;
 
@@ -139,7 +359,7 @@ namespace EpicGames.Horde.Compute
 		/// <summary>
 		/// Encodes the state of a reader in a 64-bit value
 		/// </summary>
-		internal record struct ReaderState(ulong Value)
+		protected internal record struct ReaderState(ulong Value)
 		{
 			public ReaderState(int chunkIdx, int offset) : this(((ulong)chunkIdx << 32) | (uint)offset) { }
 
@@ -152,7 +372,7 @@ namespace EpicGames.Horde.Compute
 		/// <summary>
 		/// Pointer to a reader state value
 		/// </summary>
-		internal readonly unsafe struct ReaderStatePtr
+		protected internal readonly unsafe struct ReaderStatePtr
 		{
 			readonly ulong* _data;
 
@@ -172,377 +392,274 @@ namespace EpicGames.Horde.Compute
 			public override string ToString() => Value.ToString();
 		}
 
+		readonly HeaderPtr _headerPtr;
+		readonly Memory<byte>[] _chunks;
+		int _refCount = 1;
+
 		/// <summary>
-		/// State shared between buffer instances
+		/// Constructor
 		/// </summary>
-		internal abstract class ResourcesBase : IDisposable
+		protected ComputeBufferDetail(HeaderPtr headerPtr, Memory<byte>[] chunks)
 		{
-			public HeaderPtr HeaderPtr { get; }
-			public Memory<byte>[] Chunks { get; }
-			int _refCount = 1;
+			_headerPtr = headerPtr;
+			_chunks = chunks;
+		}
 
-			public ResourcesBase(HeaderPtr headerPtr, Memory<byte>[] chunks)
-			{
-				HeaderPtr = headerPtr;
-				Chunks = chunks;
-			}
+		/// <summary>
+		/// Increment the reference count on this object
+		/// </summary>
+		public void AddRef()
+		{
+			Interlocked.Increment(ref _refCount);
+		}
 
-			public void AddRef()
+		/// <summary>
+		/// Decrement the reference count on this object, and dispose of it once it reaches zero
+		/// </summary>
+		public void Release()
+		{
+			if (Interlocked.Decrement(ref _refCount) == 0)
 			{
-				Interlocked.Increment(ref _refCount);
+				Dispose();
 			}
+		}
 
-			public void Release()
-			{
-				if (Interlocked.Decrement(ref _refCount) == 0)
-				{
-					Dispose();
-				}
-			}
+		/// <inheritdoc/>
+		public void Dispose()
+		{
+			Dispose(true);
+			GC.SuppressFinalize(this);
+		}
 
-			/// <summary>
-			/// Dispose of the 
-			/// </summary>
-			public abstract void Dispose();
+		/// <summary>
+		/// Overridable dispose method
+		/// </summary>
+		protected abstract void Dispose(bool disposing);
 
-			/// <summary>
-			/// Signals a read event
-			/// </summary>
-			public abstract void SetReadEvent(int readerIdx);
+		/// <summary>
+		/// Signals a read event
+		/// </summary>
+		public abstract void SetReadEvent(int readerIdx);
 
-			/// <summary>
-			/// Signals read events for every reader
-			/// </summary>
-			public void SetAllReadEvents()
+		/// <summary>
+		/// Signals read events for every reader
+		/// </summary>
+		public void SetAllReadEvents()
+		{
+			for (int readerIdx = 0; readerIdx < _headerPtr.NumReaders; readerIdx++)
 			{
-				for (int readerIdx = 0; readerIdx < HeaderPtr.NumReaders; readerIdx++)
-				{
-					SetReadEvent(readerIdx);
-				}
+				SetReadEvent(readerIdx);
 			}
+		}
 
-			/// <summary>
-			/// Resets a read event
-			/// </summary>
-			public abstract void ResetReadEvent(int readerIdx);
+		/// <summary>
+		/// Resets a read event
+		/// </summary>
+		public abstract void ResetReadEvent(int readerIdx);
 
-			/// <summary>
-			/// Waits for a read event to be signalled
-			/// </summary>
-			public abstract Task WaitForReadEvent(int readerIdx, CancellationToken cancellationToken);
+		/// <summary>
+		/// Waits for a read event to be signalled
+		/// </summary>
+		public abstract Task WaitForReadEvent(int readerIdx, CancellationToken cancellationToken);
 
-			/// <summary>
-			/// Signals the write event
-			/// </summary>
-			public abstract void SetWriteEvent();
+		/// <summary>
+		/// Signals the write event
+		/// </summary>
+		public abstract void SetWriteEvent();
 
-			/// <summary>
-			/// Resets the write event
-			/// </summary>
-			public abstract void ResetWriteEvent();
+		/// <summary>
+		/// Resets the write event
+		/// </summary>
+		public abstract void ResetWriteEvent();
 
-			/// <summary>
-			/// Waits for the write event to be signalled
-			/// </summary>
-			public abstract Task WaitForWriteEvent(CancellationToken cancellationToken);
+		/// <summary>
+		/// Waits for the write event to be signalled
+		/// </summary>
+		public abstract Task WaitForWriteEvent(CancellationToken cancellationToken);
 
 #pragma warning disable IDE0051 // Remove unused private members
-			// For debugging purposes only
-			ChunkState[] ChunkStates => Enumerable.Range(0, HeaderPtr.NumChunks).Select(x => HeaderPtr.GetChunkStatePtr(x).Value).ToArray();
-			ReaderState[] ReaderStates => Enumerable.Range(0, HeaderPtr.NumReaders).Select(x => HeaderPtr.GetReaderStatePtr(x).Value).ToArray();
-			int WriteChunkIdx => HeaderPtr.WriteChunkIdx;
+		// For debugging purposes only
+		ChunkState[] ChunkStates => Enumerable.Range(0, _headerPtr.NumChunks).Select(x => _headerPtr.GetChunkStatePtr(x).Value).ToArray();
+		ReaderState[] ReaderStates => Enumerable.Range(0, _headerPtr.NumReaders).Select(x => _headerPtr.GetReaderStatePtr(x).Value).ToArray();
+		int WriteChunkIdx => _headerPtr.WriteChunkIdx;
 #pragma warning restore IDE0051 // Remove unused private members
+
+		public bool IsReadingComplete(int readerIdx)
+		{
+			ReaderState readerState = _headerPtr.GetReaderStatePtr(readerIdx).Value;
+			ChunkState chunkState = _headerPtr.GetChunkStatePtr(readerState.ChunkIdx).Value;
+			return chunkState.WriteState == WriteState.Complete && readerState.Offset == chunkState.Length;
 		}
 
-		internal class ReaderImpl : IComputeBufferReader
+		/// <inheritdoc/>
+		public void AdvanceReadPosition(int readerIdx, int length)
 		{
-			readonly HeaderPtr _headerPtr;
-			readonly Memory<byte>[] _chunks;
-			ResourcesBase _resources;
-			readonly int _readerIdx;
+			ReaderStatePtr readerStatePtr = _headerPtr.GetReaderStatePtr(readerIdx);
+			readerStatePtr.Advance(length);
+		}
 
-			public ReaderImpl(ResourcesBase resources, int readerIdx)
-			{
-				_headerPtr = resources.HeaderPtr;
-				_chunks = resources.Chunks;
-				_resources = resources;
-				_readerIdx = readerIdx;
-			}
+		/// <inheritdoc/>
+		public ReadOnlyMemory<byte> GetReadBuffer(int readerIdx)
+		{
+			ReaderStatePtr readerStatePtr = _headerPtr.GetReaderStatePtr(readerIdx);
+			ReaderState readerState = readerStatePtr.Value;
 
-			public ResourcesBase GetResources() => _resources;
+			ChunkStatePtr chunkStatePtr = _headerPtr.GetChunkStatePtr(readerState.ChunkIdx);
+			ChunkState chunkState = chunkStatePtr.Value;
 
-			/// <inheritdoc/>
-			public IComputeBufferReader AddRef()
+			if (chunkState.HasReaderFlag(readerIdx))
 			{
-				_resources.AddRef();
-				return new ReaderImpl(_resources, _readerIdx);
+				return _chunks[readerState.ChunkIdx].Slice(readerState.Offset, chunkState.Length - readerState.Offset);
 			}
-
-			/// <inheritdoc/>
-			public void Dispose()
+			else
 			{
-				if (_resources != null)
-				{
-					_resources.Release();
-					_resources = null!;
-				}
-			}
-
-			/// <inheritdoc/>
-			public bool IsComplete
-			{
-				get
-				{
-					ReaderState readerState = _headerPtr.GetReaderStatePtr(_readerIdx).Value;
-					ChunkState chunkState = _headerPtr.GetChunkStatePtr(readerState.ChunkIdx).Value;
-					return chunkState.WriteState == WriteState.Complete && readerState.Offset == chunkState.Length;
-				}
-			}
-
-			/// <inheritdoc/>
-			public void AdvanceReadPosition(int length)
-			{
-				ReaderStatePtr readerStatePtr = _headerPtr.GetReaderStatePtr(_readerIdx);
-				readerStatePtr.Advance(length);
+				return ReadOnlyMemory<byte>.Empty;
 			}
+		}
 
-			/// <inheritdoc/>
-			public ReadOnlyMemory<byte> GetReadBuffer()
+		/// <inheritdoc/>
+		public async ValueTask<bool> WaitToReadAsync(int readerIdx, int minLength, CancellationToken cancellationToken = default)
+		{
+			for (; ; )
 			{
-				ReaderStatePtr readerStatePtr = _headerPtr.GetReaderStatePtr(_readerIdx);
+				ReaderStatePtr readerStatePtr = _headerPtr.GetReaderStatePtr(readerIdx);
 				ReaderState readerState = readerStatePtr.Value;
 
 				ChunkStatePtr chunkStatePtr = _headerPtr.GetChunkStatePtr(readerState.ChunkIdx);
 				ChunkState chunkState = chunkStatePtr.Value;
 
-				if (chunkState.HasReaderFlag(_readerIdx))
+				if (!chunkState.HasReaderFlag(readerIdx))
 				{
-					return _chunks[readerState.ChunkIdx].Slice(readerState.Offset, chunkState.Length - readerState.Offset);
+					// Wait until the current chunk is readable
+					ResetReadEvent(readerIdx);
+					if (!chunkState.HasReaderFlag(readerIdx))
+					{
+						await WaitForReadEvent(readerIdx, cancellationToken);
+					}
 				}
-				else
+				else if (readerState.Offset + minLength <= chunkState.Length)
 				{
-					return ReadOnlyMemory<byte>.Empty;
+					// We have enough data in the chunk to be able to read a message
+					return true;
 				}
-			}
-
-			/// <inheritdoc/>
-			public async ValueTask<bool> WaitToReadAsync(int minLength, CancellationToken cancellationToken = default)
-			{
-				for (; ; )
+				else if (chunkState.WriteState == WriteState.Writing)
 				{
-					ReaderStatePtr readerStatePtr = _headerPtr.GetReaderStatePtr(_readerIdx);
-					ReaderState readerState = readerStatePtr.Value;
-
-					ChunkStatePtr chunkStatePtr = _headerPtr.GetChunkStatePtr(readerState.ChunkIdx);
-					ChunkState chunkState = chunkStatePtr.Value;
-
-					if (!chunkState.HasReaderFlag(_readerIdx))
-					{
-						// Wait until the current chunk is readable
-						_resources.ResetReadEvent(_readerIdx);
-						if (!chunkState.HasReaderFlag(_readerIdx))
-						{
-							await _resources.WaitForReadEvent(_readerIdx, cancellationToken);
-						}
-					}
-					else if (readerState.Offset + minLength <= chunkState.Length)
-					{
-						// We have enough data in the chunk to be able to read a message
-						return true;
-					}
-					else if (chunkState.WriteState == WriteState.Writing)
-					{
-						// Wait until there is more data in the chunk
-						_resources.ResetReadEvent(_readerIdx);
-						if (_headerPtr.GetChunkStatePtr(readerState.ChunkIdx).Value == chunkState)
-						{
-							await _resources.WaitForReadEvent(_readerIdx, cancellationToken);
-						}
-					}
-					else if (readerState.Offset < chunkState.Length || chunkState.WriteState == WriteState.Complete)
+					// Wait until there is more data in the chunk
+					ResetReadEvent(readerIdx);
+					if (_headerPtr.GetChunkStatePtr(readerState.ChunkIdx).Value == chunkState)
 					{
-						// Cannot read the requested amount of data from this chunk.
-						return false;
+						await WaitForReadEvent(readerIdx, cancellationToken);
 					}
-					else if (chunkState.WriteState == WriteState.MovedToNext)
-					{
-						// Move to the next chunk
-						chunkStatePtr.FinishReading(_readerIdx);
-						_resources.SetWriteEvent();
-
-						int chunkIdx = readerStatePtr.ChunkIdx + 1;
-						if (chunkIdx == _chunks.Length)
-						{
-							chunkIdx = 0;
-						}
+				}
+				else if (readerState.Offset < chunkState.Length || chunkState.WriteState == WriteState.Complete)
+				{
+					// Cannot read the requested amount of data from this chunk.
+					return false;
+				}
+				else if (chunkState.WriteState == WriteState.MovedToNext)
+				{
+					// Move to the next chunk
+					chunkStatePtr.FinishReading(readerIdx);
+					SetWriteEvent();
 
-						readerStatePtr.Value = new ReaderState(chunkIdx, 0);
-					}
-					else
+					int chunkIdx = readerStatePtr.ChunkIdx + 1;
+					if (chunkIdx == _chunks.Length)
 					{
-						throw new NotImplementedException($"Invalid write state for buffer: {chunkState.WriteState}");
+						chunkIdx = 0;
 					}
+
+					readerStatePtr.Value = new ReaderState(chunkIdx, 0);
+				}
+				else
+				{
+					throw new NotImplementedException($"Invalid write state for buffer: {chunkState.WriteState}");
 				}
 			}
 		}
 
-		internal class WriterImpl : IComputeBufferWriter
+		/// <inheritdoc/>
+		public void AdvanceWritePosition(int size)
 		{
-			readonly HeaderPtr _headerPtr;
-			readonly Memory<byte>[] _chunks;
-			ResourcesBase _resources;
-
-			public WriterImpl(ResourcesBase resources)
+			if (size > 0)
 			{
-				_headerPtr = resources.HeaderPtr;
-				_chunks = resources.Chunks;
-				_resources = resources;
-			}
+				ChunkStatePtr chunkStatePtr = _headerPtr.GetChunkStatePtr(_headerPtr.WriteChunkIdx);
+				ChunkState chunkState = chunkStatePtr.Value;
 
-			public ResourcesBase GetResources() => _resources;
+				Debug.Assert(chunkState.WriteState == WriteState.Writing);
+				chunkStatePtr.Append(size);
 
-			/// <inheritdoc/>
-			public IComputeBufferWriter AddRef()
-			{
-				_resources.AddRef();
-				return new WriterImpl(_resources);
+				SetAllReadEvents();
 			}
+		}
 
-			/// <inheritdoc/>
-			public void Dispose()
+		/// <inheritdoc/>
+		public Memory<byte> GetWriteBuffer()
+		{
+			ChunkState state = _headerPtr.GetChunkStatePtr(_headerPtr.WriteChunkIdx).Value;
+			if (state.WriteState == WriteState.Writing)
 			{
-				if (_resources != null)
-				{
-					_resources.Release();
-					_resources = null!;
-				}
+				return _chunks[_headerPtr.WriteChunkIdx].Slice(state.Length);
 			}
-
-			/// <inheritdoc/>
-			public void AdvanceWritePosition(int size)
+			else
 			{
-				if (size > 0)
-				{
-					ChunkStatePtr chunkStatePtr = _headerPtr.GetChunkStatePtr(_headerPtr.WriteChunkIdx);
-					ChunkState chunkState = chunkStatePtr.Value;
-
-					Debug.Assert(chunkState.WriteState == WriteState.Writing);
-					chunkStatePtr.Append(size);
-
-					_resources.SetAllReadEvents();
-				}
+				return Memory<byte>.Empty;
 			}
+		}
 
-			/// <inheritdoc/>
-			public Memory<byte> GetWriteBuffer()
+		/// <inheritdoc/>
+		public bool MarkComplete()
+		{
+			ChunkStatePtr chunkStatePtr = _headerPtr.GetChunkStatePtr(_headerPtr.WriteChunkIdx);
+			if (chunkStatePtr.WriteState != WriteState.Complete)
 			{
-				ChunkState state = _headerPtr.GetChunkStatePtr(_headerPtr.WriteChunkIdx).Value;
-				if (state.WriteState == WriteState.Writing)
-				{
-					return _chunks[_headerPtr.WriteChunkIdx].Slice(state.Length);
-				}
-				else
-				{
-					return Memory<byte>.Empty;
-				}
+				chunkStatePtr.MarkComplete();
+				SetAllReadEvents();
+				return true;
 			}
+			return false;
+		}
 
-			/// <inheritdoc/>
-			public bool MarkComplete()
+		/// <inheritdoc/>
+		public async ValueTask WaitToWriteAsync(int minSize, CancellationToken cancellationToken = default)
+		{
+			if (minSize > _headerPtr.ChunkLength)
 			{
-				ChunkStatePtr chunkStatePtr = _headerPtr.GetChunkStatePtr(_headerPtr.WriteChunkIdx);
-				if (chunkStatePtr.WriteState != WriteState.Complete)
-				{
-					chunkStatePtr.MarkComplete();
-					_resources.SetAllReadEvents();
-					return true;
-				}
-				return false;
+				throw new ArgumentException("Requested read size is larger than chunk size.", nameof(minSize));
 			}
 
-			/// <inheritdoc/>
-			public async ValueTask WaitToWriteAsync(int minSize, CancellationToken cancellationToken = default)
+			for (; ; )
 			{
-				if (minSize > _headerPtr.ChunkLength)
-				{
-					throw new ArgumentException("Requested read size is larger than chunk size.", nameof(minSize));
-				}
+				ChunkStatePtr chunkStatePtr = _headerPtr.GetChunkStatePtr(_headerPtr.WriteChunkIdx);
 
-				for (; ; )
+				ChunkState chunkState = chunkStatePtr.Value;
+				if (chunkState.WriteState == WriteState.Writing)
 				{
-					ChunkStatePtr chunkStatePtr = _headerPtr.GetChunkStatePtr(_headerPtr.WriteChunkIdx);
-
-					ChunkState chunkState = chunkStatePtr.Value;
-					if (chunkState.WriteState == WriteState.Writing)
-					{
-						int length = chunkState.Length;
-						if (length + minSize <= _headerPtr.ChunkLength)
-						{
-							return;
-						}
-						chunkStatePtr.FinishWriting();
-					}
-
-					int nextChunkIdx = _headerPtr.WriteChunkIdx + 1;
-					if (nextChunkIdx == _chunks.Length)
+					int length = chunkState.Length;
+					if (length + minSize <= _headerPtr.ChunkLength)
 					{
-						nextChunkIdx = 0;
+						return;
 					}
-
-					ChunkStatePtr nextChunkStatePtr = _headerPtr.GetChunkStatePtr(nextChunkIdx);
-					while (nextChunkStatePtr.ReaderFlags != 0)
-					{
-						await _resources.WaitForWriteEvent(cancellationToken);
-						_resources.ResetWriteEvent();
-					}
-
-					_headerPtr.WriteChunkIdx = nextChunkIdx;
-					nextChunkStatePtr.StartWriting(_headerPtr.NumReaders);
-
-					_resources.SetAllReadEvents();
+					chunkStatePtr.FinishWriting();
 				}
-			}
-		}
-
-		internal ResourcesBase _resources;
-
-		/// <inheritdoc/>
-		public IComputeBufferReader Reader { get; }
 
-		/// <inheritdoc/>
-		public IComputeBufferWriter Writer { get; }
+				int nextChunkIdx = _headerPtr.WriteChunkIdx + 1;
+				if (nextChunkIdx == _chunks.Length)
+				{
+					nextChunkIdx = 0;
+				}
 
-		/// <summary>
-		/// Constructor
-		/// </summary>
-		/// <param name="resources">Resources shared between instances of the buffer</param>
-		internal ComputeBuffer(ResourcesBase resources)
-		{
-			_resources = resources;
-			Reader = new ReaderImpl(resources, 0);
-			Writer = new WriterImpl(resources);
-		}
+				ChunkStatePtr nextChunkStatePtr = _headerPtr.GetChunkStatePtr(nextChunkIdx);
+				while (nextChunkStatePtr.ReaderFlags != 0)
+				{
+					await WaitForWriteEvent(cancellationToken);
+					ResetWriteEvent();
+				}
 
-		/// <inheritdoc/>
-		public void Dispose()
-		{
-			Dispose(true);
-			GC.SuppressFinalize(this);
-		}
+				_headerPtr.WriteChunkIdx = nextChunkIdx;
+				nextChunkStatePtr.StartWriting(_headerPtr.NumReaders);
 
-		/// <summary>
-		/// Overriable dispose method
-		/// </summary>
-		protected virtual void Dispose(bool disposing)
-		{
-			if (_resources != null)
-			{
-				_resources.Release();
-				_resources = null!;
+				SetAllReadEvents();
 			}
 		}
-
-		/// <inheritdoc/>
-		public abstract IComputeBuffer AddRef();
 	}
 }
diff --git a/Engine/Source/Programs/Shared/EpicGames.Horde/Compute/ComputeChannel.cs b/Engine/Source/Programs/Shared/EpicGames.Horde/Compute/ComputeChannel.cs
--- a/Engine/Source/Programs/Shared/EpicGames.Horde/Compute/ComputeChannel.cs
+++ b/Engine/Source/Programs/Shared/EpicGames.Horde/Compute/ComputeChannel.cs
@@ -16,19 +16,19 @@ namespace EpicGames.Horde.Compute
 		/// <summary>
 		/// Reader for the channel
 		/// </summary>
-		public IComputeBufferReader Reader { get; }
+		public ComputeBufferReader Reader { get; }
 
 		/// <summary>
 		/// Writer for the channel
 		/// </summary>
-		public IComputeBufferWriter Writer { get; }
+		public ComputeBufferWriter Writer { get; }
 
 		/// <summary>
 		/// Constructor
 		/// </summary>
 		/// <param name="recvBufferReader"></param>
 		/// <param name="sendBufferWriter"></param>
-		internal ComputeChannel(IComputeBufferReader recvBufferReader, IComputeBufferWriter sendBufferWriter)
+		internal ComputeChannel(ComputeBufferReader recvBufferReader, ComputeBufferWriter sendBufferWriter)
 		{
 			Reader = recvBufferReader.AddRef();
 			Writer = sendBufferWriter.AddRef();
@@ -94,36 +94,4 @@ namespace EpicGames.Horde.Compute
 		/// </summary>
 		public void MarkComplete() => Writer.MarkComplete();
 	}
-
-	/// <summary>
-	/// Opens a channel for compute workers
-	/// </summary>
-	public static class ComputeChannelExtensions
-	{
-		/// <summary>
-		/// Creates a channel using a socket and receive buffer
-		/// </summary>
-		/// <param name="socket">Socket to use for sending data</param>
-		/// <param name="channelId">Channel id to send and receive data</param>
-		public static ComputeChannel CreateChannel(this ComputeSocket socket, int channelId)
-		{
-			using SharedMemoryBuffer recvBuffer = SharedMemoryBuffer.CreateNew(null, 65536);
-			using SharedMemoryBuffer sendBuffer = SharedMemoryBuffer.CreateNew(null, 65536);
-			return CreateChannel(socket, channelId, recvBuffer, sendBuffer);
-		}
-
-		/// <summary>
-		/// Creates a channel using a socket and receive buffer
-		/// </summary>
-		/// <param name="socket">Socket to use for sending data</param>
-		/// <param name="channelId">Channel id to send and receive data</param>
-		/// <param name="recvBuffer">Buffer for receiving data</param>
-		/// <param name="sendBuffer">Buffer for sending data</param>
-		public static ComputeChannel CreateChannel(this ComputeSocket socket, int channelId, IComputeBuffer recvBuffer, IComputeBuffer sendBuffer)
-		{
-			socket.AttachRecvBuffer(channelId, recvBuffer.Writer);
-			socket.AttachSendBuffer(channelId, sendBuffer.Reader);
-			return new ComputeChannel(recvBuffer.Reader, sendBuffer.Writer);
-		}
-	}
 }
diff --git a/Engine/Source/Programs/Shared/EpicGames.Horde/Compute/ComputeSocket.cs b/Engine/Source/Programs/Shared/EpicGames.Horde/Compute/ComputeSocket.cs
--- a/Engine/Source/Programs/Shared/EpicGames.Horde/Compute/ComputeSocket.cs
+++ b/Engine/Source/Programs/Shared/EpicGames.Horde/Compute/ComputeSocket.cs
@@ -23,14 +23,38 @@ namespace EpicGames.Horde.Compute
 		/// </summary>
 		/// <param name="channelId">Channel to receive data on</param>
 		/// <param name="recvBufferWriter">Writer for the buffer to store received data</param>
-		public abstract void AttachRecvBuffer(int channelId, IComputeBufferWriter recvBufferWriter);
+		public abstract void AttachRecvBuffer(int channelId, ComputeBufferWriter recvBufferWriter);
 
 		/// <summary>
 		/// Attaches a buffer to send data.
 		/// </summary>
 		/// <param name="channelId">Channel to receive data on</param>
 		/// <param name="sendBufferReader">Reader for the buffer to send data from</param>
-		public abstract void AttachSendBuffer(int channelId, IComputeBufferReader sendBufferReader);
+		public abstract void AttachSendBuffer(int channelId, ComputeBufferReader sendBufferReader);
+
+		/// <summary>
+		/// Creates a channel using a socket and receive buffer
+		/// </summary>
+		/// <param name="channelId">Channel id to send and receive data</param>
+		public ComputeChannel CreateChannel(int channelId)
+		{
+			using SharedMemoryBuffer recvBuffer = SharedMemoryBuffer.CreateNew(null, 65536);
+			using SharedMemoryBuffer sendBuffer = SharedMemoryBuffer.CreateNew(null, 65536);
+			return CreateChannel(channelId, recvBuffer, sendBuffer);
+		}
+
+		/// <summary>
+		/// Creates a channel using a socket and receive buffer
+		/// </summary>
+		/// <param name="channelId">Channel id to send and receive data</param>
+		/// <param name="recvBuffer">Buffer for receiving data</param>
+		/// <param name="sendBuffer">Buffer for sending data</param>
+		public ComputeChannel CreateChannel(int channelId, ComputeBuffer recvBuffer, ComputeBuffer sendBuffer)
+		{
+			AttachRecvBuffer(channelId, recvBuffer.Writer);
+			AttachSendBuffer(channelId, sendBuffer.Reader);
+			return new ComputeChannel(recvBuffer.Reader, sendBuffer.Writer);
+		}
 	}
 
 	internal enum IpcMessage
@@ -81,16 +105,16 @@ namespace EpicGames.Horde.Compute
 		}
 
 		/// <inheritdoc/>
-		public override void AttachRecvBuffer(int channelId, IComputeBufferWriter writer)
+		public override void AttachRecvBuffer(int channelId, ComputeBufferWriter writer)
 		{
-			string bufferName = SharedMemoryBuffer.GetName(writer);
+			string bufferName = ((SharedMemoryBufferDetail)writer.Detail).Name;
 			AttachBuffer(IpcMessage.AttachRecvBuffer, channelId, bufferName);
 		}
 
 		/// <inheritdoc/>
-		public override void AttachSendBuffer(int channelId, IComputeBufferReader reader)
+		public override void AttachSendBuffer(int channelId, ComputeBufferReader reader)
 		{
-			string bufferName = SharedMemoryBuffer.GetName(reader);
+			string bufferName = ((SharedMemoryBufferDetail)reader.Detail).Name;
 			AttachBuffer(IpcMessage.AttachSendBuffer, channelId, bufferName);
 		}
 
@@ -141,11 +165,11 @@ namespace EpicGames.Horde.Compute
 		readonly CancellationTokenSource _cancellationSource = new CancellationTokenSource();
 
 		BackgroundTask? _recvTask;
-		readonly Dictionary<int, IComputeBufferWriter> _recvBufferWriters = new Dictionary<int, IComputeBufferWriter>();
+		readonly Dictionary<int, ComputeBufferWriter> _recvBufferWriters = new Dictionary<int, ComputeBufferWriter>();
 
 		readonly SemaphoreSlim _sendSemaphore = new SemaphoreSlim(1, 1);
 		readonly Dictionary<int, Task> _sendTasks = new Dictionary<int, Task>();
-		readonly Dictionary<IComputeBufferReader, int> _sendBufferReaders = new Dictionary<IComputeBufferReader, int>();
+		readonly Dictionary<ComputeBufferReader, int> _sendBufferReaders = new Dictionary<ComputeBufferReader, int>();
 
 		string Tag => (_endpoint == ComputeSocketEndpoint.Local)? "LOCAL": "REMOTE";
 
@@ -215,7 +239,7 @@ namespace EpicGames.Horde.Compute
 			try
 			{
 				// Maintain a local cache of buffers to be able to query for them without having to acquire a global lock
-				Dictionary<int, IComputeBufferWriter> cachedWriters = new Dictionary<int, IComputeBufferWriter>();
+				Dictionary<int, ComputeBufferWriter> cachedWriters = new Dictionary<int, ComputeBufferWriter>();
 
 				Memory<byte> last = Memory<byte>.Empty;
 
@@ -236,7 +260,7 @@ namespace EpicGames.Horde.Compute
 					// Dispatch it to the correct place
 					if (size >= 0)
 					{
-						IComputeBufferWriter writer = GetReceiveBuffer(cachedWriters, id);
+						ComputeBufferWriter writer = GetReceiveBuffer(cachedWriters, id);
 						await ReadPacketAsync(transport, id, size, writer, cancellationToken);
 					}
 					else if (size == (int)ControlMessageType.Detach)
@@ -258,7 +282,7 @@ namespace EpicGames.Horde.Compute
 			lock (_lockObject)
 			{
 				_complete = true;
-				foreach (IComputeBufferWriter writer in _recvBufferWriters.Values)
+				foreach (ComputeBufferWriter writer in _recvBufferWriters.Values)
 				{
 					writer.MarkComplete();
 				}
@@ -267,7 +291,7 @@ namespace EpicGames.Horde.Compute
 			_logger.LogTrace("[{Tag}] Closing reader", Tag);
 		}
 
-		async Task ReadPacketAsync(ComputeTransport transport, int id, int size, IComputeBufferWriter writer, CancellationToken cancellationToken)
+		async Task ReadPacketAsync(ComputeTransport transport, int id, int size, ComputeBufferWriter writer, CancellationToken cancellationToken)
 		{
 			Memory<byte> memory = writer.GetWriteBuffer();
 			while (memory.Length < size)
@@ -286,9 +310,9 @@ namespace EpicGames.Horde.Compute
 			writer.AdvanceWritePosition(size);
 		}
 
-		IComputeBufferWriter GetReceiveBuffer(Dictionary<int, IComputeBufferWriter> cachedWriters, int id)
+		ComputeBufferWriter GetReceiveBuffer(Dictionary<int, ComputeBufferWriter> cachedWriters, int id)
 		{
-			IComputeBufferWriter? writer;
+			ComputeBufferWriter? writer;
 			if (cachedWriters.TryGetValue(id, out writer))
 			{
 				return writer;
@@ -296,7 +320,7 @@ namespace EpicGames.Horde.Compute
 
 			lock (_lockObject)
 			{
-				IComputeBufferWriter? recvBufferWriter;
+				ComputeBufferWriter? recvBufferWriter;
 				if (_recvBufferWriters.TryGetValue(id, out recvBufferWriter))
 				{
 					writer = recvBufferWriter;
@@ -354,7 +378,7 @@ namespace EpicGames.Horde.Compute
 		}
 
 		/// <inheritdoc/>
-		public override void AttachRecvBuffer(int channelId, IComputeBufferWriter recvBufferWriter)
+		public override void AttachRecvBuffer(int channelId, ComputeBufferWriter recvBufferWriter)
 		{
 			bool complete;
 			lock (_lockObject)
@@ -373,11 +397,11 @@ namespace EpicGames.Horde.Compute
 			}
 		}
 
-		void DetachRecvBuffer(Dictionary<int, IComputeBufferWriter> cachedWriters, int id)
+		void DetachRecvBuffer(Dictionary<int, ComputeBufferWriter> cachedWriters, int id)
 		{
 			cachedWriters.Remove(id);
 
-			IComputeBufferWriter? recvBufferWriter;
+			ComputeBufferWriter? recvBufferWriter;
 			lock (_lockObject)
 			{
 #pragma warning disable CA2000 // Dispose objects before losing scope
@@ -392,7 +416,7 @@ namespace EpicGames.Horde.Compute
 		}
 
 		/// <inheritdoc/>
-		public override void AttachSendBuffer(int channelId, IComputeBufferReader sendBufferReader)
+		public override void AttachSendBuffer(int channelId, ComputeBufferReader sendBufferReader)
 		{
 			sendBufferReader = sendBufferReader.AddRef();
 			lock (_lockObject)
@@ -402,9 +426,9 @@ namespace EpicGames.Horde.Compute
 			}
 		}
 
-		async Task SendFromBufferAsync(int channelId, IComputeBufferReader reader, CancellationToken cancellationToken)
+		async Task SendFromBufferAsync(int channelId, ComputeBufferReader reader, CancellationToken cancellationToken)
 		{
-			using IComputeBufferReader _ = reader;
+			using ComputeBufferReader _ = reader;
 
 			while (!cancellationToken.IsCancellationRequested)
 			{
diff --git a/Engine/Source/Programs/Shared/EpicGames.Horde/Compute/IComputeBuffer.cs b/Engine/Source/Programs/Shared/EpicGames.Horde/Compute/IComputeBuffer.cs
deleted file mode 100644
