Fixed issue where UBT would fail with a concurrency error when doing log once logging.

diff --git a/Engine/Source/Programs/Shared/EpicGames.Core/Log.cs b/Engine/Source/Programs/Shared/EpicGames.Core/Log.cs
--- a/Engine/Source/Programs/Shared/EpicGames.Core/Log.cs
+++ b/Engine/Source/Programs/Shared/EpicGames.Core/Log.cs
@@ -1,6 +1,7 @@
 // Copyright Epic Games, Inc. All Rights Reserved.
 
 using System;
+using System.Collections.Concurrent;
 using System.Collections.Generic;
 using System.Diagnostics;
 using System.IO;
@@ -184,7 +185,7 @@ namespace EpicGames.Core
 		/// <summary>
 		/// A collection of strings that have been already written once
 		/// </summary>
-		private static readonly HashSet<string> s_writeOnceSet = new HashSet<string>();
+		private static readonly ConcurrentDictionary<string, bool> s_writeOnceSet = new();
 
 		/// <summary>
 		/// Overrides the logger used for formatting output, after event parsing
@@ -410,7 +411,7 @@ namespace EpicGames.Core
 				}
 
 				// if we want this message only written one time, check if it was already written out
-				if (bWriteOnce && !s_writeOnceSet.Add(message.ToString()))
+				if (bWriteOnce && !s_writeOnceSet.TryAdd(message.ToString(), true))
 				{
 					return;
 				}
