Added missing class cast flag

diff --git a/Engine/Source/Programs/Shared/EpicGames.Core/UnrealEngineTypes.cs b/Engine/Source/Programs/Shared/EpicGames.Core/UnrealEngineTypes.cs
--- a/Engine/Source/Programs/Shared/EpicGames.Core/UnrealEngineTypes.cs
+++ b/Engine/Source/Programs/Shared/EpicGames.Core/UnrealEngineTypes.cs
@@ -382,6 +382,7 @@ namespace EpicGames.Core
 		FObjectPtrProperty = 0x0020000000000000,
 		FClassPtrProperty = 0x0040000000000000,
 		FLargeWorldCoordinatesRealProperty = 0x0080000000000000,
+		FOptionalProperty = 0x0100000000000000,
 		AllFlags = UInt64.MaxValue,
 	};
 
