Horde: Sync C# memory layout of compute buffer with the C++ version.

diff --git a/Engine/Source/Programs/Horde/Samples/RemoteClient/RemoteClientApp.cs b/Engine/Source/Programs/Horde/Samples/RemoteClient/RemoteClientApp.cs
--- a/Engine/Source/Programs/Horde/Samples/RemoteClient/RemoteClientApp.cs
+++ b/Engine/Source/Programs/Horde/Samples/RemoteClient/RemoteClientApp.cs
@@ -1,6 +1,8 @@
 // Copyright Epic Games, Inc. All Rights Reserved.
 
+using System.Buffers;
 using System.Buffers.Binary;
+using System.IO.Pipes;
 using System.Net.Http.Headers;
 using System.Reflection;
 using EpicGames.Core;
@@ -8,6 +10,7 @@ using EpicGames.Horde.Common;
 using EpicGames.Horde.Compute;
 using EpicGames.Horde.Compute.Buffers;
 using EpicGames.Horde.Compute.Clients;
+using EpicGames.Horde.Compute.Transports;
 using EpicGames.Horde.Storage;
 using EpicGames.Horde.Storage.Backends;
 using EpicGames.Horde.Storage.Bundles;
@@ -105,10 +108,9 @@ namespace RemoteClient
 				}
 
 				// Run the task remotely in the background and echo the output to the console
+				await using BackgroundTask tickTask = BackgroundTask.StartNew(ctx => WriteNumbersAsync(lease.Socket, logger, ctx));
 				await using (AgentManagedProcess process = await channel.ExecuteAsync(executable, arguments, null, null))
 				{
-					await using BackgroundTask tickTask = BackgroundTask.StartNew(ctx => WriteNumbersAsync(lease.Socket, logger, ctx));
-
 					string? line;
 					while ((line = await process.ReadLineAsync()) != null)
 					{
@@ -124,8 +126,10 @@ namespace RemoteClient
 			// Wait until the remote sends a message indicating that it's ready
 			using (PooledBuffer recvBuffer = new PooledBuffer(1, 20))
 			{
-				socket.AttachRecvBuffer(ChildProcessChannelId, recvBuffer.Writer);
-				await recvBuffer.Reader.WaitToReadAsync(1, cancellationToken);
+				socket.AttachRecvBuffer(ChildProcessChannelId, recvBuffer);
+			
+				using ComputeBufferReader reader = recvBuffer.CreateReader();
+				await reader.WaitToReadAsync(1, cancellationToken);
 			}
 
 			// Write data to the child process channel. The remote server will echo them back to us as it receives them, then exit when the channel is complete/closed.
diff --git a/Engine/Source/Programs/Shared/EpicGames.Horde.Tests/BufferTests.cs b/Engine/Source/Programs/Shared/EpicGames.Horde.Tests/BufferTests.cs
--- a/Engine/Source/Programs/Shared/EpicGames.Horde.Tests/BufferTests.cs
+++ b/Engine/Source/Programs/Shared/EpicGames.Horde.Tests/BufferTests.cs
@@ -25,15 +25,18 @@ namespace EpicGames.Horde.Tests
 		{
 			using PooledBuffer buffer = new PooledBuffer(2, 1024);
 			buffer.Writer.AdvanceWritePosition(10);
-			await buffer.Reader.WaitToReadAsync(9);
-			buffer.Reader.AdvanceReadPosition(9);
-			await buffer.Reader.WaitToReadAsync(1);
+
+			using ComputeBufferReader bufferReader = buffer.CreateReader();
+			await bufferReader.WaitToReadAsync(9);
+			bufferReader.AdvanceReadPosition(9);
+			await bufferReader.WaitToReadAsync(1);
 		}
 
 		[TestMethod]
 		public void TestOverflow()
 		{
 			using PooledBuffer buffer = new PooledBuffer(2, 20);
+			using ComputeBufferReader bufferReader = buffer.CreateReader();
 
 			// Fill up the first chunk
 			Assert.AreEqual(20, buffer.Writer.GetWriteBuffer().Length);
@@ -56,22 +59,22 @@ namespace EpicGames.Horde.Tests
 			Assert.IsFalse(waitToWriteTask.IsCompleted);
 
 			// Wait for data to be read
-			Assert.AreEqual(20, buffer.Reader.GetReadBuffer().Length);
-			buffer.Reader.AdvanceReadPosition(10);
+			Assert.AreEqual(20, bufferReader.GetReadBuffer().Length);
+			bufferReader.AdvanceReadPosition(10);
 			Assert.IsFalse(waitToWriteTask.IsCompleted);
 
-			Assert.AreEqual(10, buffer.Reader.GetReadBuffer().Length);
-			buffer.Reader.AdvanceReadPosition(10);
-			Assert.AreEqual(0, buffer.Reader.GetReadBuffer().Length);
+			Assert.AreEqual(10, bufferReader.GetReadBuffer().Length);
+			bufferReader.AdvanceReadPosition(10);
+			Assert.AreEqual(0, bufferReader.GetReadBuffer().Length);
 
-			Task waitToReadTask = buffer.Reader.WaitToReadAsync(1).AsTask();
+			Task waitToReadTask = bufferReader.WaitToReadAsync(1).AsTask();
 			Assert.IsTrue(waitToReadTask.IsCompleted);
 
-			Assert.AreEqual(20, buffer.Reader.GetReadBuffer().Length);
+			Assert.AreEqual(20, bufferReader.GetReadBuffer().Length);
 			Assert.IsTrue(waitToWriteTask.IsCompleted);
 
 			// Make sure both reader and writer have something to work with
-			Assert.AreEqual(20, buffer.Reader.GetReadBuffer().Length);
+			Assert.AreEqual(20, bufferReader.GetReadBuffer().Length);
 			Assert.AreEqual(20, buffer.Writer.GetWriteBuffer().Length);
 		}
 
@@ -101,13 +104,15 @@ namespace EpicGames.Horde.Tests
 			await using RemoteComputeSocket consumerSocket = new RemoteComputeSocket(new PipeTransport(sourceToTargetPipe.Reader, targetToSourcePipe.Writer), ComputeSocketEndpoint.Remote, NullLogger.Instance);
 
 			using ComputeBuffer consumerBuffer = createBuffer(Length);
-			consumerSocket.AttachRecvBuffer(ChannelId, consumerBuffer.Writer);
+			consumerSocket.AttachRecvBuffer(ChannelId, consumerBuffer);
 
 			byte[] input = RandomNumberGenerator.GetBytes(Length);
 			Task producerTask = RunProducerAsync(producerSocket, input);
 
+			using ComputeBufferReader consumerBufferReader = consumerBuffer.CreateReader();
+
 			byte[] output = new byte[Length];
-			await RunConsumerAsync(consumerBuffer.Reader, output);
+			await RunConsumerAsync(consumerBufferReader, output);
 
 			await producerTask;
 			Assert.IsTrue(input.SequenceEqual(output));
@@ -144,5 +149,35 @@ namespace EpicGames.Horde.Tests
 				offset += length;
 			}
 		}
+
+		[TestMethod]
+		public async Task TestSendBufferComplete()
+		{
+			Pipe recvPipe = new Pipe();
+			Pipe sendPipe = new Pipe();
+			await using RemoteComputeSocket localSocket = new RemoteComputeSocket(new PipeTransport(sendPipe.Reader, recvPipe.Writer), ComputeSocketEndpoint.Local, NullLogger.Instance);
+			await using RemoteComputeSocket remoteSocket = new RemoteComputeSocket(new PipeTransport(recvPipe.Reader, sendPipe.Writer), ComputeSocketEndpoint.Remote, NullLogger.Instance);
+
+			using (PooledBuffer remoteBuffer = new PooledBuffer(1024))
+			{
+				remoteSocket.AttachRecvBuffer(1, remoteBuffer);
+
+				using ComputeBufferReader reader = remoteBuffer.CreateReader();
+
+				// Disposing of the buffer should mark the channel as complete
+				using (PooledBuffer localBuffer = new PooledBuffer(1024))
+				{
+					localSocket.AttachSendBuffer(1, localBuffer);
+					await localBuffer.Writer.WriteAsync(new byte[] { 1, 2, 3 });
+				}
+
+				Assert.IsTrue(await reader.WaitToReadAsync(3));
+				Assert.IsTrue(reader.GetReadBuffer().Slice(0, 3).Span.SequenceEqual(new byte[] { 1, 2, 3 }));
+				reader.AdvanceReadPosition(3);
+
+				Assert.IsFalse(await reader.WaitToReadAsync(1));
+				Assert.IsTrue(reader.IsComplete);
+			}
+		}
 	}
 }
diff --git a/Engine/Source/Programs/Shared/EpicGames.Horde/Compute/AgentMessageChannel.cs b/Engine/Source/Programs/Shared/EpicGames.Horde/Compute/AgentMessageChannel.cs
--- a/Engine/Source/Programs/Shared/EpicGames.Horde/Compute/AgentMessageChannel.cs
+++ b/Engine/Source/Programs/Shared/EpicGames.Horde/Compute/AgentMessageChannel.cs
@@ -110,9 +110,9 @@ namespace EpicGames.Horde.Compute
 		public AgentMessageChannel(ComputeSocket socket, int channelId, ComputeBuffer recvBuffer, ComputeBuffer sendBuffer, ILogger logger)
 		{
 			_socket = socket;
-			_socket.AttachRecvBuffer(channelId, recvBuffer.Writer);
-			_socket.AttachSendBuffer(channelId, sendBuffer.Reader);
-			_recvBufferReader = recvBuffer.Reader.AddRef();
+			_socket.AttachRecvBuffer(channelId, recvBuffer);
+			_socket.AttachSendBuffer(channelId, sendBuffer);
+			_recvBufferReader = recvBuffer.CreateReader();
 			_sendBufferWriter = sendBuffer.Writer.AddRef();
 			_logger = logger;
 		}
@@ -124,9 +124,7 @@ namespace EpicGames.Horde.Compute
 		{
 			_currentBuilder?.Dispose();
 
-			_sendBufferWriter.MarkComplete();
 			_sendBufferWriter.Dispose();
-
 			_recvBufferReader.Dispose();
 		}
 
diff --git a/Engine/Source/Programs/Shared/EpicGames.Horde/Compute/AgentMessageHandler.cs b/Engine/Source/Programs/Shared/EpicGames.Horde/Compute/AgentMessageHandler.cs
--- a/Engine/Source/Programs/Shared/EpicGames.Horde/Compute/AgentMessageHandler.cs
+++ b/Engine/Source/Programs/Shared/EpicGames.Horde/Compute/AgentMessageHandler.cs
@@ -13,6 +13,11 @@ using EpicGames.Horde.Storage.Nodes;
 using Microsoft.Extensions.Caching.Memory;
 using Microsoft.Extensions.Logging;
 using EpicGames.Horde.Storage.Bundles;
+using System.Text;
+using System.IO;
+using System.Reflection;
+using System.Linq;
+using System.Buffers.Binary;
 
 namespace EpicGames.Horde.Compute
 {
@@ -25,6 +30,8 @@ namespace EpicGames.Horde.Compute
 		readonly IMemoryCache _memoryCache;
 		readonly ILogger _logger;
 
+		readonly bool _executeLocally = false;
+
 		/// <summary>
 		/// Constructor
 		/// </summary>
@@ -186,13 +193,13 @@ namespace EpicGames.Horde.Compute
 			{
 				newEnvVars[WorkerComputeSocket.IpcEnvVar] = ipcBuffer.Name;
 
-				using (BackgroundTask backgroundTask = BackgroundTask.StartNew(ctx => ProcessIpcMessagesAsync(socket, ipcBuffer.Reader, cancellationToken)))
-				{
-					_logger.LogInformation("Launching {Executable} {Arguments}", CommandLineArguments.Quote(executable), CommandLineArguments.Join(arguments));
-					await ExecuteProcessInternalAsync(channel, executable, arguments, workingDir, newEnvVars, cancellationToken);
-					_logger.LogInformation("Finished executing process");
-					ipcBuffer.Writer.MarkComplete();
-				}
+				using ComputeBufferReader ipcBufferReader = ipcBuffer.CreateReader();
+				using BackgroundTask backgroundTask = BackgroundTask.StartNew(ctx => ProcessIpcMessagesAsync(socket, ipcBufferReader, cancellationToken));
+
+				_logger.LogInformation("Launching {Executable} {Arguments}", CommandLineArguments.Quote(executable), CommandLineArguments.Join(arguments));
+				await ExecuteProcessInternalAsync(channel, executable, arguments, workingDir, newEnvVars, cancellationToken);
+				_logger.LogInformation("Finished executing process");
+				ipcBuffer.Writer.MarkComplete();
 			}
 
 			_logger.LogInformation("Child process has shut down");
@@ -200,7 +207,7 @@ namespace EpicGames.Horde.Compute
 
 		async Task ProcessIpcMessagesAsync(ComputeSocket socket, ComputeBufferReader ipcReader, CancellationToken cancellationToken)
 		{
-			List<(IpcMessage, int, SharedMemoryBuffer)> buffers = new();
+			List<SharedMemoryBuffer> buffers = new();
 			try
 			{
 				List<(int, ComputeBufferWriter)> writers = new List<(int, ComputeBufferWriter)>();
@@ -221,9 +228,9 @@ namespace EpicGames.Horde.Compute
 									_logger.LogDebug("Attaching send buffer for channel {ChannelId} to {Name}", channelId, name);
 
 									SharedMemoryBuffer buffer = SharedMemoryBuffer.OpenExisting(name);
-									buffers.Add((message, channelId, buffer));
+									buffers.Add(buffer);
 
-									socket.AttachSendBuffer(channelId, buffer.Reader);
+									socket.AttachSendBuffer(channelId, buffer);
 								}
 								break;
 							case IpcMessage.AttachRecvBuffer:
@@ -233,9 +240,9 @@ namespace EpicGames.Horde.Compute
 									_logger.LogDebug("Attaching recv buffer for channel {ChannelId} to {Name}", channelId, name);
 
 									SharedMemoryBuffer buffer = SharedMemoryBuffer.OpenExisting(name);
-									buffers.Add((message, channelId, buffer));
+									buffers.Add(buffer);
 
-									socket.AttachRecvBuffer(channelId, buffer.Writer);
+									socket.AttachRecvBuffer(channelId, buffer);
 								}
 								break;
 							default:
@@ -252,12 +259,8 @@ namespace EpicGames.Horde.Compute
 			}
 			finally
 			{
-				foreach ((IpcMessage message, int channelId, SharedMemoryBuffer buffer) in buffers)
+				foreach (SharedMemoryBuffer buffer in buffers)
 				{
-					if (buffer.Writer.MarkComplete())
-					{
-						_logger.LogWarning("Buffer added via {Message} on channel {ChannelId} was not marked complete", message, channelId);
-					}
 					buffer.Dispose();
 				}
 			}
@@ -265,40 +268,80 @@ namespace EpicGames.Horde.Compute
 
 		async Task ExecuteProcessInternalAsync(AgentMessageChannel channel, string executable, IReadOnlyList<string> arguments, string? workingDir, IReadOnlyDictionary<string, string?>? envVars, CancellationToken cancellationToken)
 		{
-			string resolvedExecutable = FileReference.Combine(_sandboxDir, executable).FullName;
-			string resolvedCommandLine = CommandLineArguments.Join(arguments);
 			string resolvedWorkingDir = DirectoryReference.Combine(_sandboxDir, workingDir ?? String.Empty).FullName;
-
-			Dictionary<string, string> resolvedEnvVars = ManagedProcess.GetCurrentEnvVars();
-			if (envVars != null)
+			if (_executeLocally)
 			{
-				foreach ((string key, string? value) in envVars)
+				List<(string, string?)> prevEnvVars = new List<(string, string?)>();
+				if (envVars != null)
 				{
-					if (value == null)
+					foreach ((string key, string? value) in envVars)
 					{
-						resolvedEnvVars.Remove(key);
+						prevEnvVars.Add((key, Environment.GetEnvironmentVariable(key)));
+						Environment.SetEnvironmentVariable(key, value);
 					}
-					else
+				}
+
+				string prevWorkingDir = Directory.GetCurrentDirectory();
+				Directory.SetCurrentDirectory(resolvedWorkingDir);
+
+				try
+				{
+					string resolvedExecutable = FileReference.Combine(_sandboxDir, arguments[0]).FullName;
+					string[] mainArgs = arguments.Skip(1).ToArray();
+
+					TaskCompletionSource<int> resultTcs = new TaskCompletionSource<int>();
+
+					Thread thread = new Thread(() => resultTcs.SetResult(AppDomain.CurrentDomain.ExecuteAssembly(resolvedExecutable, mainArgs)));
+					thread.Start();
+
+					int result = await resultTcs.Task;
+					await channel.SendExecuteResultAsync(result, cancellationToken);
+				}
+				finally
+				{
+					Directory.SetCurrentDirectory(prevWorkingDir);
+					foreach((string key, string? value) in prevEnvVars)
 					{
-						resolvedEnvVars[key] = value;
+						Environment.SetEnvironmentVariable(key, value);
 					}
 				}
 			}
-
-			using (ManagedProcessGroup group = new ManagedProcessGroup())
+			else
 			{
-				using (ManagedProcess process = new ManagedProcess(group, resolvedExecutable, resolvedCommandLine, resolvedWorkingDir, resolvedEnvVars, null, ProcessPriorityClass.Normal))
+				string resolvedExecutable = FileReference.Combine(_sandboxDir, executable).FullName;
+				string resolvedCommandLine = CommandLineArguments.Join(arguments);
+
+				Dictionary<string, string> resolvedEnvVars = ManagedProcess.GetCurrentEnvVars();
+				if (envVars != null)
 				{
-					byte[] buffer = new byte[1024];
-					for (; ; )
+					foreach ((string key, string? value) in envVars)
 					{
-						int length = await process.ReadAsync(buffer, 0, buffer.Length, cancellationToken);
-						if (length == 0)
+						if (value == null)
 						{
-							await channel.SendExecuteResultAsync(process.ExitCode, cancellationToken);
-							return;
+							resolvedEnvVars.Remove(key);
+						}
+						else
+						{
+							resolvedEnvVars[key] = value;
+						}
+					}
+				}
+
+				using (ManagedProcessGroup group = new ManagedProcessGroup())
+				{
+					using (ManagedProcess process = new ManagedProcess(group, resolvedExecutable, resolvedCommandLine, resolvedWorkingDir, resolvedEnvVars, null, ProcessPriorityClass.Normal))
+					{
+						byte[] buffer = new byte[1024];
+						for (; ; )
+						{
+							int length = await process.ReadAsync(buffer, 0, buffer.Length, cancellationToken);
+							if (length == 0)
+							{
+								await channel.SendExecuteResultAsync(process.ExitCode, cancellationToken);
+								return;
+							}
+							await channel.SendExecuteOutputAsync(buffer.AsMemory(0, length), cancellationToken);
 						}
-						await channel.SendExecuteOutputAsync(buffer.AsMemory(0, length), cancellationToken);
 					}
 				}
 			}
diff --git a/Engine/Source/Programs/Shared/EpicGames.Horde/Compute/Buffers/PooledBuffer.cs b/Engine/Source/Programs/Shared/EpicGames.Horde/Compute/Buffers/PooledBuffer.cs
--- a/Engine/Source/Programs/Shared/EpicGames.Horde/Compute/Buffers/PooledBuffer.cs
+++ b/Engine/Source/Programs/Shared/EpicGames.Horde/Compute/Buffers/PooledBuffer.cs
@@ -42,7 +42,7 @@ namespace EpicGames.Horde.Compute.Buffers
 		/// <inheritdoc/>
 		public override ComputeBuffer AddRef()
 		{
-			_detail.AddRef();
+			_detail.AddWriterRef();
 			return new PooledBuffer(_detail);
 		}
 	}
diff --git a/Engine/Source/Programs/Shared/EpicGames.Horde/Compute/Buffers/SharedMemoryBuffer.cs b/Engine/Source/Programs/Shared/EpicGames.Horde/Compute/Buffers/SharedMemoryBuffer.cs
--- a/Engine/Source/Programs/Shared/EpicGames.Horde/Compute/Buffers/SharedMemoryBuffer.cs
+++ b/Engine/Source/Programs/Shared/EpicGames.Horde/Compute/Buffers/SharedMemoryBuffer.cs
@@ -29,7 +29,7 @@ namespace EpicGames.Horde.Compute.Buffers
 		/// <inheritdoc/>
 		public override ComputeBuffer AddRef()
 		{
-			_detail.AddRef();
+			_detail.AddWriterRef();
 			return new SharedMemoryBuffer(_detail);
 		}
 
diff --git a/Engine/Source/Programs/Shared/EpicGames.Horde/Compute/ComputeBuffer.cs b/Engine/Source/Programs/Shared/EpicGames.Horde/Compute/ComputeBuffer.cs
--- a/Engine/Source/Programs/Shared/EpicGames.Horde/Compute/ComputeBuffer.cs
+++ b/Engine/Source/Programs/Shared/EpicGames.Horde/Compute/ComputeBuffer.cs
@@ -24,16 +24,12 @@ namespace EpicGames.Horde.Compute
 		public const int MaxReaders = 16;
 
 		internal ComputeBufferDetail _detail;
-
-		/// <summary>
-		/// Reader for this buffer
-		/// </summary>
-		public ComputeBufferReader Reader { get; }
+		private ComputeBufferWriter _writer;
 
 		/// <summary>
 		/// Writer for this buffer
 		/// </summary>
-		public ComputeBufferWriter Writer { get; }
+		public ComputeBufferWriter Writer => _writer;
 
 		/// <summary>
 		/// Constructor
@@ -42,8 +38,7 @@ namespace EpicGames.Horde.Compute
 		internal ComputeBuffer(ComputeBufferDetail detail)
 		{
 			_detail = detail;
-			Reader = new ComputeBufferReader(detail, 0);
-			Writer = new ComputeBufferWriter(detail);
+			_writer = new ComputeBufferWriter(detail);
 		}
 
 		/// <inheritdoc/>
@@ -54,17 +49,28 @@ namespace EpicGames.Horde.Compute
 		}
 
 		/// <summary>
-		/// Overriable dispose method
+		/// Overridable dispose method
 		/// </summary>
 		protected virtual void Dispose(bool disposing)
 		{
 			if (_detail != null)
 			{
-				_detail.Release();
+				_detail.ReleaseWriter();
 				_detail = null!;
 			}
 		}
 
+		/// <summary>
+		/// Creates a new reader for this buffer
+		/// </summary>
+		/// <returns></returns>
+		public ComputeBufferReader CreateReader()
+		{
+			int readerIdx = _detail.AllocateReader();
+			_detail.AddRef();
+			return new ComputeBufferReader(_detail, new ComputeBufferDetail.ReaderState(readerIdx));
+		}
+
 		/// <summary>
 		/// Creates a new reference to the underlying buffer. The underlying resources will only be destroyed once all instances are disposed of.
 		/// </summary>
@@ -76,15 +82,15 @@ namespace EpicGames.Horde.Compute
 	/// </summary>
 	public sealed class ComputeBufferReader : IDisposable
 	{
-		ComputeBufferDetail _detail;
-		readonly int _readerIdx;
+		ComputeBufferDetail _buffer;
+		readonly ComputeBufferDetail.ReaderState _readerState;
 
-		internal ComputeBufferDetail Detail => _detail;
+		internal ComputeBufferDetail Detail => _buffer;
 
-		internal ComputeBufferReader(ComputeBufferDetail detail, int readerIdx)
+		internal ComputeBufferReader(ComputeBufferDetail buffer, ComputeBufferDetail.ReaderState readerState)
 		{
-			_detail = detail;
-			_readerIdx = readerIdx;
+			_buffer = buffer;
+			_readerState = readerState;
 		}
 
 		/// <summary>
@@ -92,36 +98,41 @@ namespace EpicGames.Horde.Compute
 		/// </summary>
 		public ComputeBufferReader AddRef()
 		{
-			_detail.AddRef();
-			return new ComputeBufferReader(_detail, _readerIdx);
+			_buffer.AddRef();
+			return new ComputeBufferReader(_buffer, _readerState);
 		}
 
 		/// <inheritdoc/>
 		public void Dispose()
 		{
-			if (_detail != null)
+			if (_buffer != null)
 			{
-				_detail.Release();
-				_detail = null!;
+				_buffer.Release();
+				_buffer = null!;
 			}
 		}
 
+		/// <summary>
+		/// Detaches this reader from the underlying buffer
+		/// </summary>
+		public void Detach() => _buffer.DetachReader(_readerState);
+
 		/// <summary>
 		/// Whether this buffer is complete (no more data will be added)
 		/// </summary>
-		public bool IsComplete => _detail.IsReadingComplete(_readerIdx);
+		public bool IsComplete => _buffer.IsComplete(_readerState);
 
 		/// <summary>
 		/// Updates the read position
 		/// </summary>
 		/// <param name="length">Size of data that was read</param>
-		public void AdvanceReadPosition(int length) => _detail.AdvanceReadPosition(_readerIdx, length);
+		public void AdvanceReadPosition(int length) => _readerState.Offset += length;
 
 		/// <summary>
 		/// Gets the next data to read
 		/// </summary>
 		/// <returns>Memory to read from</returns>
-		public ReadOnlyMemory<byte> GetReadBuffer() => _detail.GetReadBuffer(_readerIdx);
+		public ReadOnlyMemory<byte> GetReadBuffer() => _buffer.GetReadBuffer(_readerState);
 
 		/// <summary>
 		/// Read from a buffer into another buffer
@@ -151,7 +162,7 @@ namespace EpicGames.Horde.Compute
 		/// <param name="minLength">Minimum amount of data to read</param>
 		/// <param name="cancellationToken">Cancellation token for the operation</param>
 		/// <returns>True if new data is available, false if the buffer is complete</returns>
-		public ValueTask<bool> WaitToReadAsync(int minLength, CancellationToken cancellationToken = default) => _detail.WaitToReadAsync(_readerIdx, minLength, cancellationToken);
+		public ValueTask<bool> WaitToReadAsync(int minLength, CancellationToken cancellationToken = default) => _buffer.WaitToReadAsync(_readerState, minLength, cancellationToken);
 	}
 
 	/// <summary>
@@ -163,17 +174,14 @@ namespace EpicGames.Horde.Compute
 
 		internal ComputeBufferDetail Detail => _detail;
 
-		internal ComputeBufferWriter(ComputeBufferDetail detail)
-		{
-			_detail = detail;
-		}
+		internal ComputeBufferWriter(ComputeBufferDetail detail) => _detail = detail;
 
 		/// <summary>
 		/// Create a new writer instance using the same underlying buffer
 		/// </summary>
 		public ComputeBufferWriter AddRef()
 		{
-			_detail.AddRef();
+			_detail.AddWriterRef();
 			return new ComputeBufferWriter(_detail);
 		}
 
@@ -182,7 +190,7 @@ namespace EpicGames.Horde.Compute
 		{
 			if (_detail != null)
 			{
-				_detail.Release();
+				_detail.ReleaseWriter();
 				_detail = null!;
 			}
 		}
@@ -236,40 +244,7 @@ namespace EpicGames.Horde.Compute
 	/// </summary>
 	internal abstract class ComputeBufferDetail : IDisposable
 	{
-		internal const int HeaderSize = (2 + ComputeBuffer.MaxChunks + ComputeBuffer.MaxReaders) * sizeof(ulong);
-
-		/// <summary>
-		/// Tracked state of the buffer
-		/// </summary>
-		protected internal readonly unsafe struct HeaderPtr
-		{
-			readonly ulong* _data;
-
-			public HeaderPtr(ulong* data) => _data = data;
-
-			public HeaderPtr(ulong* data, int numReaders, int numChunks, int chunkLength)
-			{
-				_data = data;
-
-				data[0] = ((ulong)(uint)numChunks << 32) | (uint)numReaders;
-				data[1] = (uint)chunkLength;
-
-				GetChunkStatePtr(0).StartWriting(numReaders);
-			}
-
-			public int NumReaders => (int)_data[0];
-			public int NumChunks => (int)(_data[0] >> 32);
-			public int ChunkLength => (int)_data[1];
-			public int WriteChunkIdx
-			{
-				get => (int)(_data[1] >> 32);
-				set => _data[1] = ((ulong)value << 32) | (uint)ChunkLength;
-			}
-
-			public ChunkStatePtr GetChunkStatePtr(int chunkIdx) => new ChunkStatePtr(_data + 2 + chunkIdx);
-
-			public ReaderStatePtr GetReaderStatePtr(int readerIdx) => new ReaderStatePtr(_data + 2 + ComputeBuffer.MaxChunks + readerIdx);
-		}
+		internal const int HeaderSize = (3 + ComputeBuffer.MaxChunks) * sizeof(ulong);
 
 		/// <summary>
 		/// Write state for a chunk
@@ -326,75 +301,155 @@ namespace EpicGames.Horde.Compute
 			public ChunkStatePtr(ulong* data) => _data = data;
 
 			// Current value of the chunk state
-			public ChunkState Value => new ChunkState(Interlocked.CompareExchange(ref *_data, 0, 0));
-
-			// Written length of this chunk
-			public int Length => Value.Length;
+			public ChunkState Value
+			{
+				get => new ChunkState(Interlocked.CompareExchange(ref *_data, 0, 0));
+				set => Interlocked.Exchange(ref *_data, value.Value);
+			}
 
-			// Set of flags which are set for each reader that still has to read from a chunk
-			public int ReaderFlags => Value.ReaderFlags;
-
-			// State of the writer
-			public WriteState WriteState => Value.WriteState;
+			// Attempt to update the chunk state
+			public bool TryUpdate(ChunkState prevState, ChunkState nextState) => Interlocked.CompareExchange(ref *_data, nextState.Value, prevState.Value) == prevState.Value;
 
 			// Append data to the chunk
-			public void Append(int length) => Interlocked.Add(ref *_data, (uint)length);
-
-			// Mark the chunk as being written to
-			public void StartWriting(int numReaders) => Interlocked.Exchange(ref *_data, new ChunkState(WriteState.Writing, (1 << numReaders) - 1, 0).Value);
+			public void Append(int length) => Interlocked.Add(ref *_data, (ulong)length);
 
 			// Move to the next chunk
-			public void MarkComplete() => Interlocked.Or(ref *_data, (ulong)WriteState.Complete << 62);
+			public void MarkComplete() => Interlocked.Or(ref *_data, new ChunkState(WriteState.Complete, 0, 0).Value);
+
+			// Start reading the chunk with the given reader
+			public void StartReading(int readerIdx) => Interlocked.Or(ref *_data, new ChunkState(0, 1 << readerIdx, 0).Value);
 
 			// Clear the reader flag
-			public void FinishReading(int readerIdx) => Interlocked.And(ref *_data, ~(1UL << (readerIdx + 31)));
+			public void FinishReading(int readerIdx) => Interlocked.And(ref *_data, ~new ChunkState(0, 1 << readerIdx, 0).Value);
 
 			// Move to the next chunk
-			public void FinishWriting() => Interlocked.And(ref *_data, ~((ulong)WriteState.Writing << 62));
+			public void FinishWriting() => Interlocked.And(ref *_data, ~new ChunkState(WriteState.Writing, 0, 0).Value);
 
 			/// <inheritdoc/>
 			public override string ToString() => Value.ToString();
 		}
 
 		/// <summary>
-		/// Encodes the state of a reader in a 64-bit value
+		/// State of a reader
 		/// </summary>
-		protected internal record struct ReaderState(ulong Value)
+		internal class ReaderState
 		{
-			public ReaderState(int chunkIdx, int offset) : this(((ulong)chunkIdx << 32) | (uint)offset) { }
+			public int ReaderIdx { get; }
 
-			public int ChunkIdx => (int)(Value >> 32);
-			public int Offset => (int)Value;
+			public int Offset { get; set; }
+			public int ChunkIdx { get; set; }
+			public bool Detached { get; set; }
 
-			public override string ToString() => $"Chunk {ChunkIdx}, Offset {Offset}";
+			public ReaderState(int readerIdx) => ReaderIdx = readerIdx;
 		}
 
 		/// <summary>
-		/// Pointer to a reader state value
+		/// State of the writer
 		/// </summary>
-		protected internal readonly unsafe struct ReaderStatePtr
+		public record struct WriterState(ulong Value)
+		{
+			public WriterState(int chunkIdx, int readerFlags, bool hasWrapped)
+				: this((ulong)(uint)chunkIdx | ((ulong)(uint)readerFlags << 32) | (hasWrapped ? (1UL << 63) : 0))
+			{ }
+
+			public int ReaderFlags => (int)(Value >> 32) & 0xffff;
+			public int ChunkIdx => (int)(Value & 0x7fffffff);
+			public bool HasWrapped => (Value & (1UL << 63)) != 0;
+		}
+
+		/// <summary>
+		/// Wraps a pointer to the state of a writer
+		/// </summary>
+		protected internal readonly unsafe struct WriterStatePtr
 		{
 			readonly ulong* _data;
 
-			public ReaderStatePtr(ulong* data) => _data = data;
+			public WriterStatePtr(ulong* data) => _data = data;
 
-			public ReaderState Value
+			// Current value of the chunk state
+			public WriterState Value
 			{
-				get => new ReaderState(Interlocked.CompareExchange(ref *_data, 0, 0));
+				get => new WriterState(Interlocked.CompareExchange(ref *_data, 0, 0));
 				set => Interlocked.Exchange(ref *_data, value.Value);
 			}
 
-			public int ChunkIdx => Value.ChunkIdx;
-			public int Offset => Value.Offset;
+			// Compare and swap
+			public bool TryUpdate(WriterState prevState, WriterState nextState) => Interlocked.CompareExchange(ref *_data, nextState.Value, prevState.Value) == prevState.Value;
+		}
 
-			public void Advance(int length) => Interlocked.Add(ref *_data, (ulong)length);
+		/// <summary>
+		/// Tracked state of the buffer
+		/// </summary>
+		protected internal readonly unsafe struct HeaderPtr
+		{
+			readonly ulong* _data;
 
-			public override string ToString() => Value.ToString();
+			public HeaderPtr(ulong* data) => _data = data;
+
+			public HeaderPtr(ulong* data, int numReaders, int numChunks, int chunkLength)
+			{
+				_data = data;
+
+				data[0] = ((ulong)(uint)numChunks << 32) | (uint)numReaders;
+				data[1] = (uint)chunkLength;
+
+				GetChunkStatePtr(0).Value = new ChunkState(WriteState.Writing, 0, 0);
+			}
+
+			public int NumReaders => (int)_data[0];
+			public int NumChunks => (int)(_data[0] >> 32);
+			public int ChunkLength => (int)_data[1];
+			public uint AllocatedReaders => (uint)(_data[1] >> 32);
+
+			public WriterStatePtr Writer => new WriterStatePtr(_data + 2);
+
+			public ChunkStatePtr GetChunkStatePtr(int chunkIdx) => new ChunkStatePtr(_data + 3 + chunkIdx);
+
+			public int AllocateReader()
+			{
+				for (int readerIdx = 0; readerIdx < NumReaders; readerIdx++)
+				{
+					ulong allocatedReaders = _data[1];
+					ulong readerFlag = 1UL << (32 + readerIdx);
+
+					if((allocatedReaders & readerFlag) == 0 && Interlocked.CompareExchange(ref _data[1], allocatedReaders | readerFlag, allocatedReaders) == allocatedReaders)
+					{
+						for (; ; )
+						{
+							WriterState state = Writer.Value;
+							if (state.HasWrapped)
+							{
+								throw new InvalidOperationException("Cannot create a new reader after writer has wrapped back to the first chunk");
+							}
+
+							if (Writer.TryUpdate(state, new WriterState(state.ChunkIdx, state.ReaderFlags | (1 << readerIdx), state.HasWrapped)))
+							{
+								for (int WriteChunkIdx = 0; WriteChunkIdx <= state.ChunkIdx; WriteChunkIdx++)
+								{
+									GetChunkStatePtr(WriteChunkIdx).StartReading(readerIdx);
+								}
+								return readerIdx;
+							}
+						}
+					}
+				}
+				throw new InvalidOperationException("Unable to allocate reader; all available readers are in use.");
+			}
+
+			void ReleaseReader(int readerIdx)
+			{
+				for (int Idx = 0; Idx < NumChunks; Idx++)
+				{
+					GetChunkStatePtr(Idx).FinishReading(readerIdx);
+				}
+				Interlocked.And(ref _data[1], ~(1UL << (readerIdx + 32)));
+			}
 		}
 
 		readonly HeaderPtr _headerPtr;
 		readonly Memory<byte>[] _chunks;
 		int _refCount = 1;
+		int _writerRefCount = 1;
 
 		/// <summary>
 		/// Constructor
@@ -424,6 +479,27 @@ namespace EpicGames.Horde.Compute
 			}
 		}
 
+		/// <summary>
+		/// Increments the writer ref count
+		/// </summary>
+		internal void AddWriterRef()
+		{
+			AddRef();
+			Interlocked.Increment(ref _writerRefCount);
+		}
+
+		/// <summary>
+		/// Decrement the reference count on this object, and dispose of it once it reaches zero
+		/// </summary>
+		public void ReleaseWriter()
+		{
+			if (Interlocked.Decrement(ref _writerRefCount) == 0)
+			{
+				MarkComplete();
+			}
+			Release();
+		}
+
 		/// <inheritdoc/>
 		public void Dispose()
 		{
@@ -480,34 +556,34 @@ namespace EpicGames.Horde.Compute
 #pragma warning disable IDE0051 // Remove unused private members
 		// For debugging purposes only
 		ChunkState[] ChunkStates => Enumerable.Range(0, _headerPtr.NumChunks).Select(x => _headerPtr.GetChunkStatePtr(x).Value).ToArray();
-		ReaderState[] ReaderStates => Enumerable.Range(0, _headerPtr.NumReaders).Select(x => _headerPtr.GetReaderStatePtr(x).Value).ToArray();
-		int WriteChunkIdx => _headerPtr.WriteChunkIdx;
 #pragma warning restore IDE0051 // Remove unused private members
 
-		public bool IsReadingComplete(int readerIdx)
+		/// <summary>
+		/// Allocate a new reader
+		/// </summary>
+		public int AllocateReader() => _headerPtr.AllocateReader();
+
+		/// <inheritdoc/>
+		public bool IsComplete(ReaderState readerState)
 		{
-			ReaderState readerState = _headerPtr.GetReaderStatePtr(readerIdx).Value;
 			ChunkState chunkState = _headerPtr.GetChunkStatePtr(readerState.ChunkIdx).Value;
 			return chunkState.WriteState == WriteState.Complete && readerState.Offset == chunkState.Length;
 		}
 
 		/// <inheritdoc/>
-		public void AdvanceReadPosition(int readerIdx, int length)
+		public void DetachReader(ReaderState readerState)
 		{
-			ReaderStatePtr readerStatePtr = _headerPtr.GetReaderStatePtr(readerIdx);
-			readerStatePtr.Advance(length);
+			readerState.Detached = true;
+			SetReadEvent(readerState.ReaderIdx);
 		}
 
 		/// <inheritdoc/>
-		public ReadOnlyMemory<byte> GetReadBuffer(int readerIdx)
+		public ReadOnlyMemory<byte> GetReadBuffer(ReaderState readerState)
 		{
-			ReaderStatePtr readerStatePtr = _headerPtr.GetReaderStatePtr(readerIdx);
-			ReaderState readerState = readerStatePtr.Value;
-
 			ChunkStatePtr chunkStatePtr = _headerPtr.GetChunkStatePtr(readerState.ChunkIdx);
 			ChunkState chunkState = chunkStatePtr.Value;
 
-			if (chunkState.HasReaderFlag(readerIdx))
+			if (chunkState.HasReaderFlag(readerState.ReaderIdx))
 			{
 				return _chunks[readerState.ChunkIdx].Slice(readerState.Offset, chunkState.Length - readerState.Offset);
 			}
@@ -518,12 +594,15 @@ namespace EpicGames.Horde.Compute
 		}
 
 		/// <inheritdoc/>
-		public async ValueTask<bool> WaitToReadAsync(int readerIdx, int minLength, CancellationToken cancellationToken = default)
+		public async ValueTask<bool> WaitToReadAsync(ReaderState readerState, int minLength, CancellationToken cancellationToken = default)
 		{
+			int readerIdx = readerState.ReaderIdx;
 			for (; ; )
 			{
-				ReaderStatePtr readerStatePtr = _headerPtr.GetReaderStatePtr(readerIdx);
-				ReaderState readerState = readerStatePtr.Value;
+				if (readerState.Detached)
+				{
+					return false;
+				}
 
 				ChunkStatePtr chunkStatePtr = _headerPtr.GetChunkStatePtr(readerState.ChunkIdx);
 				ChunkState chunkState = chunkStatePtr.Value;
@@ -562,13 +641,12 @@ namespace EpicGames.Horde.Compute
 					chunkStatePtr.FinishReading(readerIdx);
 					SetWriteEvent();
 
-					int chunkIdx = readerStatePtr.ChunkIdx + 1;
-					if (chunkIdx == _chunks.Length)
+					if (++readerState.ChunkIdx == _chunks.Length)
 					{
-						chunkIdx = 0;
+						readerState.ChunkIdx = 0;
 					}
 
-					readerStatePtr.Value = new ReaderState(chunkIdx, 0);
+					readerState.Offset = 0;
 				}
 				else
 				{
@@ -577,12 +655,30 @@ namespace EpicGames.Horde.Compute
 			}
 		}
 
+		/// <inheritdoc/>
+		public bool MarkComplete()
+		{
+			WriterState writerState = _headerPtr.Writer.Value;
+
+			ChunkStatePtr chunkStatePtr = _headerPtr.GetChunkStatePtr(writerState.ChunkIdx);
+			if (chunkStatePtr.Value.WriteState != WriteState.Complete)
+			{
+				chunkStatePtr.MarkComplete();
+				SetAllReadEvents();
+				return true;
+			}
+
+			return false;
+		}
+
 		/// <inheritdoc/>
 		public void AdvanceWritePosition(int size)
 		{
 			if (size > 0)
 			{
-				ChunkStatePtr chunkStatePtr = _headerPtr.GetChunkStatePtr(_headerPtr.WriteChunkIdx);
+				WriterState writerState = _headerPtr.Writer.Value;
+
+				ChunkStatePtr chunkStatePtr = _headerPtr.GetChunkStatePtr(writerState.ChunkIdx);
 				ChunkState chunkState = chunkStatePtr.Value;
 
 				Debug.Assert(chunkState.WriteState == WriteState.Writing);
@@ -595,10 +691,12 @@ namespace EpicGames.Horde.Compute
 		/// <inheritdoc/>
 		public Memory<byte> GetWriteBuffer()
 		{
-			ChunkState state = _headerPtr.GetChunkStatePtr(_headerPtr.WriteChunkIdx).Value;
-			if (state.WriteState == WriteState.Writing)
+			WriterState writerState = _headerPtr.Writer.Value;
+
+			ChunkState chunkState = _headerPtr.GetChunkStatePtr(writerState.ChunkIdx).Value;
+			if (chunkState.WriteState == WriteState.Writing)
 			{
-				return _chunks[_headerPtr.WriteChunkIdx].Slice(state.Length);
+				return _chunks[writerState.ChunkIdx].Slice(chunkState.Length);
 			}
 			else
 			{
@@ -606,19 +704,6 @@ namespace EpicGames.Horde.Compute
 			}
 		}
 
-		/// <inheritdoc/>
-		public bool MarkComplete()
-		{
-			ChunkStatePtr chunkStatePtr = _headerPtr.GetChunkStatePtr(_headerPtr.WriteChunkIdx);
-			if (chunkStatePtr.WriteState != WriteState.Complete)
-			{
-				chunkStatePtr.MarkComplete();
-				SetAllReadEvents();
-				return true;
-			}
-			return false;
-		}
-
 		/// <inheritdoc/>
 		public async ValueTask WaitToWriteAsync(int minSize, CancellationToken cancellationToken = default)
 		{
@@ -627,39 +712,60 @@ namespace EpicGames.Horde.Compute
 				throw new ArgumentException("Requested read size is larger than chunk size.", nameof(minSize));
 			}
 
-			for (; ; )
-			{
-				ChunkStatePtr chunkStatePtr = _headerPtr.GetChunkStatePtr(_headerPtr.WriteChunkIdx);
+			// Get the current chunk we're writing to
+			WriterState writerState = _headerPtr.Writer.Value;
+			int writeChunkIdx = writerState.ChunkIdx;
 
-				ChunkState chunkState = chunkStatePtr.Value;
-				if (chunkState.WriteState == WriteState.Writing)
-				{
-					int length = chunkState.Length;
-					if (length + minSize <= _headerPtr.ChunkLength)
-					{
-						return;
-					}
-					chunkStatePtr.FinishWriting();
-				}
+			ChunkStatePtr writeChunkStatePtr = _headerPtr.GetChunkStatePtr(writeChunkIdx);
 
-				int nextChunkIdx = _headerPtr.WriteChunkIdx + 1;
-				if (nextChunkIdx == _chunks.Length)
+			// Check if we can append to this chunk
+			ChunkState chunkState = writeChunkStatePtr.Value;
+			if (chunkState.WriteState == WriteState.Writing)
+			{
+				int length = chunkState.Length;
+				if (length + minSize <= _headerPtr.ChunkLength)
 				{
-					nextChunkIdx = 0;
+					return;
 				}
 
-				ChunkStatePtr nextChunkStatePtr = _headerPtr.GetChunkStatePtr(nextChunkIdx);
-				while (nextChunkStatePtr.ReaderFlags != 0)
+				writeChunkStatePtr.FinishWriting();
+				SetAllReadEvents();
+			}
+
+			if (chunkState.WriteState == WriteState.Complete)
+			{
+				return;
+			}
+
+			// Otherwise get the next chunk to write to
+			int nextWriteChunkIdx = writeChunkIdx + 1;
+			if (nextWriteChunkIdx == _chunks.Length)
+			{
+				nextWriteChunkIdx = 0;
+			}
+
+			// Wait until all readers have finished with the chunk, and we can update the writer to match
+			ChunkStatePtr nextWriteChunkStatePtr = _headerPtr.GetChunkStatePtr(nextWriteChunkIdx);
+			for (; ; )
+			{
+				ChunkState nextWriteChunkState = nextWriteChunkStatePtr.Value;
+				if (nextWriteChunkState.ReaderFlags != 0)
 				{
 					await WaitForWriteEvent(cancellationToken);
 					ResetWriteEvent();
 				}
-
-				_headerPtr.WriteChunkIdx = nextChunkIdx;
-				nextChunkStatePtr.StartWriting(_headerPtr.NumReaders);
-
-				SetAllReadEvents();
+				else if (nextWriteChunkStatePtr.TryUpdate(nextWriteChunkState, new ChunkState(WriteState.Writing, writerState.ReaderFlags, 0)))
+				{
+					if (_headerPtr.Writer.TryUpdate(writerState, new WriterState(nextWriteChunkIdx, writerState.ReaderFlags, nextWriteChunkIdx == 0)))
+					{
+						break;
+					}
+					else
+					{
+						writerState = _headerPtr.Writer.Value;
+					}
+				}
 			}
 		}
 	}
-}
+}
\ No newline at end of file
diff --git a/Engine/Source/Programs/Shared/EpicGames.Horde/Compute/ComputeSocket.cs b/Engine/Source/Programs/Shared/EpicGames.Horde/Compute/ComputeSocket.cs
--- a/Engine/Source/Programs/Shared/EpicGames.Horde/Compute/ComputeSocket.cs
+++ b/Engine/Source/Programs/Shared/EpicGames.Horde/Compute/ComputeSocket.cs
@@ -5,6 +5,7 @@ using System.Buffers;
 using System.Buffers.Binary;
 using System.Collections.Generic;
 using System.Linq;
+using System.Text;
 using System.Threading;
 using System.Threading.Tasks;
 using EpicGames.Core;
@@ -22,15 +23,15 @@ namespace EpicGames.Horde.Compute
 		/// Attaches a buffer to receive data.
 		/// </summary>
 		/// <param name="channelId">Channel to receive data on</param>
-		/// <param name="recvBufferWriter">Writer for the buffer to store received data</param>
-		public abstract void AttachRecvBuffer(int channelId, ComputeBufferWriter recvBufferWriter);
+		/// <param name="recvBuffer">Writer for the buffer to store received data</param>
+		public abstract void AttachRecvBuffer(int channelId, ComputeBuffer recvBuffer);
 
 		/// <summary>
 		/// Attaches a buffer to send data.
 		/// </summary>
 		/// <param name="channelId">Channel to receive data on</param>
-		/// <param name="sendBufferReader">Reader for the buffer to send data from</param>
-		public abstract void AttachSendBuffer(int channelId, ComputeBufferReader sendBufferReader);
+		/// <param name="sendBuffer">Reader for the buffer to send data from</param>
+		public abstract void AttachSendBuffer(int channelId, ComputeBuffer sendBuffer);
 
 		/// <summary>
 		/// Creates a channel using a socket and receive buffer
@@ -51,9 +52,13 @@ namespace EpicGames.Horde.Compute
 		/// <param name="sendBuffer">Buffer for sending data</param>
 		public ComputeChannel CreateChannel(int channelId, ComputeBuffer recvBuffer, ComputeBuffer sendBuffer)
 		{
-			AttachRecvBuffer(channelId, recvBuffer.Writer);
-			AttachSendBuffer(channelId, sendBuffer.Reader);
-			return new ComputeChannel(recvBuffer.Reader, sendBuffer.Writer);
+			AttachRecvBuffer(channelId, recvBuffer);
+			AttachSendBuffer(channelId, sendBuffer);
+
+			using (ComputeBufferReader recvBufferReader = recvBuffer.CreateReader())
+			{
+				return new ComputeChannel(recvBufferReader, sendBuffer.Writer);
+			}
 		}
 	}
 
@@ -74,6 +79,7 @@ namespace EpicGames.Horde.Compute
 		public const string IpcEnvVar = "UE_HORDE_COMPUTE_IPC";
 
 		readonly SharedMemoryBuffer _commandBuffer;
+		readonly List<ComputeBuffer> _buffers = new List<ComputeBuffer>();
 
 		/// <summary>
 		/// Creates a socket for a worker
@@ -114,16 +120,18 @@ namespace EpicGames.Horde.Compute
 		}
 
 		/// <inheritdoc/>
-		public override void AttachRecvBuffer(int channelId, ComputeBufferWriter writer)
+		public override void AttachRecvBuffer(int channelId, ComputeBuffer buffer)
 		{
-			string bufferName = ((SharedMemoryBufferDetail)writer.Detail).Name;
+			_buffers.Add(buffer.AddRef());
+			string bufferName = ((SharedMemoryBufferDetail)buffer._detail).Name;
 			AttachBuffer(IpcMessage.AttachRecvBuffer, channelId, bufferName);
 		}
 
 		/// <inheritdoc/>
-		public override void AttachSendBuffer(int channelId, ComputeBufferReader reader)
+		public override void AttachSendBuffer(int channelId, ComputeBuffer buffer)
 		{
-			string bufferName = ((SharedMemoryBufferDetail)reader.Detail).Name;
+			_buffers.Add(buffer.AddRef());
+			string bufferName = ((SharedMemoryBufferDetail)buffer._detail).Name;
 			AttachBuffer(IpcMessage.AttachSendBuffer, channelId, bufferName);
 		}
 
@@ -387,7 +395,7 @@ namespace EpicGames.Horde.Compute
 		}
 
 		/// <inheritdoc/>
-		public override void AttachRecvBuffer(int channelId, ComputeBufferWriter recvBufferWriter)
+		public override void AttachRecvBuffer(int channelId, ComputeBuffer recvBuffer)
 		{
 			bool complete;
 			lock (_lockObject)
@@ -395,14 +403,14 @@ namespace EpicGames.Horde.Compute
 				complete = _complete;
 				if (!complete)
 				{
-					_recvBufferWriters.Add(channelId, recvBufferWriter.AddRef());
+					_recvBufferWriters.Add(channelId, recvBuffer.Writer.AddRef());
 					_recvTask ??= BackgroundTask.StartNew(ctx => RunRecvTaskAsync(_transport, ctx));
 				}
 			}
 
-			if (recvBufferWriter != null && complete)
+			if (complete)
 			{
-				recvBufferWriter.MarkComplete();
+				recvBuffer.Writer.MarkComplete();
 			}
 		}
 
@@ -425,9 +433,9 @@ namespace EpicGames.Horde.Compute
 		}
 
 		/// <inheritdoc/>
-		public override void AttachSendBuffer(int channelId, ComputeBufferReader sendBufferReader)
+		public override void AttachSendBuffer(int channelId, ComputeBuffer sendBuffer)
 		{
-			sendBufferReader = sendBufferReader.AddRef();
+			ComputeBufferReader sendBufferReader = sendBuffer.CreateReader();
 			lock (_lockObject)
 			{
 				_sendTasks.Add(channelId, Task.Run(() => SendFromBufferAsync(channelId, sendBufferReader, _cancellationSource.Token), CancellationToken.None)); // No cancellation token; need to ensure dispose runs
