Horde: Fix tests.

diff --git a/Engine/Source/Programs/Shared/EpicGames.Horde.Tests/BufferTests.cs b/Engine/Source/Programs/Shared/EpicGames.Horde.Tests/BufferTests.cs
--- a/Engine/Source/Programs/Shared/EpicGames.Horde.Tests/BufferTests.cs
+++ b/Engine/Source/Programs/Shared/EpicGames.Horde.Tests/BufferTests.cs
@@ -108,8 +108,11 @@ namespace EpicGames.Horde.Tests
 			using ComputeBuffer consumerBuffer = createBuffer(Length);
 			consumerSocket.AttachRecvBuffer(ChannelId, consumerBuffer);
 
+			using ComputeBuffer producerBuffer = createBuffer(Length);
+			producerSocket.AttachSendBuffer(ChannelId, producerBuffer);
+
 			byte[] input = RandomNumberGenerator.GetBytes(Length);
-			Task producerTask = RunProducerAsync(producerSocket, input);
+			Task producerTask = RunProducerAsync(producerBuffer, input);
 
 			using ComputeBufferReader consumerBufferReader = consumerBuffer.CreateReader();
 
@@ -120,17 +123,20 @@ namespace EpicGames.Horde.Tests
 			Assert.IsTrue(input.SequenceEqual(output));
 		}
 
-		static async Task RunProducerAsync(RemoteComputeSocket socket, ReadOnlyMemory<byte> input)
+		static async Task RunProducerAsync(ComputeBuffer buffer, ReadOnlyMemory<byte> input)
 		{
+			using ComputeBufferWriter writer = buffer.CreateWriter();
+
 			int offset = 0;
 			while (offset < input.Length)
 			{
 				int length = Math.Min(input.Length - offset, 100);
-				await socket.SendAsync(ChannelId, input.Slice(offset, length));
+				await writer.WriteAsync(input.Slice(offset, length));
 				await Task.Delay(10);
 				offset += length;
 			}
-			await socket.MarkCompleteAsync(ChannelId);
+
+			writer.MarkComplete();
 		}
 
 		static async Task RunConsumerAsync(ComputeBufferReader reader, Memory<byte> output)
diff --git a/Engine/Source/Programs/Shared/EpicGames.Horde/Compute/ComputeBuffer.cs b/Engine/Source/Programs/Shared/EpicGames.Horde/Compute/ComputeBuffer.cs
--- a/Engine/Source/Programs/Shared/EpicGames.Horde/Compute/ComputeBuffer.cs
+++ b/Engine/Source/Programs/Shared/EpicGames.Horde/Compute/ComputeBuffer.cs
@@ -232,7 +232,7 @@ namespace EpicGames.Horde.Compute
 					AdvanceWritePosition(buffer.Length);
 					break;
 				}
-				await WaitToWriteAsync(writeMemory.Length, cancellationToken);
+				await WaitToWriteAsync(buffer.Length, cancellationToken);
 			}
 		}
 
