EpicGames.Core: Record the .NET activity stack on log events. Can be used to supplement logs with context-specific information allowing more granular information about errors and warnings.

diff --git a/Engine/Source/Programs/AutomationTool/AutomationUtils/CommandUtils.cs b/Engine/Source/Programs/AutomationTool/AutomationUtils/CommandUtils.cs
--- a/Engine/Source/Programs/AutomationTool/AutomationUtils/CommandUtils.cs
+++ b/Engine/Source/Programs/AutomationTool/AutomationUtils/CommandUtils.cs
@@ -107,6 +107,25 @@ namespace AutomationTool
 		/// </summary>
 		public static ILogger Logger => Log.Logger;
 
+		/// <summary>
+		/// Global activity source for UAT. Can be used to add trace data and additional context to logs.
+		/// </summary>
+		public static ActivitySource ActivitySource => CreateActivitySource();
+
+		static ActivityListener _activityListener;
+
+		static ActivitySource CreateActivitySource()
+		{
+			ActivitySource activitySource = new ActivitySource("AutomationTool");
+
+			_activityListener = new ActivityListener();
+			_activityListener.ShouldListenTo = x => x == activitySource;
+			_activityListener.Sample = (ref ActivityCreationOptions<ActivityContext> x) => ActivitySamplingResult.AllData;
+			ActivitySource.AddActivityListener(_activityListener);
+
+			return activitySource;
+		}
+
 		/// <summary>
 		/// Writes formatted text to log (with LogEventType.Console).
 		/// </summary>
diff --git a/Engine/Source/Programs/Shared/EpicGames.Core/Log.cs b/Engine/Source/Programs/Shared/EpicGames.Core/Log.cs
--- a/Engine/Source/Programs/Shared/EpicGames.Core/Log.cs
+++ b/Engine/Source/Programs/Shared/EpicGames.Core/Log.cs
@@ -1216,7 +1216,19 @@ namespace EpicGames.Core
 					}
 				}
 
-				JsonLogEvent jsonLogEvent = JsonLogEvent.FromLoggerState(logLevel, eventId, state, exception, formatter);
+				Activity? activity = Activity.Current;
+
+				JsonLogEvent jsonLogEvent;
+				if (activity == null)
+				{
+					jsonLogEvent = JsonLogEvent.FromLoggerState(logLevel, eventId, state, exception, formatter);
+				}
+				else
+				{
+					LogEvent logEvent = LogEvent.FromState(logLevel, eventId, state, exception, formatter);
+					logEvent.Properties = Enumerable.Append(logEvent.Properties ?? Array.Empty<KeyValuePair<string, object>>(), new KeyValuePair<string, object>("Activity", activity));
+					jsonLogEvent = new JsonLogEvent(logEvent);
+				}
 				_eventChannel.Writer.TryWrite(jsonLogEvent);
 
 				// Handle the console output separately; we format things differently
diff --git a/Engine/Source/Programs/Shared/EpicGames.Core/LogValueFormatter.cs b/Engine/Source/Programs/Shared/EpicGames.Core/LogValueFormatter.cs
--- a/Engine/Source/Programs/Shared/EpicGames.Core/LogValueFormatter.cs
+++ b/Engine/Source/Programs/Shared/EpicGames.Core/LogValueFormatter.cs
@@ -1,8 +1,12 @@
 // Copyright Epic Games, Inc. All Rights Reserved.
 
 using System;
+using System.Collections;
 using System.Collections.Concurrent;
 using System.Collections.Generic;
+using System.Diagnostics;
+using System.Diagnostics.CodeAnalysis;
+using System.Linq;
 using System.Reflection;
 using System.Text.Json;
 
@@ -87,6 +91,54 @@ namespace EpicGames.Core
 			public void Format(object value, Utf8JsonWriter writer) => writer.WriteStringValue(value.ToString());
 		}
 
+		class ActivityFormatter : ILogValueFormatter
+		{
+			public void Format(object value, Utf8JsonWriter writer)
+			{
+				writer.WriteStartArray();
+				Format((Activity)value, writer);
+				writer.WriteEndArray();
+			}
+
+			void Format(Activity activity, Utf8JsonWriter writer)
+			{
+				if (activity.Parent != null)
+				{
+					Format(activity.Parent, writer);
+				}
+
+				writer.WriteStartObject();
+				writer.WriteString("name", activity.OperationName);
+				if (activity.Tags.Any())
+				{
+					writer.WriteStartObject("tags");
+					foreach (KeyValuePair<string, string?> pair in activity.Tags)
+					{
+						writer.WriteString(pair.Key, pair.Value);
+					}
+					writer.WriteEndObject();
+				}
+				writer.WriteEndObject();
+			}
+		}
+
+		class EnumerableFormatter : ILogValueFormatter
+		{
+			readonly ILogValueFormatter _elementFormatter;
+
+			public EnumerableFormatter(ILogValueFormatter elementFormatter) => _elementFormatter = elementFormatter;
+
+			public void Format(object value, Utf8JsonWriter writer)
+			{
+				writer.WriteStartArray();
+				foreach (object element in (IEnumerable)value)
+				{
+					_elementFormatter.Format(element, writer);
+				}
+				writer.WriteEndArray();
+			}
+		}
+
 		class StructuredLogValueFormatter : ILogValueFormatter
 		{
 			public void Format(object value, Utf8JsonWriter writer)
@@ -208,6 +260,7 @@ namespace EpicGames.Core
 			formatters.TryAdd(typeof(string), s_stringFormatter);
 			formatters.TryAdd(typeof(LogValue), new StructuredLogValueFormatter());
 			formatters.TryAdd(typeof(FileReference), new FileReferenceFormatter());
+			formatters.TryAdd(typeof(Activity), new ActivityFormatter());
 			return formatters;
 		}
 
@@ -264,6 +317,10 @@ namespace EpicGames.Core
 				{
 					formatter = (ILogValueFormatter)Activator.CreateInstance(formatterAttribute.Type)!;
 				}
+				else if (TryGetEnumerableType(type, out Type? elementType))
+				{
+					formatter = new EnumerableFormatter(GetFormatter(elementType));
+				}
 				else
 				{
 					formatter = s_stringFormatter;
@@ -276,6 +333,21 @@ namespace EpicGames.Core
 			}
 		}
 
+		static bool TryGetEnumerableType(Type type, [NotNullWhen(true)] out Type? elementType)
+		{
+			foreach (Type interfaceType in type.GetInterfaces())
+			{
+				if (interfaceType.IsGenericType && interfaceType.GetGenericTypeDefinition() == typeof(IEnumerable<>))
+				{
+					elementType = interfaceType.GetGenericArguments()[0];
+					return true;
+				}
+			}
+
+			elementType = null;
+			return false;
+		}
+
 		/// <summary>
 		/// Formats a value to a Json log stream
 		/// </summary>
