Horde: Set the UE_HORDE_SHARED_DIR environment variable when running compute jobs.

diff --git a/Engine/Source/Programs/Horde/Horde.Agent/Commands/Compute/ComputeCommand.cs b/Engine/Source/Programs/Horde/Horde.Agent/Commands/Compute/ComputeCommand.cs
--- a/Engine/Source/Programs/Horde/Horde.Agent/Commands/Compute/ComputeCommand.cs
+++ b/Engine/Source/Programs/Horde/Horde.Agent/Commands/Compute/ComputeCommand.cs
@@ -146,7 +146,7 @@ namespace Horde.Agent.Commands.Compute
 		{
 			DirectoryReference sandboxDir = DirectoryReference.Combine(Program.DataDir, "Sandbox");
 
-			AgentMessageHandler worker = new AgentMessageHandler(sandboxDir, memoryCache, false, logger);
+			AgentMessageHandler worker = new AgentMessageHandler(sandboxDir, memoryCache, null, false, logger);
 			await worker.RunAsync(socket, cancellationToken);
 		}
 	}
diff --git a/Engine/Source/Programs/Horde/Horde.Agent/Leases/Handlers/ComputeHandler.cs b/Engine/Source/Programs/Horde/Horde.Agent/Leases/Handlers/ComputeHandler.cs
--- a/Engine/Source/Programs/Horde/Horde.Agent/Leases/Handlers/ComputeHandler.cs
+++ b/Engine/Source/Programs/Horde/Horde.Agent/Leases/Handlers/ComputeHandler.cs
@@ -2,6 +2,7 @@
 
 using System;
 using System.Buffers;
+using System.Collections.Generic;
 using System.Diagnostics;
 using System.Net.Sockets;
 using System.Threading;
@@ -103,7 +104,13 @@ namespace Horde.Agent.Leases.Handlers
 							{
 								DirectoryReference.CreateDirectory(sandboxDir);
 
-								AgentMessageHandler worker = new AgentMessageHandler(sandboxDir, _memoryCache, false, _logger);
+								DirectoryReference sharedDir = DirectoryReference.Combine(session.WorkingDir, "Saved");
+								DirectoryReference.CreateDirectory(sharedDir);
+
+								Dictionary<string, string?> newEnvVars = new Dictionary<string, string?>();
+								newEnvVars["UE_HORDE_SHARED_DIR"] = sharedDir.FullName;
+
+								AgentMessageHandler worker = new AgentMessageHandler(sandboxDir, _memoryCache, newEnvVars, false, _logger);
 								await worker.RunAsync(socket, cts.Token);
 								await socket.CloseAsync(cts.Token);
 								return LeaseResult.Success;
diff --git a/Engine/Source/Programs/Shared/EpicGames.Horde/Compute/AgentMessageHandler.cs b/Engine/Source/Programs/Shared/EpicGames.Horde/Compute/AgentMessageHandler.cs
--- a/Engine/Source/Programs/Shared/EpicGames.Horde/Compute/AgentMessageHandler.cs
+++ b/Engine/Source/Programs/Shared/EpicGames.Horde/Compute/AgentMessageHandler.cs
@@ -24,6 +24,7 @@ namespace EpicGames.Horde.Compute
 	{
 		readonly DirectoryReference _sandboxDir;
 		readonly IMemoryCache _memoryCache;
+		readonly Dictionary<string, string?> _envVars;
 		readonly bool _executeInProcess;
 		readonly ILogger _logger;
 
@@ -32,12 +33,14 @@ namespace EpicGames.Horde.Compute
 		/// </summary>
 		/// <param name="sandboxDir">Directory to use for reading/writing files</param>
 		/// <param name="memoryCache">Cache for nodes read from storage</param>
+		/// <param name="envVars">Environment variables to set for any child processes</param>
 		/// <param name="executeInProcess">Whether to execute any external assemblies in the current process</param>
 		/// <param name="logger">Logger for diagnostics</param>
-		public AgentMessageHandler(DirectoryReference sandboxDir, IMemoryCache memoryCache, bool executeInProcess, ILogger logger)
+		public AgentMessageHandler(DirectoryReference sandboxDir, IMemoryCache memoryCache, Dictionary<string, string?>? envVars, bool executeInProcess, ILogger logger)
 		{
 			_sandboxDir = sandboxDir;
 			_memoryCache = memoryCache;
+			_envVars = envVars ?? new Dictionary<string, string?>();
 			_executeInProcess = executeInProcess;
 			_logger = logger;
 		}
@@ -177,7 +180,7 @@ namespace EpicGames.Horde.Compute
 
 		async Task ExecuteProcessWindowsAsync(ComputeSocket socket, AgentMessageChannel channel, string executable, IReadOnlyList<string> arguments, string? workingDir, IReadOnlyDictionary<string, string?>? envVars, CancellationToken cancellationToken)
 		{
-			Dictionary<string, string?> newEnvVars = new Dictionary<string, string?>();
+			Dictionary<string, string?> newEnvVars = new Dictionary<string, string?>(_envVars);
 			if (envVars != null)
 			{
 				foreach ((string name, string? value) in envVars)
diff --git a/Engine/Source/Programs/Shared/EpicGames.Horde/Compute/Clients/LocalComputeClient.cs b/Engine/Source/Programs/Shared/EpicGames.Horde/Compute/Clients/LocalComputeClient.cs
--- a/Engine/Source/Programs/Shared/EpicGames.Horde/Compute/Clients/LocalComputeClient.cs
+++ b/Engine/Source/Programs/Shared/EpicGames.Horde/Compute/Clients/LocalComputeClient.cs
@@ -82,7 +82,7 @@ namespace EpicGames.Horde.Compute.Clients
 
 			await using (RemoteComputeSocket socket = new RemoteComputeSocket(new TcpTransport(tcpSocket), ComputeSocketEndpoint.Remote, logger))
 			{
-				AgentMessageHandler worker = new AgentMessageHandler(sandboxDir, memoryCache, executeInProcess, logger);
+				AgentMessageHandler worker = new AgentMessageHandler(sandboxDir, memoryCache, null, executeInProcess, logger);
 				await worker.RunAsync(socket, cancellationToken);
 				await socket.CloseAsync(cancellationToken);
 			}
