Horde: Prevent timeout exception from propagating to callers of RemoteComputeSocket.CloseAsync().

diff --git a/Engine/Source/Programs/Shared/EpicGames.Horde/Compute/ComputeSocket.cs b/Engine/Source/Programs/Shared/EpicGames.Horde/Compute/ComputeSocket.cs
--- a/Engine/Source/Programs/Shared/EpicGames.Horde/Compute/ComputeSocket.cs
+++ b/Engine/Source/Programs/Shared/EpicGames.Horde/Compute/ComputeSocket.cs
@@ -230,7 +230,14 @@ namespace EpicGames.Horde.Compute
 			// Wait for the reader to stop
 			if (_recvTask != null)
 			{
-				await _recvTask.Task.WaitAsync(TimeSpan.FromSeconds(5), cancellationToken);
+				try
+				{
+					await _recvTask.Task.WaitAsync(TimeSpan.FromSeconds(5), cancellationToken);
+				}
+				catch (TimeoutException)
+				{
+					_logger.LogWarning("Receive task did not complete gracefully within 5s; triggering cancellation.");
+				}
 				await _recvTask.StopAsync();
 			}
 		}
