Horde: Use async disposer for compute IPC message pump; can cause deadlocks using the synchronous version when running as a continuation on another task.

diff --git a/Engine/Source/Programs/Shared/EpicGames.Horde/Compute/AgentMessageHandler.cs b/Engine/Source/Programs/Shared/EpicGames.Horde/Compute/AgentMessageHandler.cs
--- a/Engine/Source/Programs/Shared/EpicGames.Horde/Compute/AgentMessageHandler.cs
+++ b/Engine/Source/Programs/Shared/EpicGames.Horde/Compute/AgentMessageHandler.cs
@@ -194,11 +194,14 @@ namespace EpicGames.Horde.Compute
 				newEnvVars[WorkerComputeSocket.IpcEnvVar] = ipcBuffer.Name;
 
 				using ComputeBufferReader ipcBufferReader = ipcBuffer.CreateReader();
-				using BackgroundTask backgroundTask = BackgroundTask.StartNew(ctx => ProcessIpcMessagesAsync(socket, ipcBufferReader, cancellationToken));
+				await using (BackgroundTask backgroundTask = BackgroundTask.StartNew(ctx => ProcessIpcMessagesAsync(socket, ipcBufferReader, cancellationToken)))
+				{
+					_logger.LogInformation("Launching {Executable} {Arguments}", CommandLineArguments.Quote(executable), CommandLineArguments.Join(arguments));
+					await ExecuteProcessInternalAsync(channel, executable, arguments, workingDir, newEnvVars, cancellationToken);
+					_logger.LogInformation("Finished executing process");
+				}
 
-				_logger.LogInformation("Launching {Executable} {Arguments}", CommandLineArguments.Quote(executable), CommandLineArguments.Join(arguments));
-				await ExecuteProcessInternalAsync(channel, executable, arguments, workingDir, newEnvVars, cancellationToken);
-				_logger.LogInformation("Finished executing process");
+				_logger.LogInformation("Ipc message loop is complete");
 			}
 
 			_logger.LogInformation("Child process has shut down");
