Horde: Correctly handle IPC message loop termination in compute tasks.

diff --git a/Engine/Source/Programs/Shared/EpicGames.Horde/Compute/AgentMessageHandler.cs b/Engine/Source/Programs/Shared/EpicGames.Horde/Compute/AgentMessageHandler.cs
--- a/Engine/Source/Programs/Shared/EpicGames.Horde/Compute/AgentMessageHandler.cs
+++ b/Engine/Source/Programs/Shared/EpicGames.Horde/Compute/AgentMessageHandler.cs
@@ -194,10 +194,17 @@ namespace EpicGames.Horde.Compute
 				newEnvVars[WorkerComputeSocket.IpcEnvVar] = ipcBuffer.Name;
 
 				using ComputeBufferReader ipcBufferReader = ipcBuffer.CreateReader();
-				await using (BackgroundTask backgroundTask = BackgroundTask.StartNew(ctx => ProcessIpcMessagesAsync(socket, ipcBufferReader, cancellationToken)))
+				await using (BackgroundTask backgroundTask = BackgroundTask.StartNew(ctx => ProcessIpcMessagesAsync(socket, ipcBufferReader, new[] { cancellationToken, ctx })))
 				{
 					_logger.LogInformation("Launching {Executable} {Arguments}", CommandLineArguments.Quote(executable), CommandLineArguments.Join(arguments));
-					await ExecuteProcessInternalAsync(channel, executable, arguments, workingDir, newEnvVars, cancellationToken);
+					try
+					{
+						await ExecuteProcessInternalAsync(channel, executable, arguments, workingDir, newEnvVars, cancellationToken);
+					}
+					finally
+					{
+						ipcBufferReader.Detach();
+					}
 					_logger.LogInformation("Finished executing process");
 				}
 
@@ -207,9 +214,12 @@ namespace EpicGames.Horde.Compute
 			_logger.LogInformation("Child process has shut down");
 		}
 
-		async Task ProcessIpcMessagesAsync(ComputeSocket socket, ComputeBufferReader ipcReader, CancellationToken cancellationToken)
+		async Task ProcessIpcMessagesAsync(ComputeSocket socket, ComputeBufferReader ipcReader, CancellationToken[] cancellationTokens)
 		{
-			List<SharedMemoryBuffer> buffers = new();
+			using CancellationTokenSource cancellationTokenSource = CancellationTokenSource.CreateLinkedTokenSource(cancellationTokens);
+			CancellationToken cancellationToken = cancellationTokenSource.Token;
+
+			List <SharedMemoryBuffer> buffers = new();
 			try
 			{
 				List<(int, ComputeBufferWriter)> writers = new List<(int, ComputeBufferWriter)>();
@@ -259,6 +269,10 @@ namespace EpicGames.Horde.Compute
 					ipcReader.AdvanceReadPosition(memory.Length - reader.RemainingMemory.Length);
 				}
 			}
+			catch (OperationCanceledException)
+			{
+				_logger.LogDebug("Ipc message loop cancelled");
+			}
 			finally
 			{
 				foreach (SharedMemoryBuffer buffer in buffers)
