Horde: Add tests for remote execution agent message loop.

diff --git a/Engine/Source/Programs/Shared/EpicGames.Horde.Tests/ComputeSocketTests.cs b/Engine/Source/Programs/Shared/EpicGames.Horde.Tests/ComputeSocketTests.cs
new file mode 100644
--- /dev/null
+++ b/Engine/Source/Programs/Shared/EpicGames.Horde.Tests/ComputeSocketTests.cs
@@ -0,0 +1,153 @@
+// Copyright Epic Games, Inc. All Rights Reserved.
+
+using System;
+using System.Buffers;
+using System.Collections.Generic;
+using System.IO.Pipelines;
+using System.IO.Pipes;
+using System.Linq;
+using System.Net.Sockets;
+using System.Text;
+using System.Threading;
+using System.Threading.Tasks;
+using EpicGames.Core;
+using EpicGames.Horde.Compute;
+using EpicGames.Horde.Compute.Buffers;
+using EpicGames.Horde.Compute.Transports;
+using EpicGames.Horde.Storage;
+using EpicGames.Horde.Storage.Backends;
+using EpicGames.Horde.Storage.Bundles;
+using EpicGames.Horde.Storage.Nodes;
+using Microsoft.Extensions.Caching.Memory;
+using Microsoft.Extensions.Logging.Abstractions;
+using Microsoft.VisualStudio.TestTools.UnitTesting;
+
+namespace EpicGames.Horde.Tests
+{
+	[TestClass]
+	public class ComputeSocketTests
+	{
+		class TestComputeSocket : ComputeSocket, IDisposable
+		{
+			public Dictionary<int, ComputeBufferWriter> RecvBufferWriters { get; } = new Dictionary<int, ComputeBufferWriter>();
+			public Dictionary<int, ComputeBufferReader> SendBufferReaders { get; } = new Dictionary<int, ComputeBufferReader>();
+
+			public void Dispose()
+			{
+				foreach (ComputeBufferWriter writer in RecvBufferWriters.Values)
+				{
+					writer.Dispose();
+				}
+				foreach (ComputeBufferReader reader in SendBufferReaders.Values)
+				{
+					reader.Dispose();
+				}
+			}
+
+			public override void AttachRecvBuffer(int channelId, ComputeBuffer recvBuffer)
+			{
+				RecvBufferWriters.Add(channelId, recvBuffer.CreateWriter());
+			}
+
+			public override void AttachSendBuffer(int channelId, ComputeBuffer sendBuffer)
+			{
+				SendBufferReaders.Add(channelId, sendBuffer.CreateReader());
+			}
+		}
+
+		[TestMethod]
+		public async Task TestAgentMessageLoop()
+		{
+			Pipe recvPipe = new Pipe();
+			Pipe sendPipe = new Pipe();
+			await using RemoteComputeSocket localSocket = new RemoteComputeSocket(new PipeTransport(sendPipe.Reader, recvPipe.Writer), ComputeSocketEndpoint.Local, NullLogger.Instance);
+			await using RemoteComputeSocket agentSocket = new RemoteComputeSocket(new PipeTransport(recvPipe.Reader, sendPipe.Writer), ComputeSocketEndpoint.Remote, NullLogger.Instance);
+
+			DirectoryReference tempDir = new DirectoryReference("test-temp");
+			await using (BackgroundTask agentTask = BackgroundTask.StartNew(ctx => RunAgent(agentSocket, tempDir, ctx)))
+			{
+				const int PrimaryChannelId = 0;
+				using (AgentMessageChannel channel = localSocket.CreateAgentMessageChannel(PrimaryChannelId, 4 * 1024 * 1024, NullLogger.Instance))
+				{
+					await channel.WaitForAttachAsync();
+
+					await channel.PingAsync();
+					using (AgentMessage message = await channel.ReceiveAsync(CancellationToken.None))
+					{
+						Assert.AreEqual(AgentMessageType.Ping, message.Type);
+						Assert.IsTrue(message.Data.Span.SequenceEqual(ReadOnlySpan<byte>.Empty));
+					}
+
+					await channel.SendXorRequestAsync(new byte[] { 1, 2, 3 }, 44);
+					using (AgentMessage message = await channel.ReceiveAsync(CancellationToken.None))
+					{
+						Assert.AreEqual(AgentMessageType.XorResponse, message.Type);
+						Assert.IsTrue(message.Data.Span.SequenceEqual(new byte[] { 1 ^ 44, 2 ^ 44, 3 ^ 44 }));
+					}
+
+					const int SecondaryChannelId = 1;
+					using (AgentMessageChannel channel2 = localSocket.CreateAgentMessageChannel(SecondaryChannelId, 4 * 1024 * 1024, NullLogger.Instance))
+					{
+						await channel.ForkAsync(SecondaryChannelId, 4 * 1024 * 1024);
+
+						await channel2.WaitForAttachAsync();
+
+						await channel2.SendXorRequestAsync(new byte[] { 1, 2, 3 }, 44);
+						using (AgentMessage message = await channel2.ReceiveAsync(CancellationToken.None))
+						{
+							Assert.AreEqual(AgentMessageType.XorResponse, message.Type);
+							Assert.IsTrue(message.Data.Span.SequenceEqual(new byte[] { 1 ^ 44, 2 ^ 44, 3 ^ 44 }));
+						}
+
+						await channel2.CloseAsync();
+					}
+
+					MemoryStorageClient storage = new MemoryStorageClient();
+					await using (BundleWriter treeWriter = storage.CreateWriter())
+					{
+						byte[] data = Encoding.UTF8.GetBytes("Hello world");
+
+						ChunkedDataWriter writer = new ChunkedDataWriter(treeWriter, new ChunkingOptions());
+						NodeRef<ChunkedDataNode> nodeRef = await writer.CreateAsync(data, CancellationToken.None);
+
+						DirectoryNode directory = new DirectoryNode();
+						directory.AddFile("hello.txt", FileEntryFlags.None, data.Length, nodeRef);
+
+						NodeRef<DirectoryNode> directoryRef = await treeWriter.WriteNodeAsync(directory);
+
+						DirectoryNode root = new DirectoryNode();
+						root.AddDirectory(new DirectoryEntry("subdir", directory.Length, directoryRef));
+
+						BundleNodeHandle handle = await treeWriter.FlushAsync(root);
+						await channel.UploadFilesAsync("", handle.GetLocator(), storage);
+
+						FileReference file = FileReference.Combine(tempDir, "subdir/hello.txt");
+						Assert.IsTrue(FileReference.Exists(file));
+						byte[] readData = FileReference.ReadAllBytes(file);
+						Assert.IsTrue(readData.SequenceEqual(data));
+
+						await channel.DeleteFilesAsync(new[] { "subdir/hello.txt" }, CancellationToken.None);
+
+						await channel.PingAsync();
+						using (AgentMessage message = await channel.ReceiveAsync(CancellationToken.None))
+						{
+							Assert.AreEqual(AgentMessageType.Ping, message.Type);
+							Assert.IsTrue(message.Data.Span.SequenceEqual(ReadOnlySpan<byte>.Empty));
+						}
+
+						Assert.IsFalse(FileReference.Exists(file));
+					}
+
+
+				}
+			}
+		}
+
+		static async Task RunAgent(ComputeSocket socket, DirectoryReference tempDir, CancellationToken cancellationToken)
+		{
+			using MemoryCache memoryCache = new MemoryCache(new MemoryCacheOptions());
+			AgentMessageHandler handler = new AgentMessageHandler(tempDir, memoryCache, null, true, NullLogger.Instance);
+			await handler.RunAsync(socket, cancellationToken);
+		}
+	}
+}
diff --git a/Engine/Source/Programs/Shared/EpicGames.Horde/Compute/AgentMessage.cs b/Engine/Source/Programs/Shared/EpicGames.Horde/Compute/AgentMessage.cs
--- a/Engine/Source/Programs/Shared/EpicGames.Horde/Compute/AgentMessage.cs
+++ b/Engine/Source/Programs/Shared/EpicGames.Horde/Compute/AgentMessage.cs
@@ -387,7 +387,7 @@ namespace EpicGames.Horde.Compute
 		/// <param name="cancellationToken">Cancellation token for the operation</param>
 		public static async ValueTask DeleteFilesAsync(this AgentMessageChannel channel, IReadOnlyList<string> paths, CancellationToken cancellationToken)
 		{
-			IAgentMessageBuilder request = await channel.CreateMessageAsync(AgentMessageType.DeleteFiles, cancellationToken);
+			using IAgentMessageBuilder request = await channel.CreateMessageAsync(AgentMessageType.DeleteFiles, cancellationToken);
 			request.WriteList(paths, MemoryWriterExtensions.WriteString);
 			request.Send();
 		}
diff --git a/Engine/Source/Programs/Shared/EpicGames.Horde/Compute/AgentMessageHandler.cs b/Engine/Source/Programs/Shared/EpicGames.Horde/Compute/AgentMessageHandler.cs
--- a/Engine/Source/Programs/Shared/EpicGames.Horde/Compute/AgentMessageHandler.cs
+++ b/Engine/Source/Programs/Shared/EpicGames.Horde/Compute/AgentMessageHandler.cs
@@ -194,7 +194,7 @@ namespace EpicGames.Horde.Compute
 				newEnvVars[WorkerComputeSocket.IpcEnvVar] = ipcBuffer.Name;
 
 				using ComputeBufferReader ipcBufferReader = ipcBuffer.CreateReader();
-				await using (BackgroundTask backgroundTask = BackgroundTask.StartNew(ctx => ProcessIpcMessagesAsync(socket, ipcBufferReader, new[] { cancellationToken, ctx })))
+				await using (BackgroundTask backgroundTask = BackgroundTask.StartNew(ctx => ProcessIpcMessagesAsync(socket, ipcBufferReader, new[] { cancellationToken, ctx }, _logger)))
 				{
 					_logger.LogInformation("Launching {Executable} {Arguments}", CommandLineArguments.Quote(executable), CommandLineArguments.Join(arguments));
 					try
@@ -214,7 +214,7 @@ namespace EpicGames.Horde.Compute
 			_logger.LogInformation("Child process has shut down");
 		}
 
-		async Task ProcessIpcMessagesAsync(ComputeSocket socket, ComputeBufferReader ipcReader, CancellationToken[] cancellationTokens)
+		internal static async Task ProcessIpcMessagesAsync(ComputeSocket socket, ComputeBufferReader ipcReader, CancellationToken[] cancellationTokens, ILogger logger)
 		{
 			using CancellationTokenSource cancellationTokenSource = CancellationTokenSource.CreateLinkedTokenSource(cancellationTokens);
 			CancellationToken cancellationToken = cancellationTokenSource.Token;
@@ -237,7 +237,7 @@ namespace EpicGames.Horde.Compute
 								{
 									int channelId = (int)reader.ReadUnsignedVarInt();
 									string name = reader.ReadString();
-									_logger.LogDebug("Attaching send buffer for channel {ChannelId} to {Name}", channelId, name);
+									logger.LogDebug("Attaching send buffer for channel {ChannelId} to {Name}", channelId, name);
 
 									SharedMemoryBuffer buffer = SharedMemoryBuffer.OpenExisting(name);
 									buffers.Add(buffer);
@@ -249,7 +249,7 @@ namespace EpicGames.Horde.Compute
 								{
 									int channelId = (int)reader.ReadUnsignedVarInt();
 									string name = reader.ReadString();
-									_logger.LogDebug("Attaching recv buffer for channel {ChannelId} to {Name}", channelId, name);
+									logger.LogDebug("Attaching recv buffer for channel {ChannelId} to {Name}", channelId, name);
 
 									SharedMemoryBuffer buffer = SharedMemoryBuffer.OpenExisting(name);
 									buffers.Add(buffer);
@@ -263,7 +263,7 @@ namespace EpicGames.Horde.Compute
 					}
 					catch (Exception ex)
 					{
-						_logger.LogError(ex, "Exception while processing messages from child process: {Message}", ex.Message);
+						logger.LogError(ex, "Exception while processing messages from child process: {Message}", ex.Message);
 					}
 
 					ipcReader.AdvanceReadPosition(memory.Length - reader.RemainingMemory.Length);
@@ -271,7 +271,7 @@ namespace EpicGames.Horde.Compute
 			}
 			catch (OperationCanceledException)
 			{
-				_logger.LogDebug("Ipc message loop cancelled");
+				logger.LogDebug("Ipc message loop cancelled");
 			}
 			finally
 			{
