Horde: More aggressively terminate remote execution sockets. Now detaches buffers immediately rather than waiting for the receive thread to do so in response to messages received from the remote.

diff --git a/Engine/Source/Programs/Shared/EpicGames.Horde.Tests/ComputeSocketTests.cs b/Engine/Source/Programs/Shared/EpicGames.Horde.Tests/ComputeSocketTests.cs
--- a/Engine/Source/Programs/Shared/EpicGames.Horde.Tests/ComputeSocketTests.cs
+++ b/Engine/Source/Programs/Shared/EpicGames.Horde.Tests/ComputeSocketTests.cs
@@ -6,6 +6,7 @@ using System.Collections.Generic;
 using System.IO.Pipelines;
 using System.IO.Pipes;
 using System.Linq;
+using System.Net;
 using System.Net.Sockets;
 using System.Text;
 using System.Threading;
@@ -19,6 +20,7 @@ using EpicGames.Horde.Storage.Backends;
 using EpicGames.Horde.Storage.Bundles;
 using EpicGames.Horde.Storage.Nodes;
 using Microsoft.Extensions.Caching.Memory;
+using Microsoft.Extensions.Logging;
 using Microsoft.Extensions.Logging.Abstractions;
 using Microsoft.VisualStudio.TestTools.UnitTesting;
 
@@ -55,14 +57,52 @@ namespace EpicGames.Horde.Tests
 			}
 		}
 
+		class TestLogger : ILogger
+		{
+			public IDisposable BeginScope<TState>(TState state) => null!;
+
+			public bool IsEnabled(LogLevel logLevel) => true;
+
+			public void Log<TState>(LogLevel logLevel, EventId eventId, TState state, Exception? exception, Func<TState, Exception?, string> formatter)
+			{
+				Console.WriteLine($"{logLevel}: {formatter(state, exception)}");
+				Assert.IsFalse(logLevel == LogLevel.Error);
+				Assert.IsFalse(logLevel == LogLevel.Warning);
+			}
+		}
+
 		[TestMethod]
-		public async Task TestAgentMessageLoop()
+		public async Task TestAgentMessageLoopPipe()
 		{
 			Pipe recvPipe = new Pipe();
 			Pipe sendPipe = new Pipe();
-			await using RemoteComputeSocket localSocket = new RemoteComputeSocket(new PipeTransport(sendPipe.Reader, recvPipe.Writer), ComputeSocketEndpoint.Local, NullLogger.Instance);
-			await using RemoteComputeSocket agentSocket = new RemoteComputeSocket(new PipeTransport(recvPipe.Reader, sendPipe.Writer), ComputeSocketEndpoint.Remote, NullLogger.Instance);
+			await using RemoteComputeSocket localSocket = new RemoteComputeSocket(new PipeTransport(sendPipe.Reader, recvPipe.Writer), ComputeSocketEndpoint.Local, new TestLogger());
+			await using RemoteComputeSocket agentSocket = new RemoteComputeSocket(new PipeTransport(recvPipe.Reader, sendPipe.Writer), ComputeSocketEndpoint.Remote, new TestLogger());
+
+			await RunAgentTests(localSocket, agentSocket);
+		}
+
+		[TestMethod]
+		public async Task TestAgentMessageLoopTcp()
+		{
+			const int Port = 9990;
+			TcpListener listener = new TcpListener(IPAddress.Loopback, Port);
+			listener.Start();
+
+			using Socket clientSocket = new Socket(SocketType.Stream, ProtocolType.Tcp);
+			Task clientConnectTask = clientSocket.ConnectAsync(IPAddress.Loopback, Port, CancellationToken.None).AsTask();
 
+			using Socket serverSocket = await listener.AcceptSocketAsync(CancellationToken.None);
+			await clientConnectTask;
+
+			await using RemoteComputeSocket localSocket = new RemoteComputeSocket(new TcpTransport(clientSocket), ComputeSocketEndpoint.Local, new TestLogger());
+			await using RemoteComputeSocket agentSocket = new RemoteComputeSocket(new TcpTransport(serverSocket), ComputeSocketEndpoint.Remote, new TestLogger());
+
+			await RunAgentTests(localSocket, agentSocket);
+		}
+
+		static async Task RunAgentTests(RemoteComputeSocket localSocket, RemoteComputeSocket agentSocket)
+		{
 			DirectoryReference tempDir = new DirectoryReference("test-temp");
 			await using (BackgroundTask agentTask = BackgroundTask.StartNew(ctx => RunAgent(agentSocket, tempDir, ctx)))
 			{
@@ -105,6 +145,10 @@ namespace EpicGames.Horde.Tests
 					MemoryStorageClient storage = new MemoryStorageClient();
 					await using (BundleWriter treeWriter = storage.CreateWriter())
 					{
+						FileReference file = FileReference.Combine(tempDir, "subdir/hello.txt");
+						FileReference.Delete(file);
+						Assert.IsFalse(FileReference.Exists(file));
+
 						byte[] data = Encoding.UTF8.GetBytes("Hello world");
 
 						ChunkedDataWriter writer = new ChunkedDataWriter(treeWriter, new ChunkingOptions());
@@ -121,7 +165,6 @@ namespace EpicGames.Horde.Tests
 						BundleNodeHandle handle = await treeWriter.FlushAsync(root);
 						await channel.UploadFilesAsync("", handle.GetLocator(), storage);
 
-						FileReference file = FileReference.Combine(tempDir, "subdir/hello.txt");
 						Assert.IsTrue(FileReference.Exists(file));
 						byte[] readData = FileReference.ReadAllBytes(file);
 						Assert.IsTrue(readData.SequenceEqual(data));
@@ -137,10 +180,11 @@ namespace EpicGames.Horde.Tests
 
 						Assert.IsFalse(FileReference.Exists(file));
 					}
-
-
 				}
 			}
+
+			await localSocket.CloseAsync(CancellationToken.None);
+			await agentSocket.CloseAsync(CancellationToken.None);
 		}
 
 		static async Task RunAgent(ComputeSocket socket, DirectoryReference tempDir, CancellationToken cancellationToken)
diff --git a/Engine/Source/Programs/Shared/EpicGames.Horde/Compute/ComputeSocket.cs b/Engine/Source/Programs/Shared/EpicGames.Horde/Compute/ComputeSocket.cs
--- a/Engine/Source/Programs/Shared/EpicGames.Horde/Compute/ComputeSocket.cs
+++ b/Engine/Source/Programs/Shared/EpicGames.Horde/Compute/ComputeSocket.cs
@@ -4,12 +4,14 @@ using System;
 using System.Buffers;
 using System.Buffers.Binary;
 using System.Collections.Generic;
+using System.Diagnostics.CodeAnalysis;
 using System.Linq;
 using System.Text;
 using System.Threading;
 using System.Threading.Tasks;
 using EpicGames.Core;
 using EpicGames.Horde.Compute.Buffers;
+using JetBrains.Annotations;
 using Microsoft.Extensions.Logging;
 
 namespace EpicGames.Horde.Compute
@@ -172,6 +174,61 @@ namespace EpicGames.Horde.Compute
 			Detach = -2,
 		}
 
+		class RecvBuffer : IDisposable
+		{
+			public ComputeBufferWriter? _writer;
+			public readonly SemaphoreSlim _semaphore = new SemaphoreSlim(1);
+			public int _refCount = 1;
+
+			public RecvBuffer(ComputeBufferWriter writer) => _writer = writer;
+
+			public void AddRef() => Interlocked.Increment(ref _refCount);
+
+			public void Release()
+			{
+				if (Interlocked.Decrement(ref _refCount) == 0)
+				{
+					Dispose();
+				}
+			}
+
+			public void Dispose()
+			{
+				_writer?.Dispose();
+				_semaphore.Dispose();
+			}
+		}
+
+		class SendBuffer : IDisposable
+		{
+			public ComputeBufferReader? _reader;
+			public readonly SemaphoreSlim _semaphore = new SemaphoreSlim(1);
+			public readonly BackgroundTask _task;
+			int _refCount = 1;
+
+			public SendBuffer(ComputeBufferReader reader, Func<SendBuffer, CancellationToken, Task> func)
+			{
+				_reader = reader;
+				_task = BackgroundTask.StartNew(ctx => func(this, ctx));
+			}
+
+			public void AddRef() => Interlocked.Increment(ref _refCount);
+
+			public void Release()
+			{
+				if (Interlocked.Decrement(ref _refCount) == 0)
+				{
+					Dispose();
+				}
+			}
+
+			public void Dispose()
+			{
+				_reader?.Dispose();
+				_semaphore.Dispose();
+			}
+		}
+
 		readonly object _lockObject = new object();
 
 		bool _complete;
@@ -180,14 +237,11 @@ namespace EpicGames.Horde.Compute
 		readonly ComputeSocketEndpoint _endpoint;
 		readonly ILogger _logger;
 
-		readonly CancellationTokenSource _cancellationSource = new CancellationTokenSource();
-
 		BackgroundTask? _recvTask;
-		readonly Dictionary<int, ComputeBufferWriter> _recvBufferWriters = new Dictionary<int, ComputeBufferWriter>();
+		readonly Dictionary<int, RecvBuffer> _recvBuffers = new Dictionary<int, RecvBuffer>();
 
 		readonly SemaphoreSlim _sendSemaphore = new SemaphoreSlim(1, 1);
-		readonly Dictionary<int, Task> _sendTasks = new Dictionary<int, Task>();
-		readonly Dictionary<ComputeBufferReader, int> _sendBufferReaders = new Dictionary<ComputeBufferReader, int>();
+		readonly Dictionary<int, SendBuffer> _sendBuffers = new Dictionary<int, SendBuffer>();
 
 		string Tag => (_endpoint == ComputeSocketEndpoint.Local)? "LOCAL": "REMOTE";
 
@@ -209,35 +263,15 @@ namespace EpicGames.Horde.Compute
 		/// </summary>
 		public async ValueTask CloseAsync(CancellationToken cancellationToken)
 		{
-			_cancellationSource.CancelAfter(TimeSpan.FromSeconds(2.0));
-
-			// Make sure we close all buffers that are attached, otherwise we'll lock up waiting for send tasks to complete
-			Task[] sendTasks;
-			lock (_lockObject)
-			{
-				sendTasks = _sendTasks.Values.ToArray();
-				_sendBufferReaders.Clear();
-			}
+			// Close all the buffers
+			await DetachAllBuffersAsync(true, true, cancellationToken);
 
-			// Wait for all the individual send tasks to complete
-			await Task.WhenAll(sendTasks);
-			cancellationToken.ThrowIfCancellationRequested();
-
-			// Send a final message indicating that the lease is done. This will allow the senders on the remote end to terminate, and trigger
-			// a shutdown event to be sent back to our read task allowing it to shut down gracefully.
+			// Send a final message indicating that the lease is done. This will allow the senders on the remote end to terminate.
 			await _transport.MarkCompleteAsync(cancellationToken);
 
 			// Wait for the reader to stop
 			if (_recvTask != null)
 			{
-				try
-				{
-					await _recvTask.Task.WaitAsync(TimeSpan.FromSeconds(5), cancellationToken);
-				}
-				catch (TimeoutException)
-				{
-					_logger.LogWarning("Receive task did not complete gracefully within 5s; triggering cancellation.");
-				}
 				await _recvTask.StopAsync();
 			}
 		}
@@ -245,14 +279,8 @@ namespace EpicGames.Horde.Compute
 		/// <inheritdoc/>
 		public async ValueTask DisposeAsync()
 		{
-			_cancellationSource.Cancel();
-			await Task.WhenAll(_sendTasks.Values);
-			if (_recvTask != null)
-			{
-				await _recvTask.DisposeAsync();
-			}
+			await CloseAsync(CancellationToken.None);
 			_sendSemaphore.Dispose();
-			_cancellationSource.Dispose();
 			GC.SuppressFinalize(this);
 		}
 
@@ -260,17 +288,18 @@ namespace EpicGames.Horde.Compute
 		{
 			_logger.LogTrace("[{Tag}] Started socket reader", Tag);
 
+			List<Task> detachTasks = new List<Task>();
+
 			byte[] header = new byte[8];
 			try
 			{
-				// Maintain a local cache of buffers to be able to query for them without having to acquire a global lock
-				Dictionary<int, ComputeBufferWriter> cachedWriters = new Dictionary<int, ComputeBufferWriter>();
-
 				Memory<byte> last = Memory<byte>.Empty;
 
 				// Process messages from the remote
 				for (; ; )
 				{
+					detachTasks.RemoveCompleteTasks();
+
 					// Read the next packet header
 					if (!await transport.RecvOptionalAsync(header, cancellationToken))
 					{
@@ -285,13 +314,11 @@ namespace EpicGames.Horde.Compute
 					// Dispatch it to the correct place
 					if (size >= 0)
 					{
-						ComputeBufferWriter writer = GetReceiveBuffer(cachedWriters, id);
-						await ReadPacketAsync(transport, id, size, writer, cancellationToken);
+						await ReadPacketAsync(transport, id, size, cancellationToken);
 					}
 					else if (size == (int)ControlMessageType.Detach)
 					{
-						_logger.LogTrace("[{Tag}] Detaching recv buffer {Id}", Tag, id);
-						DetachRecvBuffer(cachedWriters, id);
+						detachTasks.Add(DetachRecvBufferAsync(id, cancellationToken));
 					}
 					else
 					{
@@ -307,54 +334,95 @@ namespace EpicGames.Horde.Compute
 			lock (_lockObject)
 			{
 				_complete = true;
-				foreach (ComputeBufferWriter writer in _recvBufferWriters.Values)
+				foreach (int channelIdx in _recvBuffers.Keys)
 				{
-					writer.MarkComplete();
+					detachTasks.Add(DetachRecvBufferAsync(channelIdx, cancellationToken));
 				}
 			}
 
+			// Wait for all the detach tasks to finish
+			if (detachTasks.Count > 0)
+			{
+				_logger.LogTrace("[{Tag}] Waiting for detach tasks to complete...", Tag);
+				await Task.WhenAll(detachTasks).WaitAsync(cancellationToken);
+			}
+
 			_logger.LogTrace("[{Tag}] Closing reader", Tag);
 		}
 
-		async Task ReadPacketAsync(ComputeTransport transport, int id, int size, ComputeBufferWriter writer, CancellationToken cancellationToken)
+		async Task ReadPacketAsync(ComputeTransport transport, int id, int size, CancellationToken cancellationToken)
 		{
-			Memory<byte> memory = writer.GetWriteBuffer();
-			while (memory.Length < size)
+			if (!await TryReadPacketAsync(transport, id, size, cancellationToken))
 			{
-				_logger.LogTrace("[{Tag}] No space in buffer {Id}, flushing", Tag, id);
-				await writer.WaitToWriteAsync(size, cancellationToken);
-				memory = writer.GetWriteBuffer();
-			}
+				_logger.LogWarning("Discarding {Size} bytes received on compute channel {Id}", size, id);
 
-			for (int offset = 0; offset < size;)
-			{
-				int read = await transport.RecvAsync(memory.Slice(offset, size - offset), cancellationToken);
-				offset += read;
+				int bufferSize = Math.Min(size, 65536);
+				using (IMemoryOwner<byte> buffer = MemoryPool<byte>.Shared.Rent(bufferSize))
+				{
+					for (int remaining = size; remaining > 0;)
+					{
+						int chunkSize = Math.Min(bufferSize, remaining);
+						await transport.RecvAsync(buffer.Memory.Slice(0, chunkSize), cancellationToken);
+						remaining -= chunkSize;
+					}
+				}
 			}
-
-			writer.AdvanceWritePosition(size);
 		}
 
-		ComputeBufferWriter GetReceiveBuffer(Dictionary<int, ComputeBufferWriter> cachedWriters, int id)
+		async ValueTask<bool> TryReadPacketAsync(ComputeTransport transport, int id, int size, CancellationToken cancellationToken)
 		{
-			ComputeBufferWriter? writer;
-			if (cachedWriters.TryGetValue(id, out writer))
+			// Try to get the receive buffer for this channel
+			RecvBuffer? recvBuffer;
+			lock (_lockObject)
 			{
-				return writer;
+				if (_recvBuffers.TryGetValue(id, out recvBuffer))
+				{
+					recvBuffer.AddRef();
+				}
+				else
+				{
+					return false;
+				}
 			}
 
-			lock (_lockObject)
+			// Lock it and read a packet
+			bool result = false;
+			try
 			{
-				ComputeBufferWriter? recvBufferWriter;
-				if (_recvBufferWriters.TryGetValue(id, out recvBufferWriter))
+				await recvBuffer._semaphore.WaitAsync(cancellationToken);
+				try
 				{
-					writer = recvBufferWriter;
-					cachedWriters.Add(id, writer);
-					return writer;
+					ComputeBufferWriter? writer = recvBuffer._writer;
+					if (writer != null)
+					{
+						Memory<byte> memory = writer.GetWriteBuffer();
+						while (memory.Length < size)
+						{
+							_logger.LogTrace("[{Tag}] No space in buffer {Id}, flushing", Tag, id);
+							await writer.WaitToWriteAsync(size, cancellationToken);
+							memory = writer.GetWriteBuffer();
+						}
+
+						for (int offset = 0; offset < size;)
+						{
+							int read = await transport.RecvAsync(memory.Slice(offset, size - offset), cancellationToken);
+							offset += read;
+						}
+
+						writer.AdvanceWritePosition(size);
+						result = true;
+					}
+				}
+				finally
+				{
+					recvBuffer._semaphore.Release();
 				}
 			}
-
-			throw new ComputeInternalException($"No buffer is attached to channel {id}");
+			finally
+			{
+				recvBuffer.Release();
+			}
+			return result;
 		}
 
 		class SendSegment : ReadOnlySequenceSegment<byte>
@@ -405,49 +473,153 @@ namespace EpicGames.Horde.Compute
 		/// <inheritdoc/>
 		public override void AttachRecvBuffer(int channelId, ComputeBuffer recvBuffer)
 		{
+			_logger.LogTrace("[{Tag}] Attaching recv buffer {Id}", Tag, channelId);
 			lock (_lockObject)
 			{
-				if (!_complete)
+				if (_complete)
 				{
-					_recvBufferWriters.Add(channelId, recvBuffer.CreateWriter());
-					_recvTask ??= BackgroundTask.StartNew(ctx => RunRecvTaskAsync(_transport, ctx));
+					throw new InvalidOperationException($"Cannot attach new buffer to channel {channelId} after socket is closed");
 				}
+				if (_recvBuffers.ContainsKey(channelId))
+				{
+					throw new InvalidOperationException($"Buffer is already attached to channel {channelId}");
+				}
+
+				_recvBuffers.Add(channelId, new RecvBuffer(recvBuffer.CreateWriter()));
+				_recvTask ??= BackgroundTask.StartNew(ctx => RunRecvTaskAsync(_transport, ctx));
 			}
 		}
 
-		void DetachRecvBuffer(Dictionary<int, ComputeBufferWriter> cachedWriters, int id)
+		[SuppressMessage("Reliability", "CA2000:Dispose objects before losing scope")]
+		async Task DetachRecvBufferAsync(int id, CancellationToken cancellationToken)
 		{
-			cachedWriters.Remove(id);
+			_logger.LogTrace("[{Tag}] Detaching recv buffer {Id}", Tag, id);
 
-			ComputeBufferWriter? recvBufferWriter;
+			// Get the current receive buffer
+			RecvBuffer? recvBuffer;
 			lock (_lockObject)
 			{
-#pragma warning disable CA2000 // Dispose objects before losing scope
-				_recvBufferWriters.Remove(id, out recvBufferWriter);
-#pragma warning restore CA2000 // Dispose objects before losing scope
+				if (!_recvBuffers.TryGetValue(id, out recvBuffer))
+				{
+					_logger.LogTrace("[{Tag}] Buffer {Id} has already been detached", Tag, id);
+					return;
+				}
+				recvBuffer.AddRef(); // Note: adding extra ref here
+			}
+
+			// Release the writer
+			await recvBuffer._semaphore.WaitAsync(cancellationToken);
+			try
+			{
+				recvBuffer._writer?.Dispose();
+				recvBuffer._writer = null;
 			}
-			if (recvBufferWriter != null)
+			finally
 			{
-				recvBufferWriter.MarkComplete();
-				recvBufferWriter.Dispose();
+				recvBuffer._semaphore.Release();
+				recvBuffer.Release(); // Ref added above
+			}
+
+			// Remove the buffer
+			lock (_lockObject)
+			{
+				if (_recvBuffers.Remove(id, out recvBuffer))
+				{
+					recvBuffer.Release(); // Original ref from _recvBuffers
+				}
 			}
 		}
 
+		async Task DetachAllBuffersAsync(bool recvBuffers, bool sendBuffers, CancellationToken cancellationToken)
+		{
+			int[] recvChannelIds;
+			int[] sendChannelIds;
+			lock (_lockObject)
+			{
+				_complete = true;
+				recvChannelIds = _recvBuffers.Keys.ToArray();
+				sendChannelIds = _sendBuffers.Keys.ToArray();
+			}
+
+			List<Task> tasks = new List<Task>();
+			if (recvBuffers)
+			{
+				foreach (int recvChannelId in recvChannelIds)
+				{
+					tasks.Add(DetachRecvBufferAsync(recvChannelId, cancellationToken));
+				}
+			}
+			if (sendBuffers)
+			{
+				foreach (int sendChannelId in sendChannelIds)
+				{
+					tasks.Add(DetachSendBufferAsync(sendChannelId, cancellationToken));
+				}
+			}
+			await Task.WhenAll(tasks).WaitAsync(cancellationToken);
+		}
+
 		/// <inheritdoc/>
 		public override void AttachSendBuffer(int channelId, ComputeBuffer sendBuffer)
 		{
-			ComputeBufferReader sendBufferReader = sendBuffer.CreateReader();
+			_logger.LogTrace("[{Tag}] Attaching send buffer {Id}", Tag, channelId);
 			lock (_lockObject)
 			{
-				_sendTasks.Add(channelId, Task.Run(() => SendFromBufferAsync(channelId, sendBufferReader, _cancellationSource.Token), CancellationToken.None)); // No cancellation token; need to ensure dispose runs
-				_sendBufferReaders.Add(sendBufferReader, channelId);
+				if (_sendBuffers.ContainsKey(channelId))
+				{
+					throw new InvalidOperationException($"Buffer is already attached to channel {channelId}");
+				}
+
+				ComputeBufferReader sendBufferReader = sendBuffer.CreateReader();
+				SendBuffer sendBufferInfo = new SendBuffer(sendBufferReader, (buffer, ctx) => SendFromBufferAsync(channelId, buffer, ctx));
+				_sendBuffers.Add(channelId, sendBufferInfo);
 			}
 		}
 
-		async Task SendFromBufferAsync(int channelId, ComputeBufferReader reader, CancellationToken cancellationToken)
+		[SuppressMessage("Reliability", "CA2000:Dispose objects before losing scope")]
+		async Task DetachSendBufferAsync(int channelId, CancellationToken cancellationToken)
 		{
-			using ComputeBufferReader _ = reader;
+			// Get the current send buffer state
+			SendBuffer? sendBuffer;
+			lock (_lockObject)
+			{
+				if (!_sendBuffers.TryGetValue(channelId, out sendBuffer))
+				{
+					_logger.LogWarning("No buffer is attached to channel {ChannelId}", channelId);
+					return;
+				}
+				sendBuffer.AddRef();
+			}
+
+			// Release the reader
+			await sendBuffer._semaphore.WaitAsync(cancellationToken);
+			try
+			{
+				sendBuffer._reader?.Dispose();
+				sendBuffer._reader = null;
+			}
+			finally
+			{
+				sendBuffer._semaphore.Release();
+				sendBuffer.Release(); // Added above
+			}
+
+			// Wait for the send task to complete
+			await sendBuffer._task.DisposeAsync();
+
+			// Remove the buffer from the dictionary
+			lock (_lockObject)
+			{
+				if (_sendBuffers.Remove(channelId, out sendBuffer))
+				{
+					sendBuffer.Release(); // For _sendBuffers
+				}
+			}
+		}
 
+		async Task SendFromBufferAsync(int channelId, SendBuffer sendBuffer, CancellationToken cancellationToken)
+		{
+			ComputeBufferReader reader = sendBuffer._reader!;
 			while (!cancellationToken.IsCancellationRequested)
 			{
 				ReadOnlyMemory<byte> memory = reader.GetReadBuffer();
@@ -463,12 +635,6 @@ namespace EpicGames.Horde.Compute
 				}
 				await reader.WaitToReadAsync(1, cancellationToken);
 			}
-
-			lock (_lockObject)
-			{
-				_sendTasks.Remove(channelId);
-				_sendBufferReaders.Remove(reader);
-			}
 		}
 	}
 }
