Horde: Suppress directory not found exception on Linux.

diff --git a/Engine/Source/Programs/Shared/EpicGames.Horde.Tests/ComputeSocketTests.cs b/Engine/Source/Programs/Shared/EpicGames.Horde.Tests/ComputeSocketTests.cs
--- a/Engine/Source/Programs/Shared/EpicGames.Horde.Tests/ComputeSocketTests.cs
+++ b/Engine/Source/Programs/Shared/EpicGames.Horde.Tests/ComputeSocketTests.cs
@@ -146,7 +146,10 @@ namespace EpicGames.Horde.Tests
 					await using (BundleWriter treeWriter = storage.CreateWriter())
 					{
 						FileReference file = FileReference.Combine(tempDir, "subdir/hello.txt");
-						FileReference.Delete(file);
+						if (FileReference.Exists(file))
+						{
+							FileReference.Delete(file);
+						}
 						Assert.IsFalse(FileReference.Exists(file));
 
 						byte[] data = Encoding.UTF8.GetBytes("Hello world");
