Horde: Suppress static analysis warning.

diff --git a/Engine/Source/Programs/Shared/EpicGames.Horde/Compute/ComputeSocket.cs b/Engine/Source/Programs/Shared/EpicGames.Horde/Compute/ComputeSocket.cs
--- a/Engine/Source/Programs/Shared/EpicGames.Horde/Compute/ComputeSocket.cs
+++ b/Engine/Source/Programs/Shared/EpicGames.Horde/Compute/ComputeSocket.cs
@@ -203,6 +203,7 @@ namespace EpicGames.Horde.Compute
 		{
 			public ComputeBufferReader? _reader;
 			public readonly SemaphoreSlim _semaphore = new SemaphoreSlim(1);
+			[SuppressMessage("Usage", "CA2213:Disposable fields should be disposed")]
 			public readonly BackgroundTask _task;
 			int _refCount = 1;
 
