Horde: Allocate the correct buffer size before writing output from a child process.

diff --git a/Engine/Source/Programs/Shared/EpicGames.Horde/Compute/AgentMessage.cs b/Engine/Source/Programs/Shared/EpicGames.Horde/Compute/AgentMessage.cs
--- a/Engine/Source/Programs/Shared/EpicGames.Horde/Compute/AgentMessage.cs
+++ b/Engine/Source/Programs/Shared/EpicGames.Horde/Compute/AgentMessage.cs
@@ -440,7 +440,7 @@ namespace EpicGames.Horde.Compute
 		/// </summary>
 		public static async ValueTask SendExecuteOutputAsync(this AgentMessageChannel channel, ReadOnlyMemory<byte> data, CancellationToken cancellationToken = default)
 		{
-			using IAgentMessageBuilder message = await channel.CreateMessageAsync(AgentMessageType.ExecuteOutput, cancellationToken);
+			using IAgentMessageBuilder message = await channel.CreateMessageAsync(AgentMessageType.ExecuteOutput, data.Length + 20, cancellationToken);
 			message.WriteFixedLengthBytes(data.Span);
 			message.Send();
 		}
