Horde: Invalidate compute buffer pointers when disposed. Will hopefully result in better use-after-free errors.

diff --git a/Engine/Source/Programs/Shared/EpicGames.Horde/Compute/Buffers/PooledBuffer.cs b/Engine/Source/Programs/Shared/EpicGames.Horde/Compute/Buffers/PooledBuffer.cs
--- a/Engine/Source/Programs/Shared/EpicGames.Horde/Compute/Buffers/PooledBuffer.cs
+++ b/Engine/Source/Programs/Shared/EpicGames.Horde/Compute/Buffers/PooledBuffer.cs
@@ -52,8 +52,8 @@ namespace EpicGames.Horde.Compute.Buffers
 	/// </summary>
 	class PooledBufferDetail : ComputeBufferDetail
 	{
-		readonly GCHandle _headerHandle;
-		readonly IMemoryOwner<byte>[] _chunkOwners;
+		GCHandle _headerHandle;
+		IMemoryOwner<byte>[] _chunkOwners;
 
 		readonly AsyncEvent _writerEvent = new AsyncEvent();
 		readonly AsyncEvent _readerEvent = new AsyncEvent();
@@ -86,13 +86,18 @@ namespace EpicGames.Horde.Compute.Buffers
 		/// <inheritdoc/>
 		protected override void Dispose(bool disposing)
 		{
+			base.Dispose(disposing);
+
 			if (disposing)
 			{
 				for (int idx = 0; idx < _chunkOwners.Length; idx++)
 				{
 					_chunkOwners[idx].Dispose();
 				}
+				_chunkOwners = Array.Empty<IMemoryOwner<byte>>();
+
 				_headerHandle.Free();
+				_headerHandle = default;
 			}
 		}
 
diff --git a/Engine/Source/Programs/Shared/EpicGames.Horde/Compute/Buffers/SharedMemoryBuffer.cs b/Engine/Source/Programs/Shared/EpicGames.Horde/Compute/Buffers/SharedMemoryBuffer.cs
--- a/Engine/Source/Programs/Shared/EpicGames.Horde/Compute/Buffers/SharedMemoryBuffer.cs
+++ b/Engine/Source/Programs/Shared/EpicGames.Horde/Compute/Buffers/SharedMemoryBuffer.cs
@@ -149,6 +149,8 @@ namespace EpicGames.Horde.Compute.Buffers
 		/// <inheritdoc/>
 		protected override void Dispose(bool disposing)
 		{
+			base.Dispose(disposing);
+
 			if (disposing)
 			{
 				_readerEvent.Dispose();
diff --git a/Engine/Source/Programs/Shared/EpicGames.Horde/Compute/ComputeBuffer.cs b/Engine/Source/Programs/Shared/EpicGames.Horde/Compute/ComputeBuffer.cs
--- a/Engine/Source/Programs/Shared/EpicGames.Horde/Compute/ComputeBuffer.cs
+++ b/Engine/Source/Programs/Shared/EpicGames.Horde/Compute/ComputeBuffer.cs
@@ -468,8 +468,8 @@ namespace EpicGames.Horde.Compute
 			}
 		}
 
-		readonly HeaderPtr _headerPtr;
-		readonly Memory<byte>[] _chunks;
+		HeaderPtr _headerPtr;
+		Memory<byte>[] _chunks;
 		int _refCount = 1;
 
 		/// <summary>
@@ -510,7 +510,11 @@ namespace EpicGames.Horde.Compute
 		/// <summary>
 		/// Overridable dispose method
 		/// </summary>
-		protected abstract void Dispose(bool disposing);
+		protected virtual void Dispose(bool disposing)
+		{
+			_headerPtr = default;
+			_chunks = Array.Empty<Memory<byte>>();
+		}
 
 		/// <summary>
 		/// Signals a read event
