UHT: Generating GC write barriers for UFUNCTIONs where needed.

diff --git a/Engine/Source/Programs/Shared/EpicGames.UHT/Exporters/CodeGen/UhtHeaderCodeGeneratorCppFile.cs b/Engine/Source/Programs/Shared/EpicGames.UHT/Exporters/CodeGen/UhtHeaderCodeGeneratorCppFile.cs
--- a/Engine/Source/Programs/Shared/EpicGames.UHT/Exporters/CodeGen/UhtHeaderCodeGeneratorCppFile.cs
+++ b/Engine/Source/Programs/Shared/EpicGames.UHT/Exporters/CodeGen/UhtHeaderCodeGeneratorCppFile.cs
@@ -2161,12 +2161,25 @@ namespace EpicGames.UHT.Exporters.CodeGen
 			{
 				if (parameter is UhtProperty property)
 				{
-					if (!first)
+					if (first)
+					{
+						first = false;
+					}
+					else
 					{
 						builder.Append(',');
 					}
+
+					bool needsGCBarrier = property.NeedsGCBarrierWhenPassedToFunction(function);
+					if (needsGCBarrier)
+					{
+						builder.Append("P_ARG_GC_BARRIER(");
+					}
 					builder.AppendFunctionThunkParameterArg(property);
-					first = false;
+					if (needsGCBarrier)
+					{
+						builder.Append(")");
+					}
 				}
 			}
 			return builder;
diff --git a/Engine/Source/Programs/Shared/EpicGames.UHT/Types/Properties/UhtContainerProperty.cs b/Engine/Source/Programs/Shared/EpicGames.UHT/Types/Properties/UhtContainerProperty.cs
--- a/Engine/Source/Programs/Shared/EpicGames.UHT/Types/Properties/UhtContainerProperty.cs
+++ b/Engine/Source/Programs/Shared/EpicGames.UHT/Types/Properties/UhtContainerProperty.cs
@@ -26,6 +26,12 @@ namespace EpicGames.UHT.Types
 				(ValueProperty.PropertyCaps & (UhtPropertyCaps.CanBeInstanced | UhtPropertyCaps.CanHaveConfig));
 		}
 
+		/// <inheritdoc/>
+		protected override bool NeedsGCBarrierWhenPassedToFunctionImpl(UhtFunction function)
+		{
+			return ValueProperty.NeedsGCBarrierWhenPassedToFunction(function);
+		}
+		
 		/// <inheritdoc/>
 		public override void Validate(UhtStruct outerStruct, UhtProperty outermostProperty, UhtValidationOptions options)
 		{
diff --git a/Engine/Source/Programs/Shared/EpicGames.UHT/Types/Properties/UhtMapProperty.cs b/Engine/Source/Programs/Shared/EpicGames.UHT/Types/Properties/UhtMapProperty.cs
--- a/Engine/Source/Programs/Shared/EpicGames.UHT/Types/Properties/UhtMapProperty.cs
+++ b/Engine/Source/Programs/Shared/EpicGames.UHT/Types/Properties/UhtMapProperty.cs
@@ -234,6 +234,12 @@ namespace EpicGames.UHT.Types
 			KeyProperty.ValidateDeprecated();
 		}
 
+		/// <inheritdoc/>
+		protected override bool NeedsGCBarrierWhenPassedToFunctionImpl(UhtFunction function)
+		{
+			return KeyProperty.NeedsGCBarrierWhenPassedToFunction(function) || ValueProperty.NeedsGCBarrierWhenPassedToFunction(function);
+		}
+
 		/// <inheritdoc/>
 		public override void Validate(UhtStruct outerStruct, UhtProperty outermostProperty, UhtValidationOptions options)
 		{
diff --git a/Engine/Source/Programs/Shared/EpicGames.UHT/Types/Properties/UhtObjectPropertyBase.cs b/Engine/Source/Programs/Shared/EpicGames.UHT/Types/Properties/UhtObjectPropertyBase.cs
--- a/Engine/Source/Programs/Shared/EpicGames.UHT/Types/Properties/UhtObjectPropertyBase.cs
+++ b/Engine/Source/Programs/Shared/EpicGames.UHT/Types/Properties/UhtObjectPropertyBase.cs
@@ -294,5 +294,13 @@ namespace EpicGames.UHT.Types
 			}
 		}
 		#endregion
+
+		/// <inheritdoc/>
+		protected override bool NeedsGCBarrierWhenPassedToFunctionImpl(UhtFunction function)
+		{
+			var type = this.GetType();
+			return type == typeof(UhtObjectProperty)
+				|| (type == typeof(UhtClassProperty) && !this.PropertyFlags.HasFlag(EPropertyFlags.UObjectWrapper));
+		}
 	}
 }
diff --git a/Engine/Source/Programs/Shared/EpicGames.UHT/Types/UhtProperty.cs b/Engine/Source/Programs/Shared/EpicGames.UHT/Types/UhtProperty.cs
--- a/Engine/Source/Programs/Shared/EpicGames.UHT/Types/UhtProperty.cs
+++ b/Engine/Source/Programs/Shared/EpicGames.UHT/Types/UhtProperty.cs
@@ -1951,6 +1951,30 @@ namespace EpicGames.UHT.Types
 		}
 		#endregion
 
+		#region Incremental GC Support
+		/// <summary>
+		/// Determines whether or not GC barriers need to run after passing this to a function
+		/// </summary>
+		/// <returns>True if GC barriers need to run</returns>				
+		public bool NeedsGCBarrierWhenPassedToFunction(UhtFunction function)
+		{
+			if (this.RefQualifier != UhtPropertyRefQualifier.NonConstRef)
+			{
+				return false;
+			}
+			return NeedsGCBarrierWhenPassedToFunctionImpl(function);
+		}
+		
+		/// <summary>
+		/// Customization point for subclasses for NeedsGCBarrierWhenPassedToFunction
+		/// </summary>
+		/// <returns>True if GC barriers need to run</returns>		
+		protected virtual bool NeedsGCBarrierWhenPassedToFunctionImpl(UhtFunction function)
+		{
+			return false;
+		}
+		#endregion
+		
 		#region Helper methods
 		/// <summary>
 		/// Generate a new name suffix based on the current suffix and the new suffix
diff --git a/Engine/Source/Runtime/CoreUObject/Public/UObject/ScriptMacros.h b/Engine/Source/Runtime/CoreUObject/Public/UObject/ScriptMacros.h
--- a/Engine/Source/Runtime/CoreUObject/Public/UObject/ScriptMacros.h
+++ b/Engine/Source/Runtime/CoreUObject/Public/UObject/ScriptMacros.h
@@ -111,6 +111,16 @@ enum {MAX_VARIABLE_SIZE = 0x0FFF };
 #define P_NATIVE_BEGIN { SCOPED_SCRIPT_NATIVE_TIMER(ScopedNativeCallTimer);
 #define P_NATIVE_END   }
 
+#ifndef UE_SCRIPT_ARGS_GC_BARRIER
+#define UE_SCRIPT_ARGS_GC_BARRIER 0
+#endif
+
+#if UE_SCRIPT_ARGS_GC_BARRIER
+#define P_ARG_GC_BARRIER(X) MutableView(ObjectPtrWrap(X))
+#else
+#define P_ARG_GC_BARRIER(X) X
+#endif
+
 #if UE_ENABLE_INCLUDE_ORDER_DEPRECATED_IN_5_2
 #include "CoreMinimal.h"
-#endif
\ No newline at end of file
+#endif
