Horde: Downgrade warning for data received on disconnected channels to debug.

diff --git a/Engine/Source/Programs/Shared/EpicGames.Horde/Compute/ComputeSocket.cs b/Engine/Source/Programs/Shared/EpicGames.Horde/Compute/ComputeSocket.cs
--- a/Engine/Source/Programs/Shared/EpicGames.Horde/Compute/ComputeSocket.cs
+++ b/Engine/Source/Programs/Shared/EpicGames.Horde/Compute/ComputeSocket.cs
@@ -355,7 +355,7 @@ namespace EpicGames.Horde.Compute
 		{
 			if (!await TryReadPacketAsync(transport, id, size, cancellationToken))
 			{
-				_logger.LogWarning("Discarding {Size} bytes received on compute channel {Id}", size, id);
+				_logger.LogDebug("Discarding {Size} bytes received on compute channel {Id}", size, id);
 
 				int bufferSize = Math.Min(size, 65536);
 				using (IMemoryOwner<byte> buffer = MemoryPool<byte>.Shared.Rent(bufferSize))
