Unreal Cloud DDC - Proper fix for server timing headers getting very large for certain operations (like uploading a very large oplog)

diff --git a/Engine/Source/Programs/Shared/EpicGames.AspNet/ServerTimings.cs b/Engine/Source/Programs/Shared/EpicGames.AspNet/ServerTimings.cs
--- a/Engine/Source/Programs/Shared/EpicGames.AspNet/ServerTimings.cs
+++ b/Engine/Source/Programs/Shared/EpicGames.AspNet/ServerTimings.cs
@@ -148,7 +148,7 @@ namespace EpicGames.AspNet
 			context.Response.OnStarting(() => {
 				if (serverTiming.Metrics.Any())
 				{
-					string serverTimingValue = String.Join(",", serverTiming.Metrics);
+					string serverTimingValue = String.Join(",", serverTiming.Metrics.Take(10));
 					context.Response.Headers.Append("Server-Timing", serverTimingValue);
 				}
 
