EpicGames.Core: Fix compiler warning

diff --git a/Engine/Source/Programs/Shared/EpicGames.Core.Tests/JsonObjectTests.cs b/Engine/Source/Programs/Shared/EpicGames.Core.Tests/JsonObjectTests.cs
--- a/Engine/Source/Programs/Shared/EpicGames.Core.Tests/JsonObjectTests.cs
+++ b/Engine/Source/Programs/Shared/EpicGames.Core.Tests/JsonObjectTests.cs
@@ -537,7 +537,7 @@ namespace EpicGames.Core.Tests
 		public void Constructor()
 		{
 			JsonObject obj= new JsonObject();
-			Assert.IsTrue(obj.KeyNames.Count() == 0);
+			Assert.AreEqual(0, obj.KeyNames.Count());
 		}
 
 		[TestMethod]
