EpicGames.Core: Fix some style warnings.

diff --git a/Engine/Source/Programs/Shared/EpicGames.Core.Tests/BuzHashTests.cs b/Engine/Source/Programs/Shared/EpicGames.Core.Tests/BuzHashTests.cs
--- a/Engine/Source/Programs/Shared/EpicGames.Core.Tests/BuzHashTests.cs
+++ b/Engine/Source/Programs/Shared/EpicGames.Core.Tests/BuzHashTests.cs
@@ -1,11 +1,7 @@
 // Copyright Epic Games, Inc. All Rights Reserved.
 
 using System;
-using System.Collections.Generic;
-using System.Linq;
 using System.Security.Cryptography;
-using System.Text;
-using System.Threading.Tasks;
 using Microsoft.VisualStudio.TestTools.UnitTesting;
 
 namespace EpicGames.Core.Tests
diff --git a/Engine/Source/Programs/Shared/EpicGames.Core.Tests/ILoggerPropertyTests.cs b/Engine/Source/Programs/Shared/EpicGames.Core.Tests/ILoggerPropertyTests.cs
--- a/Engine/Source/Programs/Shared/EpicGames.Core.Tests/ILoggerPropertyTests.cs
+++ b/Engine/Source/Programs/Shared/EpicGames.Core.Tests/ILoggerPropertyTests.cs
@@ -48,7 +48,7 @@ public class SpyLogger : ILogger
 
 public sealed class FakeLoggerProvider : ILoggerProvider
 {
-	private ILogger _logger;
+	private readonly ILogger _logger;
 
 	public FakeLoggerProvider(ILogger logger)
 	{
diff --git a/Engine/Source/Programs/Shared/EpicGames.Core.Tests/JsonObjectTests.cs b/Engine/Source/Programs/Shared/EpicGames.Core.Tests/JsonObjectTests.cs
--- a/Engine/Source/Programs/Shared/EpicGames.Core.Tests/JsonObjectTests.cs
+++ b/Engine/Source/Programs/Shared/EpicGames.Core.Tests/JsonObjectTests.cs
@@ -2,14 +2,12 @@
 
 using System;
 using System.Collections.Generic;
-using System.ComponentModel.DataAnnotations;
 using System.IO;
 using System.Linq;
 using System.Text;
 using System.Text.Json;
 using System.Text.Json.Serialization;
 using Microsoft.VisualStudio.TestTools.UnitTesting;
-using NuGet.Frameworks;
 
 namespace EpicGames.Core.Tests
 {
@@ -308,7 +306,7 @@ namespace EpicGames.Core.Tests
 		public void GetBoolValidField()
 		{
 			string fieldName = "bIsAlive";
-			bool correctValue = _validJsonDataSource.bIsAlive;
+			bool correctValue = _validJsonDataSource.IsAlive;
 			Assert.AreEqual(correctValue, _validTestObject.GetBoolField(fieldName));
 			bool outValue;
 			Assert.IsTrue(_validTestObject.TryGetBoolField(fieldName, out outValue));
@@ -614,7 +612,7 @@ namespace EpicGames.Core.Tests
 			enemy1.ImmuneStatusEffects.Add(StatusEffect.Frozen);
 			enemy1.ImmuneStatusEffects.Add(StatusEffect.Sleep);
 
-			enemy1.bIsAlive = true;
+			enemy1.IsAlive = true;
 			enemy1.CurrentStatusEffect = StatusEffect.Poison;
 
 			return enemy1;
@@ -666,7 +664,7 @@ namespace EpicGames.Core.Tests
 		public List<StatusEffect> ImmuneStatusEffects { get; set; } = new List<StatusEffect>();
 		public StatusEffect CurrentStatusEffect { get; set; } = StatusEffect.None;
 		public List<string> Skills { get; set; } = new List<string>();
-		public bool bIsAlive { get; set; } = false;
+		public bool IsAlive { get; set; } = false;
 	}
 
 	class EnemyAttack
diff --git a/Engine/Source/Programs/Shared/EpicGames.Core.Tests/MemorySerializerTests.cs b/Engine/Source/Programs/Shared/EpicGames.Core.Tests/MemorySerializerTests.cs
--- a/Engine/Source/Programs/Shared/EpicGames.Core.Tests/MemorySerializerTests.cs
+++ b/Engine/Source/Programs/Shared/EpicGames.Core.Tests/MemorySerializerTests.cs
@@ -1,7 +1,6 @@
 // Copyright Epic Games, Inc. All Rights Reserved.
 
 using System;
-using System.Collections.Generic;
 using System.Text;
 using Microsoft.VisualStudio.TestTools.UnitTesting;
 
