EpicGames.Core: Fix JSON tests.

diff --git a/Engine/Source/Programs/Shared/EpicGames.Core.Tests/JsonObjectTests.cs b/Engine/Source/Programs/Shared/EpicGames.Core.Tests/JsonObjectTests.cs
--- a/Engine/Source/Programs/Shared/EpicGames.Core.Tests/JsonObjectTests.cs
+++ b/Engine/Source/Programs/Shared/EpicGames.Core.Tests/JsonObjectTests.cs
@@ -305,7 +305,7 @@ namespace EpicGames.Core.Tests
 		[TestMethod]
 		public void GetBoolValidField()
 		{
-			string fieldName = "bIsAlive";
+			string fieldName = "IsAlive";
 			bool correctValue = _validJsonDataSource.IsAlive;
 			Assert.AreEqual(correctValue, _validTestObject.GetBoolField(fieldName));
 			bool outValue;
@@ -518,7 +518,7 @@ namespace EpicGames.Core.Tests
 		[TestMethod]
 		public void SetBoolField()
 		{
-			string fieldName = "bIsAlive";
+			string fieldName = "IsAlive";
 			bool setValue = false;
 			Assert.IsTrue(_validTestObject.ContainsField(fieldName));
 			_validTestObject.AddOrSetFieldValue(fieldName, setValue);
