Horde: Support running compute tasks under Wine on Linux

Compute tasks can request Wine to be used and agent can choose to respect that if conditions are right (running on Linux and the Wine executable path set).

diff --git a/Engine/Source/Programs/Horde/Horde.Agent/AgentSettings.cs b/Engine/Source/Programs/Horde/Horde.Agent/AgentSettings.cs
--- a/Engine/Source/Programs/Horde/Horde.Agent/AgentSettings.cs
+++ b/Engine/Source/Programs/Horde/Horde.Agent/AgentSettings.cs
@@ -230,6 +230,11 @@ namespace Horde.Agent
 		/// List of process names to terminate after a lease completes, but not after a job step
 		/// </summary>
 		public List<ProcessToTerminate> ProcessesToTerminate { get; } = new List<ProcessToTerminate>();
+		
+		/// <summary>
+		/// Path to Wine executable. Set to null to disable.
+		/// </summary>
+		public string? WineExecutablePath { get; set; }
 
 		/// <summary>
 		/// Whether to write step output to the logging device
diff --git a/Engine/Source/Programs/Horde/Horde.Agent/Commands/Compute/ComputeCommand.cs b/Engine/Source/Programs/Horde/Horde.Agent/Commands/Compute/ComputeCommand.cs
--- a/Engine/Source/Programs/Horde/Horde.Agent/Commands/Compute/ComputeCommand.cs
+++ b/Engine/Source/Programs/Horde/Horde.Agent/Commands/Compute/ComputeCommand.cs
@@ -146,7 +146,7 @@ namespace Horde.Agent.Commands.Compute
 		{
 			DirectoryReference sandboxDir = DirectoryReference.Combine(Program.DataDir, "Sandbox");
 
-			AgentMessageHandler worker = new AgentMessageHandler(sandboxDir, memoryCache, null, false, logger);
+			AgentMessageHandler worker = new AgentMessageHandler(sandboxDir, memoryCache, null, false, null, logger);
 			await worker.RunAsync(socket, cancellationToken);
 		}
 	}
diff --git a/Engine/Source/Programs/Horde/Horde.Agent/Commands/Compute/CppComputeCommand.cs b/Engine/Source/Programs/Horde/Horde.Agent/Commands/Compute/CppComputeCommand.cs
--- a/Engine/Source/Programs/Horde/Horde.Agent/Commands/Compute/CppComputeCommand.cs
+++ b/Engine/Source/Programs/Horde/Horde.Agent/Commands/Compute/CppComputeCommand.cs
@@ -63,7 +63,7 @@ namespace Horde.Agent.Commands.Compute
 			{
 				await channel.UploadFilesAsync("", sandbox, storage, cancellationToken);
 
-				await using (AgentManagedProcess process = await channel.ExecuteAsync(jsonComputeTask.Executable, jsonComputeTask.Arguments, jsonComputeTask.WorkingDir, jsonComputeTask.EnvVars, cancellationToken))
+				await using (AgentManagedProcess process = await channel.ExecuteAsync(jsonComputeTask.Executable, jsonComputeTask.Arguments, jsonComputeTask.WorkingDir, jsonComputeTask.EnvVars, ExecuteProcessFlags.None, cancellationToken))
 				{
 					string? line;
 					while ((line = await process.ReadLineAsync(cancellationToken)) != null)
diff --git a/Engine/Source/Programs/Horde/Horde.Agent/Leases/Handlers/ComputeHandler.cs b/Engine/Source/Programs/Horde/Horde.Agent/Leases/Handlers/ComputeHandler.cs
--- a/Engine/Source/Programs/Horde/Horde.Agent/Leases/Handlers/ComputeHandler.cs
+++ b/Engine/Source/Programs/Horde/Horde.Agent/Leases/Handlers/ComputeHandler.cs
@@ -14,6 +14,7 @@ using Horde.Agent.Services;
 using HordeCommon.Rpc.Tasks;
 using Microsoft.Extensions.Caching.Memory;
 using Microsoft.Extensions.Logging;
+using Microsoft.Extensions.Options;
 
 namespace Horde.Agent.Leases.Handlers
 {
@@ -60,15 +61,17 @@ namespace Horde.Agent.Leases.Handlers
 
 		readonly ComputeListenerService _listenerService;
 		readonly IMemoryCache _memoryCache;
+		readonly AgentSettings _settings;
 		readonly ILogger _logger;
 
 		/// <summary>
 		/// Constructor
 		/// </summary>
-		public ComputeHandler(ComputeListenerService listenerService, IMemoryCache memoryCache, ILogger<ComputeHandler> logger)
+		public ComputeHandler(ComputeListenerService listenerService, IMemoryCache memoryCache, IOptions<AgentSettings> settings, ILogger<ComputeHandler> logger)
 		{
 			_listenerService = listenerService;
 			_memoryCache = memoryCache;
+			_settings = settings.Value;
 			_logger = logger;
 		}
 
@@ -110,7 +113,7 @@ namespace Horde.Agent.Leases.Handlers
 								Dictionary<string, string?> newEnvVars = new Dictionary<string, string?>();
 								newEnvVars["UE_HORDE_SHARED_DIR"] = sharedDir.FullName;
 
-								AgentMessageHandler worker = new AgentMessageHandler(sandboxDir, _memoryCache, newEnvVars, false, _logger);
+								AgentMessageHandler worker = new AgentMessageHandler(sandboxDir, _memoryCache, newEnvVars, false, _settings.WineExecutablePath, _logger);
 								await worker.RunAsync(socket, cts.Token);
 								await socket.CloseAsync(cts.Token);
 								return LeaseResult.Success;
diff --git a/Engine/Source/Programs/Horde/Samples/RemoteClient/RemoteClientApp.cs b/Engine/Source/Programs/Horde/Samples/RemoteClient/RemoteClientApp.cs
--- a/Engine/Source/Programs/Horde/Samples/RemoteClient/RemoteClientApp.cs
+++ b/Engine/Source/Programs/Horde/Samples/RemoteClient/RemoteClientApp.cs
@@ -110,7 +110,7 @@ namespace RemoteClient
 				using ComputeChannel childProcessChannel = lease.Socket.CreateChannel(ChildProcessChannelId);
 				await using BackgroundTask tickTask = BackgroundTask.StartNew(ctx => WriteNumbersAsync(childProcessChannel, logger, ctx));
 
-				await using (AgentManagedProcess process = await channel.ExecuteAsync(executable, arguments, null, null))
+				await using (AgentManagedProcess process = await channel.ExecuteAsync(executable, arguments, null, null, ExecuteProcessFlags.None))
 				{
 					string? line;
 					while ((line = await process.ReadLineAsync()) != null)
diff --git a/Engine/Source/Programs/Shared/EpicGames.Horde.Tests/ComputeSocketTests.cs b/Engine/Source/Programs/Shared/EpicGames.Horde.Tests/ComputeSocketTests.cs
--- a/Engine/Source/Programs/Shared/EpicGames.Horde.Tests/ComputeSocketTests.cs
+++ b/Engine/Source/Programs/Shared/EpicGames.Horde.Tests/ComputeSocketTests.cs
@@ -193,7 +193,7 @@ namespace EpicGames.Horde.Tests
 		static async Task RunAgent(ComputeSocket socket, DirectoryReference tempDir, CancellationToken cancellationToken)
 		{
 			using MemoryCache memoryCache = new MemoryCache(new MemoryCacheOptions());
-			AgentMessageHandler handler = new AgentMessageHandler(tempDir, memoryCache, null, true, NullLogger.Instance);
+			AgentMessageHandler handler = new AgentMessageHandler(tempDir, memoryCache, null, true, null, NullLogger.Instance);
 			await handler.RunAsync(socket, cancellationToken);
 		}
 	}
diff --git a/Engine/Source/Programs/Shared/EpicGames.Horde/Compute/AgentMessage.cs b/Engine/Source/Programs/Shared/EpicGames.Horde/Compute/AgentMessage.cs
--- a/Engine/Source/Programs/Shared/EpicGames.Horde/Compute/AgentMessage.cs
+++ b/Engine/Source/Programs/Shared/EpicGames.Horde/Compute/AgentMessage.cs
@@ -64,7 +64,12 @@ namespace EpicGames.Horde.Compute
 		/// <summary>
 		/// Execute a process in a sandbox (Initiator -> Remote)
 		/// </summary>
-		Execute = 0x16,
+		ExecuteV1 = 0x16,
+		
+		/// <summary>
+		/// Execute a process in a sandbox (Initiator -> Remote)
+		/// </summary>
+		ExecuteV2 = 0x22,
 
 		/// <summary>
 		/// Returns output from the child process to the caller (Remote -> Initiator)
@@ -106,6 +111,24 @@ namespace EpicGames.Horde.Compute
 
 		#endregion
 	}
+	
+	/// <summary>
+	/// Flags describing how to execute a compute task process on the agent
+	/// </summary>
+	[Flags]
+	public enum ExecuteProcessFlags
+	{
+		/// <summary>
+		/// No execute flags set
+		/// </summary>
+		None = 0,
+		
+		/// <summary>
+		/// Request execution to be wrapped under Wine when running on Linux.
+		/// Agent still reserves the right to refuse it (e.g no Wine executable configured, mismatching OS etc)
+		/// </summary>
+		UseWine = 1,
+	}
 
 	/// <summary>
 	/// Standard implementation of a message
@@ -207,7 +230,8 @@ namespace EpicGames.Horde.Compute
 	/// <param name="Arguments">Arguments for the executable</param>
 	/// <param name="WorkingDir">Working directory to execute in</param>
 	/// <param name="EnvVars">Environment variables for the child process. Null values unset variables.</param>
-	public record struct ExecuteProcessMessage(string Executable, IReadOnlyList<string> Arguments, string? WorkingDir, IReadOnlyDictionary<string, string?> EnvVars);
+	/// <param name="Flags">Additional execution flags</param>
+	public record struct ExecuteProcessMessage(string Executable, IReadOnlyList<string> Arguments, string? WorkingDir, IReadOnlyDictionary<string, string?> EnvVars, ExecuteProcessFlags Flags);
 
 	/// <summary>
 	/// Response from executing a child process
@@ -402,7 +426,7 @@ namespace EpicGames.Horde.Compute
 		}
 
 		/// <summary>
-		/// Executes a remote process
+		/// Executes a remote process (using ExecuteV1)
 		/// </summary>
 		/// <param name="channel">Current channel</param>
 		/// <param name="executable">Executable to run, relative to the sandbox root</param>
@@ -412,12 +436,36 @@ namespace EpicGames.Horde.Compute
 		/// <param name="cancellationToken">Cancellation token for the operation</param>
 		public static async Task<AgentManagedProcess> ExecuteAsync(this AgentMessageChannel channel, string executable, IReadOnlyList<string> arguments, string? workingDir, IReadOnlyDictionary<string, string?>? envVars, CancellationToken cancellationToken = default)
 		{
-			using (IAgentMessageBuilder request = await channel.CreateMessageAsync(AgentMessageType.Execute, cancellationToken))
+			using (IAgentMessageBuilder request = await channel.CreateMessageAsync(AgentMessageType.ExecuteV1, cancellationToken))
+			{
+				request.WriteString(executable);
+				request.WriteList(arguments, MemoryWriterExtensions.WriteString);
+				request.WriteOptionalString(workingDir);
+				request.WriteDictionary(envVars ?? new Dictionary<string, string?>(), MemoryWriterExtensions.WriteString, MemoryWriterExtensions.WriteOptionalString);
+				request.Send();
+			}
+			return new AgentManagedProcess(channel);
+		}
+		
+		/// <summary>
+		/// Executes a remote process (using ExecuteV2)
+		/// </summary>
+		/// <param name="channel">Current channel</param>
+		/// <param name="executable">Executable to run, relative to the sandbox root</param>
+		/// <param name="arguments">Arguments for the child process</param>
+		/// <param name="workingDir">Working directory for the process</param>
+		/// <param name="envVars">Environment variables for the child process</param>
+		/// <param name="flags">Additional execution flags</param>
+		/// <param name="cancellationToken">Cancellation token for the operation</param>
+		public static async Task<AgentManagedProcess> ExecuteAsync(this AgentMessageChannel channel, string executable, IReadOnlyList<string> arguments, string? workingDir, IReadOnlyDictionary<string, string?>? envVars, ExecuteProcessFlags flags = ExecuteProcessFlags.None, CancellationToken cancellationToken = default)
+		{
+			using (IAgentMessageBuilder request = await channel.CreateMessageAsync(AgentMessageType.ExecuteV2, cancellationToken))
 			{
 				request.WriteString(executable);
 				request.WriteList(arguments, MemoryWriterExtensions.WriteString);
 				request.WriteOptionalString(workingDir);
 				request.WriteDictionary(envVars ?? new Dictionary<string, string?>(), MemoryWriterExtensions.WriteString, MemoryWriterExtensions.WriteOptionalString);
+				request.WriteInt32((int)flags);
 				request.Send();
 			}
 			return new AgentManagedProcess(channel);
@@ -426,13 +474,26 @@ namespace EpicGames.Horde.Compute
 		/// <summary>
 		/// Parses a message as a <see cref="ExecuteProcessMessage"/>
 		/// </summary>
-		public static ExecuteProcessMessage ParseExecuteProcessMessage(this AgentMessage message)
+		public static ExecuteProcessMessage ParseExecuteProcessV1Message(this AgentMessage message)
+		{
+			string executable = message.ReadString();
+			List<string> arguments = message.ReadList(MemoryReaderExtensions.ReadString);
+			string? workingDir = message.ReadOptionalString();
+			Dictionary<string, string?> envVars = message.ReadDictionary(MemoryReaderExtensions.ReadString, MemoryReaderExtensions.ReadOptionalString);
+			return new ExecuteProcessMessage(executable, arguments, workingDir, envVars, ExecuteProcessFlags.None);
+		}
+		
+		/// <summary>
+		/// Parses a message as a <see cref="ExecuteProcessMessage"/>
+		/// </summary>
+		public static ExecuteProcessMessage ParseExecuteProcessV2Message(this AgentMessage message)
 		{
 			string executable = message.ReadString();
 			List<string> arguments = message.ReadList(MemoryReaderExtensions.ReadString);
 			string? workingDir = message.ReadOptionalString();
 			Dictionary<string, string?> envVars = message.ReadDictionary(MemoryReaderExtensions.ReadString, MemoryReaderExtensions.ReadOptionalString);
-			return new ExecuteProcessMessage(executable, arguments, workingDir, envVars);
+			ExecuteProcessFlags flags = (ExecuteProcessFlags)message.ReadInt32();
+			return new ExecuteProcessMessage(executable, arguments, workingDir, envVars, flags);
 		}
 
 		/// <summary>
diff --git a/Engine/Source/Programs/Shared/EpicGames.Horde/Compute/AgentMessageHandler.cs b/Engine/Source/Programs/Shared/EpicGames.Horde/Compute/AgentMessageHandler.cs
--- a/Engine/Source/Programs/Shared/EpicGames.Horde/Compute/AgentMessageHandler.cs
+++ b/Engine/Source/Programs/Shared/EpicGames.Horde/Compute/AgentMessageHandler.cs
@@ -26,6 +26,7 @@ namespace EpicGames.Horde.Compute
 		readonly IMemoryCache _memoryCache;
 		readonly Dictionary<string, string?> _envVars;
 		readonly bool _executeInProcess;
+		readonly string? _wineExecutablePath;
 		readonly ILogger _logger;
 
 		/// <summary>
@@ -35,13 +36,15 @@ namespace EpicGames.Horde.Compute
 		/// <param name="memoryCache">Cache for nodes read from storage</param>
 		/// <param name="envVars">Environment variables to set for any child processes</param>
 		/// <param name="executeInProcess">Whether to execute any external assemblies in the current process</param>
+		/// <param name="wineExecutablePath">Path to Wine executable. If null, execution under Wine is disabled</param>
 		/// <param name="logger">Logger for diagnostics</param>
-		public AgentMessageHandler(DirectoryReference sandboxDir, IMemoryCache memoryCache, Dictionary<string, string?>? envVars, bool executeInProcess, ILogger logger)
+		public AgentMessageHandler(DirectoryReference sandboxDir, IMemoryCache memoryCache, Dictionary<string, string?>? envVars, bool executeInProcess, string? wineExecutablePath, ILogger logger)
 		{
 			_sandboxDir = sandboxDir;
 			_memoryCache = memoryCache;
 			_envVars = envVars ?? new Dictionary<string, string?>();
 			_executeInProcess = executeInProcess;
+			_wineExecutablePath = wineExecutablePath;
 			_logger = logger;
 		}
 
@@ -93,10 +96,16 @@ namespace EpicGames.Horde.Compute
 								DeleteFiles(deleteFiles.Filter);
 							}
 							break;
-						case AgentMessageType.Execute:
+						case AgentMessageType.ExecuteV1:
 							{
-								ExecuteProcessMessage executeProcess = message.ParseExecuteProcessMessage();
-								await ExecuteProcessAsync(socket, channel, executeProcess.Executable, executeProcess.Arguments, executeProcess.WorkingDir, executeProcess.EnvVars, cancellationToken);
+								ExecuteProcessMessage executeProcess = message.ParseExecuteProcessV1Message();
+								await ExecuteProcessAsync(socket, channel, executeProcess.Executable, executeProcess.Arguments, executeProcess.WorkingDir, executeProcess.EnvVars, executeProcess.Flags, cancellationToken);
+							}
+							break;
+						case AgentMessageType.ExecuteV2:
+							{
+								ExecuteProcessMessage executeProcess = message.ParseExecuteProcessV2Message();
+								await ExecuteProcessAsync(socket, channel, executeProcess.Executable, executeProcess.Arguments, executeProcess.WorkingDir, executeProcess.EnvVars, executeProcess.Flags, cancellationToken);
 							}
 							break;
 						case AgentMessageType.XorRequest:
@@ -159,17 +168,17 @@ namespace EpicGames.Horde.Compute
 			}
 		}
 
-		async Task ExecuteProcessAsync(ComputeSocket socket, AgentMessageChannel channel, string executable, IReadOnlyList<string> arguments, string? workingDir, IReadOnlyDictionary<string, string?>? envVars, CancellationToken cancellationToken)
+		async Task ExecuteProcessAsync(ComputeSocket socket, AgentMessageChannel channel, string executable, IReadOnlyList<string> arguments, string? workingDir, IReadOnlyDictionary<string, string?>? envVars, ExecuteProcessFlags flags, CancellationToken cancellationToken)
 		{
 			try
 			{
 				if (RuntimeInformation.IsOSPlatform(OSPlatform.Windows))
 				{
-					await ExecuteProcessWindowsAsync(socket, channel, executable, arguments, workingDir, envVars, cancellationToken);
+					await ExecuteProcessWindowsAsync(socket, channel, executable, arguments, workingDir, envVars, flags, cancellationToken);
 				}
 				else
 				{
-					await ExecuteProcessInternalAsync(channel, executable, arguments, workingDir, envVars, cancellationToken);
+					await ExecuteProcessInternalAsync(channel, executable, arguments, workingDir, envVars, flags, cancellationToken);
 				}
 			}
 			catch (Exception ex)
@@ -178,7 +187,7 @@ namespace EpicGames.Horde.Compute
 			}
 		}
 
-		async Task ExecuteProcessWindowsAsync(ComputeSocket socket, AgentMessageChannel channel, string executable, IReadOnlyList<string> arguments, string? workingDir, IReadOnlyDictionary<string, string?>? envVars, CancellationToken cancellationToken)
+		async Task ExecuteProcessWindowsAsync(ComputeSocket socket, AgentMessageChannel channel, string executable, IReadOnlyList<string> arguments, string? workingDir, IReadOnlyDictionary<string, string?>? envVars, ExecuteProcessFlags flags, CancellationToken cancellationToken)
 		{
 			Dictionary<string, string?> newEnvVars = new Dictionary<string, string?>(_envVars);
 			if (envVars != null)
@@ -199,7 +208,7 @@ namespace EpicGames.Horde.Compute
 					_logger.LogInformation("Launching {Executable} {Arguments}", CommandLineArguments.Quote(executable), CommandLineArguments.Join(arguments));
 					try
 					{
-						await ExecuteProcessInternalAsync(channel, executable, arguments, workingDir, newEnvVars, cancellationToken);
+						await ExecuteProcessInternalAsync(channel, executable, arguments, workingDir, newEnvVars, flags, cancellationToken);
 					}
 					finally
 					{
@@ -282,7 +291,7 @@ namespace EpicGames.Horde.Compute
 			}
 		}
 
-		async Task ExecuteProcessInternalAsync(AgentMessageChannel channel, string executable, IReadOnlyList<string> arguments, string? workingDir, IReadOnlyDictionary<string, string?>? envVars, CancellationToken cancellationToken)
+		async Task ExecuteProcessInternalAsync(AgentMessageChannel channel, string executable, IReadOnlyList<string> arguments, string? workingDir, IReadOnlyDictionary<string, string?>? envVars, ExecuteProcessFlags flags, CancellationToken cancellationToken)
 		{
 			string resolvedExecutable = FileReference.Combine(_sandboxDir, executable).FullName;
 			string resolvedWorkingDir = DirectoryReference.Combine(_sandboxDir, workingDir ?? String.Empty).FullName;
@@ -329,6 +338,13 @@ namespace EpicGames.Horde.Compute
 			{
 				string resolvedCommandLine = CommandLineArguments.Join(arguments);
 
+				if (flags.HasFlag(ExecuteProcessFlags.UseWine) && _wineExecutablePath != null)
+				{
+					// Path to the original Windows executable is prepended to the argument list so Wine can run it
+					resolvedCommandLine = CommandLineArguments.Join(new[] { resolvedExecutable }.Concat(arguments).ToList());
+					resolvedExecutable = _wineExecutablePath;
+				}
+
 				Dictionary<string, string> resolvedEnvVars = ManagedProcess.GetCurrentEnvVars();
 				if (envVars != null)
 				{
diff --git a/Engine/Source/Programs/Shared/EpicGames.Horde/Compute/Clients/LocalComputeClient.cs b/Engine/Source/Programs/Shared/EpicGames.Horde/Compute/Clients/LocalComputeClient.cs
--- a/Engine/Source/Programs/Shared/EpicGames.Horde/Compute/Clients/LocalComputeClient.cs
+++ b/Engine/Source/Programs/Shared/EpicGames.Horde/Compute/Clients/LocalComputeClient.cs
@@ -82,7 +82,7 @@ namespace EpicGames.Horde.Compute.Clients
 
 			await using (RemoteComputeSocket socket = new RemoteComputeSocket(new TcpTransport(tcpSocket), ComputeSocketEndpoint.Remote, logger))
 			{
-				AgentMessageHandler worker = new AgentMessageHandler(sandboxDir, memoryCache, null, executeInProcess, logger);
+				AgentMessageHandler worker = new AgentMessageHandler(sandboxDir, memoryCache, null, executeInProcess, null, logger);
 				await worker.RunAsync(socket, cancellationToken);
 				await socket.CloseAsync(cancellationToken);
 			}
