Horde: Preserve executable permission when copying files to temp storage on Mac and Linux.

diff --git a/Engine/Source/Programs/Shared/EpicGames.Horde/Storage/Nodes/DirectoryNode.cs b/Engine/Source/Programs/Shared/EpicGames.Horde/Storage/Nodes/DirectoryNode.cs
--- a/Engine/Source/Programs/Shared/EpicGames.Horde/Storage/Nodes/DirectoryNode.cs
+++ b/Engine/Source/Programs/Shared/EpicGames.Horde/Storage/Nodes/DirectoryNode.cs
@@ -14,6 +14,7 @@ using System.IO;
 using System.IO.Compression;
 using System.IO.Pipelines;
 using System.Linq;
+using System.Runtime.InteropServices;
 using System.Threading;
 using System.Threading.Tasks;
 
@@ -517,7 +518,26 @@ namespace EpicGames.Horde.Storage.Nodes
 			for(int idx = 0; idx < files.Count; idx++)
 			{
 				FileInfo file = files[idx];
-				FileEntry entry = new FileEntry(file.Name, FileEntryFlags.None, file.Length, chunks[idx]);
+
+				FileEntryFlags flags = FileEntryFlags.None;
+				if (RuntimeInformation.IsOSPlatform(OSPlatform.Linux))
+				{
+					int mode = FileUtils.GetFileMode_Linux(file.FullName);
+					if ((mode & ((1 << 0) | (1 << 3) | (1 << 6))) != 0)
+					{
+						flags |= FileEntryFlags.Executable;
+					}
+				}
+				else if (RuntimeInformation.IsOSPlatform(OSPlatform.OSX))
+				{
+					int mode = FileUtils.GetFileMode_Mac(file.FullName);
+					if ((mode & ((1 << 0) | (1 << 3) | (1 << 6))) != 0)
+					{
+						flags |= FileEntryFlags.Executable;
+					}
+				}
+
+				FileEntry entry = new FileEntry(file.Name, flags, file.Length, chunks[idx]);
 				update.AddFile(new FileReference(file).MakeRelativeTo(baseDir), entry);
 			}
 
@@ -707,8 +727,7 @@ namespace EpicGames.Horde.Storage.Nodes
 			foreach (FileEntry fileEntry in _nameToFileEntry.Values)
 			{
 				FileInfo fileInfo = new FileInfo(Path.Combine(directoryInfo.FullName, fileEntry.Name.ToString()));
-				ChunkedDataNode fileNode = await fileEntry.ExpandAsync(cancellationToken);
-				tasks.Add(Task.Run(() => fileNode.CopyToFileAsync(fileInfo, cancellationToken), cancellationToken));
+				tasks.Add(Task.Run(() => fileEntry.CopyToFileAsync(fileInfo, cancellationToken), cancellationToken));
 			}
 			foreach (DirectoryEntry directoryEntry in _nameToDirectoryEntry.Values)
 			{
diff --git a/Engine/Source/Programs/Shared/EpicGames.Horde/Storage/Nodes/FileEntry.cs b/Engine/Source/Programs/Shared/EpicGames.Horde/Storage/Nodes/FileEntry.cs
--- a/Engine/Source/Programs/Shared/EpicGames.Horde/Storage/Nodes/FileEntry.cs
+++ b/Engine/Source/Programs/Shared/EpicGames.Horde/Storage/Nodes/FileEntry.cs
@@ -3,6 +3,7 @@
 using System;
 using System.IO;
 using System.IO.Pipelines;
+using System.Runtime.InteropServices;
 using System.Threading;
 using System.Threading.Tasks;
 using EpicGames.Core;
@@ -155,6 +156,18 @@ namespace EpicGames.Horde.Storage.Nodes
 		{
 			ChunkedDataNode node = await ExpandAsync(cancellationToken);
 			await node.CopyToFileAsync(file, cancellationToken);
+
+			if ((Flags & FileEntryFlags.Executable) != 0)
+			{
+				if (RuntimeInformation.IsOSPlatform(OSPlatform.Linux))
+				{
+					FileUtils.SetFileMode_Linux(file.FullName, 0b_111_111_111);
+				}
+				else if (RuntimeInformation.IsOSPlatform(OSPlatform.OSX))
+				{
+					FileUtils.SetFileMode_Mac(file.FullName, 0b_111_111_111);
+				}
+			}
 		}
 
 		/// <inheritdoc/>
