Horde: Various static analyzer fixes.

diff --git a/Engine/Source/Programs/Horde/Horde.Server/Ddc/ContentIdService.cs b/Engine/Source/Programs/Horde/Horde.Server/Ddc/ContentIdService.cs
--- a/Engine/Source/Programs/Horde/Horde.Server/Ddc/ContentIdService.cs
+++ b/Engine/Source/Programs/Horde/Horde.Server/Ddc/ContentIdService.cs
@@ -10,12 +10,12 @@ namespace Horde.Server.Ddc
 	{
 		public Task<BlobId[]?> ResolveAsync(NamespaceId ns, ContentId contentId, bool mustBeContentId = false)
 		{
-			throw new NotImplementedException();
+			return Task.FromResult<BlobId[]?>(null);
 		}
 
 		public Task PutAsync(NamespaceId ns, ContentId contentId, BlobId blobIdentifier, int contentWeight)
 		{
-			throw new NotImplementedException();
+			return Task.CompletedTask;
 		}
 	}
 }
diff --git a/Engine/Source/Programs/Horde/Horde.Server/Ddc/RefId.cs b/Engine/Source/Programs/Horde/Horde.Server/Ddc/RefId.cs
--- a/Engine/Source/Programs/Horde/Horde.Server/Ddc/RefId.cs
+++ b/Engine/Source/Programs/Horde/Horde.Server/Ddc/RefId.cs
@@ -3,6 +3,7 @@
 using System;
 using System.ComponentModel;
 using System.Globalization;
+using System.Text;
 using System.Text.Json;
 using System.Text.Json.Serialization;
 using EpicGames.Core;
@@ -50,6 +51,11 @@ namespace Horde.Server.Ddc
 
 		/// <inheritdoc cref="IoHash.op_Inequality"/>
 		public static bool operator !=(RefId left, RefId right) => left.Hash != right.Hash;
+
+		/// <summary>
+		/// Creates a RefId by hashing a string
+		/// </summary>
+		public static RefId FromName(string name) => new RefId(IoHash.Compute(Encoding.UTF8.GetBytes(name)));
 	}
 
 	/// <summary>
diff --git a/Engine/Source/Programs/Shared/EpicGames.Core.Tests/JsonTracerTests.cs b/Engine/Source/Programs/Shared/EpicGames.Core.Tests/JsonTracerTests.cs
--- a/Engine/Source/Programs/Shared/EpicGames.Core.Tests/JsonTracerTests.cs
+++ b/Engine/Source/Programs/Shared/EpicGames.Core.Tests/JsonTracerTests.cs
@@ -18,91 +18,103 @@ namespace EpicGames.Core.Tests
 		[TestMethod]
 		public void Run()
 		{
-			JsonTracer Tracer = new JsonTracer();
-			using (IScope Scope1 = Tracer.BuildSpan("bogusOp1").StartActive()) { }
-			using (IScope Scope2 = Tracer.BuildSpan("bogusOp2").StartActive()) { }
+			JsonTracer tracer = new JsonTracer();
+			using (IScope scope1 = tracer.BuildSpan("bogusOp1").StartActive())
+			{
+			}
+			using (IScope scope2 = tracer.BuildSpan("bogusOp2").StartActive())
+			{
+			}
 
-			Assert.AreEqual(2, Tracer.GetFinishedSpans().Count);
+			Assert.AreEqual(2, tracer.GetFinishedSpans().Count);
 		}
 		
 		[TestMethod]
 		public void SpansWithTags()
 		{
-			JsonTracer Tracer = new JsonTracer();
-			using (IScope Scope1 = Tracer.BuildSpan("bogusOp1").WithTag("age", 94).StartActive()) { }
-			using (IScope Scope2 = Tracer.BuildSpan("bogusOp1").WithTag("foo", "bar").StartActive()) { }
+			JsonTracer tracer = new JsonTracer();
+			using (IScope scope1 = tracer.BuildSpan("bogusOp1").WithTag("age", 94).StartActive())
+			{ 
+			}
+			using (IScope scope2 = tracer.BuildSpan("bogusOp1").WithTag("foo", "bar").StartActive())
+			{ 
+			}
 
-			List<JsonTracerSpan> Spans = Tracer.GetFinishedSpans();
+			List<JsonTracerSpan> spans = tracer.GetFinishedSpans();
 
-			Assert.AreEqual(2, Spans.Count);
-			Assert.AreEqual("bogusOp1", Spans[0].OperationName);
+			Assert.AreEqual(2, spans.Count);
+			Assert.AreEqual("bogusOp1", spans[0].OperationName);
 		}
 
 		[TestMethod]
 		public void SerializeToJson()
 		{
-			JsonTracer Tracer = new JsonTracer();
-			using (IScope Scope1 = Tracer.BuildSpan("bogusOp1").WithTag("age", 94).WithTag("foo", "bar").StartActive()) { }
-			using (IScope Scope2 = Tracer.BuildSpan("bogusOp2").WithTag("Service", "myService").WithTag("Resource", "myResource").StartActive()) { }
+			JsonTracer tracer = new JsonTracer();
+			using (IScope scope1 = tracer.BuildSpan("bogusOp1").WithTag("age", 94).WithTag("foo", "bar").StartActive()) 
+			{
+			}
+			using (IScope scope2 = tracer.BuildSpan("bogusOp2").WithTag("Service", "myService").WithTag("Resource", "myResource").StartActive()) 
+			{ 
+			}
 
-			using StringWriter Sw = new StringWriter();
-			using JsonWriter Jw = new JsonWriter(Sw);
-			Tracer.GetFinishedSpansAsJson(Jw);
-			byte[] Data = Encoding.ASCII.GetBytes(Sw.ToString());
+			using StringWriter sw = new StringWriter();
+			using JsonWriter jw = new JsonWriter(sw);
+			tracer.GetFinishedSpansAsJson(jw);
+			byte[] data = Encoding.ASCII.GetBytes(sw.ToString());
 
-			using (JsonDocument Document = JsonDocument.Parse(Data.AsMemory()))
+			using (JsonDocument document = JsonDocument.Parse(data.AsMemory()))
 			{
-				JsonElement Root = Document.RootElement;
-				JsonElement SpansElement = Root.GetProperty("Spans");
-				List<JsonElement> SpansList = SpansElement.EnumerateArray().ToList();
-				Console.WriteLine(SpansList[0]);
-				Assert.AreEqual("bogusOp1", SpansList[0].GetProperty("Name").GetString());
-				Assert.IsNotNull(SpansList[0].GetProperty("StartTime").GetDateTime());
-				Assert.IsNotNull(SpansList[0].GetProperty("FinishTime").GetDateTime());
-				Assert.AreEqual("myService", SpansList[1].GetProperty("Service").GetString());
-				Assert.AreEqual("myResource", SpansList[1].GetProperty("Resource").GetString());
+				JsonElement root = document.RootElement;
+				JsonElement spansElement = root.GetProperty("Spans");
+				List<JsonElement> spansList = spansElement.EnumerateArray().ToList();
+				Console.WriteLine(spansList[0]);
+				Assert.AreEqual("bogusOp1", spansList[0].GetProperty("Name").GetString());
+				Assert.IsNotNull(spansList[0].GetProperty("StartTime").GetDateTime());
+				Assert.IsNotNull(spansList[0].GetProperty("FinishTime").GetDateTime());
+				Assert.AreEqual("myService", spansList[1].GetProperty("Service").GetString());
+				Assert.AreEqual("myResource", spansList[1].GetProperty("Resource").GetString());
 			}
 		}
 
 		[TestMethod]
 		public void SpansWithParallelFor()
 		{
-			JsonTracer Tracer = new JsonTracer();
-			using (IScope ParentScope = Tracer.BuildSpan("Parent").StartActive())
+			JsonTracer tracer = new JsonTracer();
+			using (IScope parentScope = tracer.BuildSpan("Parent").StartActive())
 			{
-				Parallel.For(0, 100, Index =>
+				Parallel.For(0, 100, index =>
 				{
-					using (IScope ChildScope = Tracer.BuildSpan("Child").StartActive())
+					using (IScope childScope = tracer.BuildSpan("Child").StartActive())
 					{
-						Assert.AreEqual(Tracer.ActiveSpan, ChildScope.Span);
-						Assert.AreEqual(((JsonTracerSpan)ChildScope.Span).ParentId, ((JsonTracerSpan)ParentScope.Span).Context.SpanId);
+						Assert.AreEqual(tracer.ActiveSpan, childScope.Span);
+						Assert.AreEqual(((JsonTracerSpan)childScope.Span).ParentId, ((JsonTracerSpan)parentScope.Span).Context.SpanId);
 					}
-					Assert.AreEqual(Tracer.ActiveSpan, ParentScope.Span);
+					Assert.AreEqual(tracer.ActiveSpan, parentScope.Span);
 				});
-				Assert.AreEqual(Tracer.ActiveSpan, ParentScope.Span);
+				Assert.AreEqual(tracer.ActiveSpan, parentScope.Span);
 			}
-			Assert.AreEqual(Tracer.ActiveSpan, null);
+			Assert.AreEqual(tracer.ActiveSpan, null);
 		}
 
-		private static async Task<bool> AsyncSpanTest(JsonTracer Tracer, IScope ParentScope)
+		private static async Task<bool> AsyncSpanTest(JsonTracer tracer, IScope parentScope)
 		{
-			for (int Index = 0; Index < 100; Index++)
+			for (int index = 0; index < 100; index++)
 			{
-				using (IScope AsyncChildScope = Tracer.BuildSpan("AsyncChild").StartActive())
+				using (IScope asyncChildScope = tracer.BuildSpan("AsyncChild").StartActive())
 				{
 					await Task.Run(() => 
 					{
-						using (IScope TaskChildScope = Tracer.BuildSpan("TaskChild").StartActive())
+						using (IScope taskChildScope = tracer.BuildSpan("TaskChild").StartActive())
 						{
-							Assert.AreEqual(Tracer.ActiveSpan, TaskChildScope.Span);
-							Assert.AreEqual(((JsonTracerSpan)TaskChildScope.Span).ParentId, ((JsonTracerSpan)AsyncChildScope.Span).Context.SpanId);
+							Assert.AreEqual(tracer.ActiveSpan, taskChildScope.Span);
+							Assert.AreEqual(((JsonTracerSpan)taskChildScope.Span).ParentId, ((JsonTracerSpan)asyncChildScope.Span).Context.SpanId);
 							return true;
 						}
 					});
-					Assert.AreEqual(Tracer.ActiveSpan, AsyncChildScope.Span);
-					Assert.AreEqual(((JsonTracerSpan)AsyncChildScope.Span).ParentId, ((JsonTracerSpan)ParentScope.Span).Context.SpanId);
+					Assert.AreEqual(tracer.ActiveSpan, asyncChildScope.Span);
+					Assert.AreEqual(((JsonTracerSpan)asyncChildScope.Span).ParentId, ((JsonTracerSpan)parentScope.Span).Context.SpanId);
 				}
-				Assert.AreEqual(Tracer.ActiveSpan, ParentScope.Span);
+				Assert.AreEqual(tracer.ActiveSpan, parentScope.Span);
 			}
 			return true;
 		}
@@ -110,44 +122,44 @@ namespace EpicGames.Core.Tests
 		[TestMethod]
 		public void SpansWithAsync()
 		{
-			JsonTracer Tracer = new JsonTracer();
-			using (IScope ParentScope = Tracer.BuildSpan("Parent").StartActive())
+			JsonTracer tracer = new JsonTracer();
+			using (IScope parentScope = tracer.BuildSpan("Parent").StartActive())
 			{
-				Assert.AreEqual(Tracer.ActiveSpan, ParentScope.Span);
-				AsyncSpanTest(Tracer, ParentScope).Wait();
-				Assert.AreEqual(Tracer.ActiveSpan, ParentScope.Span);
+				Assert.AreEqual(tracer.ActiveSpan, parentScope.Span);
+				AsyncSpanTest(tracer, parentScope).Wait();
+				Assert.AreEqual(tracer.ActiveSpan, parentScope.Span);
 			}
 
-			Assert.AreEqual(Tracer.ActiveSpan, null);
+			Assert.AreEqual(tracer.ActiveSpan, null);
 		}
 
-		private static async Task ForEachAsyncSpanTest(JsonTracer Tracer, IScope ParentScope)
+		private static async Task ForEachAsyncSpanTest(JsonTracer tracer, IScope parentScope)
 		{
-			int[] Values = Enumerable.Range(0, 10000).ToArray();
-			await Parallel.ForEachAsync(Values, new ParallelOptions { MaxDegreeOfParallelism = 10000 }, async (i, token) =>
+			int[] values = Enumerable.Range(0, 10000).ToArray();
+			await Parallel.ForEachAsync(values, new ParallelOptions { MaxDegreeOfParallelism = 10000 }, async (i, token) =>
 			{
-				using (IScope ChildScope = Tracer.BuildSpan("Child").StartActive())
+				using (IScope childScope = tracer.BuildSpan("Child").StartActive())
 				{
-					Assert.AreEqual(Tracer.ActiveSpan, ChildScope.Span);
-					Assert.AreEqual(((JsonTracerSpan)ChildScope.Span).ParentId, ((JsonTracerSpan)ParentScope.Span).Context.SpanId);
+					Assert.AreEqual(tracer.ActiveSpan, childScope.Span);
+					Assert.AreEqual(((JsonTracerSpan)childScope.Span).ParentId, ((JsonTracerSpan)parentScope.Span).Context.SpanId);
 					await Task.Delay(10, token);
 				}
-				Assert.AreEqual(Tracer.ActiveSpan, ParentScope.Span);
+				Assert.AreEqual(tracer.ActiveSpan, parentScope.Span);
 			});
 		}
 
 		[TestMethod]
 		public void SpansWithParallelForEachAsync()
 		{
-			JsonTracer Tracer = new JsonTracer();
-			using (IScope ParentScope = Tracer.BuildSpan("Parent").StartActive())
+			JsonTracer tracer = new JsonTracer();
+			using (IScope parentScope = tracer.BuildSpan("Parent").StartActive())
 			{
-				Assert.AreEqual(Tracer.ActiveSpan, ParentScope.Span);
-				ForEachAsyncSpanTest(Tracer, ParentScope).Wait();
-				Assert.AreEqual(Tracer.ActiveSpan, ParentScope.Span);
+				Assert.AreEqual(tracer.ActiveSpan, parentScope.Span);
+				ForEachAsyncSpanTest(tracer, parentScope).Wait();
+				Assert.AreEqual(tracer.ActiveSpan, parentScope.Span);
 			}
 
-			Assert.AreEqual(Tracer.ActiveSpan, null);
+			Assert.AreEqual(tracer.ActiveSpan, null);
 		}
 	}
 }
diff --git a/Engine/Source/Programs/Shared/EpicGames.Horde/Storage/Nodes/DdcRefNode.cs b/Engine/Source/Programs/Shared/EpicGames.Horde/Storage/Nodes/DdcRefNode.cs
--- a/Engine/Source/Programs/Shared/EpicGames.Horde/Storage/Nodes/DdcRefNode.cs
+++ b/Engine/Source/Programs/Shared/EpicGames.Horde/Storage/Nodes/DdcRefNode.cs
@@ -18,7 +18,7 @@ namespace Horde.Server.Ddc
 		public IoHash RootHash { get; }
 
 		/// <summary>
-		/// References to attachments
+		/// References to attachments. We embed this in the ref node to ensure any aliased blobs have a hard reference from the root.
 		/// </summary>
 		public List<(IoHash Hash, BlobHandle Handle)> References { get; }
 
diff --git a/Engine/Source/Programs/Shared/EpicGames.Serialization.Tests/EngineTests.cs b/Engine/Source/Programs/Shared/EpicGames.Serialization.Tests/EngineTests.cs
--- a/Engine/Source/Programs/Shared/EpicGames.Serialization.Tests/EngineTests.cs
+++ b/Engine/Source/Programs/Shared/EpicGames.Serialization.Tests/EngineTests.cs
@@ -244,16 +244,16 @@ namespace EpicGames.Serialization.Tests
 
 			// Test CbField(Object, Empty)
 			{
-				CbObject @object = CbObject.Empty;
-				TestIntObject(@object, 0, 1);
+				CbObject obj = CbObject.Empty;
+				TestIntObject(obj, 0, 1);
 
 				// Find fields that do not exist.
-				Assert.IsFalse(@object.Find("Field").HasValue());
-				Assert.IsFalse(@object.FindIgnoreCase("Field").HasValue());
-				Assert.IsFalse(@object["Field"].HasValue());
+				Assert.IsFalse(obj.Find("Field").HasValue());
+				Assert.IsFalse(obj.FindIgnoreCase("Field").HasValue());
+				Assert.IsFalse(obj["Field"].HasValue());
 
 				// Advance an iterator past the last field.
-				CbFieldIterator it = @object.CreateIterator();
+				CbFieldIterator it = obj.CreateIterator();
 				Assert.IsFalse((bool)it);
 				Assert.IsTrue(!it);
 				for (int count = 16; count > 0; --count)
