Horde: Fix exception when serializing empty files to temp storage.

diff --git a/Engine/Source/Programs/Shared/EpicGames.Horde.Tests/ChunkingTests.cs b/Engine/Source/Programs/Shared/EpicGames.Horde.Tests/ChunkingTests.cs
--- a/Engine/Source/Programs/Shared/EpicGames.Horde.Tests/ChunkingTests.cs
+++ b/Engine/Source/Programs/Shared/EpicGames.Horde.Tests/ChunkingTests.cs
@@ -12,6 +12,7 @@ using EpicGames.Horde.Storage;
 using EpicGames.Horde.Storage.Backends;
 using Microsoft.Extensions.Caching.Memory;
 using Microsoft.Extensions.Logging.Abstractions;
+using System.Collections.Generic;
 
 namespace EpicGames.Horde.Tests
 {
@@ -47,6 +48,29 @@ namespace EpicGames.Horde.Tests
 			}
 		}
 
+		[TestMethod]
+		public async Task EmptyNodeTest()
+		{
+			MemoryStorageClient store = new MemoryStorageClient();
+
+			await using (IStorageWriter writer = store.CreateWriter(new RefName("hello")))
+			{
+				ChunkingOptions options = new ChunkingOptions();
+				options.LeafOptions = new LeafChunkedDataNodeOptions(64, 64, 64);
+				options.InteriorOptions = new InteriorChunkedDataNodeOptions(4, 4, 4);
+
+				using MemoryStream emptyStream = new MemoryStream();
+				List<NodeRef<ChunkedDataNode>> leafNodes = await LeafChunkedDataNode.CreateFromStreamAsync(writer, emptyStream, new LeafChunkedDataNodeOptions(64, 64, 64), CancellationToken.None);
+				NodeRef<ChunkedDataNode> dataNode = await InteriorChunkedDataNode.CreateTreeAsync(leafNodes, new InteriorChunkedDataNodeOptions(4, 4, 4), writer, CancellationToken.None); 
+
+				DirectoryNode directory = new DirectoryNode();
+				directory.AddFile("test.foo", FileEntryFlags.None, 0, dataNode);
+
+				NodeRef<DirectoryNode> directoryRef = await writer.WriteNodeAsync(directory);
+				await writer.WriteRefAsync(directoryRef.Handle);
+			}
+		}
+
 		[TestMethod]
 		public async Task FixedSizeChunkingTests()
 		{
diff --git a/Engine/Source/Programs/Shared/EpicGames.Horde/Storage/Nodes/ChunkedDataNode.cs b/Engine/Source/Programs/Shared/EpicGames.Horde/Storage/Nodes/ChunkedDataNode.cs
--- a/Engine/Source/Programs/Shared/EpicGames.Horde/Storage/Nodes/ChunkedDataNode.cs
+++ b/Engine/Source/Programs/Shared/EpicGames.Horde/Storage/Nodes/ChunkedDataNode.cs
@@ -187,7 +187,7 @@ namespace EpicGames.Horde.Storage.Nodes
 			for (; ; )
 			{
 				size += await stream.ReadGreedyAsync(readBuffer.Memory.Slice(size), cancellationToken);
-				if (size == 0)
+				if (size == 0 && handles.Count > 0)
 				{
 					break;
 				}
