Horde: Missing file.

diff --git a/Engine/Source/Programs/Shared/EpicGames.Horde/Storage/Bundles/BundleStorageClient.cs b/Engine/Source/Programs/Shared/EpicGames.Horde/Storage/Bundles/BundleStorageClient.cs
--- a/Engine/Source/Programs/Shared/EpicGames.Horde/Storage/Bundles/BundleStorageClient.cs
+++ b/Engine/Source/Programs/Shared/EpicGames.Horde/Storage/Bundles/BundleStorageClient.cs
@@ -74,10 +74,10 @@ namespace EpicGames.Horde.Storage.Bundles
 		#region Aliases
 
 		/// <inheritdoc/>
-		Task IStorageClient.AddAliasAsync(Utf8String name, BlobHandle handle, CancellationToken cancellationToken) => AddAliasAsync(name, (BundleNodeHandle)handle, cancellationToken);
+		Task IStorageClient.AddAliasAsync(Utf8String name, BlobHandle handle, int rank, CancellationToken cancellationToken) => AddAliasAsync(name, (BundleNodeHandle)handle, rank, cancellationToken);
 
 		/// <inheritdoc/>
-		public abstract Task AddAliasAsync(Utf8String name, BundleNodeHandle handle, CancellationToken cancellationToken = default);
+		public abstract Task AddAliasAsync(Utf8String name, BundleNodeHandle handle, int rank = 0, CancellationToken cancellationToken = default);
 
 		/// <inheritdoc/>
 		Task IStorageClient.RemoveAliasAsync(Utf8String name, BlobHandle handle, CancellationToken cancellationToken) => RemoveAliasAsync(name, (BundleNodeHandle)handle, cancellationToken);
