Horde: Add a BinaryId type to EpicGames.Horde, which mirrors functionality from ObjectId without introducing a MongoDB dependency.

diff --git a/Engine/Source/Programs/Horde/Horde.Server/Artifacts/ArtifactId.cs b/Engine/Source/Programs/Horde/Horde.Server/Artifacts/ArtifactId.cs
--- a/Engine/Source/Programs/Horde/Horde.Server/Artifacts/ArtifactId.cs
+++ b/Engine/Source/Programs/Horde/Horde.Server/Artifacts/ArtifactId.cs
@@ -1,6 +1,7 @@
 // Copyright Epic Games, Inc. All Rights Reserved.
 
 using System.ComponentModel;
+using EpicGames.Horde;
 using Horde.Server.Utilities;
 using MongoDB.Bson;
 
@@ -10,18 +11,18 @@ namespace Horde.Server.Artifacts
 	/// Unique id for an artifact
 	/// </summary>
 	/// <param name="Id">Identifier for the artifact</param>
-	[TypeConverter(typeof(ObjectIdTypeConverter<ArtifactId, ArtifactIdConverter>))]
-	[ObjectIdConverter(typeof(ArtifactIdConverter))]
-	public record struct ArtifactId(ObjectId Id)
+	[TypeConverter(typeof(BinaryIdTypeConverter<ArtifactId, ArtifactIdConverter>))]
+	[BinaryIdConverter(typeof(ArtifactIdConverter))]
+	public record struct ArtifactId(BinaryId Id)
 	{
 		/// <summary>
 		/// Creates a new random artifact id
 		/// </summary>
 		/// <returns>New artifact id</returns>
-		public static ArtifactId GenerateNewId() => new ArtifactId(ObjectId.GenerateNewId());
+		public static ArtifactId GenerateNewId() => new ArtifactId(BinaryIdUtils.CreateNew());
 
 		/// <inheritdoc cref="ObjectId.Parse(System.String)"/>
-		public static ArtifactId Parse(string text) => new ArtifactId(ObjectId.Parse(text));
+		public static ArtifactId Parse(string text) => new ArtifactId(BinaryId.Parse(text));
 
 		/// <inheritdoc/>
 		public override string ToString() => Id.ToString();
@@ -30,12 +31,12 @@ namespace Horde.Server.Artifacts
 	/// <summary>
 	/// Converter class to and from ObjectId values
 	/// </summary>
-	class ArtifactIdConverter : ObjectIdConverter<ArtifactId>
+	class ArtifactIdConverter : BinaryIdConverter<ArtifactId>
 	{
 		/// <inheritdoc/>
-		public override ArtifactId FromObjectId(ObjectId id) => new ArtifactId(id);
+		public override ArtifactId FromBinaryId(BinaryId id) => new ArtifactId(id);
 
 		/// <inheritdoc/>
-		public override ObjectId ToObjectId(ArtifactId value) => value.Id;
+		public override BinaryId ToBinaryId(ArtifactId value) => value.Id;
 	}
 }
diff --git a/Engine/Source/Programs/Horde/Horde.Server/Startup.cs b/Engine/Source/Programs/Horde/Horde.Server/Startup.cs
--- a/Engine/Source/Programs/Horde/Horde.Server/Startup.cs
+++ b/Engine/Source/Programs/Horde/Horde.Server/Startup.cs
@@ -1034,6 +1034,7 @@ namespace Horde.Server
 				BsonSerializer.RegisterSerializer(new ConditionSerializer());
 				BsonSerializer.RegisterSerializationProvider(new BsonSerializationProvider());
 				BsonSerializer.RegisterSerializationProvider(new StringIdBsonSerializationProvider());
+				BsonSerializer.RegisterSerializationProvider(new BinaryIdBsonSerializationProvider());
 				BsonSerializer.RegisterSerializationProvider(new ObjectIdBsonSerializationProvider());
 			}
 		}
diff --git a/Engine/Source/Programs/Horde/Horde.Server/Utilities/BinaryIdConverter.cs b/Engine/Source/Programs/Horde/Horde.Server/Utilities/BinaryIdConverter.cs
new file mode 100644
--- /dev/null
+++ b/Engine/Source/Programs/Horde/Horde.Server/Utilities/BinaryIdConverter.cs
@@ -0,0 +1,46 @@
+// Copyright Epic Games, Inc. All Rights Reserved.
+
+using System;
+using System.Collections.Generic;
+using System.Linq;
+using System.Reflection;
+using System.Text;
+using System.Threading.Tasks;
+using EpicGames.Horde;
+using MongoDB.Bson;
+using MongoDB.Bson.Serialization;
+using MongoDB.Bson.Serialization.Serializers;
+
+namespace Horde.Server.Utilities
+{
+	/// <summary>
+	/// Class which serializes object id types to BSON
+	/// </summary>
+	sealed class BinaryIdBsonSerializer<TValue, TConverter> : SerializerBase<TValue> where TValue : struct where TConverter : BinaryIdConverter<TValue>, new()
+	{
+		readonly TConverter _converter = new TConverter();
+
+		/// <inheritdoc/>
+		public override TValue Deserialize(BsonDeserializationContext context, BsonDeserializationArgs args) => _converter.FromBinaryId(BinaryIdUtils.FromObjectId(context.Reader.ReadObjectId()));
+
+		/// <inheritdoc/>
+		public override void Serialize(BsonSerializationContext context, BsonSerializationArgs args, TValue value) => context.Writer.WriteObjectId(BinaryIdUtils.ToObjectId(_converter.ToBinaryId(value)));
+	}
+
+	/// <summary>
+	/// Class which serializes object id types to BSON
+	/// </summary>
+	sealed class BinaryIdBsonSerializationProvider : BsonSerializationProviderBase
+	{
+		/// <inheritdoc/>
+		public override IBsonSerializer? GetSerializer(Type type, IBsonSerializerRegistry serializerRegistry)
+		{
+			BinaryIdConverterAttribute? attribute = type.GetCustomAttribute<BinaryIdConverterAttribute>();
+			if (attribute == null)
+			{
+				return null;
+			}
+			return (IBsonSerializer?)Activator.CreateInstance(typeof(BinaryIdBsonSerializer<,>).MakeGenericType(type, attribute.ConverterType));
+		}
+	}
+}
diff --git a/Engine/Source/Programs/Horde/Horde.Server/Utilities/BinaryIdUtils.cs b/Engine/Source/Programs/Horde/Horde.Server/Utilities/BinaryIdUtils.cs
new file mode 100644
--- /dev/null
+++ b/Engine/Source/Programs/Horde/Horde.Server/Utilities/BinaryIdUtils.cs
@@ -0,0 +1,43 @@
+// Copyright Epic Games, Inc. All Rights Reserved.
+
+using EpicGames.Horde;
+using MongoDB.Bson;
+using System;
+using System.Collections.Generic;
+using System.Linq;
+using System.Text;
+using System.Threading.Tasks;
+
+namespace Horde.Server.Utilities
+{
+	/// <summary>
+	/// Utility methods for <see cref="BinaryId"/>
+	/// </summary>
+	static class BinaryIdUtils
+	{
+		/// <summary>
+		/// Creates a new BinaryId
+		/// </summary>
+		public static BinaryId CreateNew()
+		{
+			ObjectId objectId = ObjectId.GenerateNewId();
+			return new BinaryId(objectId.ToByteArray());
+		}
+
+		/// <summary>
+		/// Creates a BinaryId from an ObjectId
+		/// </summary>
+		public static BinaryId FromObjectId(ObjectId objectId)
+		{
+			return new BinaryId(objectId.ToByteArray());
+		}
+
+		/// <summary>
+		/// Creates an ObjectId from a BinaryId
+		/// </summary>
+		public static ObjectId ToObjectId(BinaryId binaryId)
+		{
+			return new ObjectId(binaryId.ToByteArray());
+		}
+	}
+}
diff --git a/Engine/Source/Programs/Shared/EpicGames.Core/StringUtils.cs b/Engine/Source/Programs/Shared/EpicGames.Core/StringUtils.cs
--- a/Engine/Source/Programs/Shared/EpicGames.Core/StringUtils.cs
+++ b/Engine/Source/Programs/Shared/EpicGames.Core/StringUtils.cs
@@ -365,24 +365,42 @@ namespace EpicGames.Core
 		/// <returns></returns>
 		public static bool TryParseHexString(ReadOnlySpan<char> text, [NotNullWhen(true)] out byte[]? outBytes)
 		{
-			if((text.Length & 1) != 0)
+			byte[] bytes = new byte[text.Length / 2];
+			if (TryParseHexString(text, bytes))
+			{
+				outBytes = bytes;
+				return true;
+			}
+			else
 			{
 				outBytes = null;
 				return false;
 			}
+		}
+
+		/// <summary>
+		/// Parses a hexadecimal string into an array of bytes
+		/// </summary>
+		/// <param name="text">Text to parse</param>
+		/// <param name="bytes">Receives the parsed string</param>
+		/// <returns></returns>
+		public static bool TryParseHexString(ReadOnlySpan<char> text, Span<byte> bytes)
+		{
+			if((text.Length & 1) != 0)
+			{
+				return false;
+			}
 
-			byte[] bytes = new byte[text.Length / 2];
 			for(int idx = 0; idx < text.Length; idx += 2)
 			{
 				int value = (GetHexDigit(text[idx]) << 4) | GetHexDigit(text[idx + 1]);
 				if(value < 0)
 				{
-					outBytes = null;
 					return false;
 				}
 				bytes[idx / 2] = (byte)value;
 			}
-			outBytes = bytes;
+
 			return true;
 		}
 
@@ -394,24 +412,42 @@ namespace EpicGames.Core
 		/// <returns></returns>
 		public static bool TryParseHexString(ReadOnlySpan<byte> text, [NotNullWhen(true)] out byte[]? outBytes)
 		{
-			if ((text.Length & 1) != 0)
+			byte[] bytes = new byte[text.Length / 2];
+			if (TryParseHexString(text, bytes))
+			{
+				outBytes = bytes;
+				return true;
+			}
+			else
 			{
 				outBytes = null;
 				return false;
 			}
+		}
+
+		/// <summary>
+		/// Parses a hexadecimal string into an array of bytes
+		/// </summary>
+		/// <param name="text">Text to parse</param>
+		/// <param name="bytes">Receives the parsed string</param>
+		/// <returns></returns>
+		public static bool TryParseHexString(ReadOnlySpan<byte> text, Span<byte> bytes)
+		{
+			if ((text.Length & 1) != 0)
+			{
+				return false;
+			}
 
-			byte[] bytes = new byte[text.Length / 2];
 			for (int idx = 0; idx < text.Length; idx += 2)
 			{
 				int value = ParseHexByte(text, idx);
 				if (value < 0)
 				{
-					outBytes = null;
 					return false;
 				}
 				bytes[idx / 2] = (byte)value;
 			}
-			outBytes = bytes;
+
 			return true;
 		}
 
diff --git a/Engine/Source/Programs/Shared/EpicGames.Horde/BinaryId.cs b/Engine/Source/Programs/Shared/EpicGames.Horde/BinaryId.cs
new file mode 100644
--- /dev/null
+++ b/Engine/Source/Programs/Shared/EpicGames.Horde/BinaryId.cs
@@ -0,0 +1,207 @@
+// Copyright Epic Games, Inc. All Rights Reserved.
+
+using EpicGames.Core;
+using System;
+using System.Buffers.Binary;
+using System.ComponentModel;
+using System.Globalization;
+using System.Text;
+using System.Text.Json;
+using System.Text.Json.Serialization;
+
+namespace EpicGames.Horde
+{
+	/// <summary>
+	/// Normalized string identifier for a resource
+	/// </summary>
+	[JsonSchemaString]
+	[JsonConverter(typeof(StringIdJsonConverter))]
+	[TypeConverter(typeof(StringIdTypeConverter))]
+	public readonly struct BinaryId : IEquatable<BinaryId>
+	{
+		readonly int _a;
+		readonly int _b;
+		readonly int _c;
+
+		/// <summary>
+		/// Constructor
+		/// </summary>
+		/// <param name="bytes">Bytes to parse</param>
+		public BinaryId(ReadOnlySpan<byte> bytes)
+		{
+			_a = BinaryPrimitives.ReadInt32LittleEndian(bytes);
+			_b = BinaryPrimitives.ReadInt32LittleEndian(bytes[4..]);
+			_c = BinaryPrimitives.ReadInt32LittleEndian(bytes[8..]);
+		}
+
+		/// <summary>
+		/// Parse a binary id from a string
+		/// </summary>
+		public static BinaryId Parse(string text)
+		{
+			BinaryId id;
+			if (!TryParse(text, out id))
+			{
+				throw new FormatException($"Invalid BinaryId: {text}");
+			}
+			return id;
+		}
+
+		/// <summary>
+		/// Parse a binary id from a string
+		/// </summary>
+		public static BinaryId Parse(Utf8String text) => Parse(text.Span);
+
+		/// <summary>
+		/// Parse a binary id from a string
+		/// </summary>
+		public static BinaryId Parse(ReadOnlySpan<byte> text)
+		{
+			BinaryId id;
+			if (!TryParse(text, out id))
+			{
+				throw new FormatException($"Invalid BinaryId: {Encoding.UTF8.GetString(text)}");
+			}
+			return id;
+		}
+
+		/// <summary>
+		/// Attempt to parse a binary id from a string
+		/// </summary>
+		public static bool TryParse(string text, out BinaryId result)
+		{
+			Span<byte> bytes = stackalloc byte[12];
+			if (StringUtils.TryParseHexString(text, bytes))
+			{
+				result = new BinaryId(bytes);
+				return true;
+			}
+			else
+			{
+				result = default;
+				return false;
+			}
+		}
+
+		/// <summary>
+		/// Attempt to parse a binary id from a string
+		/// </summary>
+		public static bool TryParse(Utf8String text, out BinaryId result) => TryParse(text.Span, out result);
+
+		/// <summary>
+		/// Attempt to parse a binary id from a string
+		/// </summary>
+		public static bool TryParse(ReadOnlySpan<byte> text, out BinaryId result)
+		{
+			Span<byte> bytes = stackalloc byte[12];
+			if (StringUtils.TryParseHexString(text, bytes))
+			{
+				result = new BinaryId(bytes);
+				return true;
+			}
+			else
+			{
+				result = default;
+				return false;
+			}
+		}
+
+		/// <summary>
+		/// Checks whether this StringId is set
+		/// </summary>
+		public bool IsEmpty => (_a | _b | _c) == 0;
+
+		/// <inheritdoc/>
+		public override bool Equals(object? obj) => obj is BinaryId id && Equals(id);
+
+		/// <inheritdoc/>
+		public override int GetHashCode() => HashCode.Combine(_a, _b, _c);
+
+		/// <inheritdoc/>
+		public bool Equals(BinaryId other) => _a == other._a && _b == other._b && _c == other._c;
+
+		/// <inheritdoc/>
+		public override string ToString() => $"{_a:08x}{_b:08x}{_c:08x}";
+
+		/// <summary>
+		/// Format this id as a sequence of UTF8 characters
+		/// </summary>
+		public void ToUtf8String(Span<byte> chars)
+		{
+			StringUtils.FormatUtf8HexString((uint)_a, chars);
+			StringUtils.FormatUtf8HexString((uint)_b, chars[4..]);
+			StringUtils.FormatUtf8HexString((uint)_c, chars[8..]);
+		}
+
+		/// <summary>
+		/// Format this id as a sequence of UTF8 characters
+		/// </summary>
+		public byte[] ToByteArray()
+		{
+			byte[] bytes = new byte[12];
+			ToByteArray(bytes);
+			return bytes;
+		}
+
+		/// <summary>
+		/// Format this id as a sequence of UTF8 characters
+		/// </summary>
+		public void ToByteArray(Span<byte> chars)
+		{
+			BinaryPrimitives.WriteInt32LittleEndian(chars, _a);
+			BinaryPrimitives.WriteInt32LittleEndian(chars[4..], _b);
+			BinaryPrimitives.WriteInt32LittleEndian(chars[8..], _c);
+		}
+
+		/// <summary>
+		/// Compares two binary ids for equality
+		/// </summary>
+		/// <param name="left">The first string id</param>
+		/// <param name="right">Second string id</param>
+		/// <returns>True if the two string ids are equal</returns>
+		public static bool operator ==(BinaryId left, BinaryId right) => left.Equals(right);
+
+		/// <summary>
+		/// Compares two binary ids for inequality
+		/// </summary>
+		/// <param name="left">The first string id</param>
+		/// <param name="right">Second string id</param>
+		/// <returns>True if the two string ids are not equal</returns>
+		public static bool operator !=(BinaryId left, BinaryId right) => !left.Equals(right);
+	}
+
+	/// <summary>
+	/// Class which serializes <see cref="BinaryId"/> types
+	/// </summary>
+	sealed class BinaryIdJsonConverter : JsonConverter<BinaryId>
+	{
+		/// <inheritdoc/>
+		public override BinaryId Read(ref Utf8JsonReader reader, Type typeToConvert, JsonSerializerOptions options) => BinaryId.Parse(reader.GetUtf8String());
+
+		/// <inheritdoc/>
+		public override void Write(Utf8JsonWriter writer, BinaryId value, JsonSerializerOptions options)
+		{
+			Span<byte> bytes = stackalloc byte[12 * 2];
+			value.ToUtf8String(bytes);
+			writer.WriteStringValue(bytes);
+		}
+	}
+
+	/// <summary>
+	/// Class which serializes <see cref="BinaryId"/> types
+	/// </summary>
+	sealed class BinaryIdTypeConverter : TypeConverter
+	{
+		/// <inheritdoc/>
+		public override bool CanConvertFrom(ITypeDescriptorContext? context, Type sourceType) => sourceType == typeof(string);
+
+		/// <inheritdoc/>
+		public override object? ConvertFrom(ITypeDescriptorContext? context, CultureInfo? culture, object value) => BinaryId.Parse((string)value);
+
+		/// <inheritdoc/>
+		public override bool CanConvertTo(ITypeDescriptorContext? context, Type? destinationType) => destinationType == typeof(string);
+
+		/// <inheritdoc/>
+		public override object? ConvertTo(ITypeDescriptorContext? context, CultureInfo? culture, object? value, Type destinationType) => ((BinaryId)value!).ToString();
+	}
+}
diff --git a/Engine/Source/Programs/Shared/EpicGames.Horde/BinaryIdConverter.cs b/Engine/Source/Programs/Shared/EpicGames.Horde/BinaryIdConverter.cs
new file mode 100644
--- /dev/null
+++ b/Engine/Source/Programs/Shared/EpicGames.Horde/BinaryIdConverter.cs
@@ -0,0 +1,159 @@
+// Copyright Epic Games, Inc. All Rights Reserved.
+
+using System;
+using System.ComponentModel;
+using System.Globalization;
+using System.Reflection;
+using System.Text.Json;
+using System.Text.Json.Serialization;
+using EpicGames.Core;
+using EpicGames.Serialization;
+
+namespace EpicGames.Horde
+{
+	/// <summary>
+	/// Base class for converting to and from types containing a <see cref="BinaryId"/>. Useful pattern for reducing boilerplate with strongly typed records.
+	/// </summary>
+	/// <typeparam name="T"></typeparam>
+	public abstract class BinaryIdConverter<T> where T : struct
+	{
+		/// <summary>
+		/// Converts a type to a <see cref="BinaryId"/>
+		/// </summary>
+		public abstract BinaryId ToBinaryId(T value);
+
+		/// <summary>
+		/// Constructs a type from a <see cref="BinaryId"/>
+		/// </summary>
+		public abstract T FromBinaryId(BinaryId id);
+	}
+
+	/// <summary>
+	/// Attribute declaring a <see cref="BinaryIdConverter{T}"/> for a particular type
+	/// </summary>
+	[AttributeUsage(AttributeTargets.Struct)]
+	public sealed class BinaryIdConverterAttribute : Attribute
+	{
+		/// <summary>
+		/// The converter type
+		/// </summary>
+		public Type ConverterType { get; }
+
+		/// <summary>
+		/// Constructor
+		/// </summary>
+		public BinaryIdConverterAttribute(Type converterType) => ConverterType = converterType;
+	}
+
+	/// <summary>
+	/// Converter to compact binary objects
+	/// </summary>
+	public sealed class BinaryIdCbConverter<TValue, TConverter> : CbConverterBase<TValue> where TValue : struct where TConverter : BinaryIdConverter<TValue>, new()
+	{
+		readonly TConverter _converter = new TConverter();
+
+		/// <inheritdoc/>
+		public override TValue Read(CbField field)
+		{
+			return _converter.FromBinaryId(BinaryId.Parse(field.AsUtf8String()));
+		}
+
+		/// <inheritdoc/>
+		public override void Write(CbWriter writer, TValue value)
+		{
+			writer.WriteStringValue(_converter.ToBinaryId(value).ToString());
+		}
+
+		/// <inheritdoc/>
+		public override void WriteNamed(CbWriter writer, Utf8String name, TValue value)
+		{
+			writer.WriteString(name, _converter.ToBinaryId(value).ToString());
+		}
+	}
+
+	/// <summary>
+	/// Class which serializes types with a <see cref="BinaryIdConverter{T}"/> to Json
+	/// </summary>
+	public sealed class BinaryIdTypeConverter<TValue, TConverter> : TypeConverter where TValue : struct where TConverter : BinaryIdConverter<TValue>, new()
+	{
+		readonly TConverter _converter = new TConverter();
+
+		/// <inheritdoc/>
+		public override bool CanConvertFrom(ITypeDescriptorContext? context, Type sourceType)
+		{
+			return sourceType == typeof(string) || sourceType == typeof(BinaryId);
+		}
+
+		/// <inheritdoc/>
+		public override object? ConvertFrom(ITypeDescriptorContext? context, CultureInfo? culture, object value)
+		{
+			if (value is string str)
+			{
+				return _converter.FromBinaryId(BinaryId.Parse(str));
+			}
+			if (value is BinaryId stringId)
+			{
+				return _converter.FromBinaryId(stringId);
+			}
+			return null;
+		}
+
+		/// <inheritdoc/>
+		public override bool CanConvertTo(ITypeDescriptorContext? context, Type? destinationType) => destinationType == typeof(string) || destinationType == typeof(BinaryId);
+
+		/// <inheritdoc/>
+		public override object? ConvertTo(ITypeDescriptorContext? context, CultureInfo? culture, object? value, Type destinationType)
+		{
+			if (destinationType == typeof(string))
+			{
+				return _converter.ToBinaryId((TValue)value!).ToString();
+			}
+			if (destinationType == typeof(BinaryId))
+			{
+				return _converter.ToBinaryId((TValue)value!);
+			}
+			return null;
+		}
+	}
+
+	/// <summary>
+	/// Class which serializes types with a <see cref="BinaryIdConverter{T}"/> to Json
+	/// </summary>
+	public sealed class BinaryIdJsonConverter<TValue, TConverter> : JsonConverter<TValue> where TValue : struct where TConverter : BinaryIdConverter<TValue>, new()
+	{
+		readonly TConverter _converter = new TConverter();
+
+		/// <inheritdoc/>
+		public override TValue Read(ref Utf8JsonReader reader, Type typeToConvert, JsonSerializerOptions options) => _converter.FromBinaryId(BinaryId.Parse(reader.GetUtf8String()));
+
+		/// <inheritdoc/>
+		public override void Write(Utf8JsonWriter writer, TValue value, JsonSerializerOptions options)
+		{
+			BinaryId binaryId = _converter.ToBinaryId(value);
+
+			Span<byte> span = stackalloc byte[12 * 2];
+			binaryId.ToUtf8String(span);
+			writer.WriteStringValue(span);
+		}
+	}
+
+	/// <summary>
+	/// Creates constructors for types with a <see cref="BinaryIdConverter{T}"/> to Json
+	/// </summary>
+	public sealed class BinaryIdJsonConverterFactory : JsonConverterFactory
+	{
+		/// <inheritdoc/>
+		public override bool CanConvert(Type typeToConvert) => typeToConvert.GetCustomAttribute<BinaryIdConverterAttribute>() != null;
+
+		/// <inheritdoc/>
+		public override JsonConverter? CreateConverter(Type type, JsonSerializerOptions options)
+		{
+			BinaryIdConverterAttribute? attribute = type.GetCustomAttribute<BinaryIdConverterAttribute>();
+			if (attribute == null)
+			{
+				return null;
+			}
+			return (JsonConverter?)Activator.CreateInstance(typeof(BinaryIdJsonConverter<,>).MakeGenericType(type, attribute.ConverterType));
+		}
+	}
+}
