Horde: Fix BinaryId conversion to other types.

diff --git a/Engine/Source/Programs/Shared/EpicGames.Horde/BinaryId.cs b/Engine/Source/Programs/Shared/EpicGames.Horde/BinaryId.cs
--- a/Engine/Source/Programs/Shared/EpicGames.Horde/BinaryId.cs
+++ b/Engine/Source/Programs/Shared/EpicGames.Horde/BinaryId.cs
@@ -15,8 +15,8 @@ namespace EpicGames.Horde
 	/// Normalized string identifier for a resource
 	/// </summary>
 	[JsonSchemaString]
-	[JsonConverter(typeof(StringIdJsonConverter))]
-	[TypeConverter(typeof(StringIdTypeConverter))]
+	[JsonConverter(typeof(BinaryIdJsonConverter))]
+	[TypeConverter(typeof(BinaryIdTypeConverter))]
 	public readonly struct BinaryId : IEquatable<BinaryId>
 	{
 		readonly int _a;
@@ -107,7 +107,7 @@ namespace EpicGames.Horde
 		}
 
 		/// <summary>
-		/// Checks whether this StringId is set
+		/// Checks whether this BinaryId is set
 		/// </summary>
 		public bool IsEmpty => (_a | _b | _c) == 0;
 
@@ -121,7 +121,12 @@ namespace EpicGames.Horde
 		public bool Equals(BinaryId other) => _a == other._a && _b == other._b && _c == other._c;
 
 		/// <inheritdoc/>
-		public override string ToString() => $"{_a:08x}{_b:08x}{_c:08x}";
+		public override string ToString()
+		{
+			Span<byte> bytes = stackalloc byte[12];
+			ToByteArray(bytes);
+			return StringUtils.FormatHexString(bytes);
+		}
 
 		/// <summary>
 		/// Format this id as a sequence of UTF8 characters
