Horde: Move artifact messages into EpicGames.Horde, and add an 'artifacts find' command to the Horde command line tool.

diff --git a/Engine/Source/Programs/Horde/Horde.Server.Tests/ArtifactTests.cs b/Engine/Source/Programs/Horde/Horde.Server.Tests/ArtifactTests.cs
--- a/Engine/Source/Programs/Horde/Horde.Server.Tests/ArtifactTests.cs
+++ b/Engine/Source/Programs/Horde/Horde.Server.Tests/ArtifactTests.cs
@@ -5,10 +5,12 @@ using System.Collections.Generic;
 using System.Linq;
 using System.Threading;
 using System.Threading.Tasks;
+using EpicGames.Horde.Api;
 using EpicGames.Horde.Storage;
 using Horde.Server.Acls;
 using Horde.Server.Artifacts;
 using Horde.Server.Storage;
+using Horde.Server.Utilities;
 using HordeCommon;
 using Microsoft.Extensions.DependencyInjection;
 using Microsoft.VisualStudio.TestTools.UnitTesting;
@@ -19,10 +21,10 @@ namespace Horde.Server.Tests
 	public class ArtifactTests : TestSetup
 	{
 		[TestMethod]
-		public async Task CreateArtifact()
+		public async Task CreateArtifactAsync()
 		{
 			IArtifactCollection artifactCollection = ServiceProvider.GetRequiredService<IArtifactCollection>();
-			IArtifact artifact = await artifactCollection.AddAsync(ArtifactId.GenerateNewId(), ArtifactType.StepOutput, new string[] { "test1", "test2" }, Namespace.Artifacts, new RefName("test"), null, AclScopeName.Root);
+			IArtifact artifact = await artifactCollection.AddAsync(new ArtifactId(BinaryIdUtils.CreateNew()), ArtifactType.StepOutput, new string[] { "test1", "test2" }, Namespace.Artifacts, new RefName("test"), null, AclScopeName.Root);
 
 			{
 				List<IArtifact> artifacts = await artifactCollection.FindAsync(keys: new[] { "test1" }).ToListAsync();
@@ -43,7 +45,7 @@ namespace Horde.Server.Tests
 		}
 
 		[TestMethod]
-		public async Task ExpireArtifact()
+		public async Task ExpireArtifactAsync()
 		{
 			FakeClock clock = ServiceProvider.GetRequiredService<FakeClock>();
 			ArtifactExpirationService expirationService = ServiceProvider.GetRequiredService<ArtifactExpirationService>();
@@ -51,7 +53,7 @@ namespace Horde.Server.Tests
 			await expirationService.StartAsync(CancellationToken.None);
 			
 			IArtifactCollection artifactCollection = ServiceProvider.GetRequiredService<IArtifactCollection>();
-			IArtifact artifact = await artifactCollection.AddAsync(ArtifactId.GenerateNewId(), ArtifactType.StepOutput, new string[] { "test1", "test2" }, Namespace.Artifacts, new RefName("test"), clock.UtcNow + TimeSpan.FromHours(1.0), AclScopeName.Root);
+			IArtifact artifact = await artifactCollection.AddAsync(new ArtifactId(BinaryIdUtils.CreateNew()), ArtifactType.StepOutput, new string[] { "test1", "test2" }, Namespace.Artifacts, new RefName("test"), clock.UtcNow + TimeSpan.FromHours(1.0), AclScopeName.Root);
 
 			{
 				List<IArtifact> artifacts = await artifactCollection.FindAsync(keys: new[] { "test1" }).ToListAsync();
diff --git a/Engine/Source/Programs/Horde/Horde.Server/Artifacts/ArtifactCollection.cs b/Engine/Source/Programs/Horde/Horde.Server/Artifacts/ArtifactCollection.cs
--- a/Engine/Source/Programs/Horde/Horde.Server/Artifacts/ArtifactCollection.cs
+++ b/Engine/Source/Programs/Horde/Horde.Server/Artifacts/ArtifactCollection.cs
@@ -6,6 +6,7 @@ using System.Linq;
 using System.Runtime.CompilerServices;
 using System.Threading;
 using System.Threading.Tasks;
+using EpicGames.Horde.Api;
 using EpicGames.Horde.Storage;
 using Horde.Server.Acls;
 using Horde.Server.Server;
diff --git a/Engine/Source/Programs/Horde/Horde.Server/Artifacts/ArtifactId.cs b/Engine/Source/Programs/Horde/Horde.Server/Artifacts/ArtifactId.cs
deleted file mode 100644
diff --git a/Engine/Source/Programs/Horde/Horde.Server/Artifacts/ArtifactType.cs b/Engine/Source/Programs/Horde/Horde.Server/Artifacts/ArtifactType.cs
deleted file mode 100644
diff --git a/Engine/Source/Programs/Horde/Horde.Server/Artifacts/ArtifactsController.cs b/Engine/Source/Programs/Horde/Horde.Server/Artifacts/ArtifactsController.cs
--- a/Engine/Source/Programs/Horde/Horde.Server/Artifacts/ArtifactsController.cs
+++ b/Engine/Source/Programs/Horde/Horde.Server/Artifacts/ArtifactsController.cs
@@ -8,6 +8,7 @@ using System.Net.Http.Headers;
 using System.Threading;
 using System.Threading.Tasks;
 using EpicGames.Core;
+using EpicGames.Horde.Api;
 using EpicGames.Horde.Storage;
 using EpicGames.Horde.Storage.Bundles;
 using EpicGames.Horde.Storage.Nodes;
@@ -23,101 +24,6 @@ using Microsoft.Extensions.Options;
 
 namespace Horde.Server.Artifacts
 {
-	/// <summary>
-	/// Describes an artifact
-	/// </summary>
-	public class GetArtifactResponse
-	{
-		readonly IArtifact _artifact;
-
-		/// <inheritdoc cref="IArtifact.Id"/>
-		public string Id => _artifact.Id.ToString();
-
-		/// <inheritdoc cref="IArtifact.Type"/>
-		public ArtifactType Type => _artifact.Type;
-
-		/// <inheritdoc cref="IArtifact.Keys"/>
-		public IReadOnlyList<string> Keys => _artifact.Keys;
-
-		internal GetArtifactResponse(IArtifact artifact) => _artifact = artifact;
-	}
-
-	/// <summary>
-	/// Result of an artifact search
-	/// </summary>
-	public class FindArtifactsResponse
-	{
-		/// <summary>
-		/// List of artifacts matching the search criteria
-		/// </summary>
-		public List<GetArtifactResponse> Artifacts { get; } = new List<GetArtifactResponse>();
-	}
-
-	/// <summary>
-	/// Describes a file within an artifact
-	/// </summary>
-	public class GetArtifactFileEntryResponse
-	{
-		readonly FileEntry _entry;
-
-		/// <inheritdoc cref="FileEntry.Name"/>
-		public string Name => _entry.Name.ToString();
-
-		/// <inheritdoc cref="FileEntry.Length"/>
-		public long Length => _entry.Length;
-
-		/// <inheritdoc cref="FileEntry.Hash"/>
-		public IoHash Hash => _entry.Hash;
-
-		internal GetArtifactFileEntryResponse(FileEntry entry) => _entry = entry;
-	}
-
-	/// <summary>
-	/// Describes a file within an artifact
-	/// </summary>
-	public class GetArtifactDirectoryEntryResponse : GetArtifactDirectoryResponse
-	{
-		readonly DirectoryEntry _entry;
-
-		/// <inheritdoc cref="FileEntry.Name"/>
-		public string Name => _entry.Name.ToString();
-
-		/// <inheritdoc cref="FileEntry.Length"/>
-		public long Length => _entry.Length;
-
-		/// <inheritdoc cref="FileEntry.Hash"/>
-		public IoHash Hash => _entry.Handle.Hash;
-
-		internal GetArtifactDirectoryEntryResponse(DirectoryEntry entry) => _entry = entry;
-	}
-
-	/// <summary>
-	/// Describes a directory within an artifact
-	/// </summary>
-	public class GetArtifactDirectoryResponse
-	{
-		/// <summary>
-		/// Names of sub-directories
-		/// </summary>
-		public List<GetArtifactDirectoryEntryResponse>? Directories { get; internal set; }
-			
-		/// <summary>
-		/// Files within the directory
-		/// </summary>
-		public List<GetArtifactFileEntryResponse>? Files { get; internal set; }
-	}
-
-	/// <summary>
-	/// Request to create a zip file with artifact data
-	/// </summary>
-	public class CreateZipRequest
-	{
-		/// <summary>
-		/// Filter lines for the zip. Uses standard <see cref="FileFilter"/> syntax.
-		/// </summary>
-		public List<string> Filter { get; set; } = new List<string>();
-	}
-
 	/// <summary>
 	/// Public interface for artifacts
 	/// </summary>
@@ -160,7 +66,7 @@ namespace Horde.Server.Artifacts
 				return Forbid(ArtifactAclAction.ReadArtifact, artifact.AclScope);
 			}
 
-			return PropertyFilter.Apply(new GetArtifactResponse(artifact), filter);
+			return PropertyFilter.Apply(new GetArtifactResponse(artifact.Id, artifact.Type, artifact.Keys), filter);
 		}
 
 		/// <summary>
@@ -280,7 +186,7 @@ namespace Horde.Server.Artifacts
 				{
 					DirectoryNode subDirectoryNode = await subDirectoryEntry.ExpandAsync(cancellationToken);
 
-					GetArtifactDirectoryEntryResponse subDirectoryEntryResponse = new GetArtifactDirectoryEntryResponse(subDirectoryEntry);
+					GetArtifactDirectoryEntryResponse subDirectoryEntryResponse = new GetArtifactDirectoryEntryResponse(subDirectoryEntry.Name.ToString(), subDirectoryEntry.Length, subDirectoryEntry.Handle.Hash);
 					if (depth == 0)
 					{
 						if (subDirectoryNode.Directories.Count + subDirectoryNode.Files.Count < 16)
@@ -308,7 +214,7 @@ namespace Horde.Server.Artifacts
 
 			if (directoryNode.Files.Count > 0)
 			{
-				response.Files = directoryNode.Files.Select(x => new GetArtifactFileEntryResponse(x)).ToList();
+				response.Files = directoryNode.Files.Select(x => new GetArtifactFileEntryResponse(x.Name.ToString(), x.Length, x.Handle.Hash)).ToList();
 			}
 		}
 
@@ -435,7 +341,7 @@ namespace Horde.Server.Artifacts
 			{
 				if (_globalConfig.Authorize(artifact.AclScope, ArtifactAclAction.ReadArtifact, User))
 				{
-					response.Artifacts.Add(new GetArtifactResponse(artifact));
+					response.Artifacts.Add(new GetArtifactResponse(artifact.Id, artifact.Type, artifact.Keys));
 				}
 			}
 
diff --git a/Engine/Source/Programs/Horde/Horde.Server/Artifacts/IArtifact.cs b/Engine/Source/Programs/Horde/Horde.Server/Artifacts/IArtifact.cs
--- a/Engine/Source/Programs/Horde/Horde.Server/Artifacts/IArtifact.cs
+++ b/Engine/Source/Programs/Horde/Horde.Server/Artifacts/IArtifact.cs
@@ -2,6 +2,7 @@
 
 using System;
 using System.Collections.Generic;
+using EpicGames.Horde.Api;
 using EpicGames.Horde.Storage;
 using Horde.Server.Acls;
 
diff --git a/Engine/Source/Programs/Horde/Horde.Server/Artifacts/IArtifactCollection.cs b/Engine/Source/Programs/Horde/Horde.Server/Artifacts/IArtifactCollection.cs
--- a/Engine/Source/Programs/Horde/Horde.Server/Artifacts/IArtifactCollection.cs
+++ b/Engine/Source/Programs/Horde/Horde.Server/Artifacts/IArtifactCollection.cs
@@ -4,6 +4,7 @@ using System;
 using System.Collections.Generic;
 using System.Threading;
 using System.Threading.Tasks;
+using EpicGames.Horde.Api;
 using EpicGames.Horde.Storage;
 using Horde.Server.Acls;
 
diff --git a/Engine/Source/Programs/Horde/Horde.Server/Jobs/JobRpcService.cs b/Engine/Source/Programs/Horde/Horde.Server/Jobs/JobRpcService.cs
--- a/Engine/Source/Programs/Horde/Horde.Server/Jobs/JobRpcService.cs
+++ b/Engine/Source/Programs/Horde/Horde.Server/Jobs/JobRpcService.cs
@@ -13,6 +13,7 @@ using System.Text.Json.Serialization;
 using System.Threading;
 using System.Threading.Tasks;
 using EpicGames.Core;
+using EpicGames.Horde.Api;
 using EpicGames.Horde.Storage;
 using Google.Protobuf.WellKnownTypes;
 using Grpc.Core;
@@ -92,7 +93,7 @@ namespace Horde.Server.Jobs
 				throw new StructuredRpcException(StatusCode.NotFound, "Couldn't find template {TemplateId} in stream {StreamId}", job.TemplateId, job.StreamId);
 			}
 
-			ArtifactId artifactId = ArtifactId.GenerateNewId();
+			ArtifactId artifactId = new ArtifactId(BinaryIdUtils.CreateNew());
 
 			NamespaceId namespaceId = String.IsNullOrEmpty(request.NamespaceId)? Namespace.Artifacts : new NamespaceId(request.NamespaceId);
 			RefName refName = new RefName(String.IsNullOrEmpty(request.RefName) ? $"job-{job.Id}/step-{step.Id}/{artifactId}" : request.RefName);
diff --git a/Engine/Source/Programs/Horde/Horde/Commands/Artifacts/ArtifactFind.cs b/Engine/Source/Programs/Horde/Horde/Commands/Artifacts/ArtifactFind.cs
new file mode 100644
--- /dev/null
+++ b/Engine/Source/Programs/Horde/Horde/Commands/Artifacts/ArtifactFind.cs
@@ -0,0 +1,34 @@
+// Copyright Epic Games, Inc. All Rights Reserved.
+
+using EpicGames.Core;
+using EpicGames.Horde;
+using EpicGames.Horde.Api;
+using Microsoft.Extensions.Logging;
+
+namespace Horde.Commands.Bundles
+{
+	[Command("artifact", "find", "Finds artifacts matching the given query parameters")]
+	class ArtifactFind : Command
+	{
+		[CommandLine("-Key=")]
+		public List<string> Keys { get; } = new List<string>();
+
+		public override async Task<int> ExecuteAsync(ILogger logger)
+		{
+			using HordeHttpClient horde = await Settings.GetHttpCientAsync(logger);
+
+			List<GetArtifactResponse> artifacts = await horde.FindArtifactsAsync(keys: Keys);
+			foreach (GetArtifactResponse artifact in artifacts)
+			{
+				logger.LogInformation("");
+				logger.LogInformation("Artifact {Id} ({Type})", artifact.Id, artifact.Type);
+				foreach (string key in artifact.Keys)
+				{
+					logger.LogInformation("  {Key}", key);
+				}
+			}
+
+			return 0;
+		}
+	}
+}
diff --git a/Engine/Source/Programs/Horde/Horde/Settings.cs b/Engine/Source/Programs/Horde/Horde/Settings.cs
--- a/Engine/Source/Programs/Horde/Horde/Settings.cs
+++ b/Engine/Source/Programs/Horde/Horde/Settings.cs
@@ -2,6 +2,7 @@
 
 using System.Text.Json;
 using EpicGames.Core;
+using EpicGames.Horde;
 using EpicGames.OIDC;
 using Microsoft.Extensions.Configuration;
 using Microsoft.Extensions.Logging;
@@ -138,5 +139,22 @@ namespace Horde
 			logger.LogInformation("Received bearer token for {Server}", settings.Server);
 			return result.AccessToken;
 		}
+
+		public static async Task<HordeHttpClient> GetHttpCientAsync(ILogger logger, CancellationToken cancellationToken = default)
+		{
+			Uri? server = await GetServerAsync(cancellationToken);
+			if (server == null)
+			{
+				throw new Exception("No server is configured. Run 'horde login -server=...' to set up.");
+			}
+
+			string? token = await GetAccessTokenAsync(logger, cancellationToken);
+			if (token == null)
+			{
+				throw new Exception("Unable to log in to server.");
+			}
+
+			return new HordeHttpClient(server, token);
+		}
 	}
 }
diff --git a/Engine/Source/Programs/Shared/EpicGames.Horde/Api/ArtifactId.cs b/Engine/Source/Programs/Shared/EpicGames.Horde/Api/ArtifactId.cs
new file mode 100644
--- /dev/null
+++ b/Engine/Source/Programs/Shared/EpicGames.Horde/Api/ArtifactId.cs
@@ -0,0 +1,34 @@
+// Copyright Epic Games, Inc. All Rights Reserved.
+
+using System.ComponentModel;
+using System.Text.Json.Serialization;
+
+namespace EpicGames.Horde.Api
+{
+	/// <summary>
+	/// Unique id for an artifact
+	/// </summary>
+	/// <param name="Id">Identifier for the artifact</param>
+	[TypeConverter(typeof(BinaryIdTypeConverter<ArtifactId, ArtifactIdConverter>))]
+	[BinaryIdConverter(typeof(ArtifactIdConverter))]
+	public record struct ArtifactId(BinaryId Id)
+	{
+		/// <inheritdoc cref="BinaryId.Parse(System.String)"/>
+		public static ArtifactId Parse(string text) => new ArtifactId(BinaryId.Parse(text));
+
+		/// <inheritdoc/>
+		public override string ToString() => Id.ToString();
+	}
+
+	/// <summary>
+	/// Converter class to and from ObjectId values
+	/// </summary>
+	class ArtifactIdConverter : BinaryIdConverter<ArtifactId>
+	{
+		/// <inheritdoc/>
+		public override ArtifactId FromBinaryId(BinaryId id) => new ArtifactId(id);
+
+		/// <inheritdoc/>
+		public override BinaryId ToBinaryId(ArtifactId value) => value.Id;
+	}
+}
diff --git a/Engine/Source/Programs/Shared/EpicGames.Horde/Api/ArtifactMessages.cs b/Engine/Source/Programs/Shared/EpicGames.Horde/Api/ArtifactMessages.cs
new file mode 100644
--- /dev/null
+++ b/Engine/Source/Programs/Shared/EpicGames.Horde/Api/ArtifactMessages.cs
@@ -0,0 +1,290 @@
+﻿// Copyright Epic Games, Inc. All Rights Reserved.
+
+using System;
+using System.Collections.Generic;
+using System.Security.Cryptography;
+using System.Text;
+using System.Threading;
+using System.Threading.Tasks;
+using EpicGames.Core;
+
+#pragma warning disable CA2227
+
+namespace EpicGames.Horde.Api
+{
+	/// <summary>
+	/// Describes an artifact
+	/// </summary>
+	public class GetArtifactResponse
+	{
+		/// <summary>
+		/// Identifier for the Artifact. Randomly generated.
+		/// </summary>
+		public ArtifactId Id { get; }
+
+		/// <summary>
+		/// Type of artifact
+		/// </summary>
+		public ArtifactType Type { get; }
+
+		/// <summary>
+		/// Keys used to collate artifacts
+		/// </summary>
+		public IReadOnlyList<string> Keys { get; }
+
+		/// <summary>
+		/// Constructor
+		/// </summary>
+		public GetArtifactResponse(ArtifactId id, ArtifactType type, IReadOnlyList<string> keys)
+		{
+			Id = id;
+			Type = type;
+			Keys = keys;
+		}
+	}
+
+	/// <summary>
+	/// Result of an artifact search
+	/// </summary>
+	public class FindArtifactsResponse
+	{
+		/// <summary>
+		/// List of artifacts matching the search criteria
+		/// </summary>
+		public List<GetArtifactResponse> Artifacts { get; set; } = new List<GetArtifactResponse>();
+	}
+
+	/// <summary>
+	/// Describes a file within an artifact
+	/// </summary>
+	public class GetArtifactFileEntryResponse
+	{
+		/// <summary>
+		/// Name of this file
+		/// </summary>
+		public string Name { get; }
+
+		/// <summary>
+		/// Length of this entry
+		/// </summary>
+		public long Length { get; }
+
+		/// <summary>
+		/// Hash of the target node
+		/// </summary>
+		public IoHash Hash { get; }
+
+		/// <summary>
+		/// Constructor
+		/// </summary>
+		public GetArtifactFileEntryResponse(string name, long length, IoHash hash)
+		{
+			Name = name;
+			Length = length;
+			Hash = hash;
+		}
+	}
+
+	/// <summary>
+	/// Describes a file within an artifact
+	/// </summary>
+	public class GetArtifactDirectoryEntryResponse : GetArtifactDirectoryResponse
+	{
+		/// <summary>
+		/// Name of this file
+		/// </summary>
+		public string Name { get; }
+
+		/// <summary>
+		/// Length of this entry
+		/// </summary>
+		public long Length { get; }
+
+		/// <summary>
+		/// Hash of the target node
+		/// </summary>
+		public IoHash Hash { get; }
+
+		/// <summary>
+		/// Constructor
+		/// </summary>
+		public GetArtifactDirectoryEntryResponse(string name, long length, IoHash hash)
+		{
+			Name = name;
+			Length = length;
+			Hash = hash;
+		}
+	}
+
+	/// <summary>
+	/// Describes a directory within an artifact
+	/// </summary>
+	public class GetArtifactDirectoryResponse
+	{
+		/// <summary>
+		/// Names of sub-directories
+		/// </summary>
+		public List<GetArtifactDirectoryEntryResponse>? Directories { get; set; }
+
+		/// <summary>
+		/// Files within the directory
+		/// </summary>
+		public List<GetArtifactFileEntryResponse>? Files { get; set; }
+	}
+
+	/// <summary>
+	/// Request to create a zip file with artifact data
+	/// </summary>
+	public class CreateZipRequest
+	{
+		/// <summary>
+		/// Filter lines for the zip. Uses standard <see cref="FileFilter"/> syntax.
+		/// </summary>
+		public List<string> Filter { get; set; } = new List<string>();
+	}
+
+	/// <summary>
+	/// Extension methods for the artifacts endpoint
+	/// </summary>
+	public static class ArtifactExtensions
+	{
+		/// <summary>
+		/// Gets metadata about an artifact object
+		/// </summary>
+		/// <param name="horde">The Horde client instance</param>
+		/// <param name="id">Identifier for the artifact</param>
+		/// <param name="cancellationToken">Cancellation token for the operation</param>
+		public static Task<GetArtifactResponse> GetArtifactAsync(this HordeHttpClient horde, ArtifactId id, CancellationToken cancellationToken = default)
+		{
+			return horde.GetAsync<GetArtifactResponse>($"api/v2/artifacts/{id}", cancellationToken);
+		}
+
+		/// <summary>
+		/// Finds artifacts with a set of ids or keys
+		/// </summary>
+		/// <param name="horde">The Horde client instance</param>
+		/// <param name="ids">Artifact ids to return</param>
+		/// <param name="keys">Keys to find</param>
+		/// <param name="cancellationToken">Cancellation token for the operation</param>
+		/// <returns>Information about all the artifacts</returns>
+		public static async Task<List<GetArtifactResponse>> FindArtifactsAsync(this HordeHttpClient horde, IEnumerable<ArtifactId>? ids = null, IEnumerable<string>? keys = null, CancellationToken cancellationToken = default)
+		{
+			QueryStringBuilder queryParams = new QueryStringBuilder();
+			if (ids != null)
+			{
+				foreach (ArtifactId id in ids)
+				{
+					queryParams.Add("id", id.ToString());
+				}
+			}
+			if (keys != null)
+			{
+				foreach (string key in keys)
+				{
+					queryParams.Add("key", key);
+				}
+			}
+
+			FindArtifactsResponse response = await horde.GetAsync<FindArtifactsResponse>($"api/v2/artifacts?{queryParams}", cancellationToken);
+			return response.Artifacts;
+		}
+
+
+#if false
+			/// <summary>
+			/// Gets metadata about an artifact object
+			/// </summary>
+			/// <param name="id">Identifier of the artifact to retrieve</param>
+			/// <param name="filter">Filter for returned properties</param>
+			/// <returns>Information about all the artifacts</returns>
+			[HttpGet]
+			[Route("/api/v2/artifacts/{id}")]
+			[ProducesResponseType(typeof(FindArtifactsResponse), 200)]
+			public async Task<ActionResult<object>> GetArtifactAsync(ArtifactId id, [FromQuery] PropertyFilter? filter = null);
+
+			/// <summary>
+			/// Retrieves bundles for a particular artifact
+			/// </summary>
+			/// <param name="id">Identifier of the artifact to retrieve</param>
+			/// <param name="locator">The blob locator</param>
+			/// <param name="length">Length of data to return</param>
+			/// <param name="offset">Offset of the data to return</param>
+			/// <param name="cancellationToken">Cancellation token for the operation</param>
+			/// <returns>Information about all the artifacts</returns>
+			[HttpGet]
+			[Route("/api/v2/artifacts/{id}/blobs/{*locator}")]
+			[Route("/api/v2/artifacts/{id}/bundles/{*locator}")]
+			public async Task<ActionResult<object>> ReadArtifactBlobAsync(ArtifactId id, BundleLocator locator, [FromQuery] int? offset = null, [FromQuery] int? length = null, CancellationToken cancellationToken = default)
+
+			/// <summary>
+			/// Retrieves the root blob for an artifact
+			/// </summary>
+			/// <param name="id">Identifier of the artifact to retrieve</param>
+			/// <param name="cancellationToken">Cancellation token for the operation</param>
+			/// <returns>Information about all the artifacts</returns>
+			[HttpGet]
+			[Route("/api/v2/artifacts/{id}/ref")]
+			public async Task<ActionResult<object>> ReadArtifactRefAsync(ArtifactId id, CancellationToken cancellationToken = default)
+
+			/// <summary>
+			/// Gets metadata about an artifact object
+			/// </summary>
+			/// <param name="id">Identifier of the artifact to retrieve</param>
+			/// <param name="path">Path to fetch</param>
+			/// <param name="filter">Filter for returned properties</param>
+			/// <param name="cancellationToken">Cancellation token for the operation</param>
+			/// <returns>Information about all the artifacts</returns>
+			[HttpGet]
+			[Route("/api/v2/artifacts/{id}/browse")]
+			[ProducesResponseType(typeof(GetArtifactDirectoryResponse), 200)]
+			public async Task<ActionResult<object>> BrowseArtifactAsync(ArtifactId id, [FromQuery] string? path = null, [FromQuery] PropertyFilter? filter = null, CancellationToken cancellationToken = default)
+
+			/// <summary>
+			/// Downloads an individual file from an artifact
+			/// </summary>
+			/// <param name="id">Identifier of the artifact to retrieve</param>
+			/// <param name="path">Path to fetch</param>
+			/// <param name="inline">Whether to request the file be downloaded vs displayed inline</param>
+			/// <param name="cancellationToken">Cancellation token for the operation</param>
+			/// <returns>Information about all the artifacts</returns>
+			[HttpGet]
+			[Route("/api/v2/artifacts/{id}/file")]
+			public async Task<ActionResult<object>> GetFileAsync(ArtifactId id, [FromQuery] string path, [FromQuery] bool inline = false, CancellationToken cancellationToken = default)
+
+			/// <summary>
+			/// Downloads an individual file from an artifact
+			/// </summary>
+			/// <param name="id">Identifier of the artifact to retrieve</param>
+			/// <param name="filter">Paths to include in the zip file. The post version of this request allows for more parameters than can fit in a request string.</param>
+			/// <param name="cancellationToken">Cancellation token for the operation</param>
+			/// <returns>Information about all the artifacts</returns>
+			[HttpGet]
+			[Route("/api/v2/artifacts/{id}/zip")]
+			public async Task<ActionResult<object>> GetZipAsync(ArtifactId id, [FromQuery(Name = "filter")] string[]? filter, CancellationToken cancellationToken = default)
+
+			/// <summary>
+			/// Downloads an individual file from an artifact
+			/// </summary>
+			/// <param name="id">Identifier of the artifact to retrieve</param>
+			/// <param name="request">Filter for the zip file</param>
+			/// <param name="cancellationToken">Cancellation token for the operation</param>
+			/// <returns>Information about all the artifacts</returns>
+			[HttpPost]
+			[Route("/api/v2/artifacts/{id}/zip")]
+			public async Task<ActionResult<object>> CreateZipFromFilterAsync(ArtifactId id, CreateZipRequest request, CancellationToken cancellationToken = default)
+
+			/// <summary>
+			/// Gets metadata about an artifact object
+			/// </summary>
+			/// <param name="ids">Artifact ids to return</param>
+			/// <param name="keys">Keys to find</param>
+			/// <param name="filter">Filter for returned values</param>
+			/// <returns>Information about all the artifacts</returns>
+			[HttpGet]
+			[Route("/api/v2/artifacts")]
+			[ProducesResponseType(typeof(FindArtifactsResponse), 200)]
+			public async Task<ActionResult<object>> FindArtifactsAsync([FromQuery(Name = "id")] IEnumerable<ArtifactId>? ids = null, [FromQuery(Name = "key")] IEnumerable<string>? keys = null, [FromQuery] PropertyFilter? filter = null)
+		}
+#endif
+	}
+}
diff --git a/Engine/Source/Programs/Shared/EpicGames.Horde/Api/ArtifactType.cs b/Engine/Source/Programs/Shared/EpicGames.Horde/Api/ArtifactType.cs
new file mode 100644
--- /dev/null
+++ b/Engine/Source/Programs/Shared/EpicGames.Horde/Api/ArtifactType.cs
@@ -0,0 +1,63 @@
+// Copyright Epic Games, Inc. All Rights Reserved.
+
+using System.ComponentModel;
+using EpicGames.Core;
+
+namespace EpicGames.Horde.Api
+{
+	/// <summary>
+	/// Type of an artifact
+	/// </summary>
+	/// <param name="Id">The artifact type</param>
+	[JsonSchemaString]
+	[StringIdConverter(typeof(ArtifactTypeConverter))]
+	[TypeConverter(typeof(StringIdTypeConverter<ArtifactType, ArtifactTypeConverter>))]
+	public record struct ArtifactType(StringId Id)
+	{
+		/// <summary>
+		/// Default artifact type 
+		/// </summary>
+		public static ArtifactType Unknown { get; } = default;
+
+		/// <summary>
+		/// Output from a build step
+		/// </summary>
+		public static ArtifactType StepOutput { get; } = new ArtifactType("step-output");
+
+		/// <summary>
+		/// Captured state from the machine executing a build step
+		/// </summary>
+		public static ArtifactType StepSaved { get; } = new ArtifactType("step-saved");
+
+		/// <summary>
+		/// Traces from the machine executing a build step
+		/// </summary>
+		public static ArtifactType StepTrace { get; } = new ArtifactType("step-trace");
+
+		/// <summary>
+		/// Test data generated by a machine executing a build step
+		/// </summary>
+		public static ArtifactType StepTestData { get; } = new ArtifactType("step-testdata");
+
+		/// <summary>
+		/// Constructor
+		/// </summary>
+		/// <param name="id">Identifier for the artifact type</param>
+		public ArtifactType(string id) : this(new StringId(id)) { }
+
+		/// <inheritdoc/>
+		public override string ToString() => Id.ToString();
+	}
+
+	/// <summary>
+	/// Converter to and from <see cref="StringId"/> instances.
+	/// </summary>
+	class ArtifactTypeConverter : StringIdConverter<ArtifactType>
+	{
+		/// <inheritdoc/>
+		public override ArtifactType FromStringId(StringId id) => new ArtifactType(id);
+
+		/// <inheritdoc/>
+		public override StringId ToStringId(ArtifactType value) => value.Id;
+	}
+}
diff --git a/Engine/Source/Programs/Shared/EpicGames.Horde/HordeHttpClient.cs b/Engine/Source/Programs/Shared/EpicGames.Horde/HordeHttpClient.cs
--- a/Engine/Source/Programs/Shared/EpicGames.Horde/HordeHttpClient.cs
+++ b/Engine/Source/Programs/Shared/EpicGames.Horde/HordeHttpClient.cs
@@ -65,6 +65,7 @@ namespace EpicGames.Horde
 			options.PropertyNameCaseInsensitive = true;
 			options.Converters.Add(new JsonStringEnumConverter());
 			options.Converters.Add(new StringIdJsonConverterFactory());
+			options.Converters.Add(new BinaryIdJsonConverterFactory());
 		}
 
 		/// <summary>
@@ -81,7 +82,8 @@ namespace EpicGames.Horde
 				response.EnsureSuccessStatusCode();
 
 				byte[] bytes = await response.Content.ReadAsByteArrayAsync(cancellationToken);
-				return JsonSerializer.Deserialize<TResponse>(bytes, _jsonSerializerOptions)!;
+				TResponse result = JsonSerializer.Deserialize<TResponse>(bytes, _jsonSerializerOptions)!;
+				return result;
 			}
 		}
 	}
