Horde: Fix serialization of BinaryId types to JSON.

diff --git a/Engine/Source/Programs/Shared/EpicGames.Horde/BinaryId.cs b/Engine/Source/Programs/Shared/EpicGames.Horde/BinaryId.cs
--- a/Engine/Source/Programs/Shared/EpicGames.Horde/BinaryId.cs
+++ b/Engine/Source/Programs/Shared/EpicGames.Horde/BinaryId.cs
@@ -23,6 +23,16 @@ namespace EpicGames.Horde
 		readonly int _b;
 		readonly int _c;
 
+		/// <summary>
+		/// Number of bytes in the identifier
+		/// </summary>
+		public const int NumBytes = 12;
+
+		/// <summary>
+		/// Number of characters when formatted as a string
+		/// </summary>
+		public const int NumChars = NumBytes * 2;
+
 		/// <summary>
 		/// Constructor
 		/// </summary>
@@ -70,7 +80,7 @@ namespace EpicGames.Horde
 		/// </summary>
 		public static bool TryParse(string text, out BinaryId result)
 		{
-			Span<byte> bytes = stackalloc byte[12];
+			Span<byte> bytes = stackalloc byte[NumBytes];
 			if (StringUtils.TryParseHexString(text, bytes))
 			{
 				result = new BinaryId(bytes);
@@ -93,7 +103,7 @@ namespace EpicGames.Horde
 		/// </summary>
 		public static bool TryParse(ReadOnlySpan<byte> text, out BinaryId result)
 		{
-			Span<byte> bytes = stackalloc byte[12];
+			Span<byte> bytes = stackalloc byte[NumBytes];
 			if (StringUtils.TryParseHexString(text, bytes))
 			{
 				result = new BinaryId(bytes);
@@ -123,7 +133,7 @@ namespace EpicGames.Horde
 		/// <inheritdoc/>
 		public override string ToString()
 		{
-			Span<byte> bytes = stackalloc byte[12];
+			Span<byte> bytes = stackalloc byte[NumBytes];
 			ToByteArray(bytes);
 			return StringUtils.FormatHexString(bytes);
 		}
@@ -134,8 +144,8 @@ namespace EpicGames.Horde
 		public void ToUtf8String(Span<byte> chars)
 		{
 			StringUtils.FormatUtf8HexString((uint)_a, chars);
-			StringUtils.FormatUtf8HexString((uint)_b, chars[4..]);
-			StringUtils.FormatUtf8HexString((uint)_c, chars[8..]);
+			StringUtils.FormatUtf8HexString((uint)_b, chars[8..]);
+			StringUtils.FormatUtf8HexString((uint)_c, chars[16..]);
 		}
 
 		/// <summary>
@@ -143,7 +153,7 @@ namespace EpicGames.Horde
 		/// </summary>
 		public byte[] ToByteArray()
 		{
-			byte[] bytes = new byte[12];
+			byte[] bytes = new byte[NumBytes];
 			ToByteArray(bytes);
 			return bytes;
 		}
@@ -151,11 +161,11 @@ namespace EpicGames.Horde
 		/// <summary>
 		/// Format this id as a sequence of UTF8 characters
 		/// </summary>
-		public void ToByteArray(Span<byte> chars)
+		public void ToByteArray(Span<byte> bytes)
 		{
-			BinaryPrimitives.WriteInt32LittleEndian(chars, _a);
-			BinaryPrimitives.WriteInt32LittleEndian(chars[4..], _b);
-			BinaryPrimitives.WriteInt32LittleEndian(chars[8..], _c);
+			BinaryPrimitives.WriteInt32LittleEndian(bytes, _a);
+			BinaryPrimitives.WriteInt32LittleEndian(bytes[4..], _b);
+			BinaryPrimitives.WriteInt32LittleEndian(bytes[8..], _c);
 		}
 
 		/// <summary>
@@ -186,7 +196,7 @@ namespace EpicGames.Horde
 		/// <inheritdoc/>
 		public override void Write(Utf8JsonWriter writer, BinaryId value, JsonSerializerOptions options)
 		{
-			Span<byte> bytes = stackalloc byte[12 * 2];
+			Span<byte> bytes = stackalloc byte[BinaryId.NumChars];
 			value.ToUtf8String(bytes);
 			writer.WriteStringValue(bytes);
 		}
