Horde: Add overloads to attempt getting an access token without throwing an exception on failure.

diff --git a/Engine/Source/Programs/Horde/Horde/Settings.cs b/Engine/Source/Programs/Horde/Horde/Settings.cs
--- a/Engine/Source/Programs/Horde/Horde/Settings.cs
+++ b/Engine/Source/Programs/Horde/Horde/Settings.cs
@@ -121,12 +121,8 @@ namespace Horde
 			using ITokenStore tokenStore = TokenStoreFactory.CreateTokenStore();
 			OidcTokenManager oidcTokenManager = OidcTokenManager.CreateTokenManager(configuration, tokenStore, new List<string>() { OidcProvider });
 
-			OidcTokenInfo result;
-			try
-			{
-				result = await oidcTokenManager.GetAccessToken(OidcProvider, cancellationToken);
-			}
-			catch (NotLoggedInException)
+			OidcTokenInfo? result = await oidcTokenManager.TryGetAccessToken(OidcProvider, cancellationToken);
+			if(result == null)
 			{
 				result = await oidcTokenManager.Login(OidcProvider, cancellationToken);
 			}
diff --git a/Engine/Source/Programs/Shared/EpicGames.OIDC/OidcTokenManager.cs b/Engine/Source/Programs/Shared/EpicGames.OIDC/OidcTokenManager.cs
--- a/Engine/Source/Programs/Shared/EpicGames.OIDC/OidcTokenManager.cs
+++ b/Engine/Source/Programs/Shared/EpicGames.OIDC/OidcTokenManager.cs
@@ -148,6 +148,11 @@ namespace EpicGames.OIDC
 			return _tokenClients[providerIdentifier].GetAccessTokenAsync(cancellationToken);
 		}
 
+		public Task<OidcTokenInfo?> TryGetAccessToken(string providerIdentifier, CancellationToken cancellationToken = default)
+		{
+			return _tokenClients[providerIdentifier].TryGetAccessTokenAsync(cancellationToken);
+		}
+
 		public OidcStatus GetStatusForProvider(string providerIdentifier)
 		{
 			return _tokenClients[providerIdentifier].GetStatus();
@@ -449,7 +454,7 @@ namespace EpicGames.OIDC
 			return false;
 		}
 
-		private async Task<OidcTokenInfo> DoRefreshToken(string inRefreshToken, CancellationToken cancellationToken)
+		private async Task<OidcTokenInfo?> TryDoRefreshToken(string inRefreshToken, CancellationToken cancellationToken)
 		{
 			// redirect uri is not used for refrehs tokens so we can just pick one of them to configure the client
 			OidcClientOptions options = await BuildClientOptions(_redirectUris.First(), cancellationToken);
@@ -464,7 +469,7 @@ namespace EpicGames.OIDC
 				{
 					// the refresh token is no logger valid, resetting it and treating us as not logged in
 					_refreshToken = null;
-					throw new NotLoggedInException();
+					return null;
 				}
 				else
 				{
@@ -507,7 +512,18 @@ namespace EpicGames.OIDC
 
 			return disco;
 		}
+
 		public async Task<OidcTokenInfo> GetAccessTokenAsync(CancellationToken cancellationToken)
+		{
+			OidcTokenInfo? tokenInfo = await TryGetAccessTokenAsync(cancellationToken);
+			if (tokenInfo == null)
+			{
+				throw new NotLoggedInException();
+			}
+			return tokenInfo;
+		}
+
+		public async Task<OidcTokenInfo?> TryGetAccessTokenAsync(CancellationToken cancellationToken)
 		{
 			if (String.IsNullOrEmpty(_refreshToken))
 			{
@@ -526,7 +542,7 @@ namespace EpicGames.OIDC
 				};
 			}
 
-			return await DoRefreshToken(_refreshToken, cancellationToken);
+			return await TryDoRefreshToken(_refreshToken, cancellationToken);
 		}
 
 		public OidcStatus GetStatus()
