Horde: Add additional debug logging when extracting files from temp storage.

diff --git a/Engine/Source/Programs/Shared/EpicGames.Horde/Storage/Nodes/FileEntry.cs b/Engine/Source/Programs/Shared/EpicGames.Horde/Storage/Nodes/FileEntry.cs
--- a/Engine/Source/Programs/Shared/EpicGames.Horde/Storage/Nodes/FileEntry.cs
+++ b/Engine/Source/Programs/Shared/EpicGames.Horde/Storage/Nodes/FileEntry.cs
@@ -154,20 +154,27 @@ namespace EpicGames.Horde.Storage.Nodes
 		/// <returns></returns>
 		public async Task CopyToFileAsync(FileInfo file, CancellationToken cancellationToken)
 		{
-			ChunkedDataNode node = await ExpandAsync(cancellationToken);
-			await node.CopyToFileAsync(file, cancellationToken);
-
-			if ((Flags & FileEntryFlags.Executable) != 0)
+			try
 			{
-				if (RuntimeInformation.IsOSPlatform(OSPlatform.Linux))
-				{
-					FileUtils.SetFileMode_Linux(file.FullName, 0b_111_111_111);
-				}
-				else if (RuntimeInformation.IsOSPlatform(OSPlatform.OSX))
+				ChunkedDataNode node = await ExpandAsync(cancellationToken);
+				await node.CopyToFileAsync(file, cancellationToken);
+
+				if ((Flags & FileEntryFlags.Executable) != 0)
 				{
-					FileUtils.SetFileMode_Mac(file.FullName, 0b_111_111_111);
+					if (RuntimeInformation.IsOSPlatform(OSPlatform.Linux))
+					{
+						FileUtils.SetFileMode_Linux(file.FullName, 0b_111_111_111);
+					}
+					else if (RuntimeInformation.IsOSPlatform(OSPlatform.OSX))
+					{
+						FileUtils.SetFileMode_Mac(file.FullName, 0b_111_111_111);
+					}
 				}
 			}
+			catch (Exception ex)
+			{
+				throw new Exception($"Unable to extract file {file.FullName}", ex);
+			}
 		}
 
 		/// <inheritdoc/>
