Horde: Fix BinaryId to utf8 string conversion.

diff --git a/Engine/Source/Programs/Shared/EpicGames.Core/StringUtils.cs b/Engine/Source/Programs/Shared/EpicGames.Core/StringUtils.cs
--- a/Engine/Source/Programs/Shared/EpicGames.Core/StringUtils.cs
+++ b/Engine/Source/Programs/Shared/EpicGames.Core/StringUtils.cs
@@ -520,6 +520,26 @@ namespace EpicGames.Core
 			}
 		}
 
+		/// <summary>
+		/// Formats a 32-bit unsigned integer as a hexadecimal string
+		/// </summary>
+		/// <param name="value">Value to render</param>
+		/// <param name="characters">Buffer to receive the characters</param>
+		public static void FormatLittleEndianUtf8HexString(uint value, Span<byte> characters)
+		{
+			characters[0] = s_hexDigitToUtf8Byte[(value >> 4) & 15];
+			characters[1] = s_hexDigitToUtf8Byte[value & 15];
+
+			characters[2] = s_hexDigitToUtf8Byte[(value >> 12) & 15];
+			characters[3] = s_hexDigitToUtf8Byte[(value >> 8) & 15];
+
+			characters[4] = s_hexDigitToUtf8Byte[(value >> 20) & 15];
+			characters[5] = s_hexDigitToUtf8Byte[(value >> 16) & 15];
+
+			characters[6] = s_hexDigitToUtf8Byte[(value >> 28) & 15];
+			characters[7] = s_hexDigitToUtf8Byte[(value >> 24) & 15];
+		}
+
 		/// <summary>
 		/// Formats a 32-bit unsigned integer as a hexadecimal string
 		/// </summary>
diff --git a/Engine/Source/Programs/Shared/EpicGames.Horde/BinaryId.cs b/Engine/Source/Programs/Shared/EpicGames.Horde/BinaryId.cs
--- a/Engine/Source/Programs/Shared/EpicGames.Horde/BinaryId.cs
+++ b/Engine/Source/Programs/Shared/EpicGames.Horde/BinaryId.cs
@@ -138,14 +138,24 @@ namespace EpicGames.Horde
 			return StringUtils.FormatHexString(bytes);
 		}
 
+		/// <summary>
+		/// Format this id as a sequence of UTF8 characters
+		/// </summary>
+		public Utf8String ToUtf8String()
+		{
+			byte[] data = new byte[NumChars];
+			ToUtf8String(data.AsSpan());
+			return new Utf8String(data);
+		}
+
 		/// <summary>
 		/// Format this id as a sequence of UTF8 characters
 		/// </summary>
 		public void ToUtf8String(Span<byte> chars)
 		{
-			StringUtils.FormatUtf8HexString((uint)_a, chars);
-			StringUtils.FormatUtf8HexString((uint)_b, chars[8..]);
-			StringUtils.FormatUtf8HexString((uint)_c, chars[16..]);
+			StringUtils.FormatLittleEndianUtf8HexString((uint)_a, chars);
+			StringUtils.FormatLittleEndianUtf8HexString((uint)_b, chars[8..]);
+			StringUtils.FormatLittleEndianUtf8HexString((uint)_c, chars[16..]);
 		}
 
 		/// <summary>
