Fix for Replicating Gameplay Ability Blueprints with Push Model enabled.  Manually declare REPLICATED_BASE_CLASS in GameplayAbility and fix-up UHT to account for this.

diff --git a/Engine/Plugins/Runtime/GameplayAbilities/Source/GameplayAbilities/Public/Abilities/GameplayAbility.h b/Engine/Plugins/Runtime/GameplayAbilities/Source/GameplayAbilities/Public/Abilities/GameplayAbility.h
--- a/Engine/Plugins/Runtime/GameplayAbilities/Source/GameplayAbilities/Public/Abilities/GameplayAbility.h
+++ b/Engine/Plugins/Runtime/GameplayAbilities/Source/GameplayAbilities/Public/Abilities/GameplayAbility.h
@@ -13,6 +13,7 @@
 #include "Abilities/GameplayAbilityTypes.h"
 #include "GameplayTaskOwnerInterface.h"
 #include "Abilities/GameplayAbilityTargetTypes.h"
+#include "Net/Core/PushModel/PushModelMacros.h"
 #include "GameplayAbility.generated.h"
 
 class UAbilitySystemComponent;
@@ -107,6 +108,7 @@ UCLASS(Blueprintable)
 class GAMEPLAYABILITIES_API UGameplayAbility : public UObject, public IGameplayTaskOwnerInterface
 {
 	GENERATED_UCLASS_BODY()
+	REPLICATED_BASE_CLASS(UGameplayAbility)
 
 	friend class UAbilitySystemComponent;
 	friend class UGameplayAbilitySet;
diff --git a/Engine/Source/Programs/Shared/EpicGames.UHT/Exporters/CodeGen/UhtHeaderCodeGeneratorHFile.cs b/Engine/Source/Programs/Shared/EpicGames.UHT/Exporters/CodeGen/UhtHeaderCodeGeneratorHFile.cs
--- a/Engine/Source/Programs/Shared/EpicGames.UHT/Exporters/CodeGen/UhtHeaderCodeGeneratorHFile.cs
+++ b/Engine/Source/Programs/Shared/EpicGames.UHT/Exporters/CodeGen/UhtHeaderCodeGeneratorHFile.cs
@@ -1500,9 +1500,21 @@ namespace EpicGames.UHT.Exporters.CodeGen
 			// it's the base most replicated class. In that case, go ahead and add our interface macro.
 			if (classObj.ClassExportFlags.HasExactFlags(UhtClassExportFlags.HasReplciatedProperties, UhtClassExportFlags.SelfHasReplicatedProperties))
 			{
-				builder.Append("private: \\\r\n");
-				builder.Append("\tREPLICATED_BASE_CLASS(").Append(classObj.SourceName).Append(") \\\r\n");
-				builder.Append("public: \\\r\n");
+				// Make sure the client hasn't implemented it themselves
+				bool alreadyDefinedPushModel = false;
+				UhtClass? checkForReplicatedBaseClass = classObj;
+				while (!alreadyDefinedPushModel && checkForReplicatedBaseClass != null)
+				{
+					alreadyDefinedPushModel = checkForReplicatedBaseClass.TryGetDeclaration("REPLICATED_BASE_CLASS", out _);
+					checkForReplicatedBaseClass = checkForReplicatedBaseClass.SuperClass;
+				}
+
+				if (!alreadyDefinedPushModel)
+				{
+					builder.Append("private: \\\r\n");
+					builder.Append("\tREPLICATED_BASE_CLASS(").Append(classObj.SourceName).Append(") \\\r\n");
+					builder.Append("public: \\\r\n");
+				}
 			}
 			return builder;
 		}
