Horde: Move user response classes into EpicGames.Horde.

diff --git a/Engine/Source/Programs/Horde/Horde.Server.Tests/BisectTests.cs b/Engine/Source/Programs/Horde/Horde.Server.Tests/BisectTests.cs
--- a/Engine/Source/Programs/Horde/Horde.Server.Tests/BisectTests.cs
+++ b/Engine/Source/Programs/Horde/Horde.Server.Tests/BisectTests.cs
@@ -42,7 +42,7 @@ namespace Horde.Server.Tests
 			Assert.AreEqual(failedJob.StreamId, response.StreamId);
 			Assert.AreEqual("CompileEditor", response.NodeName);
 			Assert.AreEqual(JobStepOutcome.Failure, response.Outcome);
-			Assert.AreEqual("TestUser", response.Owner.Name);
+			Assert.AreEqual("TestUser", response.Owner?.Name ?? "");
 			Assert.AreEqual(failedJob.Change, response.InitialChange);
 			Assert.AreEqual(failedJob.Id, response.InitialJobId);
 			Assert.AreEqual(failedJob.Change, response.CurrentChange);
@@ -65,9 +65,9 @@ namespace Horde.Server.Tests
 			Assert.AreEqual(14, jobs[0].Change);
 			await SetJobOutcomeAsync(jobs[0], graph, JobStepOutcome.Failure);
 
-			List<GetBisectTaskResponse> bisectTasks = Deref(await BisectTasksController!.FindBisectTasksAsync(bisectTask.Owner.Id));
+			List<GetBisectTaskResponse> bisectTasks = Deref(await BisectTasksController!.FindBisectTasksAsync(bisectTask.Owner?.Id.ToString() ?? ""));
 			Assert.AreEqual(1, bisectTasks.Count);
-			Assert.AreEqual(bisectTask.Owner.Id, bisectTasks[0].Owner.Id);
+			Assert.AreEqual(bisectTask.Owner?.Id ?? default, bisectTasks[0].Owner?.Id ?? default);
 
 			await Clock.AdvanceAsync(TimeSpan.FromMinutes(30));
 
diff --git a/Engine/Source/Programs/Horde/Horde.Server.Tests/JobCollectionTests.cs b/Engine/Source/Programs/Horde/Horde.Server.Tests/JobCollectionTests.cs
--- a/Engine/Source/Programs/Horde/Horde.Server.Tests/JobCollectionTests.cs
+++ b/Engine/Source/Programs/Horde/Horde.Server.Tests/JobCollectionTests.cs
@@ -5,6 +5,7 @@ using System.Collections.Generic;
 using System.Linq;
 using System.Threading.Tasks;
 using EpicGames.Core;
+using EpicGames.Horde.Api;
 using Horde.Server.Jobs;
 using Horde.Server.Jobs.Graphs;
 using Horde.Server.Agents;
diff --git a/Engine/Source/Programs/Horde/Horde.Server.Tests/JobTaskSourceTests.cs b/Engine/Source/Programs/Horde/Horde.Server.Tests/JobTaskSourceTests.cs
--- a/Engine/Source/Programs/Horde/Horde.Server.Tests/JobTaskSourceTests.cs
+++ b/Engine/Source/Programs/Horde/Horde.Server.Tests/JobTaskSourceTests.cs
@@ -14,6 +14,7 @@ using System;
 using Horde.Server.Agents;
 using Horde.Server.Agents.Pools;
 using Horde.Server.Streams;
+using EpicGames.Horde.Api;
 
 namespace Horde.Server.Tests
 {
diff --git a/Engine/Source/Programs/Horde/Horde.Server.Tests/Stubs/Services/PerforceServiceStub.cs b/Engine/Source/Programs/Horde/Horde.Server.Tests/Stubs/Services/PerforceServiceStub.cs
--- a/Engine/Source/Programs/Horde/Horde.Server.Tests/Stubs/Services/PerforceServiceStub.cs
+++ b/Engine/Source/Programs/Horde/Horde.Server.Tests/Stubs/Services/PerforceServiceStub.cs
@@ -7,6 +7,7 @@ using System.Runtime.CompilerServices;
 using System.Threading;
 using System.Threading.Tasks;
 using EpicGames.Core;
+using EpicGames.Horde.Api;
 using EpicGames.Perforce;
 using Horde.Server.Perforce;
 using Horde.Server.Streams;
diff --git a/Engine/Source/Programs/Horde/Horde.Server/Authentication/AnonymousHandler.cs b/Engine/Source/Programs/Horde/Horde.Server/Authentication/AnonymousHandler.cs
--- a/Engine/Source/Programs/Horde/Horde.Server/Authentication/AnonymousHandler.cs
+++ b/Engine/Source/Programs/Horde/Horde.Server/Authentication/AnonymousHandler.cs
@@ -10,6 +10,7 @@ using Horde.Server.Utilities;
 using Microsoft.AspNetCore.Authentication;
 using Microsoft.Extensions.Logging;
 using Microsoft.Extensions.Options;
+using EpicGames.Horde.Api;
 
 namespace Horde.Server.Authentication
 {
diff --git a/Engine/Source/Programs/Horde/Horde.Server/Issues/IIssue.cs b/Engine/Source/Programs/Horde/Horde.Server/Issues/IIssue.cs
--- a/Engine/Source/Programs/Horde/Horde.Server/Issues/IIssue.cs
+++ b/Engine/Source/Programs/Horde/Horde.Server/Issues/IIssue.cs
@@ -7,6 +7,7 @@ using Horde.Server.Streams;
 using Horde.Server.Users;
 using Horde.Server.Utilities;
 using MongoDB.Bson;
+using EpicGames.Horde.Api;
 
 namespace Horde.Server.Issues
 {
diff --git a/Engine/Source/Programs/Horde/Horde.Server/Issues/IIssueCollection.cs b/Engine/Source/Programs/Horde/Horde.Server/Issues/IIssueCollection.cs
--- a/Engine/Source/Programs/Horde/Horde.Server/Issues/IIssueCollection.cs
+++ b/Engine/Source/Programs/Horde/Horde.Server/Issues/IIssueCollection.cs
@@ -13,6 +13,7 @@ using Horde.Server.Users;
 using Horde.Server.Utilities;
 using MongoDB.Bson;
 using MongoDB.Driver;
+using EpicGames.Horde.Api;
 
 namespace Horde.Server.Issues
 {
diff --git a/Engine/Source/Programs/Horde/Horde.Server/Issues/IIssueDetails.cs b/Engine/Source/Programs/Horde/Horde.Server/Issues/IIssueDetails.cs
--- a/Engine/Source/Programs/Horde/Horde.Server/Issues/IIssueDetails.cs
+++ b/Engine/Source/Programs/Horde/Horde.Server/Issues/IIssueDetails.cs
@@ -1,9 +1,11 @@
 // Copyright Epic Games, Inc. All Rights Reserved.
+
 using System;
 using System.Collections.Generic;
 using System.Threading.Tasks;
 using Horde.Server.Users;
 using Horde.Server.Utilities;
+using EpicGames.Horde.Api;
 
 namespace Horde.Server.Issues
 {
diff --git a/Engine/Source/Programs/Horde/Horde.Server/Issues/IIssueSpan.cs b/Engine/Source/Programs/Horde/Horde.Server/Issues/IIssueSpan.cs
--- a/Engine/Source/Programs/Horde/Horde.Server/Issues/IIssueSpan.cs
+++ b/Engine/Source/Programs/Horde/Horde.Server/Issues/IIssueSpan.cs
@@ -5,6 +5,7 @@ using Horde.Server.Streams;
 using Horde.Server.Users;
 using Horde.Server.Utilities;
 using MongoDB.Bson;
+using EpicGames.Horde.Api;
 
 namespace Horde.Server.Issues
 {
diff --git a/Engine/Source/Programs/Horde/Horde.Server/Issues/IssueCollection.cs b/Engine/Source/Programs/Horde/Horde.Server/Issues/IssueCollection.cs
--- a/Engine/Source/Programs/Horde/Horde.Server/Issues/IssueCollection.cs
+++ b/Engine/Source/Programs/Horde/Horde.Server/Issues/IssueCollection.cs
@@ -5,6 +5,7 @@ using System.Collections.Generic;
 using System.Diagnostics;
 using System.Linq;
 using System.Threading.Tasks;
+using EpicGames.Horde.Api;
 using EpicGames.Redis.Utility;
 using Horde.Server.Auditing;
 using Horde.Server.Jobs;
diff --git a/Engine/Source/Programs/Horde/Horde.Server/Issues/IssueMessages.cs b/Engine/Source/Programs/Horde/Horde.Server/Issues/IssueMessages.cs
--- a/Engine/Source/Programs/Horde/Horde.Server/Issues/IssueMessages.cs
+++ b/Engine/Source/Programs/Horde/Horde.Server/Issues/IssueMessages.cs
@@ -4,6 +4,7 @@ using System;
 using System.Collections.Generic;
 using System.Linq;
 using EpicGames.Core;
+using EpicGames.Horde.Api;
 using Horde.Server.Issues.External;
 using Horde.Server.Streams;
 using Horde.Server.Users;
@@ -553,24 +554,15 @@ namespace Horde.Server.Issues
 			Promoted = issue.Promoted;
 			Owner = details.Owner?.Login;
 			OwnerId = details.Owner?.Id.ToString();
-			if(details.Owner != null)
-			{
-				OwnerInfo = new GetThinUserInfoResponse(details.Owner);
-			}
+			OwnerInfo = details.Owner?.ToThinApiResponse();
 			NominatedBy = details.NominatedBy?.Login;
-			if (details.NominatedBy != null)
-			{
-				NominatedByInfo = new GetThinUserInfoResponse(details.NominatedBy);
-			}
+			NominatedByInfo = details.NominatedBy?.ToThinApiResponse();
 			AcknowledgedAt = issue.AcknowledgedAt;
 			FixChange = issue.FixChange;
 			ResolvedAt = issue.ResolvedAt;
 			ResolvedBy = details.ResolvedBy?.Login;
 			ResolvedById = details.ResolvedBy?.Id.ToString();
-			if (details.ResolvedBy != null)
-			{
-				ResolvedByInfo = new GetThinUserInfoResponse(details.ResolvedBy);
-			}
+			ResolvedByInfo = details.ResolvedBy?.ToThinApiResponse();
 			VerifiedAt = issue.VerifiedAt;
 			LastSeenAt = issue.LastSeenAt;
 			Streams = details.Spans.Select(x => x.StreamName).Distinct().ToList()!;
@@ -590,12 +582,12 @@ namespace Horde.Server.Issues
 			}
 			PrimarySuspects = details.SuspectUsers.Where(x => x.Login != null).Select(x => x.Login).ToList();
 			PrimarySuspectIds= details.SuspectUsers.Select(x => x.Id.ToString()).ToList();
-			PrimarySuspectsInfo = details.SuspectUsers.ConvertAll(x => new GetThinUserInfoResponse(x));
+			PrimarySuspectsInfo = details.SuspectUsers.ConvertAll(x => x.ToThinApiResponse());
 			ShowDesktopAlerts = showDesktopAlerts;
 			ExternalIssueKey = details.ExternalIssueKey;
-			QuarantinedByUserInfo = details.QuarantinedBy != null ? new GetThinUserInfoResponse(details.QuarantinedBy) : null;
+			QuarantinedByUserInfo = details.QuarantinedBy?.ToThinApiResponse();
 			QuarantineTimeUtc = details.QuarantineTimeUtc;
-			ForceClosedByUserInfo = details.ForceClosedBy != null ? new GetThinUserInfoResponse(details.ForceClosedBy) : null;
+			ForceClosedByUserInfo = details.ForceClosedBy?.ToThinApiResponse();
 			WorkflowThreadUrl = issue.WorkflowThreadUrl;
 		}
 	}
@@ -800,18 +792,18 @@ namespace Horde.Server.Issues
 			Promoted = issue.Promoted;
 			if (owner != null)
 			{
-				Owner = new GetThinUserInfoResponse(owner);
+				Owner = owner.ToThinApiResponse();
 			}
 			if (nominatedBy != null)
 			{
-				NominatedBy = new GetThinUserInfoResponse(nominatedBy);
+				NominatedBy = nominatedBy.ToThinApiResponse();
 			}
 			AcknowledgedAt = issue.AcknowledgedAt;
 			FixChange = issue.FixChange;
 			ResolvedAt = issue.ResolvedAt;
 			if (resolvedBy != null)
 			{
-				ResolvedBy = new GetThinUserInfoResponse(resolvedBy);
+				ResolvedBy = resolvedBy.ToThinApiResponse();
 			}
 			VerifiedAt = issue.VerifiedAt;
 			LastSeenAt = issue.LastSeenAt;
@@ -819,7 +811,7 @@ namespace Horde.Server.Issues
 			ExternalIssueKey = issue.ExternalIssueKey;
 			if (quarantinedBy != null)
 			{
-				QuarantinedBy = new GetThinUserInfoResponse(quarantinedBy);
+				QuarantinedBy = quarantinedBy.ToThinApiResponse();
 				QuarantineTimeUtc = issue.QuarantineTimeUtc;
 			}
 			WorkflowThreadUrl = issue.WorkflowThreadUrl;
diff --git a/Engine/Source/Programs/Horde/Horde.Server/Issues/IssueService.cs b/Engine/Source/Programs/Horde/Horde.Server/Issues/IssueService.cs
--- a/Engine/Source/Programs/Horde/Horde.Server/Issues/IssueService.cs
+++ b/Engine/Source/Programs/Horde/Horde.Server/Issues/IssueService.cs
@@ -10,6 +10,7 @@ using System.Text.RegularExpressions;
 using System.Threading;
 using System.Threading.Tasks;
 using EpicGames.Core;
+using EpicGames.Horde.Api;
 using Horde.Server.Issues.Handlers;
 using Horde.Server.Jobs;
 using Horde.Server.Jobs.Graphs;
diff --git a/Engine/Source/Programs/Horde/Horde.Server/Issues/IssuesController.cs b/Engine/Source/Programs/Horde/Horde.Server/Issues/IssuesController.cs
--- a/Engine/Source/Programs/Horde/Horde.Server/Issues/IssuesController.cs
+++ b/Engine/Source/Programs/Horde/Horde.Server/Issues/IssuesController.cs
@@ -21,6 +21,7 @@ using Horde.Server.Users;
 using Horde.Server.Jobs.Graphs;
 using Microsoft.Extensions.Options;
 using Horde.Server.Server;
+using EpicGames.Horde.Api;
 
 namespace Horde.Server.Issues
 {
diff --git a/Engine/Source/Programs/Horde/Horde.Server/Jobs/Bisect/BisectController.cs b/Engine/Source/Programs/Horde/Horde.Server/Jobs/Bisect/BisectController.cs
--- a/Engine/Source/Programs/Horde/Horde.Server/Jobs/Bisect/BisectController.cs
+++ b/Engine/Source/Programs/Horde/Horde.Server/Jobs/Bisect/BisectController.cs
@@ -16,6 +16,7 @@ using HordeCommon;
 using Microsoft.AspNetCore.Authorization;
 using Microsoft.AspNetCore.Mvc;
 using Microsoft.Extensions.Options;
+using EpicGames.Horde.Api;
 
 namespace Horde.Server.Jobs.Bisect
 {
@@ -77,7 +78,7 @@ namespace Horde.Server.Jobs.Bisect
 		public BisectTaskState State => _bisectTask.State;
 
 		/// <inheritdoc cref="IBisectTask.OwnerId"/>
-		public GetThinUserInfoResponse Owner { get; }
+		public GetThinUserInfoResponse? Owner { get; }
 
 		/// <inheritdoc cref="IBisectTask.StreamId"/>
 		public StreamId StreamId => _bisectTask.StreamId;
@@ -124,7 +125,7 @@ namespace Horde.Server.Jobs.Bisect
 		/// </summary>
 		public List<GetJobStepRefResponse> Steps { get; }
 
-		internal GetBisectTaskResponse(IBisectTask bisectTask, GetThinUserInfoResponse owner, List<IJobStepRef> steps, IJob? nextJob)
+		internal GetBisectTaskResponse(IBisectTask bisectTask, GetThinUserInfoResponse? owner, List<IJobStepRef> steps, IJob? nextJob)
 		{
 			_bisectTask = bisectTask;
 			Owner = owner;
@@ -371,7 +372,7 @@ namespace Horde.Server.Jobs.Bisect
 			}
 
 			IJob? nextJob = task.State == BisectTaskState.Running ? await _jobCollection.FindBisectTaskJobsAsync(task.Id, true, cancellationToken).FirstOrDefaultAsync(cancellationToken) : null;
-			return new GetBisectTaskResponse(task, new GetThinUserInfoResponse(user), steps, nextJob);
+			return new GetBisectTaskResponse(task, user?.ToThinApiResponse(), steps, nextJob);
 		}
 
 		/// <summary>
diff --git a/Engine/Source/Programs/Horde/Horde.Server/Jobs/Bisect/BisectTaskCollection.cs b/Engine/Source/Programs/Horde/Horde.Server/Jobs/Bisect/BisectTaskCollection.cs
--- a/Engine/Source/Programs/Horde/Horde.Server/Jobs/Bisect/BisectTaskCollection.cs
+++ b/Engine/Source/Programs/Horde/Horde.Server/Jobs/Bisect/BisectTaskCollection.cs
@@ -14,6 +14,7 @@ using HordeCommon;
 using MongoDB.Bson;
 using MongoDB.Bson.Serialization.Attributes;
 using MongoDB.Driver;
+using EpicGames.Horde.Api;
 
 namespace Horde.Server.Jobs.Bisect
 {
diff --git a/Engine/Source/Programs/Horde/Horde.Server/Jobs/Bisect/IBisectTask.cs b/Engine/Source/Programs/Horde/Horde.Server/Jobs/Bisect/IBisectTask.cs
--- a/Engine/Source/Programs/Horde/Horde.Server/Jobs/Bisect/IBisectTask.cs
+++ b/Engine/Source/Programs/Horde/Horde.Server/Jobs/Bisect/IBisectTask.cs
@@ -5,6 +5,7 @@ using Horde.Server.Perforce;
 using Horde.Server.Streams;
 using Horde.Server.Users;
 using HordeCommon;
+using EpicGames.Horde.Api;
 
 namespace Horde.Server.Jobs.Bisect
 {
diff --git a/Engine/Source/Programs/Horde/Horde.Server/Jobs/Bisect/IBisectTaskCollection.cs b/Engine/Source/Programs/Horde/Horde.Server/Jobs/Bisect/IBisectTaskCollection.cs
--- a/Engine/Source/Programs/Horde/Horde.Server/Jobs/Bisect/IBisectTaskCollection.cs
+++ b/Engine/Source/Programs/Horde/Horde.Server/Jobs/Bisect/IBisectTaskCollection.cs
@@ -7,6 +7,7 @@ using System.Threading.Tasks;
 using Horde.Server.Perforce;
 using Horde.Server.Users;
 using HordeCommon;
+using EpicGames.Horde.Api;
 
 namespace Horde.Server.Jobs.Bisect
 {
diff --git a/Engine/Source/Programs/Horde/Horde.Server/Jobs/IJob.cs b/Engine/Source/Programs/Horde/Horde.Server/Jobs/IJob.cs
--- a/Engine/Source/Programs/Horde/Horde.Server/Jobs/IJob.cs
+++ b/Engine/Source/Programs/Horde/Horde.Server/Jobs/IJob.cs
@@ -6,6 +6,7 @@ using System.Diagnostics.CodeAnalysis;
 using System.Linq;
 using System.Security.Claims;
 using EpicGames.Core;
+using EpicGames.Horde.Api;
 using Horde.Server.Acls;
 using Horde.Server.Agents;
 using Horde.Server.Agents.Leases;
diff --git a/Engine/Source/Programs/Horde/Horde.Server/Jobs/IJobCollection.cs b/Engine/Source/Programs/Horde/Horde.Server/Jobs/IJobCollection.cs
--- a/Engine/Source/Programs/Horde/Horde.Server/Jobs/IJobCollection.cs
+++ b/Engine/Source/Programs/Horde/Horde.Server/Jobs/IJobCollection.cs
@@ -4,6 +4,7 @@ using System;
 using System.Collections.Generic;
 using System.Threading.Tasks;
 using EpicGames.Core;
+using EpicGames.Horde.Api;
 using Horde.Server.Agents;
 using Horde.Server.Agents.Leases;
 using Horde.Server.Agents.Pools;
diff --git a/Engine/Source/Programs/Horde/Horde.Server/Jobs/JobCollection.cs b/Engine/Source/Programs/Horde/Horde.Server/Jobs/JobCollection.cs
--- a/Engine/Source/Programs/Horde/Horde.Server/Jobs/JobCollection.cs
+++ b/Engine/Source/Programs/Horde/Horde.Server/Jobs/JobCollection.cs
@@ -11,6 +11,7 @@ using System.Text;
 using System.Threading;
 using System.Threading.Tasks;
 using EpicGames.Core;
+using EpicGames.Horde.Api;
 using Horde.Server.Acls;
 using Horde.Server.Agents;
 using Horde.Server.Agents.Leases;
diff --git a/Engine/Source/Programs/Horde/Horde.Server/Jobs/JobMessages.cs b/Engine/Source/Programs/Horde/Horde.Server/Jobs/JobMessages.cs
--- a/Engine/Source/Programs/Horde/Horde.Server/Jobs/JobMessages.cs
+++ b/Engine/Source/Programs/Horde/Horde.Server/Jobs/JobMessages.cs
@@ -6,6 +6,7 @@ using System.ComponentModel.DataAnnotations;
 using System.Linq;
 using System.Text.Json.Serialization;
 using EpicGames.Core;
+using EpicGames.Horde.Api;
 using Horde.Server.Acls;
 using Horde.Server.Jobs.Bisect;
 using Horde.Server.Jobs.Graphs;
diff --git a/Engine/Source/Programs/Horde/Horde.Server/Jobs/JobService.cs b/Engine/Source/Programs/Horde/Horde.Server/Jobs/JobService.cs
--- a/Engine/Source/Programs/Horde/Horde.Server/Jobs/JobService.cs
+++ b/Engine/Source/Programs/Horde/Horde.Server/Jobs/JobService.cs
@@ -9,6 +9,7 @@ using System.Text.RegularExpressions;
 using System.Threading;
 using System.Threading.Tasks;
 using EpicGames.Core;
+using EpicGames.Horde.Api;
 using Horde.Server.Acls;
 using Horde.Server.Agents;
 using Horde.Server.Agents.Leases;
diff --git a/Engine/Source/Programs/Horde/Horde.Server/Jobs/JobsController.cs b/Engine/Source/Programs/Horde/Horde.Server/Jobs/JobsController.cs
--- a/Engine/Source/Programs/Horde/Horde.Server/Jobs/JobsController.cs
+++ b/Engine/Source/Programs/Horde/Horde.Server/Jobs/JobsController.cs
@@ -6,6 +6,7 @@ using System.Linq;
 using System.Threading;
 using System.Threading.Tasks;
 using EpicGames.Core;
+using EpicGames.Horde.Api;
 using Horde.Server.Acls;
 using Horde.Server.Agents;
 using Horde.Server.Jobs.Artifacts;
@@ -455,13 +456,13 @@ namespace Horde.Server.Jobs
 			GetThinUserInfoResponse? startedByUserInfo = null;
 			if (job.StartedByUserId != null)
 			{
-				startedByUserInfo = new GetThinUserInfoResponse(await _userCollection.GetCachedUserAsync(job.StartedByUserId.Value));
+				startedByUserInfo = (await _userCollection.GetCachedUserAsync(job.StartedByUserId.Value))?.ToThinApiResponse();
 			}
 
 			GetThinUserInfoResponse? abortedByUserInfo = null;
 			if (job.AbortedByUserId != null)
 			{
-				abortedByUserInfo = new GetThinUserInfoResponse(await _userCollection.GetCachedUserAsync(job.AbortedByUserId.Value));
+				abortedByUserInfo = (await _userCollection.GetCachedUserAsync(job.AbortedByUserId.Value))?.ToThinApiResponse();
 			}
 
 			GetJobResponse response = new GetJobResponse(job, startedByUserInfo, abortedByUserInfo);
@@ -517,13 +518,13 @@ namespace Horde.Server.Jobs
 			GetThinUserInfoResponse? abortedByUserInfo = null;
 			if (step.AbortedByUserId != null)
 			{
-				abortedByUserInfo = new GetThinUserInfoResponse(await _userCollection.GetCachedUserAsync(step.AbortedByUserId.Value));
+				abortedByUserInfo = (await _userCollection.GetCachedUserAsync(step.AbortedByUserId.Value))?.ToThinApiResponse();
 			}
 
 			GetThinUserInfoResponse? retriedByUserInfo = null;
 			if (step.RetriedByUserId != null)
 			{
-				retriedByUserInfo = new GetThinUserInfoResponse(await _userCollection.GetCachedUserAsync(step.RetriedByUserId.Value));
+				retriedByUserInfo = (await _userCollection.GetCachedUserAsync(step.RetriedByUserId.Value))?.ToThinApiResponse();
 			}
 
 			return new GetStepResponse(step, abortedByUserInfo, retriedByUserInfo);
diff --git a/Engine/Source/Programs/Horde/Horde.Server/Notifications/INotificationTrigger.cs b/Engine/Source/Programs/Horde/Horde.Server/Notifications/INotificationTrigger.cs
--- a/Engine/Source/Programs/Horde/Horde.Server/Notifications/INotificationTrigger.cs
+++ b/Engine/Source/Programs/Horde/Horde.Server/Notifications/INotificationTrigger.cs
@@ -4,6 +4,7 @@ using System.Collections.Generic;
 using Horde.Server.Users;
 using Horde.Server.Utilities;
 using MongoDB.Bson;
+using EpicGames.Horde.Api;
 
 namespace Horde.Server.Notifications
 {
diff --git a/Engine/Source/Programs/Horde/Horde.Server/Notifications/INotificationTriggerCollection.cs b/Engine/Source/Programs/Horde/Horde.Server/Notifications/INotificationTriggerCollection.cs
--- a/Engine/Source/Programs/Horde/Horde.Server/Notifications/INotificationTriggerCollection.cs
+++ b/Engine/Source/Programs/Horde/Horde.Server/Notifications/INotificationTriggerCollection.cs
@@ -5,6 +5,7 @@ using System.Threading.Tasks;
 using Horde.Server.Users;
 using Horde.Server.Utilities;
 using MongoDB.Bson;
+using EpicGames.Horde.Api;
 
 namespace Horde.Server.Notifications
 {
diff --git a/Engine/Source/Programs/Horde/Horde.Server/Notifications/ISubscription.cs b/Engine/Source/Programs/Horde/Horde.Server/Notifications/ISubscription.cs
--- a/Engine/Source/Programs/Horde/Horde.Server/Notifications/ISubscription.cs
+++ b/Engine/Source/Programs/Horde/Horde.Server/Notifications/ISubscription.cs
@@ -2,6 +2,7 @@
 
 using Horde.Server.Users;
 using Horde.Server.Utilities;
+using EpicGames.Horde.Api;
 
 namespace Horde.Server.Notifications
 {
diff --git a/Engine/Source/Programs/Horde/Horde.Server/Notifications/ISubscriptionCollection.cs b/Engine/Source/Programs/Horde/Horde.Server/Notifications/ISubscriptionCollection.cs
--- a/Engine/Source/Programs/Horde/Horde.Server/Notifications/ISubscriptionCollection.cs
+++ b/Engine/Source/Programs/Horde/Horde.Server/Notifications/ISubscriptionCollection.cs
@@ -4,6 +4,7 @@ using System.Collections.Generic;
 using System.Threading.Tasks;
 using Horde.Server.Users;
 using Horde.Server.Utilities;
+using EpicGames.Horde.Api;
 
 namespace Horde.Server.Notifications
 {
diff --git a/Engine/Source/Programs/Horde/Horde.Server/Notifications/NotificationService.cs b/Engine/Source/Programs/Horde/Horde.Server/Notifications/NotificationService.cs
--- a/Engine/Source/Programs/Horde/Horde.Server/Notifications/NotificationService.cs
+++ b/Engine/Source/Programs/Horde/Horde.Server/Notifications/NotificationService.cs
@@ -9,6 +9,7 @@ using System.Security.Claims;
 using System.Threading;
 using System.Threading.Tasks;
 using EpicGames.Core;
+using EpicGames.Horde.Api;
 using EpicGames.Redis;
 using Horde.Server.Agents.Pools;
 using Horde.Server.Configuration;
diff --git a/Engine/Source/Programs/Horde/Horde.Server/Notifications/NotificationTriggerCollection.cs b/Engine/Source/Programs/Horde/Horde.Server/Notifications/NotificationTriggerCollection.cs
--- a/Engine/Source/Programs/Horde/Horde.Server/Notifications/NotificationTriggerCollection.cs
+++ b/Engine/Source/Programs/Horde/Horde.Server/Notifications/NotificationTriggerCollection.cs
@@ -9,6 +9,7 @@ using Horde.Server.Utilities;
 using MongoDB.Bson;
 using MongoDB.Bson.Serialization.Attributes;
 using MongoDB.Driver;
+using EpicGames.Horde.Api;
 
 namespace Horde.Server.Notifications
 {
diff --git a/Engine/Source/Programs/Horde/Horde.Server/Notifications/Sinks/SlackNotificationSink.cs b/Engine/Source/Programs/Horde/Horde.Server/Notifications/Sinks/SlackNotificationSink.cs
--- a/Engine/Source/Programs/Horde/Horde.Server/Notifications/Sinks/SlackNotificationSink.cs
+++ b/Engine/Source/Programs/Horde/Horde.Server/Notifications/Sinks/SlackNotificationSink.cs
@@ -15,6 +15,7 @@ using System.Text.RegularExpressions;
 using System.Threading;
 using System.Threading.Tasks;
 using EpicGames.Core;
+using EpicGames.Horde.Api;
 using EpicGames.Redis;
 using EpicGames.Redis.Utility;
 using EpicGames.Slack;
diff --git a/Engine/Source/Programs/Horde/Horde.Server/Notifications/SubscriptionCollection.cs b/Engine/Source/Programs/Horde/Horde.Server/Notifications/SubscriptionCollection.cs
--- a/Engine/Source/Programs/Horde/Horde.Server/Notifications/SubscriptionCollection.cs
+++ b/Engine/Source/Programs/Horde/Horde.Server/Notifications/SubscriptionCollection.cs
@@ -7,6 +7,7 @@ using System.Linq;
 using System.Text;
 using System.Threading.Tasks;
 using EpicGames.Core;
+using EpicGames.Horde.Api;
 using Horde.Server.Jobs;
 using Horde.Server.Server;
 using Horde.Server.Streams;
diff --git a/Engine/Source/Programs/Horde/Horde.Server/Notifications/SubscriptionsController.cs b/Engine/Source/Programs/Horde/Horde.Server/Notifications/SubscriptionsController.cs
--- a/Engine/Source/Programs/Horde/Horde.Server/Notifications/SubscriptionsController.cs
+++ b/Engine/Source/Programs/Horde/Horde.Server/Notifications/SubscriptionsController.cs
@@ -10,6 +10,7 @@ using Horde.Server.Utilities;
 using Microsoft.AspNetCore.Authorization;
 using Microsoft.AspNetCore.Mvc;
 using Microsoft.Extensions.Options;
+using EpicGames.Horde.Api;
 
 namespace Horde.Server.Notifications
 {
diff --git a/Engine/Source/Programs/Horde/Horde.Server/Perforce/ICommit.cs b/Engine/Source/Programs/Horde/Horde.Server/Perforce/ICommit.cs
--- a/Engine/Source/Programs/Horde/Horde.Server/Perforce/ICommit.cs
+++ b/Engine/Source/Programs/Horde/Horde.Server/Perforce/ICommit.cs
@@ -6,6 +6,7 @@ using System.Linq;
 using System.Threading;
 using System.Threading.Tasks;
 using EpicGames.Core;
+using EpicGames.Horde.Api;
 using EpicGames.Perforce;
 using Horde.Server.Streams;
 using Horde.Server.Users;
diff --git a/Engine/Source/Programs/Horde/Horde.Server/Perforce/PerforceService.cs b/Engine/Source/Programs/Horde/Horde.Server/Perforce/PerforceService.cs
--- a/Engine/Source/Programs/Horde/Horde.Server/Perforce/PerforceService.cs
+++ b/Engine/Source/Programs/Horde/Horde.Server/Perforce/PerforceService.cs
@@ -14,6 +14,7 @@ using System.Text.RegularExpressions;
 using System.Threading;
 using System.Threading.Tasks;
 using EpicGames.Core;
+using EpicGames.Horde.Api;
 using EpicGames.Perforce;
 using Horde.Server.Server;
 using Horde.Server.Streams;
diff --git a/Engine/Source/Programs/Horde/Horde.Server/Perforce/PerforceServiceCache.cs b/Engine/Source/Programs/Horde/Horde.Server/Perforce/PerforceServiceCache.cs
--- a/Engine/Source/Programs/Horde/Horde.Server/Perforce/PerforceServiceCache.cs
+++ b/Engine/Source/Programs/Horde/Horde.Server/Perforce/PerforceServiceCache.cs
@@ -8,6 +8,7 @@ using System.Runtime.CompilerServices;
 using System.Threading;
 using System.Threading.Tasks;
 using EpicGames.Core;
+using EpicGames.Horde.Api;
 using EpicGames.Perforce;
 using EpicGames.Redis;
 using Horde.Server.Server;
diff --git a/Engine/Source/Programs/Horde/Horde.Server/Server/Notices/INotice.cs b/Engine/Source/Programs/Horde/Horde.Server/Server/Notices/INotice.cs
--- a/Engine/Source/Programs/Horde/Horde.Server/Server/Notices/INotice.cs
+++ b/Engine/Source/Programs/Horde/Horde.Server/Server/Notices/INotice.cs
@@ -5,6 +5,7 @@ using Horde.Server.Users;
 using Horde.Server.Utilities;
 using MongoDB.Bson;
 using MongoDB.Bson.Serialization.Attributes;
+using EpicGames.Horde.Api;
 
 namespace Horde.Server.Server.Notices
 {
diff --git a/Engine/Source/Programs/Horde/Horde.Server/Server/Notices/INoticeCollection.cs b/Engine/Source/Programs/Horde/Horde.Server/Server/Notices/INoticeCollection.cs
--- a/Engine/Source/Programs/Horde/Horde.Server/Server/Notices/INoticeCollection.cs
+++ b/Engine/Source/Programs/Horde/Horde.Server/Server/Notices/INoticeCollection.cs
@@ -6,6 +6,7 @@ using System.Threading.Tasks;
 using Horde.Server.Users;
 using Horde.Server.Utilities;
 using MongoDB.Bson;
+using EpicGames.Horde.Api;
 
 namespace Horde.Server.Server.Notices
 {
diff --git a/Engine/Source/Programs/Horde/Horde.Server/Server/Notices/NoticeCollection.cs b/Engine/Source/Programs/Horde/Horde.Server/Server/Notices/NoticeCollection.cs
--- a/Engine/Source/Programs/Horde/Horde.Server/Server/Notices/NoticeCollection.cs
+++ b/Engine/Source/Programs/Horde/Horde.Server/Server/Notices/NoticeCollection.cs
@@ -9,6 +9,7 @@ using System.Linq;
 using Horde.Server.Utilities;
 using MongoDB.Bson;
 using Horde.Server.Users;
+using EpicGames.Horde.Api;
 
 namespace Horde.Server.Server.Notices
 {
diff --git a/Engine/Source/Programs/Horde/Horde.Server/Server/Notices/NoticeMessages.cs b/Engine/Source/Programs/Horde/Horde.Server/Server/Notices/NoticeMessages.cs
--- a/Engine/Source/Programs/Horde/Horde.Server/Server/Notices/NoticeMessages.cs
+++ b/Engine/Source/Programs/Horde/Horde.Server/Server/Notices/NoticeMessages.cs
@@ -3,6 +3,7 @@
 using System;
 using System.ComponentModel.DataAnnotations;
 using Horde.Server.Users;
+using EpicGames.Horde.Api;
 
 namespace Horde.Server.Server.Notices
 {
diff --git a/Engine/Source/Programs/Horde/Horde.Server/Server/Notices/NoticeService.cs b/Engine/Source/Programs/Horde/Horde.Server/Server/Notices/NoticeService.cs
--- a/Engine/Source/Programs/Horde/Horde.Server/Server/Notices/NoticeService.cs
+++ b/Engine/Source/Programs/Horde/Horde.Server/Server/Notices/NoticeService.cs
@@ -6,6 +6,7 @@ using System.Threading.Tasks;
 using Horde.Server.Users;
 using Horde.Server.Utilities;
 using MongoDB.Bson;
+using EpicGames.Horde.Api;
 
 namespace Horde.Server.Server.Notices
 {
diff --git a/Engine/Source/Programs/Horde/Horde.Server/Server/Notices/NoticesController.cs b/Engine/Source/Programs/Horde/Horde.Server/Server/Notices/NoticesController.cs
--- a/Engine/Source/Programs/Horde/Horde.Server/Server/Notices/NoticesController.cs
+++ b/Engine/Source/Programs/Horde/Horde.Server/Server/Notices/NoticesController.cs
@@ -11,6 +11,7 @@ using Microsoft.AspNetCore.Authorization;
 using Microsoft.AspNetCore.Mvc;
 using Microsoft.Extensions.Options;
 using MongoDB.Bson;
+using EpicGames.Horde.Api;
 
 namespace Horde.Server.Server.Notices
 {
@@ -124,7 +125,7 @@ namespace Horde.Server.Server.Notices
 
 				if (notice.UserId != null)
 				{
-					userInfo = new GetThinUserInfoResponse(await _userCollection.GetCachedUserAsync(notice.UserId));
+					userInfo = (await _userCollection.GetCachedUserAsync(notice.UserId))?.ToThinApiResponse();
 				}
 
 				messages.Add(new GetNoticeResponse() { Id = notice.Id.ToString(), Active = true, Message = notice.Message, CreatedByUser = userInfo});
diff --git a/Engine/Source/Programs/Horde/Horde.Server/Streams/IStream.cs b/Engine/Source/Programs/Horde/Horde.Server/Streams/IStream.cs
--- a/Engine/Source/Programs/Horde/Horde.Server/Streams/IStream.cs
+++ b/Engine/Source/Programs/Horde/Horde.Server/Streams/IStream.cs
@@ -6,6 +6,7 @@ using System.Diagnostics.CodeAnalysis;
 using System.Linq;
 using System.Text.RegularExpressions;
 using EpicGames.Core;
+using EpicGames.Horde.Api;
 using Horde.Server.Acls;
 using Horde.Server.Agents;
 using Horde.Server.Jobs;
diff --git a/Engine/Source/Programs/Horde/Horde.Server/Streams/StreamCollection.cs b/Engine/Source/Programs/Horde/Horde.Server/Streams/StreamCollection.cs
--- a/Engine/Source/Programs/Horde/Horde.Server/Streams/StreamCollection.cs
+++ b/Engine/Source/Programs/Horde/Horde.Server/Streams/StreamCollection.cs
@@ -6,6 +6,7 @@ using System.IO;
 using System.Linq;
 using System.Threading.Tasks;
 using EpicGames.Core;
+using EpicGames.Horde.Api;
 using Horde.Server.Jobs;
 using Horde.Server.Projects;
 using Horde.Server.Server;
diff --git a/Engine/Source/Programs/Horde/Horde.Server/Streams/StreamMessages.cs b/Engine/Source/Programs/Horde/Horde.Server/Streams/StreamMessages.cs
--- a/Engine/Source/Programs/Horde/Horde.Server/Streams/StreamMessages.cs
+++ b/Engine/Source/Programs/Horde/Horde.Server/Streams/StreamMessages.cs
@@ -568,7 +568,7 @@ namespace Horde.Server.Streams
 		{
 			Number = commit.Number;
 			Author = author.Name;
-			AuthorInfo = new GetThinUserInfoResponse(author);
+			AuthorInfo = author.ToThinApiResponse();
 			Description = commit.Description;
 
 			if (tags != null)
diff --git a/Engine/Source/Programs/Horde/Horde.Server/Streams/StreamsController.cs b/Engine/Source/Programs/Horde/Horde.Server/Streams/StreamsController.cs
--- a/Engine/Source/Programs/Horde/Horde.Server/Streams/StreamsController.cs
+++ b/Engine/Source/Programs/Horde/Horde.Server/Streams/StreamsController.cs
@@ -209,7 +209,7 @@ namespace Horde.Server.Streams
 
 							stepStates ??= new List<GetTemplateStepStateResponse>();
 
-							GetThinUserInfoResponse? pausedByUserInfo = new GetThinUserInfoResponse(await _userCollection.GetCachedUserAsync(state.PausedByUserId));
+							GetThinUserInfoResponse? pausedByUserInfo = (await _userCollection.GetCachedUserAsync(state.PausedByUserId))?.ToThinApiResponse();
 							stepStates.Add(new GetTemplateStepStateResponse(state, pausedByUserInfo));
 						}
 					}
diff --git a/Engine/Source/Programs/Horde/Horde.Server/Ugs/UgsController.cs b/Engine/Source/Programs/Horde/Horde.Server/Ugs/UgsController.cs
--- a/Engine/Source/Programs/Horde/Horde.Server/Ugs/UgsController.cs
+++ b/Engine/Source/Programs/Horde/Horde.Server/Ugs/UgsController.cs
@@ -13,6 +13,7 @@ using Microsoft.AspNetCore.Mvc;
 using Microsoft.Extensions.Logging;
 using Microsoft.Extensions.Options;
 using MongoDB.Bson;
+using EpicGames.Horde.Api;
 
 namespace Horde.Server.Ugs
 {
diff --git a/Engine/Source/Programs/Horde/Horde.Server/Users/IUser.cs b/Engine/Source/Programs/Horde/Horde.Server/Users/IUser.cs
--- a/Engine/Source/Programs/Horde/Horde.Server/Users/IUser.cs
+++ b/Engine/Source/Programs/Horde/Horde.Server/Users/IUser.cs
@@ -2,6 +2,10 @@
 
 using System;
 using System.Collections.Generic;
+using System.Linq;
+using Amazon.S3.Model;
+using EpicGames.Core;
+using EpicGames.Horde.Api;
 using Horde.Server.Jobs;
 using Horde.Server.Streams;
 using MongoDB.Bson;
@@ -148,6 +152,61 @@ namespace Horde.Server.Users
 		/// List of job template preferences
 		/// </summary>
 		public IReadOnlyList<IUserJobTemplateSettings>? JobTemplateSettings { get; }
+	}
+
+	/// <summary>
+	/// Extension methods
+	/// </summary>
+	public static class UserExtensions
+	{
+		/// <summary>
+		/// Creates an API response object
+		/// </summary>
+		public static GetUserResponse ToApiResponse(this IUser user, IAvatar? avatar, IUserClaims? claims, IUserSettings? settings)
+		{
+			GetUserResponse response = new GetUserResponse(user.Id, user.Name);
+			response.Email = user.Email;
+
+			response.Image24 = avatar?.Image24;
+			response.Image32 = avatar?.Image32;
+			response.Image48 = avatar?.Image48;
+			response.Image72 = avatar?.Image72;
+
+			response.Claims = claims?.Claims.Select(x => new GetUserClaimResponse(x.Type, x.Value)).ToList();
 
+			if (settings != null)
+			{
+				response.EnableExperimentalFeatures = settings.EnableExperimentalFeatures;
+
+				response.DashboardSettings = BsonTypeMapper.MapToDotNetValue(settings.DashboardSettings);
+				response.PinnedJobIds = settings.PinnedJobIds.ConvertAll(x => x.ToString());
+
+				if (settings.JobTemplateSettings != null && settings.JobTemplateSettings.Count > 0)
+				{
+					response.JobTemplateSettings = new List<GetJobTemplateSettingsResponse>();
+					for (int i = 0; i < settings.JobTemplateSettings.Count; i++)
+					{
+						response.JobTemplateSettings.Add(settings.JobTemplateSettings[i].ToApiResponse());
+					}
+				}
+			}
+			return response;
+		}
+
+		/// <summary>
+		/// Creates an API thin user response
+		/// </summary>
+		public static GetThinUserInfoResponse ToThinApiResponse(this IUser user)
+		{
+			return new GetThinUserInfoResponse(user.Id, user.Name, user.Email, user.Login);
+		}
+
+		/// <summary>
+		/// Creates an API settings response
+		/// </summary>
+		public static GetJobTemplateSettingsResponse ToApiResponse(this IUserJobTemplateSettings settings)
+		{
+			return new GetJobTemplateSettingsResponse(settings.StreamId.ToString(), settings.TemplateId.ToString(), settings.TemplateHash.ToString(), settings.Arguments.ToList(), settings.UpdateTimeUtc);
+		}
 	}
 }
diff --git a/Engine/Source/Programs/Horde/Horde.Server/Users/IUserCollection.cs b/Engine/Source/Programs/Horde/Horde.Server/Users/IUserCollection.cs
--- a/Engine/Source/Programs/Horde/Horde.Server/Users/IUserCollection.cs
+++ b/Engine/Source/Programs/Horde/Horde.Server/Users/IUserCollection.cs
@@ -6,6 +6,7 @@ using System.Threading.Tasks;
 using Horde.Server.Jobs;
 using Horde.Server.Utilities;
 using MongoDB.Bson;
+using EpicGames.Horde.Api;
 
 namespace Horde.Server.Users
 {
diff --git a/Engine/Source/Programs/Horde/Horde.Server/Users/UserCollection.cs b/Engine/Source/Programs/Horde/Horde.Server/Users/UserCollection.cs
--- a/Engine/Source/Programs/Horde/Horde.Server/Users/UserCollection.cs
+++ b/Engine/Source/Programs/Horde/Horde.Server/Users/UserCollection.cs
@@ -12,6 +12,7 @@ using MongoDB.Bson;
 using MongoDB.Bson.Serialization;
 using MongoDB.Bson.Serialization.Attributes;
 using MongoDB.Driver;
+using EpicGames.Horde.Api;
 
 namespace Horde.Server.Users
 {
diff --git a/Engine/Source/Programs/Horde/Horde.Server/Users/UserCollectionV2.cs b/Engine/Source/Programs/Horde/Horde.Server/Users/UserCollectionV2.cs
--- a/Engine/Source/Programs/Horde/Horde.Server/Users/UserCollectionV2.cs
+++ b/Engine/Source/Programs/Horde/Horde.Server/Users/UserCollectionV2.cs
@@ -14,6 +14,7 @@ using MongoDB.Bson;
 using MongoDB.Bson.Serialization;
 using MongoDB.Bson.Serialization.Attributes;
 using MongoDB.Driver;
+using EpicGames.Horde.Api;
 
 namespace Horde.Server.Users
 {
diff --git a/Engine/Source/Programs/Horde/Horde.Server/Users/UserController.cs b/Engine/Source/Programs/Horde/Horde.Server/Users/UserController.cs
--- a/Engine/Source/Programs/Horde/Horde.Server/Users/UserController.cs
+++ b/Engine/Source/Programs/Horde/Horde.Server/Users/UserController.cs
@@ -2,13 +2,19 @@
 
 using System.Collections.Generic;
 using System.Linq;
+using System.Security.Claims;
 using System.Threading.Tasks;
+using Horde.Server.Acls;
+using Horde.Server.Agents;
+using Horde.Server.Agents.Pools;
 using Horde.Server.Jobs;
 using Horde.Server.Server;
+using Horde.Server.Server.Notices;
 using Horde.Server.Utilities;
 using Microsoft.AspNetCore.Authorization;
 using Microsoft.AspNetCore.Mvc;
 using Microsoft.Extensions.Options;
+using EpicGames.Horde.Api;
 
 namespace Horde.Server.Users
 {
@@ -53,12 +59,27 @@ namespace Horde.Server.Users
 			IUserClaims claims = await _userCollection.GetClaimsAsync(internalUser.Id);
 			IUserSettings settings = await _userCollection.GetSettingsAsync(internalUser.Id);
 
-			GetUserResponse response = new GetUserResponse(internalUser, avatar, claims, settings);
-			response.DashboardFeatures = new GetDashboardFeaturesResponse(_globalConfig.Value, User);
+			GetUserResponse response = internalUser.ToApiResponse(avatar, claims, settings);
+			response.DashboardFeatures = GetDashboardFeatures(_globalConfig.Value, User);
 
 			return PropertyFilter.Apply(response, filter);
 		}
 
+		static GetDashboardFeaturesResponse GetDashboardFeatures(GlobalConfig globalConfig, ClaimsPrincipal principal)
+		{
+			GetDashboardFeaturesResponse response = new GetDashboardFeaturesResponse();
+			response.ShowLandingPage = globalConfig.Dashboard.ShowLandingPage;
+			response.ShowCI = globalConfig.Dashboard.ShowCI;
+			response.ShowAgents = globalConfig.Dashboard.ShowAgents;
+			response.ShowPerforceServers = globalConfig.Dashboard.ShowPerforceServers;
+			response.ShowDeviceManager = globalConfig.Dashboard.ShowDeviceManager;
+			response.ShowTests = globalConfig.Dashboard.ShowTests;
+			response.ShowNoticeEditor = globalConfig.Authorize(NoticeAclAction.CreateNotice, principal) || globalConfig.Authorize(NoticeAclAction.UpdateNotice, principal);
+			response.ShowPoolEditor = globalConfig.Authorize(PoolAclAction.CreatePool, principal) || globalConfig.Authorize(PoolAclAction.UpdatePool, principal);
+			response.ShowRemoteDesktop = globalConfig.Authorize(AgentAclAction.UpdateAgent, principal);
+			return response;
+		}
+
 		/// <summary>
 		/// Updates the logged in user
 		/// </summary>
diff --git a/Engine/Source/Programs/Horde/Horde.Server/Users/UserId.cs b/Engine/Source/Programs/Horde/Horde.Server/Users/UserId.cs
deleted file mode 100644
diff --git a/Engine/Source/Programs/Horde/Horde.Server/Users/UserMessages.cs b/Engine/Source/Programs/Horde/Horde.Server/Users/UserMessages.cs
deleted file mode 100644
diff --git a/Engine/Source/Programs/Horde/Horde.Server/Users/UsersController.cs b/Engine/Source/Programs/Horde/Horde.Server/Users/UsersController.cs
--- a/Engine/Source/Programs/Horde/Horde.Server/Users/UsersController.cs
+++ b/Engine/Source/Programs/Horde/Horde.Server/Users/UsersController.cs
@@ -7,6 +7,7 @@ using System.Threading.Tasks;
 using Horde.Server.Utilities;
 using Microsoft.AspNetCore.Authorization;
 using Microsoft.AspNetCore.Mvc;
+using EpicGames.Horde.Api;
 
 namespace Horde.Server.Users
 {
@@ -63,7 +64,7 @@ namespace Horde.Server.Users
 			IAvatar? avatar = (AvatarService == null) ? (IAvatar?)null : await AvatarService.GetAvatarAsync(user);
 			IUserClaims? claims = await UserCollection.GetClaimsAsync(user.Id);
 			IUserSettings? settings = await UserCollection.GetSettingsAsync(user.Id);
-			return PropertyFilter.Apply(new GetUserResponse(user, avatar, claims, settings), filter);
+			return PropertyFilter.Apply(user.ToApiResponse(avatar, claims, settings), filter);
 		}
 
 		/// <summary>
@@ -95,7 +96,7 @@ namespace Horde.Server.Users
 			{
 				IAvatar? avatar = (AvatarService == null || !includeAvatar) ? (IAvatar?)null : await AvatarService.GetAvatarAsync(user);
 				IUserClaims? claims = (!includeClaims) ? null : await UserCollection.GetClaimsAsync(user.Id);
-				response.Add(new GetUserResponse(user, avatar, claims, null));
+				response.Add(user.ToApiResponse(avatar, claims, null));
 			}
 
 			return response;
diff --git a/Engine/Source/Programs/Horde/Horde.Server/Utilities/HordeClaimTypes.cs b/Engine/Source/Programs/Horde/Horde.Server/Utilities/HordeClaimTypes.cs
--- a/Engine/Source/Programs/Horde/Horde.Server/Utilities/HordeClaimTypes.cs
+++ b/Engine/Source/Programs/Horde/Horde.Server/Utilities/HordeClaimTypes.cs
@@ -3,6 +3,7 @@
 using System.Linq;
 using System.Security.Claims;
 using Horde.Server.Users;
+using EpicGames.Horde.Api;
 
 namespace Horde.Server.Utilities
 {
diff --git a/Engine/Source/Programs/Shared/EpicGames.Horde/Api/UserId.cs b/Engine/Source/Programs/Shared/EpicGames.Horde/Api/UserId.cs
new file mode 100644
--- /dev/null
+++ b/Engine/Source/Programs/Shared/EpicGames.Horde/Api/UserId.cs
@@ -0,0 +1,61 @@
+// Copyright Epic Games, Inc. All Rights Reserved.
+
+using System.ComponentModel;
+using EpicGames.Core;
+using EpicGames.Horde;
+
+namespace EpicGames.Horde.Api
+{
+	/// <summary>
+	/// Identifier for a user
+	/// </summary>
+	/// <param name="Id">Id to construct from</param>
+	[JsonSchemaString]
+	[TypeConverter(typeof(BinaryIdTypeConverter<UserId, UserIdConverter>))]
+	[BinaryIdConverter(typeof(UserIdConverter))]
+	public record struct UserId(BinaryId Id)
+	{
+		/// <summary>
+		/// Constant value for empty user id
+		/// </summary>
+		public static UserId Empty { get; } = default;
+
+		/// <summary>
+		/// Special user id for an anonymous administrator
+		/// </summary>
+		public static UserId Anonymous { get; } = UserId.Parse("63f7d3525119b9aa4c0f035a");
+
+		/// <inheritdoc cref="BinaryId.Parse(System.String)"/>
+		public static UserId Parse(string text) => new UserId(BinaryId.Parse(text));
+
+		/// <inheritdoc cref="BinaryId.TryParse(System.String, out BinaryId)"/>
+		public static bool TryParse(string text, out UserId id)
+		{
+			if (BinaryId.TryParse(text, out BinaryId objectId))
+			{
+				id = new UserId(objectId);
+				return true;
+			}
+			else
+			{
+				id = default;
+				return false;
+			}
+		}
+
+		/// <inheritdoc/>
+		public override string ToString() => Id.ToString();
+	}
+
+	/// <summary>
+	/// Converter to and from <see cref="BinaryId"/> instances.
+	/// </summary>
+	class UserIdConverter : BinaryIdConverter<UserId>
+	{
+		/// <inheritdoc/>
+		public override UserId FromBinaryId(BinaryId id) => new UserId(id);
+
+		/// <inheritdoc/>
+		public override BinaryId ToBinaryId(UserId value) => value.Id;
+	}
+}
diff --git a/Engine/Source/Programs/Shared/EpicGames.Horde/Api/UserMessages.cs b/Engine/Source/Programs/Shared/EpicGames.Horde/Api/UserMessages.cs
new file mode 100644
--- /dev/null
+++ b/Engine/Source/Programs/Shared/EpicGames.Horde/Api/UserMessages.cs
@@ -0,0 +1,273 @@
+// Copyright Epic Games, Inc. All Rights Reserved.
+
+using System;
+using System.Collections.Generic;
+using System.Linq;
+using System.Text.Json;
+
+#pragma warning disable CA2227
+
+namespace EpicGames.Horde.Api
+{
+	/// <summary>
+	/// Response describing the current user
+	/// </summary>
+	public class GetUserResponse
+	{
+		/// <summary>
+		/// Id of the user
+		/// </summary>
+		public UserId Id { get; set; }
+
+		/// <summary>
+		/// Name of the user
+		/// </summary>
+		public string Name { get; set; }
+
+		/// <summary>
+		/// Avatar image URL (24px)
+		/// </summary>
+		public string? Image24 { get; set; }
+
+		/// <summary>
+		/// Avatar image URL (32px)
+		/// </summary>
+		public string? Image32 { get; set; }
+
+		/// <summary>
+		/// Avatar image URL (48px)
+		/// </summary>
+		public string? Image48 { get; set; }
+
+		/// <summary>
+		/// Avatar image URL (72px)
+		/// </summary>
+		public string? Image72 { get; set; }
+
+		/// <summary>
+		/// Email of the user
+		/// </summary>
+		public string? Email { get; set; }
+
+		/// <summary>
+		/// Claims for the user
+		/// </summary>
+		public List<GetUserClaimResponse>? Claims { get; set; }
+
+		/// <summary>
+		/// Whether to enable experimental features for this user
+		/// </summary>
+		public bool? EnableExperimentalFeatures { get; set; }
+
+		/// <summary>
+		/// Settings for the dashboard
+		/// </summary>
+		public object? DashboardSettings { get; set; }
+
+		/// <summary>
+		/// Settings for whether various dashboard features should be shown for the current user
+		/// </summary>
+		public GetDashboardFeaturesResponse? DashboardFeatures { get; set; }
+
+		/// <summary>
+		/// User job template preferences
+		/// </summary>
+		public List<GetJobTemplateSettingsResponse>? JobTemplateSettings { get; set; }
+
+		/// <summary>
+		/// List of pinned job ids
+		/// </summary>
+		public List<string>? PinnedJobIds { get; set; }
+
+		/// <summary>
+		/// Constructor
+		/// </summary>
+		public GetUserResponse(UserId id, string name)
+		{
+			Id = id;
+			Name = name;
+		}
+	}
+
+	/// <summary>
+	/// New claim document
+	/// </summary>
+	public class GetUserClaimResponse
+	{
+		/// <summary>
+		/// Type of the claim
+		/// </summary>
+		public string Type { get; set; }
+
+		/// <summary>
+		/// Value for the claim
+		/// </summary>
+		public string Value { get; set; }
+
+		/// <summary>
+		/// Constructor
+		/// </summary>
+		public GetUserClaimResponse(string type, string value)
+		{
+			Type = type;
+			Value = value;
+		}
+	}
+
+	/// <summary>
+	/// Job template settings for the current user
+	/// </summary>
+	public class GetJobTemplateSettingsResponse
+	{
+		/// <summary>
+		/// The stream the job was run in
+		/// </summary>
+		public string StreamId { get; set; }
+
+		/// <summary>
+		/// The template id of the job
+		/// </summary>
+		public string TemplateId { get; set; }
+
+		/// <summary>
+		/// The hash of the template definition
+		/// </summary>
+		public string TemplateHash { get; set; }
+
+		/// <summary>
+		/// The arguments defined when creating the job
+		/// </summary>
+		public List<string> Arguments { get; set; }
+
+		/// <summary>
+		/// The last update time of the job template
+		/// </summary>
+		public DateTimeOffset UpdateTimeUtc { get; set; }
+
+		/// <summary>
+		/// Constructor
+		/// </summary>
+		public GetJobTemplateSettingsResponse(string streamId, string templateId, string templateHash, List<string> arguments, DateTime updateTimeUtc)
+		{
+			StreamId = streamId;
+			TemplateId = templateId;
+			TemplateHash = templateHash;
+			Arguments = arguments;
+			UpdateTimeUtc = new DateTimeOffset(updateTimeUtc);
+		}
+	}
+
+	/// <summary>
+	/// Settings for whether various features should be enabled on the dashboard
+	/// </summary>
+	public class GetDashboardFeaturesResponse
+	{
+		/// <summary>
+		/// Navigate to the landing page by default
+		/// </summary>
+		public bool ShowLandingPage { get; set; }
+
+		/// <summary>
+		/// Enable CI functionality
+		/// </summary>
+		public bool ShowCI { get; set; }
+
+		/// <summary>
+		/// Whether to show functionality related to agents, pools, and utilization on the dashboard.
+		/// </summary>
+		public bool ShowAgents { get; set; }
+
+		/// <summary>
+		/// Show the Perforce server option on the server menu
+		/// </summary>
+		public bool ShowPerforceServers { get; set; }
+
+		/// <summary>
+		/// Show the device manager on the server menu
+		/// </summary>
+		public bool ShowDeviceManager { get; set; }
+
+		/// <summary>
+		/// Show automated tests on the server menu
+		/// </summary>
+		public bool ShowTests { get; set; }
+
+		/// <summary>
+		/// Whether the notice editor should be listed in the server menu
+		/// </summary>
+		public bool ShowNoticeEditor { get; set; }
+
+		/// <summary>
+		/// Whether controls for modifying pools should be shown
+		/// </summary>
+		public bool ShowPoolEditor { get; set; }
+
+		/// <summary>
+		/// Whether the remote desktop button should be shown on the agent modal
+		/// </summary>
+		public bool ShowRemoteDesktop { get; set; }
+	}
+
+	/// <summary>
+	/// Basic information about a user. May be embedded in other responses.
+	/// </summary>
+	public class GetThinUserInfoResponse
+	{
+		/// <summary>
+		/// Id of the user
+		/// </summary>
+		public UserId Id { get; set; }
+
+		/// <summary>
+		/// Name of the user
+		/// </summary>
+		public string Name { get; set; }
+
+		/// <summary>
+		/// The user's email address
+		/// </summary>
+		public string? Email { get; set; }
+
+		/// <summary>
+		/// The user login [DEPRECATED]
+		/// </summary>
+		public string? Login { get; set; }
+
+		/// <summary>
+		/// Constructor
+		/// </summary>
+		public GetThinUserInfoResponse(UserId id, string name, string? email, string? login)
+		{
+			Id = id;
+			Name = name;
+			Email = email;
+			Login = login;
+		}
+	}
+
+	/// <summary>
+	/// Request to update settings for a user
+	/// </summary>
+	public class UpdateUserRequest
+	{
+		/// <summary>
+		/// Whether to enable experimental features for this user
+		/// </summary>
+		public bool? EnableExperimentalFeatures { get; set; }
+
+		/// <summary>
+		/// New dashboard settings
+		/// </summary>
+		public JsonElement? DashboardSettings { get; set; }
+
+		/// <summary>
+		/// Job ids to add to the pinned list
+		/// </summary>
+		public List<string>? AddPinnedJobIds { get; set; }
+
+		/// <summary>
+		/// Jobs ids to remove from the pinned list
+		/// </summary>
+		public List<string>? RemovePinnedJobIds { get; set; }
+	}
+}
