Horde: Forward env vars for remotely executed non-Windows processes

diff --git a/Engine/Source/Programs/Shared/EpicGames.Horde/Compute/AgentMessageHandler.cs b/Engine/Source/Programs/Shared/EpicGames.Horde/Compute/AgentMessageHandler.cs
--- a/Engine/Source/Programs/Shared/EpicGames.Horde/Compute/AgentMessageHandler.cs
+++ b/Engine/Source/Programs/Shared/EpicGames.Horde/Compute/AgentMessageHandler.cs
@@ -346,6 +346,15 @@ namespace EpicGames.Horde.Compute
 				}
 
 				Dictionary<string, string> resolvedEnvVars = ManagedProcess.GetCurrentEnvVars();
+
+				foreach ((string key, string? value) in _envVars)
+				{
+					if (value != null)
+					{
+						resolvedEnvVars[key] = value;
+					}
+				}
+				
 				if (envVars != null)
 				{
 					foreach ((string key, string? value) in envVars)
