Box: Experimental arm64 support

#rnx

diff --git a/Engine/Source/Programs/Shared/EpicGames.Box/EpicGames.Box.csproj b/Engine/Source/Programs/Shared/EpicGames.Box/EpicGames.Box.csproj
--- a/Engine/Source/Programs/Shared/EpicGames.Box/EpicGames.Box.csproj
+++ b/Engine/Source/Programs/Shared/EpicGames.Box/EpicGames.Box.csproj
@@ -1,4 +1,4 @@
-﻿<Project Sdk="Microsoft.NET.Sdk">
+<Project Sdk="Microsoft.NET.Sdk">
   <Import Project="..\UnrealEngine.csproj.props" />
 
   <PropertyGroup>
@@ -36,27 +36,7 @@
     </Compile>
   </ItemGroup>
 
-  <ItemGroup>
-    <Compile Include="..\..\..\..\Restricted\NotForLicensees\Source\Programs\Shared\EpicGames.Box\*.cs">
-      <Link>$([System.Text.RegularExpressions.Regex]::Replace(%(FullPath), '^.+?[\\/]EpicGames.Box[\\/]', ''))</Link>
-    </Compile>
-  </ItemGroup>
-
-  <Choose>
-    <When Condition="'$([System.Runtime.InteropServices.RuntimeInformation]::IsOSPlatform($([System.Runtime.InteropServices.OSPlatform]::Windows)))' == 'true'">
-      <ItemGroup Condition="Exists('..\..\..\..\Restricted\NotForLicensees\Binaries\Win64\Box\BoxExport.dll')">
-        <Content Include="..\..\..\..\Restricted\NotForLicensees\Binaries\Win64\Box\BoxExport.dll">
-          <CopyToOutputDirectory>PreserveNewest</CopyToOutputDirectory>
-        </Content>
-        <Content Include="..\..\..\..\Restricted\NotForLicensees\Binaries\Win64\Box\BoxDetours.dll">
-          <CopyToOutputDirectory>PreserveNewest</CopyToOutputDirectory>
-        </Content>
-        <Content Include="..\..\..\..\Restricted\NotForLicensees\Binaries\Win64\Box\BoxVisualizer.exe">
-          <CopyToOutputDirectory>PreserveNewest</CopyToOutputDirectory>
-        </Content>
-      </ItemGroup>
-    </When>
-  </Choose>
+  <Import Condition="Exists('..\..\..\..\Restricted\NotForLicensees\Source\Programs\Shared\EpicGames.Box\Box.props')" Project="..\..\..\..\Restricted\NotForLicensees\Source\Programs\Shared\EpicGames.Box\Box.props" />
 
   <ItemGroup>
     <ProjectReference Include="..\..\Shared\EpicGames.Core\EpicGames.Core.csproj" />
diff --git a/Engine/Source/Programs/UnrealBuildTool/UnrealBuildTool.csproj b/Engine/Source/Programs/UnrealBuildTool/UnrealBuildTool.csproj
--- a/Engine/Source/Programs/UnrealBuildTool/UnrealBuildTool.csproj
+++ b/Engine/Source/Programs/UnrealBuildTool/UnrealBuildTool.csproj
@@ -75,7 +75,7 @@
   <PropertyGroup Condition="Exists('..\..\..\Restricted\NotForLicensees\Source\Programs\UnrealBuildTool\Executors\Experimental\BoxExecutor.cs')">
     <DefineConstants>$(DefineConstants);__BOXEXECUTOR_AVAILABLE__</DefineConstants>
   </PropertyGroup>
-  <PropertyGroup Condition="Exists('..\..\..\Restricted\NotForLicensees\Source\Programs\Shared\EpicGames.Box\BoxServer.cs')">
+  <PropertyGroup Condition="Exists('..\..\..\Restricted\NotForLicensees\Source\Programs\Shared\EpicGames.Box\Box.Props')">
     <DefineConstants>$(DefineConstants);__BOX_AVAILABLE__</DefineConstants>
   </PropertyGroup>
 
