Horde: Force bundles to a complete state once 60,000 nodes have been written to it. Our internal counters are only 16-bit.

diff --git a/Engine/Source/Programs/Shared/EpicGames.Horde/Storage/Bundles/Bundle.cs b/Engine/Source/Programs/Shared/EpicGames.Horde/Storage/Bundles/Bundle.cs
--- a/Engine/Source/Programs/Shared/EpicGames.Horde/Storage/Bundles/Bundle.cs
+++ b/Engine/Source/Programs/Shared/EpicGames.Horde/Storage/Bundles/Bundle.cs
@@ -69,6 +69,11 @@ namespace EpicGames.Horde.Storage.Bundles
 	/// </summary>
 	public class Bundle
 	{
+		/// <summary>
+		/// Maximum number of exports from a single bundle
+		/// </summary>
+		public const int MaxExports = 60000;
+
 		/// <summary>
 		/// Header for the bundle
 		/// </summary>
diff --git a/Engine/Source/Programs/Shared/EpicGames.Horde/Storage/Bundles/BundleWriter.cs b/Engine/Source/Programs/Shared/EpicGames.Horde/Storage/Bundles/BundleWriter.cs
--- a/Engine/Source/Programs/Shared/EpicGames.Horde/Storage/Bundles/BundleWriter.cs
+++ b/Engine/Source/Programs/Shared/EpicGames.Horde/Storage/Bundles/BundleWriter.cs
@@ -346,6 +346,9 @@ namespace EpicGames.Horde.Storage.Bundles
 				_encodedPacketWriter.Dispose();
 			}
 
+			// Whether this bundle is full
+			public bool IsFull() => _queue.Count >= Bundle.MaxExports;
+
 			// Whether this bundle has finished writing
 			public bool IsComplete() => CompleteTask.IsCompleted;
 
@@ -787,7 +790,7 @@ namespace EpicGames.Horde.Storage.Bundles
 			}
 
 			// If the bundle is full, start the process of writing it to disk
-			if (currentBundle.UncompressedLength > _options.MaxBlobSize)
+			if (currentBundle.UncompressedLength > _options.MaxBlobSize || currentBundle.IsFull())
 			{
 				Complete();
 			}
