Horde: Fix incorrect sign on export index read from large bundles.

diff --git a/Engine/Source/Programs/Shared/EpicGames.Horde/Storage/Bundles/Bundle.cs b/Engine/Source/Programs/Shared/EpicGames.Horde/Storage/Bundles/Bundle.cs
--- a/Engine/Source/Programs/Shared/EpicGames.Horde/Storage/Bundles/Bundle.cs
+++ b/Engine/Source/Programs/Shared/EpicGames.Horde/Storage/Bundles/Bundle.cs
@@ -1126,7 +1126,9 @@ namespace EpicGames.Horde.Storage.Bundles
 		/// </summary>
 		public static BundleExportRef Read(ReadOnlySpan<byte> data)
 		{
-			return new BundleExportRef(BinaryPrimitives.ReadInt16LittleEndian(data), BinaryPrimitives.ReadInt16LittleEndian(data.Slice(2)), new IoHash(data.Slice(4)));
+			int importIdx = BinaryPrimitives.ReadInt16LittleEndian(data);
+			int nodeIdx = BinaryPrimitives.ReadUInt16LittleEndian(data.Slice(2));
+			return new BundleExportRef(importIdx, nodeIdx, new IoHash(data.Slice(4)));
 		}
 
 		/// <summary>
