Fix Asan TEnumAsByte overwrite

diff --git a/Engine/Source/Programs/Shared/EpicGames.UHT/Exporters/CodeGen/UhtHeaderCodeGeneratorCppFile.cs b/Engine/Source/Programs/Shared/EpicGames.UHT/Exporters/CodeGen/UhtHeaderCodeGeneratorCppFile.cs
--- a/Engine/Source/Programs/Shared/EpicGames.UHT/Exporters/CodeGen/UhtHeaderCodeGeneratorCppFile.cs
+++ b/Engine/Source/Programs/Shared/EpicGames.UHT/Exporters/CodeGen/UhtHeaderCodeGeneratorCppFile.cs
@@ -1345,8 +1345,26 @@ namespace EpicGames.UHT.Exporters.CodeGen
 								.Append("\t\tCopyAssignItems(Result, Source, ")
 								.Append(property.ArrayDimensions)
 								.Append(");\r\n");
-						}
-						else
+                        }
+                        else if (property is UhtByteProperty byteProperty && byteProperty.Enum != null)
+                        {
+                            // If someone passed in a TEnumAsByte instead of the actual enum value, the cast in the else clause would cause an issue.
+                            // Since this is known to be a TEnumAsByte, we just fetch the first byte.  *HOWEVER* on MSB machines where 
+                            // the actual enum value is passed in this will fail and return zero if the native size of the enum > 1 byte.
+                            builder
+                                .Append("\t\t")
+                                .Append("uint8")
+                                .Append("& Result = *(")
+                                .Append("uint8")
+                                .Append("*)OutValue;\r\n");
+                            builder
+                                .Append("\t\tResult = (")
+                                .Append("uint8")
+                                .Append(")Obj->")
+                                .Append(property.Getter!)
+                                .Append("();\r\n");
+                        }
+                        else
 						{
 							builder
 								.Append("\t\t")
