Horde: Move LeaseId into EpicGames.Horde.Api.

diff --git a/Engine/Source/Programs/Horde/Horde.Server.Tests/Fleet/AwsAutoScalingLifecycleTest.cs b/Engine/Source/Programs/Horde/Horde.Server.Tests/Fleet/AwsAutoScalingLifecycleTest.cs
--- a/Engine/Source/Programs/Horde/Horde.Server.Tests/Fleet/AwsAutoScalingLifecycleTest.cs
+++ b/Engine/Source/Programs/Horde/Horde.Server.Tests/Fleet/AwsAutoScalingLifecycleTest.cs
@@ -13,12 +13,14 @@ using Horde.Server.Agents;
 using Horde.Server.Agents.Fleet;
 using Horde.Server.Agents.Leases;
 using Horde.Server.Agents.Pools;
+using Horde.Server.Utilities;
 using HordeCommon;
 using Microsoft.Extensions.DependencyInjection;
 using Microsoft.Extensions.Logging;
 using Microsoft.VisualStudio.TestTools.UnitTesting;
 using MongoDB.Bson;
 using Moq;
+using EpicGames.Horde.Api;
 
 namespace Horde.Server.Tests.Fleet;
 
@@ -170,7 +172,7 @@ public class AwsAutoScalingLifecycleServiceTest : TestSetup
 	public async Task GetInstancesAvailableForTermination_AgentRunningJob_ReturnsNoInstanceId()
 	{
 		// Arrange
-		AgentLease lease = new(new LeaseId(ObjectId.GenerateNewId()), null, "test-lease", null, null, null, LeaseState.Active, null, false, null);
+		AgentLease lease = new(new LeaseId(BinaryIdUtils.CreateNew()), null, "test-lease", null, null, null, LeaseState.Active, null, false, null);
 		IAgent agent = await CreateAgentAsync(new PoolId("pool1"), awsInstanceId: "i-1000", lease: lease);
 		TerminationPolicyEvent e = CreateTerminationPolicyEvent("i-1000");
 
diff --git a/Engine/Source/Programs/Horde/Horde.Server.Tests/Fleet/LeaseUtilizationStrategyTest.cs b/Engine/Source/Programs/Horde/Horde.Server.Tests/Fleet/LeaseUtilizationStrategyTest.cs
--- a/Engine/Source/Programs/Horde/Horde.Server.Tests/Fleet/LeaseUtilizationStrategyTest.cs
+++ b/Engine/Source/Programs/Horde/Horde.Server.Tests/Fleet/LeaseUtilizationStrategyTest.cs
@@ -14,6 +14,7 @@ using Horde.Server.Utilities;
 using HordeCommon;
 using HordeCommon.Rpc.Tasks;
 using Microsoft.VisualStudio.TestTools.UnitTesting;
+using EpicGames.Horde.Api;
 
 namespace Horde.Server.Tests.Fleet
 {
@@ -123,7 +124,7 @@ namespace Horde.Server.Tests.Fleet
 				payload = Any.Pack(computeTask).ToByteArray();
 			}
 
-			ILease lease = await LeaseCollection.AddAsync(LeaseId.GenerateNewId(), null, "placeholderLease", agent.Id, agent.SessionId!.Value, new StreamId("placeholderStream"), pool.Id, null, startTime, payload);
+			ILease lease = await LeaseCollection.AddAsync(new LeaseId(BinaryIdUtils.CreateNew()), null, "placeholderLease", agent.Id, agent.SessionId!.Value, new StreamId("placeholderStream"), pool.Id, null, startTime, payload);
 			bool wasModified = await LeaseCollection.TrySetOutcomeAsync(lease.Id, startTime + duration, LeaseOutcome.Success, null);
 			Assert.IsTrue(wasModified);
 			return lease;
diff --git a/Engine/Source/Programs/Horde/Horde.Server.Tests/JobCollectionTests.cs b/Engine/Source/Programs/Horde/Horde.Server.Tests/JobCollectionTests.cs
--- a/Engine/Source/Programs/Horde/Horde.Server.Tests/JobCollectionTests.cs
+++ b/Engine/Source/Programs/Horde/Horde.Server.Tests/JobCollectionTests.cs
@@ -118,12 +118,12 @@ namespace Horde.Server.Tests
 
 			SessionId sessionId1 = SessionId.GenerateNewId();
 			await JobCollection.TryAssignLeaseAsync(fixture.Job1, 0, new PoolId("foo"), fixture.Agent1.Id,
-				sessionId1, LeaseId.GenerateNewId(), LogId.GenerateNewId());
+				sessionId1, new LeaseId(BinaryIdUtils.CreateNew()), LogId.GenerateNewId());
 			
 			SessionId sessionId2 = SessionId.GenerateNewId();
 			IJob job = (await JobCollection.GetAsync(fixture.Job1.Id))!;
 			await JobCollection.TryAssignLeaseAsync(job, 0, new PoolId("foo"), fixture.Agent1.Id,
-				sessionId2, LeaseId.GenerateNewId(), LogId.GenerateNewId());
+				sessionId2, new LeaseId(BinaryIdUtils.CreateNew()), LogId.GenerateNewId());
 			
 			// Manually verify the log output
 		}
diff --git a/Engine/Source/Programs/Horde/Horde.Server.Tests/LeasesControllerDbTest.cs b/Engine/Source/Programs/Horde/Horde.Server.Tests/LeasesControllerDbTest.cs
--- a/Engine/Source/Programs/Horde/Horde.Server.Tests/LeasesControllerDbTest.cs
+++ b/Engine/Source/Programs/Horde/Horde.Server.Tests/LeasesControllerDbTest.cs
@@ -10,6 +10,7 @@ using Horde.Server.Utilities;
 using HordeCommon;
 using Microsoft.AspNetCore.Mvc;
 using Microsoft.VisualStudio.TestTools.UnitTesting;
+using EpicGames.Horde.Api;
 
 namespace Horde.Server.Tests
 {
@@ -37,7 +38,7 @@ namespace Horde.Server.Tests
 
 		private async Task<ILease> CreateLease(DateTime startTime, TimeSpan duration)
 		{
-			LeaseId id = LeaseId.GenerateNewId();
+			LeaseId id = new LeaseId(BinaryIdUtils.CreateNew());
 			SessionId sessionId = SessionId.GenerateNewId();
 			ILease lease = await LeaseCollection.AddAsync(id, null, "myLease", new AgentId("agent-1"), sessionId, null, null, null, startTime, Array.Empty<byte>());
 			await LeaseCollection.TrySetOutcomeAsync(id, startTime + duration, LeaseOutcome.Success, null);
diff --git a/Engine/Source/Programs/Horde/Horde.Server.Tests/RpcServiceTest.cs b/Engine/Source/Programs/Horde/Horde.Server.Tests/RpcServiceTest.cs
--- a/Engine/Source/Programs/Horde/Horde.Server.Tests/RpcServiceTest.cs
+++ b/Engine/Source/Programs/Horde/Horde.Server.Tests/RpcServiceTest.cs
@@ -22,15 +22,17 @@ using Microsoft.AspNetCore.Http.Features;
 using Microsoft.Extensions.Hosting;
 using Microsoft.VisualStudio.TestTools.UnitTesting;
 using MongoDB.Bson;
-using AgentCapabilities = HordeCommon.Rpc.Messages.AgentCapabilities;
-using ISession = Microsoft.AspNetCore.Http.ISession;
 using Horde.Server.Jobs.Artifacts;
 using Horde.Server.Agents;
 using HordeCommon.Rpc.Messages;
 using Horde.Server.Agents.Sessions;
+using EpicGames.Horde.Api;
 
 namespace Horde.Server.Tests
 {
+	using AgentCapabilities = HordeCommon.Rpc.Messages.AgentCapabilities;
+	using ISession = Microsoft.AspNetCore.Http.ISession;
+
 	public class AppLifetimeStub : IHostApplicationLifetime
 	{
 		public CancellationToken ApplicationStarted { get; }
@@ -451,7 +453,7 @@ namespace Horde.Server.Tests
 			// Set the session ID on the job batch to pass auth later
 			Deref(await JobCollection.TryAssignLeaseAsync(fixture.Job1, 0, new PoolId("foo"),
 				new AgentId("test"), sessionId,
-				LeaseId.GenerateNewId(), LogId.GenerateNewId()));
+				new LeaseId(BinaryIdUtils.CreateNew()), LogId.GenerateNewId()));
 /*
 			TestAsyncStreamReader<UploadArtifactRequest> RequestStream = new TestAsyncStreamReader<UploadArtifactRequest>(Context);
 			Task<UploadArtifactResponse> Call = TestSetup.RpcService.UploadArtifact(RequestStream,  Context);
diff --git a/Engine/Source/Programs/Horde/Horde.Server/Acls/AclService.cs b/Engine/Source/Programs/Horde/Horde.Server/Acls/AclService.cs
--- a/Engine/Source/Programs/Horde/Horde.Server/Acls/AclService.cs
+++ b/Engine/Source/Programs/Horde/Horde.Server/Acls/AclService.cs
@@ -13,6 +13,7 @@ using Horde.Server.Server;
 using Horde.Server.Utilities;
 using Microsoft.IdentityModel.Tokens;
 using MongoDB.Driver;
+using EpicGames.Horde.Api;
 
 namespace Horde.Server.Acls
 {
diff --git a/Engine/Source/Programs/Horde/Horde.Server/Agents/AgentService.cs b/Engine/Source/Programs/Horde/Horde.Server/Agents/AgentService.cs
--- a/Engine/Source/Programs/Horde/Horde.Server/Agents/AgentService.cs
+++ b/Engine/Source/Programs/Horde/Horde.Server/Agents/AgentService.cs
@@ -9,6 +9,7 @@ using System.Security.Claims;
 using System.Threading;
 using System.Threading.Tasks;
 using EpicGames.Core;
+using EpicGames.Horde.Api;
 using EpicGames.Redis;
 using EpicGames.Serialization;
 using Google.Protobuf.WellKnownTypes;
diff --git a/Engine/Source/Programs/Horde/Horde.Server/Agents/AgentsController.cs b/Engine/Source/Programs/Horde/Horde.Server/Agents/AgentsController.cs
--- a/Engine/Source/Programs/Horde/Horde.Server/Agents/AgentsController.cs
+++ b/Engine/Source/Programs/Horde/Horde.Server/Agents/AgentsController.cs
@@ -16,6 +16,7 @@ using Microsoft.AspNetCore.Authorization;
 using Microsoft.AspNetCore.Mvc;
 using Microsoft.Extensions.Logging;
 using Microsoft.Extensions.Options;
+using EpicGames.Horde.Api;
 
 namespace Horde.Server.Agents
 {
diff --git a/Engine/Source/Programs/Horde/Horde.Server/Agents/ConformList.cs b/Engine/Source/Programs/Horde/Horde.Server/Agents/ConformList.cs
--- a/Engine/Source/Programs/Horde/Horde.Server/Agents/ConformList.cs
+++ b/Engine/Source/Programs/Horde/Horde.Server/Agents/ConformList.cs
@@ -7,6 +7,7 @@ using Horde.Server.Server;
 using Horde.Server.Utilities;
 using MongoDB.Bson;
 using MongoDB.Bson.Serialization.Attributes;
+using EpicGames.Horde.Api;
 
 namespace Horde.Server.Agents
 {
diff --git a/Engine/Source/Programs/Horde/Horde.Server/Agents/IAgent.cs b/Engine/Source/Programs/Horde/Horde.Server/Agents/IAgent.cs
--- a/Engine/Source/Programs/Horde/Horde.Server/Agents/IAgent.cs
+++ b/Engine/Source/Programs/Horde/Horde.Server/Agents/IAgent.cs
@@ -7,6 +7,7 @@ using System.Linq;
 using System.Security.Cryptography;
 using System.Threading.Tasks;
 using EpicGames.Core;
+using EpicGames.Horde.Api;
 using EpicGames.Horde.Common;
 using EpicGames.Horde.Compute;
 using Google.Protobuf.WellKnownTypes;
diff --git a/Engine/Source/Programs/Horde/Horde.Server/Agents/Leases/ILease.cs b/Engine/Source/Programs/Horde/Horde.Server/Agents/Leases/ILease.cs
--- a/Engine/Source/Programs/Horde/Horde.Server/Agents/Leases/ILease.cs
+++ b/Engine/Source/Programs/Horde/Horde.Server/Agents/Leases/ILease.cs
@@ -8,6 +8,7 @@ using Horde.Server.Agents.Sessions;
 using Horde.Server.Logs;
 using Horde.Server.Streams;
 using HordeCommon;
+using EpicGames.Horde.Api;
 
 namespace Horde.Server.Agents.Leases
 {
diff --git a/Engine/Source/Programs/Horde/Horde.Server/Agents/Leases/ILeaseCollection.cs b/Engine/Source/Programs/Horde/Horde.Server/Agents/Leases/ILeaseCollection.cs
--- a/Engine/Source/Programs/Horde/Horde.Server/Agents/Leases/ILeaseCollection.cs
+++ b/Engine/Source/Programs/Horde/Horde.Server/Agents/Leases/ILeaseCollection.cs
@@ -9,6 +9,7 @@ using Horde.Server.Logs;
 using Horde.Server.Streams;
 using Horde.Server.Utilities;
 using HordeCommon;
+using EpicGames.Horde.Api;
 
 namespace Horde.Server.Agents.Leases
 {
diff --git a/Engine/Source/Programs/Horde/Horde.Server/Agents/Leases/LeaseCollection.cs b/Engine/Source/Programs/Horde/Horde.Server/Agents/Leases/LeaseCollection.cs
--- a/Engine/Source/Programs/Horde/Horde.Server/Agents/Leases/LeaseCollection.cs
+++ b/Engine/Source/Programs/Horde/Horde.Server/Agents/Leases/LeaseCollection.cs
@@ -13,6 +13,7 @@ using Horde.Server.Utilities;
 using HordeCommon;
 using MongoDB.Bson.Serialization.Attributes;
 using MongoDB.Driver;
+using EpicGames.Horde.Api;
 
 namespace Horde.Server.Agents.Leases
 {
diff --git a/Engine/Source/Programs/Horde/Horde.Server/Agents/Leases/LeaseId.cs b/Engine/Source/Programs/Horde/Horde.Server/Agents/Leases/LeaseId.cs
deleted file mode 100644
diff --git a/Engine/Source/Programs/Horde/Horde.Server/Agents/Leases/LeaseMessages.cs b/Engine/Source/Programs/Horde/Horde.Server/Agents/Leases/LeaseMessages.cs
deleted file mode 100644
diff --git a/Engine/Source/Programs/Horde/Horde.Server/Agents/Leases/LeasesController.cs b/Engine/Source/Programs/Horde/Horde.Server/Agents/Leases/LeasesController.cs
--- a/Engine/Source/Programs/Horde/Horde.Server/Agents/Leases/LeasesController.cs
+++ b/Engine/Source/Programs/Horde/Horde.Server/Agents/Leases/LeasesController.cs
@@ -12,6 +12,7 @@ using Microsoft.AspNetCore.Authorization;
 using Microsoft.AspNetCore.Mvc;
 using Microsoft.Extensions.Options;
 using OpenTelemetry.Trace;
+using EpicGames.Horde.Api;
 
 namespace Horde.Server.Agents.Leases
 {
diff --git a/Engine/Source/Programs/Horde/Horde.Server/Compute/ComputeController.cs b/Engine/Source/Programs/Horde/Horde.Server/Compute/ComputeController.cs
--- a/Engine/Source/Programs/Horde/Horde.Server/Compute/ComputeController.cs
+++ b/Engine/Source/Programs/Horde/Horde.Server/Compute/ComputeController.cs
@@ -6,9 +6,9 @@ using System.Net;
 using System.Threading;
 using System.Threading.Tasks;
 using EpicGames.Core;
+using EpicGames.Horde.Api;
 using EpicGames.Horde.Compute;
 using Horde.Server.Acls;
-using Horde.Server.Agents.Leases;
 using Horde.Server.Server;
 using Horde.Server.Utilities;
 using Microsoft.AspNetCore.Authorization;
diff --git a/Engine/Source/Programs/Horde/Horde.Server/Compute/ComputeService.cs b/Engine/Source/Programs/Horde/Horde.Server/Compute/ComputeService.cs
--- a/Engine/Source/Programs/Horde/Horde.Server/Compute/ComputeService.cs
+++ b/Engine/Source/Programs/Horde/Horde.Server/Compute/ComputeService.cs
@@ -7,6 +7,7 @@ using System.Net;
 using System.Security.Cryptography;
 using System.Threading;
 using System.Threading.Tasks;
+using EpicGames.Horde.Api;
 using EpicGames.Horde.Compute;
 using EpicGames.Horde.Compute.Clients;
 using EpicGames.Horde.Compute.Transports;
@@ -16,6 +17,7 @@ using Horde.Server.Agents;
 using Horde.Server.Agents.Leases;
 using Horde.Server.Jobs;
 using Horde.Server.Logs;
+using Horde.Server.Utilities;
 using HordeCommon;
 using HordeCommon.Rpc.Tasks;
 
@@ -51,7 +53,7 @@ namespace Horde.Server.Compute
 				Dictionary<string, int> assignedResources = new Dictionary<string, int>();
 				if (agent.MeetsRequirements(requirements, assignedResources))
 				{
-					LeaseId leaseId = LeaseId.GenerateNewId();
+					LeaseId leaseId = new LeaseId(BinaryIdUtils.CreateNew());
 					ILogFile? log = await _logService.CreateLogFileAsync(JobId.Empty, leaseId, agent.SessionId, LogType.Json, useNewStorageBackend: true, cancellationToken: cancellationToken);
 
 					ComputeTask computeTask = CreateComputeTask(assignedResources, log?.Id, parentLeaseId);
diff --git a/Engine/Source/Programs/Horde/Horde.Server/Jobs/JobTaskSource.cs b/Engine/Source/Programs/Horde/Horde.Server/Jobs/JobTaskSource.cs
--- a/Engine/Source/Programs/Horde/Horde.Server/Jobs/JobTaskSource.cs
+++ b/Engine/Source/Programs/Horde/Horde.Server/Jobs/JobTaskSource.cs
@@ -32,6 +32,7 @@ using Microsoft.Extensions.Hosting;
 using Microsoft.Extensions.Logging;
 using Microsoft.Extensions.Options;
 using MongoDB.Driver;
+using EpicGames.Horde.Api;
 
 namespace Horde.Server.Jobs
 {
@@ -729,7 +730,7 @@ namespace Horde.Server.Jobs
 			_logger.LogDebug("Assigning job {JobId}, batch {BatchId} to waiter (agent {AgentId})", job.Id, batch.Id, agent.Id);
 
 			// Generate a new unique id for the lease
-			LeaseId leaseId = LeaseId.GenerateNewId();
+			LeaseId leaseId = new LeaseId(BinaryIdUtils.CreateNew());
 
 			// The next time to try assigning to another agent
 			DateTime backOffTime = DateTime.UtcNow + TimeSpan.FromMinutes(1.0);
diff --git a/Engine/Source/Programs/Horde/Horde.Server/Logs/ILogFile.cs b/Engine/Source/Programs/Horde/Horde.Server/Logs/ILogFile.cs
--- a/Engine/Source/Programs/Horde/Horde.Server/Logs/ILogFile.cs
+++ b/Engine/Source/Programs/Horde/Horde.Server/Logs/ILogFile.cs
@@ -3,6 +3,7 @@
 using System;
 using System.Collections.Generic;
 using EpicGames.Core;
+using EpicGames.Horde.Api;
 using EpicGames.Horde.Storage;
 using Horde.Server.Agents.Leases;
 using Horde.Server.Agents.Sessions;
diff --git a/Engine/Source/Programs/Horde/Horde.Server/Logs/ILogFileCollection.cs b/Engine/Source/Programs/Horde/Horde.Server/Logs/ILogFileCollection.cs
--- a/Engine/Source/Programs/Horde/Horde.Server/Logs/ILogFileCollection.cs
+++ b/Engine/Source/Programs/Horde/Horde.Server/Logs/ILogFileCollection.cs
@@ -7,6 +7,7 @@ using Horde.Server.Agents.Leases;
 using Horde.Server.Agents.Sessions;
 using Horde.Server.Jobs;
 using Horde.Server.Utilities;
+using EpicGames.Horde.Api;
 
 namespace Horde.Server.Logs
 {
diff --git a/Engine/Source/Programs/Horde/Horde.Server/Logs/LogFileCollection.cs b/Engine/Source/Programs/Horde/Horde.Server/Logs/LogFileCollection.cs
--- a/Engine/Source/Programs/Horde/Horde.Server/Logs/LogFileCollection.cs
+++ b/Engine/Source/Programs/Horde/Horde.Server/Logs/LogFileCollection.cs
@@ -5,6 +5,7 @@ using System.Linq;
 using System.Net;
 using System.Threading;
 using System.Threading.Tasks;
+using EpicGames.Horde.Api;
 using EpicGames.Horde.Storage;
 using Horde.Server.Agents.Leases;
 using Horde.Server.Agents.Sessions;
diff --git a/Engine/Source/Programs/Horde/Horde.Server/Logs/LogFileService.cs b/Engine/Source/Programs/Horde/Horde.Server/Logs/LogFileService.cs
--- a/Engine/Source/Programs/Horde/Horde.Server/Logs/LogFileService.cs
+++ b/Engine/Source/Programs/Horde/Horde.Server/Logs/LogFileService.cs
@@ -28,10 +28,12 @@ using Microsoft.Extensions.Logging;
 using Microsoft.Extensions.Options;
 using MongoDB.Bson;
 using OpenTelemetry.Trace;
-using Stream = System.IO.Stream;
+using EpicGames.Horde.Api;
 
 namespace Horde.Server.Logs
 {
+	using Stream = System.IO.Stream;
+
 	/// <summary>
 	/// Metadata about a log file
 	/// </summary>
diff --git a/Engine/Source/Programs/Horde/Horde.Server/Tasks/ConformTaskSource.cs b/Engine/Source/Programs/Horde/Horde.Server/Tasks/ConformTaskSource.cs
--- a/Engine/Source/Programs/Horde/Horde.Server/Tasks/ConformTaskSource.cs
+++ b/Engine/Source/Programs/Horde/Horde.Server/Tasks/ConformTaskSource.cs
@@ -20,6 +20,7 @@ using HordeCommon.Rpc.Tasks;
 using Microsoft.Extensions.Hosting;
 using Microsoft.Extensions.Logging;
 using Microsoft.Extensions.Options;
+using EpicGames.Horde.Api;
 
 namespace Horde.Server.Tasks
 {
@@ -173,7 +174,7 @@ namespace Horde.Server.Tasks
 				ConformTask task = new ConformTask();
 				if (await GetWorkspacesAsync(agent, task.Workspaces))
 				{
-					LeaseId leaseId = LeaseId.GenerateNewId();
+					LeaseId leaseId = new LeaseId(BinaryIdUtils.CreateNew());
 					if (await AllocateConformLeaseAsync(agent.Id, task.Workspaces, leaseId))
 					{
 						ILogFile log = await _logService.CreateLogFileAsync(JobId.Empty, leaseId, agent.SessionId, LogType.Json, useNewStorageBackend: false, cancellationToken: cancellationToken);
diff --git a/Engine/Source/Programs/Horde/Horde.Server/Tasks/ITaskSource.cs b/Engine/Source/Programs/Horde/Horde.Server/Tasks/ITaskSource.cs
--- a/Engine/Source/Programs/Horde/Horde.Server/Tasks/ITaskSource.cs
+++ b/Engine/Source/Programs/Horde/Horde.Server/Tasks/ITaskSource.cs
@@ -16,6 +16,7 @@ using Horde.Server.Agents;
 using Horde.Server.Agents.Leases;
 using HordeCommon;
 using Microsoft.Extensions.Logging;
+using EpicGames.Horde.Api;
 
 namespace Horde.Server.Tasks
 {
diff --git a/Engine/Source/Programs/Horde/Horde.Server/Tasks/RestartTaskSource.cs b/Engine/Source/Programs/Horde/Horde.Server/Tasks/RestartTaskSource.cs
--- a/Engine/Source/Programs/Horde/Horde.Server/Tasks/RestartTaskSource.cs
+++ b/Engine/Source/Programs/Horde/Horde.Server/Tasks/RestartTaskSource.cs
@@ -11,6 +11,7 @@ using Horde.Server.Logs;
 using Horde.Server.Utilities;
 using HordeCommon;
 using HordeCommon.Rpc.Tasks;
+using EpicGames.Horde.Api;
 
 namespace Horde.Server.Tasks
 {
@@ -40,7 +41,7 @@ namespace Horde.Server.Tasks
 				return await DrainAsync(cancellationToken);
 			}
 
-			LeaseId leaseId = LeaseId.GenerateNewId();
+			LeaseId leaseId = new LeaseId(BinaryIdUtils.CreateNew());
 			ILogFile log = await _logService.CreateLogFileAsync(JobId.Empty, leaseId, agent.SessionId, LogType.Json, useNewStorageBackend: false, cancellationToken: cancellationToken);
 
 			RestartTask task = new RestartTask();
diff --git a/Engine/Source/Programs/Horde/Horde.Server/Tasks/ShutdownTaskSource.cs b/Engine/Source/Programs/Horde/Horde.Server/Tasks/ShutdownTaskSource.cs
--- a/Engine/Source/Programs/Horde/Horde.Server/Tasks/ShutdownTaskSource.cs
+++ b/Engine/Source/Programs/Horde/Horde.Server/Tasks/ShutdownTaskSource.cs
@@ -11,6 +11,7 @@ using Horde.Server.Logs;
 using Horde.Server.Utilities;
 using HordeCommon;
 using HordeCommon.Rpc.Tasks;
+using EpicGames.Horde.Api;
 
 namespace Horde.Server.Tasks
 {
@@ -39,7 +40,7 @@ namespace Horde.Server.Tasks
 				return await DrainAsync(cancellationToken);
 			}
 
-			LeaseId leaseId = LeaseId.GenerateNewId();
+			LeaseId leaseId = new LeaseId(BinaryIdUtils.CreateNew());
 			ILogFile log = await _logService.CreateLogFileAsync(JobId.Empty, leaseId, agent.SessionId, LogType.Json, useNewStorageBackend: false, cancellationToken: cancellationToken);
 
 			ShutdownTask task = new ShutdownTask();
diff --git a/Engine/Source/Programs/Horde/Horde.Server/Tasks/UpgradeTaskSource.cs b/Engine/Source/Programs/Horde/Horde.Server/Tasks/UpgradeTaskSource.cs
--- a/Engine/Source/Programs/Horde/Horde.Server/Tasks/UpgradeTaskSource.cs
+++ b/Engine/Source/Programs/Horde/Horde.Server/Tasks/UpgradeTaskSource.cs
@@ -5,6 +5,7 @@ using System.Text;
 using System.Threading;
 using System.Threading.Tasks;
 using EpicGames.Core;
+using EpicGames.Horde.Api;
 using Google.Protobuf;
 using Google.Protobuf.WellKnownTypes;
 using Horde.Server.Agents;
@@ -13,6 +14,7 @@ using Horde.Server.Jobs;
 using Horde.Server.Logs;
 using Horde.Server.Server;
 using Horde.Server.Tools;
+using Horde.Server.Utilities;
 using HordeCommon;
 using HordeCommon.Rpc.Tasks;
 using Microsoft.Extensions.Options;
@@ -67,7 +69,7 @@ namespace Horde.Server.Tasks
 				return await DrainAsync(cancellationToken);
 			}
 
-			LeaseId leaseId = LeaseId.GenerateNewId();
+			LeaseId leaseId = new LeaseId(BinaryIdUtils.CreateNew());
 			ILogFile logFile = await _logService.CreateLogFileAsync(JobId.Empty, leaseId, agent.SessionId, LogType.Json, useNewStorageBackend: false, cancellationToken: cancellationToken);
 
 			UpgradeTask task = new UpgradeTask();
diff --git a/Engine/Source/Programs/Shared/EpicGames.Horde/Api/LeaseId.cs b/Engine/Source/Programs/Shared/EpicGames.Horde/Api/LeaseId.cs
new file mode 100644
--- /dev/null
+++ b/Engine/Source/Programs/Shared/EpicGames.Horde/Api/LeaseId.cs
@@ -0,0 +1,51 @@
+// Copyright Epic Games, Inc. All Rights Reserved.
+
+using System.ComponentModel;
+using EpicGames.Core;
+
+namespace EpicGames.Horde.Api
+{
+	/// <summary>
+	/// Identifier for a user
+	/// </summary>
+	/// <param name="Id">Id to construct from</param>
+	[JsonSchemaString]
+	[TypeConverter(typeof(BinaryIdTypeConverter<LeaseId, LeaseIdConverter>))]
+	[BinaryIdConverter(typeof(LeaseIdConverter))]
+	public record struct LeaseId(BinaryId Id)
+	{
+		/// <inheritdoc cref="BinaryId.TryParse(System.String, out BinaryId)"/>
+		public static bool TryParse(string text, out LeaseId leaseId)
+		{
+			BinaryId objectId;
+			if (BinaryId.TryParse(text, out objectId))
+			{
+				leaseId = new LeaseId(objectId);
+				return true;
+			}
+			else
+			{
+				leaseId = default;
+				return false;
+			}
+		}
+
+		/// <inheritdoc cref="BinaryId.Parse(System.String)"/>
+		public static LeaseId Parse(string text) => new LeaseId(BinaryId.Parse(text));
+
+		/// <inheritdoc/>
+		public override string ToString() => Id.ToString();
+	}
+
+	/// <summary>
+	/// Converter to and from <see cref="BinaryId"/> instances.
+	/// </summary>
+	class LeaseIdConverter : BinaryIdConverter<LeaseId>
+	{
+		/// <inheritdoc/>
+		public override LeaseId FromBinaryId(BinaryId id) => new LeaseId(id);
+
+		/// <inheritdoc/>
+		public override BinaryId ToBinaryId(LeaseId value) => value.Id;
+	}
+}
diff --git a/Engine/Source/Programs/Shared/EpicGames.Horde/Api/LeaseMessages.cs b/Engine/Source/Programs/Shared/EpicGames.Horde/Api/LeaseMessages.cs
new file mode 100644
--- /dev/null
+++ b/Engine/Source/Programs/Shared/EpicGames.Horde/Api/LeaseMessages.cs
@@ -0,0 +1,15 @@
+// Copyright Epic Games, Inc. All Rights Reserved.
+
+namespace EpicGames.Horde.Api
+{
+	/// <summary>
+	/// Updates an existing lease
+	/// </summary>
+	public class UpdateLeaseRequest
+	{
+		/// <summary>
+		/// Mark this lease as aborted
+		/// </summary>
+		public bool? Aborted { get; set; }
+	}
+}
