Horde: Add a custom log value formatter for LeaseId types, so the dashboard can hyperlink them to the correct place.

diff --git a/Engine/Source/Programs/Horde/Horde.Server/Compute/ComputeController.cs b/Engine/Source/Programs/Horde/Horde.Server/Compute/ComputeController.cs
--- a/Engine/Source/Programs/Horde/Horde.Server/Compute/ComputeController.cs
+++ b/Engine/Source/Programs/Horde/Horde.Server/Compute/ComputeController.cs
@@ -17,58 +17,6 @@ using Microsoft.Extensions.Options;
 
 namespace Horde.Server.Compute
 {
-	/// <summary>
-	/// Request a machine to execute compute requests
-	/// </summary>
-	public class AssignComputeRequest
-	{
-		/// <summary>
-		/// Condition to identify machines that can execute the request
-		/// </summary>
-		public Requirements? Requirements { get; set; }
-	}
-
-	/// <summary>
-	/// Request a machine to execute compute requests
-	/// </summary>
-	public class AssignComputeResponse
-	{
-		/// <summary>
-		/// IP address of the remote machine
-		/// </summary>
-		public string Ip { get; set; } = String.Empty;
-
-		/// <summary>
-		/// Port number on the remote machine
-		/// </summary>
-		public int Port { get; set; }
-
-		/// <summary>
-		/// Cryptographic nonce to identify the request, as a hex string
-		/// </summary>
-		public string Nonce { get; set; } = String.Empty;
-
-		/// <summary>
-		/// AES key for the channel, as a hex string
-		/// </summary>
-		public string Key { get; set; } = String.Empty;
-
-		/// <summary>
-		/// Identifier for the new lease on the remote machine
-		/// </summary>
-		public LeaseId LeaseId { get; set; }
-
-		/// <summary>
-		/// Resources assigned to this machine
-		/// </summary>
-		public Dictionary<string, int> AssignedResources { get; set; } = new Dictionary<string, int>();
-
-		/// <summary>
-		/// Properties of the agent assigned to do the work
-		/// </summary>
-		public IReadOnlyList<string> Properties { get; set; } = new List<string>();
-	}
-
 	/// <summary>
 	/// Controller for the /api/v2/compute endpoint
 	/// </summary>
diff --git a/Engine/Source/Programs/Shared/EpicGames.Horde/Api/ComputeMessages.cs b/Engine/Source/Programs/Shared/EpicGames.Horde/Api/ComputeMessages.cs
new file mode 100644
--- /dev/null
+++ b/Engine/Source/Programs/Shared/EpicGames.Horde/Api/ComputeMessages.cs
@@ -0,0 +1,62 @@
+// Copyright Epic Games, Inc. All Rights Reserved.
+
+using System;
+using System.Collections.Generic;
+using EpicGames.Horde.Compute;
+
+#pragma warning disable CA2227 // Collection properties should be read only
+
+namespace EpicGames.Horde.Api
+{
+	/// <summary>
+	/// Request a machine to execute compute requests
+	/// </summary>
+	public class AssignComputeRequest
+	{
+		/// <summary>
+		/// Condition to identify machines that can execute the request
+		/// </summary>
+		public Requirements? Requirements { get; set; }
+	}
+
+	/// <summary>
+	/// Request a machine to execute compute requests
+	/// </summary>
+	public class AssignComputeResponse
+	{
+		/// <summary>
+		/// IP address of the remote machine
+		/// </summary>
+		public string Ip { get; set; } = String.Empty;
+
+		/// <summary>
+		/// Port number on the remote machine
+		/// </summary>
+		public int Port { get; set; }
+
+		/// <summary>
+		/// Cryptographic nonce to identify the request, as a hex string
+		/// </summary>
+		public string Nonce { get; set; } = String.Empty;
+
+		/// <summary>
+		/// AES key for the channel, as a hex string
+		/// </summary>
+		public string Key { get; set; } = String.Empty;
+
+		/// <summary>
+		/// Identifier for the new lease on the remote machine
+		/// </summary>
+		public LeaseId LeaseId { get; set; }
+
+		/// <summary>
+		/// Resources assigned to this machine
+		/// </summary>
+		public Dictionary<string, int> AssignedResources { get; set; } = new Dictionary<string, int>();
+
+		/// <summary>
+		/// Properties of the agent assigned to do the work
+		/// </summary>
+		public IReadOnlyList<string> Properties { get; set; } = new List<string>();
+	}
+}
diff --git a/Engine/Source/Programs/Shared/EpicGames.Horde/Api/LeaseId.cs b/Engine/Source/Programs/Shared/EpicGames.Horde/Api/LeaseId.cs
--- a/Engine/Source/Programs/Shared/EpicGames.Horde/Api/LeaseId.cs
+++ b/Engine/Source/Programs/Shared/EpicGames.Horde/Api/LeaseId.cs
@@ -1,6 +1,7 @@
 // Copyright Epic Games, Inc. All Rights Reserved.
 
 using System.ComponentModel;
+using System.Text.Json;
 using EpicGames.Core;
 
 namespace EpicGames.Horde.Api
@@ -12,6 +13,7 @@ namespace EpicGames.Horde.Api
 	[JsonSchemaString]
 	[TypeConverter(typeof(BinaryIdTypeConverter<LeaseId, LeaseIdConverter>))]
 	[BinaryIdConverter(typeof(LeaseIdConverter))]
+	[LogValueFormatterAttribute(typeof(LeaseIdLogFormatter))]
 	public record struct LeaseId(BinaryId Id)
 	{
 		/// <inheritdoc cref="BinaryId.TryParse(System.String, out BinaryId)"/>
@@ -48,4 +50,19 @@ namespace EpicGames.Horde.Api
 		/// <inheritdoc/>
 		public override BinaryId ToBinaryId(LeaseId value) => value.Id;
 	}
+
+	/// <summary>
+	/// Formats a LeaseId as a typed log value
+	/// </summary>
+	class LeaseIdLogFormatter : ILogValueFormatter
+	{
+		/// <inheritdoc/>
+		public void Format(object value, Utf8JsonWriter writer)
+		{
+			writer.WriteStartObject();
+			writer.WriteString(LogEventPropertyName.Type, "LeaseId");
+			writer.WriteString(LogEventPropertyName.Text, ((LeaseId)value).ToString());
+			writer.WriteEndObject();
+		}
+	}
 }
diff --git a/Engine/Source/Programs/Shared/EpicGames.Horde/Compute/Clients/ServerComputeClient.cs b/Engine/Source/Programs/Shared/EpicGames.Horde/Compute/Clients/ServerComputeClient.cs
--- a/Engine/Source/Programs/Shared/EpicGames.Horde/Compute/Clients/ServerComputeClient.cs
+++ b/Engine/Source/Programs/Shared/EpicGames.Horde/Compute/Clients/ServerComputeClient.cs
@@ -1,11 +1,11 @@
 // Copyright Epic Games, Inc. All Rights Reserved.
 
 using EpicGames.Core;
+using EpicGames.Horde.Api;
 using EpicGames.Horde.Compute.Transports;
 using Microsoft.Extensions.Logging;
 using System;
 using System.Collections.Generic;
-using System.Linq;
 using System.Net;
 using System.Net.Http;
 using System.Net.Http.Headers;
@@ -27,21 +27,6 @@ namespace EpicGames.Horde.Compute.Clients
 		/// </summary>
 		public const int NonceLength = 64;
 
-		class AssignComputeRequest
-		{
-			public Requirements? Requirements { get; set; }
-		}
-
-		class AssignComputeResponse
-		{
-			public string Ip { get; set; } = String.Empty;
-			public int Port { get; set; }
-			public string Nonce { get; set; } = String.Empty;
-			public string Key { get; set; } = String.Empty;
-			public Dictionary<string, int> AssignedResources { get; set; } = new Dictionary<string, int>();
-			public List<string> Properties { get; set; } = new List<string>();
-		}
-
 		record class LeaseInfo(IReadOnlyList<string> Properties, IReadOnlyDictionary<string, int> AssignedResources, RemoteComputeSocket Socket);
 
 		class LeaseImpl : IComputeLease
@@ -181,7 +166,7 @@ namespace EpicGames.Horde.Compute.Clients
 			// Send the nonce
 			byte[] nonce = StringUtils.ParseHexString(responseMessage.Nonce);
 			await socket.SendMessageAsync(nonce, SocketFlags.None, cancellationToken);
-			_logger.LogDebug("Connection established.");
+			_logger.LogInformation("Connected to {Ip} under lease {LeaseId}", responseMessage.Ip, responseMessage.LeaseId);
 
 			// Pass the rest of the call over to the handler
 			byte[] key = StringUtils.ParseHexString(responseMessage.Key);
