Horde: Fix serialization of LeaseId types in compute responses.

diff --git a/Engine/Source/Programs/Shared/EpicGames.Horde/Compute/Clients/ServerComputeClient.cs b/Engine/Source/Programs/Shared/EpicGames.Horde/Compute/Clients/ServerComputeClient.cs
--- a/Engine/Source/Programs/Shared/EpicGames.Horde/Compute/Clients/ServerComputeClient.cs
+++ b/Engine/Source/Programs/Shared/EpicGames.Horde/Compute/Clients/ServerComputeClient.cs
@@ -12,6 +12,7 @@ using System.Net.Http.Headers;
 using System.Net.Http.Json;
 using System.Net.Sockets;
 using System.Runtime.CompilerServices;
+using System.Text.Json;
 using System.Threading;
 using System.Threading.Tasks;
 
@@ -128,6 +129,9 @@ namespace EpicGames.Horde.Compute.Clients
 		{
 			_logger.LogDebug("Requesting compute resource");
 
+			JsonSerializerOptions jsonSerializerOptions = new JsonSerializerOptions();
+			HordeHttpClient.ConfigureJsonSerializer(jsonSerializerOptions);
+
 			// Assign a compute worker
 			HttpClient client = await _createHttpClientAsync(cancellationToken);
 
@@ -150,7 +154,7 @@ namespace EpicGames.Horde.Compute.Clients
 
 				response.EnsureSuccessStatusCode();
 
-				responseMessage = await response.Content.ReadFromJsonAsync<AssignComputeResponse>(cancellationToken: cancellationToken);
+				responseMessage = await response.Content.ReadFromJsonAsync<AssignComputeResponse>(jsonSerializerOptions, cancellationToken);
 				if (responseMessage == null)
 				{
 					throw new InvalidOperationException();
