Horde: Move AgentId into EpicGames.Horde, and support rendering them as custom log values.

diff --git a/Engine/Source/Programs/Horde/Horde.Server.Tests/AgentControllerTest.cs b/Engine/Source/Programs/Horde/Horde.Server.Tests/AgentControllerTest.cs
--- a/Engine/Source/Programs/Horde/Horde.Server.Tests/AgentControllerTest.cs
+++ b/Engine/Source/Programs/Horde/Horde.Server.Tests/AgentControllerTest.cs
@@ -4,6 +4,7 @@ using System.Threading.Tasks;
 using Horde.Server.Agents;
 using Microsoft.AspNetCore.Mvc;
 using Microsoft.VisualStudio.TestTools.UnitTesting;
+using EpicGames.Horde.Api;
 
 namespace Horde.Server.Tests;
 
diff --git a/Engine/Source/Programs/Horde/Horde.Server.Tests/AgentServiceTest.cs b/Engine/Source/Programs/Horde/Horde.Server.Tests/AgentServiceTest.cs
--- a/Engine/Source/Programs/Horde/Horde.Server.Tests/AgentServiceTest.cs
+++ b/Engine/Source/Programs/Horde/Horde.Server.Tests/AgentServiceTest.cs
@@ -7,6 +7,7 @@ using System.Security.Claims;
 using System.Threading;
 using System.Threading.Tasks;
 using EpicGames.Core;
+using EpicGames.Horde.Api;
 using Horde.Server.Agents;
 using Horde.Server.Agents.Pools;
 using Horde.Server.Auditing;
diff --git a/Engine/Source/Programs/Horde/Horde.Server.Tests/Agents/Pools/PoolUpdateServiceTest.cs b/Engine/Source/Programs/Horde/Horde.Server.Tests/Agents/Pools/PoolUpdateServiceTest.cs
--- a/Engine/Source/Programs/Horde/Horde.Server.Tests/Agents/Pools/PoolUpdateServiceTest.cs
+++ b/Engine/Source/Programs/Horde/Horde.Server.Tests/Agents/Pools/PoolUpdateServiceTest.cs
@@ -7,6 +7,7 @@ using System.Threading;
 using Horde.Server.Agents;
 using Horde.Server.Agents.Pools;
 using Microsoft.Extensions.Logging.Abstractions;
+using EpicGames.Horde.Api;
 
 namespace Horde.Server.Tests.Agents.Pools
 {
diff --git a/Engine/Source/Programs/Horde/Horde.Server.Tests/TestSetup.cs b/Engine/Source/Programs/Horde/Horde.Server.Tests/TestSetup.cs
--- a/Engine/Source/Programs/Horde/Horde.Server.Tests/TestSetup.cs
+++ b/Engine/Source/Programs/Horde/Horde.Server.Tests/TestSetup.cs
@@ -11,6 +11,7 @@ using Amazon.AutoScaling;
 using Amazon.EC2;
 using Datadog.Trace;
 using EpicGames.Core;
+using EpicGames.Horde.Api;
 using EpicGames.Horde.Storage;
 using Horde.Server.Acls;
 using Horde.Server.Agents;
diff --git a/Engine/Source/Programs/Horde/Horde.Server/Agents/AgentCollection.cs b/Engine/Source/Programs/Horde/Horde.Server/Agents/AgentCollection.cs
--- a/Engine/Source/Programs/Horde/Horde.Server/Agents/AgentCollection.cs
+++ b/Engine/Source/Programs/Horde/Horde.Server/Agents/AgentCollection.cs
@@ -5,6 +5,7 @@ using System.Collections.Generic;
 using System.Linq;
 using System.Linq.Expressions;
 using System.Threading.Tasks;
+using EpicGames.Horde.Api;
 using EpicGames.Redis;
 using Google.Protobuf.WellKnownTypes;
 using Horde.Server.Agents.Pools;
diff --git a/Engine/Source/Programs/Horde/Horde.Server/Agents/AgentId.cs b/Engine/Source/Programs/Horde/Horde.Server/Agents/AgentId.cs
deleted file mode 100644
diff --git a/Engine/Source/Programs/Horde/Horde.Server/Agents/Fleet/AwsAutoScalingLifecycle.cs b/Engine/Source/Programs/Horde/Horde.Server/Agents/Fleet/AwsAutoScalingLifecycle.cs
--- a/Engine/Source/Programs/Horde/Horde.Server/Agents/Fleet/AwsAutoScalingLifecycle.cs
+++ b/Engine/Source/Programs/Horde/Horde.Server/Agents/Fleet/AwsAutoScalingLifecycle.cs
@@ -15,6 +15,7 @@ using Amazon.Runtime;
 using Amazon.SQS;
 using Amazon.SQS.Model;
 using EpicGames.Core;
+using EpicGames.Horde.Api;
 using Horde.Server.Server;
 using Horde.Server.Utilities;
 using HordeCommon;
diff --git a/Engine/Source/Programs/Horde/Horde.Server/Agents/Fleet/LeaseUtilizationStrategy.cs b/Engine/Source/Programs/Horde/Horde.Server/Agents/Fleet/LeaseUtilizationStrategy.cs
--- a/Engine/Source/Programs/Horde/Horde.Server/Agents/Fleet/LeaseUtilizationStrategy.cs
+++ b/Engine/Source/Programs/Horde/Horde.Server/Agents/Fleet/LeaseUtilizationStrategy.cs
@@ -13,6 +13,7 @@ using HordeCommon;
 using HordeCommon.Rpc.Tasks;
 using Microsoft.Extensions.Caching.Memory;
 using OpenTelemetry.Trace;
+using EpicGames.Horde.Api;
 
 namespace Horde.Server.Agents.Fleet
 {
diff --git a/Engine/Source/Programs/Horde/Horde.Server/Agents/Fleet/Providers/AwsFleetManager.cs b/Engine/Source/Programs/Horde/Horde.Server/Agents/Fleet/Providers/AwsFleetManager.cs
--- a/Engine/Source/Programs/Horde/Horde.Server/Agents/Fleet/Providers/AwsFleetManager.cs
+++ b/Engine/Source/Programs/Horde/Horde.Server/Agents/Fleet/Providers/AwsFleetManager.cs
@@ -12,6 +12,7 @@ using Horde.Server.Auditing;
 using Horde.Server.Utilities;
 using Microsoft.Extensions.Logging;
 using OpenTelemetry.Trace;
+using EpicGames.Horde.Api;
 
 namespace Horde.Server.Agents.Fleet.Providers
 {
diff --git a/Engine/Source/Programs/Horde/Horde.Server/Agents/IAgentCollection.cs b/Engine/Source/Programs/Horde/Horde.Server/Agents/IAgentCollection.cs
--- a/Engine/Source/Programs/Horde/Horde.Server/Agents/IAgentCollection.cs
+++ b/Engine/Source/Programs/Horde/Horde.Server/Agents/IAgentCollection.cs
@@ -9,6 +9,7 @@ using Horde.Server.Agents.Sessions;
 using Horde.Server.Auditing;
 using Horde.Server.Utilities;
 using HordeCommon;
+using EpicGames.Horde.Api;
 
 namespace Horde.Server.Agents
 {
diff --git a/Engine/Source/Programs/Horde/Horde.Server/Agents/Sessions/ISession.cs b/Engine/Source/Programs/Horde/Horde.Server/Agents/Sessions/ISession.cs
--- a/Engine/Source/Programs/Horde/Horde.Server/Agents/Sessions/ISession.cs
+++ b/Engine/Source/Programs/Horde/Horde.Server/Agents/Sessions/ISession.cs
@@ -3,6 +3,7 @@
 using System;
 using System.Collections.Generic;
 using Horde.Server.Utilities;
+using EpicGames.Horde.Api;
 
 namespace Horde.Server.Agents.Sessions
 {
diff --git a/Engine/Source/Programs/Horde/Horde.Server/Agents/Sessions/ISessionCollection.cs b/Engine/Source/Programs/Horde/Horde.Server/Agents/Sessions/ISessionCollection.cs
--- a/Engine/Source/Programs/Horde/Horde.Server/Agents/Sessions/ISessionCollection.cs
+++ b/Engine/Source/Programs/Horde/Horde.Server/Agents/Sessions/ISessionCollection.cs
@@ -4,6 +4,7 @@ using System;
 using System.Collections.Generic;
 using System.Threading.Tasks;
 using Horde.Server.Utilities;
+using EpicGames.Horde.Api;
 
 namespace Horde.Server.Agents.Sessions
 {
diff --git a/Engine/Source/Programs/Horde/Horde.Server/Agents/Sessions/SessionCollection.cs b/Engine/Source/Programs/Horde/Horde.Server/Agents/Sessions/SessionCollection.cs
--- a/Engine/Source/Programs/Horde/Horde.Server/Agents/Sessions/SessionCollection.cs
+++ b/Engine/Source/Programs/Horde/Horde.Server/Agents/Sessions/SessionCollection.cs
@@ -4,10 +4,12 @@ using System;
 using System.Collections.Generic;
 using System.Linq;
 using System.Threading.Tasks;
+using Amazon.EC2.Model;
 using Horde.Server.Server;
 using Horde.Server.Utilities;
 using MongoDB.Bson.Serialization.Attributes;
 using MongoDB.Driver;
+using EpicGames.Horde.Api;
 
 namespace Horde.Server.Agents.Sessions
 {
diff --git a/Engine/Source/Programs/Horde/Horde.Server/Agents/Telemetry/TelemetryService.cs b/Engine/Source/Programs/Horde/Horde.Server/Agents/Telemetry/TelemetryService.cs
--- a/Engine/Source/Programs/Horde/Horde.Server/Agents/Telemetry/TelemetryService.cs
+++ b/Engine/Source/Programs/Horde/Horde.Server/Agents/Telemetry/TelemetryService.cs
@@ -12,6 +12,7 @@ using Horde.Server.Utilities;
 using HordeCommon;
 using Microsoft.Extensions.Hosting;
 using Microsoft.Extensions.Logging;
+using EpicGames.Horde.Api;
 
 namespace Horde.Server.Agents.Telemetry
 {
diff --git a/Engine/Source/Programs/Horde/Horde.Server/Compute/ComputeController.cs b/Engine/Source/Programs/Horde/Horde.Server/Compute/ComputeController.cs
--- a/Engine/Source/Programs/Horde/Horde.Server/Compute/ComputeController.cs
+++ b/Engine/Source/Programs/Horde/Horde.Server/Compute/ComputeController.cs
@@ -74,6 +74,7 @@ namespace Horde.Server.Compute
 			response.Port = computeResource.Port;
 			response.Nonce = StringUtils.FormatHexString(computeResource.Task.Nonce.Span);
 			response.Key = StringUtils.FormatHexString(computeResource.Task.Key.Span);
+			response.AgentId = computeResource.AgentId;
 			response.LeaseId = computeResource.LeaseId;
 			response.Properties = computeResource.Properties;
 
diff --git a/Engine/Source/Programs/Horde/Horde.Server/Compute/ComputeService.cs b/Engine/Source/Programs/Horde/Horde.Server/Compute/ComputeService.cs
--- a/Engine/Source/Programs/Horde/Horde.Server/Compute/ComputeService.cs
+++ b/Engine/Source/Programs/Horde/Horde.Server/Compute/ComputeService.cs
@@ -14,7 +14,6 @@ using EpicGames.Horde.Compute.Transports;
 using Google.Protobuf;
 using Google.Protobuf.WellKnownTypes;
 using Horde.Server.Agents;
-using Horde.Server.Agents.Leases;
 using Horde.Server.Jobs;
 using Horde.Server.Logs;
 using Horde.Server.Utilities;
@@ -91,7 +90,7 @@ namespace Horde.Server.Compute
 				return null;
 			}
 
-			return new ComputeResource(ip, port, computeTask, agent.Properties, leaseId);
+			return new ComputeResource(ip, port, computeTask, agent.Properties, agent.Id, leaseId);
 		}
 
 		static ComputeTask CreateComputeTask(Dictionary<string, int> assignedResources, LogId? logId, LeaseId? parentLeaseId)
diff --git a/Engine/Source/Programs/Horde/Horde.Server/Compute/ComputeTaskSource.cs b/Engine/Source/Programs/Horde/Horde.Server/Compute/ComputeTaskSource.cs
--- a/Engine/Source/Programs/Horde/Horde.Server/Compute/ComputeTaskSource.cs
+++ b/Engine/Source/Programs/Horde/Horde.Server/Compute/ComputeTaskSource.cs
@@ -4,21 +4,13 @@ using System;
 using System.Collections.Generic;
 using System.Linq;
 using System.Net;
-using System.Security.Cryptography;
 using System.Threading;
 using System.Threading.Tasks;
-using EpicGames.Core;
 using EpicGames.Horde.Api;
 using EpicGames.Horde.Compute;
-using EpicGames.Horde.Compute.Clients;
-using EpicGames.Horde.Compute.Transports;
-using Google.Protobuf;
-using Google.Protobuf.WellKnownTypes;
 using Horde.Server.Agents;
-using Horde.Server.Agents.Leases;
 using Horde.Server.Server;
 using Horde.Server.Tasks;
-using HordeCommon;
 using HordeCommon.Rpc.Tasks;
 using Microsoft.Extensions.Logging;
 using Microsoft.Extensions.Options;
@@ -50,6 +42,11 @@ namespace Horde.Server.Compute
 		/// </summary>
 		public IReadOnlyList<string> Properties { get; }
 
+		/// <summary>
+		/// Agent id on the remote machine
+		/// </summary>
+		public AgentId AgentId { get; }
+
 		/// <summary>
 		/// Lease id on the remote machine
 		/// </summary>
@@ -58,12 +55,13 @@ namespace Horde.Server.Compute
 		/// <summary>
 		/// Constructor
 		/// </summary>
-		public ComputeResource(IPAddress ip, int port, ComputeTask task, IReadOnlyList<string> properties, LeaseId leaseId)
+		public ComputeResource(IPAddress ip, int port, ComputeTask task, IReadOnlyList<string> properties, AgentId agentId, LeaseId leaseId)
 		{
 			Ip = ip;
 			Port = port;
 			Task = task;
 			Properties = properties;
+			AgentId = agentId;
 			LeaseId = leaseId;
 		}
 	}
diff --git a/Engine/Source/Programs/Horde/Horde.Server/Jobs/IJobStepRef.cs b/Engine/Source/Programs/Horde/Horde.Server/Jobs/IJobStepRef.cs
--- a/Engine/Source/Programs/Horde/Horde.Server/Jobs/IJobStepRef.cs
+++ b/Engine/Source/Programs/Horde/Horde.Server/Jobs/IJobStepRef.cs
@@ -11,6 +11,7 @@ using Horde.Server.Streams;
 using HordeCommon;
 using MongoDB.Bson.Serialization;
 using MongoDB.Bson.Serialization.Attributes;
+using EpicGames.Horde.Api;
 
 namespace Horde.Server.Jobs
 {
diff --git a/Engine/Source/Programs/Horde/Horde.Server/Jobs/IJobStepRefCollection.cs b/Engine/Source/Programs/Horde/Horde.Server/Jobs/IJobStepRefCollection.cs
--- a/Engine/Source/Programs/Horde/Horde.Server/Jobs/IJobStepRefCollection.cs
+++ b/Engine/Source/Programs/Horde/Horde.Server/Jobs/IJobStepRefCollection.cs
@@ -12,6 +12,7 @@ using Horde.Server.Logs;
 using Horde.Server.Streams;
 using HordeCommon;
 using Microsoft.Extensions.Logging;
+using EpicGames.Horde.Api;
 
 namespace Horde.Server.Jobs
 {
diff --git a/Engine/Source/Programs/Horde/Horde.Server/Jobs/JobStepRefCollection.cs b/Engine/Source/Programs/Horde/Horde.Server/Jobs/JobStepRefCollection.cs
--- a/Engine/Source/Programs/Horde/Horde.Server/Jobs/JobStepRefCollection.cs
+++ b/Engine/Source/Programs/Horde/Horde.Server/Jobs/JobStepRefCollection.cs
@@ -16,6 +16,7 @@ using Horde.Server.Utilities;
 using HordeCommon;
 using MongoDB.Bson.Serialization.Attributes;
 using MongoDB.Driver;
+using EpicGames.Horde.Api;
 
 namespace Horde.Server.Jobs
 {
diff --git a/Engine/Source/Programs/Horde/Horde.Server/Server/ConsistencyService.cs b/Engine/Source/Programs/Horde/Horde.Server/Server/ConsistencyService.cs
--- a/Engine/Source/Programs/Horde/Horde.Server/Server/ConsistencyService.cs
+++ b/Engine/Source/Programs/Horde/Horde.Server/Server/ConsistencyService.cs
@@ -12,6 +12,7 @@ using HordeCommon;
 using Microsoft.Extensions.Hosting;
 using Microsoft.Extensions.Logging;
 using Horde.Server.Agents;
+using EpicGames.Horde.Api;
 
 namespace Horde.Server.Server
 {
diff --git a/Engine/Source/Programs/Horde/Horde.Server/Server/RpcService.cs b/Engine/Source/Programs/Horde/Horde.Server/Server/RpcService.cs
--- a/Engine/Source/Programs/Horde/Horde.Server/Server/RpcService.cs
+++ b/Engine/Source/Programs/Horde/Horde.Server/Server/RpcService.cs
@@ -11,6 +11,7 @@ using System.Threading;
 using System.Threading.Tasks;
 using Amazon.EC2.Model;
 using EpicGames.Core;
+using EpicGames.Horde.Api;
 using Google.Protobuf;
 using Google.Protobuf.WellKnownTypes;
 using Grpc.Core;
diff --git a/Engine/Source/Programs/Horde/Horde.Server/Startup.cs b/Engine/Source/Programs/Horde/Horde.Server/Startup.cs
--- a/Engine/Source/Programs/Horde/Horde.Server/Startup.cs
+++ b/Engine/Source/Programs/Horde/Horde.Server/Startup.cs
@@ -107,6 +107,8 @@ using EpicGames.Horde.Storage.Bundles;
 using Horde.Server.Ddc;
 using System.Net.Mime;
 using Microsoft.Extensions.Logging.Abstractions;
+using EpicGames.Redis;
+using StackExchange.Redis;
 
 namespace Horde.Server
 {
@@ -342,7 +344,19 @@ namespace Horde.Server
 
 			settings.Telemetry = telemetryConfigs;
 		}
-		
+
+		/// <summary>
+		/// Converter to/from Redis values
+		/// </summary>
+		sealed class AgentIdRedisConverter : IRedisConverter<AgentId>
+		{
+			/// <inheritdoc/>
+			public AgentId FromRedisValue(RedisValue value) => new AgentId((string)value!);
+
+			/// <inheritdoc/>
+			public RedisValue ToRedisValue(AgentId value) => value.ToString();
+		}
+
 		// This method gets called *multiple times* by the runtime. Use this method to add services to the container.
 		public void ConfigureServices(IServiceCollection services)
 		{
@@ -372,6 +386,8 @@ namespace Horde.Server
 				ThreadPool.SetMinThreads(min, min);
 			}
 
+			RedisSerializer.RegisterConverter<AgentId, AgentIdRedisConverter>();
+
 #pragma warning disable CA2000 // Dispose objects before losing scope
 			RedisService redisService = new RedisService(settings);
 #pragma warning restore CA2000 // Dispose objects before losing scope
@@ -994,6 +1010,30 @@ namespace Horde.Server
 			}
 		}
 
+		public sealed class AgentIdBsonSerializer : SerializerBase<AgentId>
+		{
+			/// <inheritdoc/>
+			public override AgentId Deserialize(BsonDeserializationContext context, BsonDeserializationArgs args)
+			{
+				string argument;
+				if (context.Reader.CurrentBsonType == MongoDB.Bson.BsonType.ObjectId)
+				{
+					argument = context.Reader.ReadObjectId().ToString();
+				}
+				else
+				{
+					argument = context.Reader.ReadString();
+				}
+				return new AgentId(argument);
+			}
+
+			/// <inheritdoc/>
+			public override void Serialize(BsonSerializationContext context, BsonSerializationArgs args, AgentId value)
+			{
+				context.Writer.WriteString(value.ToString());
+			}
+		}
+
 		public sealed class AclActionBsonSerializer : SerializerBase<AclAction>
 		{
 			/// <inheritdoc/>
@@ -1029,6 +1069,7 @@ namespace Horde.Server
 				BsonSerializer.RegisterSerializer(new RefNameBsonSerializer());
 				BsonSerializer.RegisterSerializer(new IoHashBsonSerializer());
 				BsonSerializer.RegisterSerializer(new NamespaceIdBsonSerializer());
+				BsonSerializer.RegisterSerializer(new AgentIdBsonSerializer());
 				BsonSerializer.RegisterSerializer(new AclActionBsonSerializer());
 				BsonSerializer.RegisterSerializer(new AclScopeNameBsonSerializer());
 				BsonSerializer.RegisterSerializer(new ConditionSerializer());
diff --git a/Engine/Source/Programs/Horde/Horde.Server/Utilities/HordeClaims.cs b/Engine/Source/Programs/Horde/Horde.Server/Utilities/HordeClaims.cs
--- a/Engine/Source/Programs/Horde/Horde.Server/Utilities/HordeClaims.cs
+++ b/Engine/Source/Programs/Horde/Horde.Server/Utilities/HordeClaims.cs
@@ -3,6 +3,7 @@
 using Horde.Server.Acls;
 using Horde.Server.Agents;
 using Horde.Server.Agents.Sessions;
+using EpicGames.Horde.Api;
 
 namespace Horde.Server.Utilities
 {
diff --git a/Engine/Source/Programs/Shared/EpicGames.Horde/Api/AgentId.cs b/Engine/Source/Programs/Shared/EpicGames.Horde/Api/AgentId.cs
new file mode 100644
--- /dev/null
+++ b/Engine/Source/Programs/Shared/EpicGames.Horde/Api/AgentId.cs
@@ -0,0 +1,171 @@
+// Copyright Epic Games, Inc. All Rights Reserved.
+
+using System;
+using System.ComponentModel;
+using System.Globalization;
+using System.Text.Json;
+using EpicGames.Core;
+
+namespace EpicGames.Horde.Api
+{
+	/// <summary>
+	/// Normalized hostname of an agent
+	/// </summary>
+	[TypeConverter(typeof(AgentIdTypeConverter))]
+	[LogValueFormatter(typeof(AgentIdLogFormatter))]
+	public struct AgentId : IEquatable<AgentId>
+	{
+		/// <summary>
+		/// The text representing this id
+		/// </summary>
+		readonly string _name;
+
+		/// <summary>
+		/// Constructor
+		/// </summary>
+		/// <param name="name">Hostname of the agent</param>
+		public AgentId(string name)
+		{
+			if (name.Length == 0)
+			{
+				throw new ArgumentException("Agent name may not be empty");
+			}
+
+			char[] result = new char[name.Length];
+			for (int idx = 0; idx < name.Length; idx++)
+			{
+				char character = name[idx];
+				if ((character >= 'A' && character <= 'Z') || (character >= '0' && character <= '9') || character == '_' || character == '-')
+				{
+					result[idx] = character;
+				}
+				else if (character >= 'a' && character <= 'z')
+				{
+					result[idx] = (char)('A' + (character - 'a'));
+				}
+				else if (character == '.')
+				{
+					if (idx == 0)
+					{
+						throw new ArgumentException("Agent name may not begin with the '.' character");
+					}
+					else if (idx == name.Length - 1)
+					{
+						throw new ArgumentException("Agent name may not end with the '.' character");
+					}
+					else if (name[idx - 1] == '.')
+					{
+						throw new ArgumentException("Agent name may not contain two consecutive '.' characters");
+					}
+					else
+					{
+						result[idx] = character;
+					}
+				}
+				else
+				{
+					throw new ArgumentException($"Character '{character}' in agent '{name}' is invalid");
+				}
+			}
+			_name = new string(result);
+		}
+
+		/// <inheritdoc/>
+		public override bool Equals(object? obj)
+		{
+			return obj is AgentId id && Equals(id);
+		}
+
+		/// <inheritdoc/>
+		public override int GetHashCode()
+		{
+			return _name.GetHashCode(StringComparison.Ordinal);
+		}
+
+		/// <inheritdoc/>
+		public bool Equals(AgentId other)
+		{
+			return _name.Equals(other._name, StringComparison.Ordinal);
+		}
+
+		/// <inheritdoc/>
+		public override string ToString()
+		{
+			return _name;
+		}
+
+		/// <summary>
+		/// Compares two string ids for equality
+		/// </summary>
+		/// <param name="left">The first string id</param>
+		/// <param name="right">Second string id</param>
+		/// <returns>True if the two string ids are equal</returns>
+		public static bool operator ==(AgentId left, AgentId right)
+		{
+			return left.Equals(right);
+		}
+
+		/// <summary>
+		/// Compares two string ids for inequality
+		/// </summary>
+		/// <param name="left">The first string id</param>
+		/// <param name="right">Second string id</param>
+		/// <returns>True if the two string ids are not equal</returns>
+		public static bool operator !=(AgentId left, AgentId right)
+		{
+			return !left.Equals(right);
+		}
+	}
+
+	/// <summary>
+	/// Type converter from strings to PropertyFilter objects
+	/// </summary>
+	sealed class AgentIdTypeConverter : TypeConverter
+	{
+		/// <inheritdoc/>
+		public override bool CanConvertFrom(ITypeDescriptorContext? context, Type sourceType)
+		{
+			return sourceType == typeof(string);
+		}
+
+		/// <inheritdoc/>
+		public override object ConvertFrom(ITypeDescriptorContext? context, CultureInfo? culture, object value)
+		{
+			return new AgentId((string)value);
+		}
+
+		/// <inheritdoc/>
+		public override bool CanConvertTo(ITypeDescriptorContext? context, Type? destinationType)
+		{
+			return destinationType == typeof(string);
+		}
+
+		/// <inheritdoc/>
+		public override object? ConvertTo(ITypeDescriptorContext? context, CultureInfo? culture, object? value, Type destinationType)
+		{
+			if (destinationType == typeof(string))
+			{
+				return value?.ToString();
+			}
+			else
+			{
+				return null;
+			}
+		}
+	}
+
+	/// <summary>
+	/// Formats an AgentId as a typed log value
+	/// </summary>
+	class AgentIdLogFormatter : ILogValueFormatter
+	{
+		/// <inheritdoc/>
+		public void Format(object value, Utf8JsonWriter writer)
+		{
+			writer.WriteStartObject();
+			writer.WriteString(LogEventPropertyName.Type, "AgentId");
+			writer.WriteString(LogEventPropertyName.Text, ((AgentId)value).ToString());
+			writer.WriteEndObject();
+		}
+	}
+}
diff --git a/Engine/Source/Programs/Shared/EpicGames.Horde/Api/ComputeMessages.cs b/Engine/Source/Programs/Shared/EpicGames.Horde/Api/ComputeMessages.cs
--- a/Engine/Source/Programs/Shared/EpicGames.Horde/Api/ComputeMessages.cs
+++ b/Engine/Source/Programs/Shared/EpicGames.Horde/Api/ComputeMessages.cs
@@ -44,6 +44,11 @@ namespace EpicGames.Horde.Api
 		/// </summary>
 		public string Key { get; set; } = String.Empty;
 
+		/// <summary>
+		/// Identifier for the remote machine
+		/// </summary>
+		public AgentId AgentId { get; set; }
+
 		/// <summary>
 		/// Identifier for the new lease on the remote machine
 		/// </summary>
diff --git a/Engine/Source/Programs/Shared/EpicGames.Horde/Api/LeaseId.cs b/Engine/Source/Programs/Shared/EpicGames.Horde/Api/LeaseId.cs
--- a/Engine/Source/Programs/Shared/EpicGames.Horde/Api/LeaseId.cs
+++ b/Engine/Source/Programs/Shared/EpicGames.Horde/Api/LeaseId.cs
@@ -13,7 +13,7 @@ namespace EpicGames.Horde.Api
 	[JsonSchemaString]
 	[TypeConverter(typeof(BinaryIdTypeConverter<LeaseId, LeaseIdConverter>))]
 	[BinaryIdConverter(typeof(LeaseIdConverter))]
-	[LogValueFormatterAttribute(typeof(LeaseIdLogFormatter))]
+	[LogValueFormatter(typeof(LeaseIdLogFormatter))]
 	public record struct LeaseId(BinaryId Id)
 	{
 		/// <inheritdoc cref="BinaryId.TryParse(System.String, out BinaryId)"/>
diff --git a/Engine/Source/Programs/Shared/EpicGames.Horde/Compute/Clients/ServerComputeClient.cs b/Engine/Source/Programs/Shared/EpicGames.Horde/Compute/Clients/ServerComputeClient.cs
--- a/Engine/Source/Programs/Shared/EpicGames.Horde/Compute/Clients/ServerComputeClient.cs
+++ b/Engine/Source/Programs/Shared/EpicGames.Horde/Compute/Clients/ServerComputeClient.cs
@@ -161,7 +161,7 @@ namespace EpicGames.Horde.Compute.Clients
 				}
 			}
 
-			_logger.LogDebug("Connecting to {Ip} with nonce {Nonce}...", responseMessage.Ip, responseMessage.Nonce);
+			_logger.LogDebug("Connecting to {AgentId} ({Ip}) with nonce {Nonce}...", responseMessage.AgentId, responseMessage.Ip, responseMessage.Nonce);
 
 			// Connect to the remote machine
 			using Socket socket = new Socket(SocketType.Stream, ProtocolType.Tcp);
@@ -170,7 +170,7 @@ namespace EpicGames.Horde.Compute.Clients
 			// Send the nonce
 			byte[] nonce = StringUtils.ParseHexString(responseMessage.Nonce);
 			await socket.SendMessageAsync(nonce, SocketFlags.None, cancellationToken);
-			_logger.LogInformation("Connected to {Ip} under lease {LeaseId}", responseMessage.Ip, responseMessage.LeaseId);
+			_logger.LogInformation("Connected to {AgentId} ({Ip}) under lease {LeaseId}", responseMessage.AgentId, responseMessage.Ip, responseMessage.LeaseId);
 
 			// Pass the rest of the call over to the handler
 			byte[] key = StringUtils.ParseHexString(responseMessage.Key);
