Horde: Better handling for default AgentId values.

diff --git a/Engine/Source/Programs/Shared/EpicGames.Horde/Api/AgentId.cs b/Engine/Source/Programs/Shared/EpicGames.Horde/Api/AgentId.cs
--- a/Engine/Source/Programs/Shared/EpicGames.Horde/Api/AgentId.cs
+++ b/Engine/Source/Programs/Shared/EpicGames.Horde/Api/AgentId.cs
@@ -79,19 +79,19 @@ namespace EpicGames.Horde.Api
 		/// <inheritdoc/>
 		public override int GetHashCode()
 		{
-			return _name.GetHashCode(StringComparison.Ordinal);
+			return (_name ?? String.Empty).GetHashCode(StringComparison.Ordinal);
 		}
 
 		/// <inheritdoc/>
 		public bool Equals(AgentId other)
 		{
-			return _name.Equals(other._name, StringComparison.Ordinal);
+			return String.Equals(_name, other._name, StringComparison.Ordinal);
 		}
 
 		/// <inheritdoc/>
 		public override string ToString()
 		{
-			return _name;
+			return _name ?? String.Empty;
 		}
 
 		/// <summary>
