Horde: Add a JSON converter for AgentId types.

diff --git a/Engine/Source/Programs/Shared/EpicGames.Horde/Api/AgentId.cs b/Engine/Source/Programs/Shared/EpicGames.Horde/Api/AgentId.cs
--- a/Engine/Source/Programs/Shared/EpicGames.Horde/Api/AgentId.cs
+++ b/Engine/Source/Programs/Shared/EpicGames.Horde/Api/AgentId.cs
@@ -4,6 +4,7 @@ using System;
 using System.ComponentModel;
 using System.Globalization;
 using System.Text.Json;
+using System.Text.Json.Serialization;
 using EpicGames.Core;
 
 namespace EpicGames.Horde.Api
@@ -12,6 +13,7 @@ namespace EpicGames.Horde.Api
 	/// Normalized hostname of an agent
 	/// </summary>
 	[TypeConverter(typeof(AgentIdTypeConverter))]
+	[JsonConverter(typeof(AgentIdJsonConverter))]
 	[LogValueFormatter(typeof(AgentIdLogFormatter))]
 	public struct AgentId : IEquatable<AgentId>
 	{
@@ -154,6 +156,18 @@ namespace EpicGames.Horde.Api
 		}
 	}
 
+	/// <summary>
+	/// Class which serializes AgentId objects to JSON
+	/// </summary>
+	public sealed class AgentIdJsonConverter : JsonConverter<AgentId>
+	{
+		/// <inheritdoc/>
+		public override AgentId Read(ref Utf8JsonReader reader, Type typeToConvert, JsonSerializerOptions options) => new AgentId(reader.GetString() ?? String.Empty);
+
+		/// <inheritdoc/>
+		public override void Write(Utf8JsonWriter writer, AgentId value, JsonSerializerOptions options) => writer.WriteStringValue(value.ToString());
+	}
+
 	/// <summary>
 	/// Formats an AgentId as a typed log value
 	/// </summary>
