OidcToken - Bug fixes

* Create output directory for temp file if its missing
* Corrected logic for tokens that do not expire (they are set to expire at 0 which we didn't treat correctly)
* Make sure to save the new refresh token we get after using a old one, resolves issues for horde.exe that had to prompt users for logins more often then it should need to.

diff --git a/Engine/Source/Programs/OidcToken/TokenService.cs b/Engine/Source/Programs/OidcToken/TokenService.cs
--- a/Engine/Source/Programs/OidcToken/TokenService.cs
+++ b/Engine/Source/Programs/OidcToken/TokenService.cs
@@ -173,6 +173,12 @@ namespace OidcToken
 		private async Task OutputToken(string token, DateTimeOffset expiresAt)
 		{
 			FileInfo fi = new(Settings.CurrentValue.OutFile);
+
+			if (fi.DirectoryName != null)
+			{
+				Directory.CreateDirectory(fi.DirectoryName);
+			}
+
 			Logger.LogInformation("Token output to \"{OutFile}\"", fi.FullName);
 
 			await using FileStream fs = fi.Open(FileMode.Create, FileAccess.Write);
diff --git a/Engine/Source/Programs/Shared/EpicGames.OIDC/OidcTokenManager.cs b/Engine/Source/Programs/Shared/EpicGames.OIDC/OidcTokenManager.cs
--- a/Engine/Source/Programs/Shared/EpicGames.OIDC/OidcTokenManager.cs
+++ b/Engine/Source/Programs/Shared/EpicGames.OIDC/OidcTokenManager.cs
@@ -42,14 +42,14 @@ namespace EpicGames.OIDC
 		{
 			lock (_lockObject)
 			{
-			OidcTokenClient? client;
-			if (!_tokenClients.TryGetValue(name, out client))
-			{
-				client = new OidcTokenClient(name, providerInfo, _tokenStore);
-				_tokenClients.Add(name, client);
+				OidcTokenClient? client;
+				if (!_tokenClients.TryGetValue(name, out client))
+				{
+					client = new OidcTokenClient(name, providerInfo, _tokenStore);
+					_tokenClients.Add(name, client);
+				}
+				return client;
 			}
-			return client;
-		}
 		}
 
 		public OidcTokenManager(IServiceProvider provider, IOptionsMonitor<OidcTokenOptions> settings, ITokenStore tokenStore, List<string>? allowedProviders = null)
@@ -481,6 +481,10 @@ namespace EpicGames.OIDC
 			_accessToken = refreshTokenResult.AccessToken;
 			_tokenExpiry = refreshTokenResult.AccessTokenExpiration;
 
+			// refresh tokens are always one time use only so we need to store this new refresh token we got so it can be used the next time
+			_tokenStore.AddRefreshToken(_name, _refreshToken);
+			_tokenStore.Save();
+
 			return new OidcTokenInfo
 			{
 				IdentityToken = refreshTokenResult.IdentityToken,
@@ -532,7 +536,9 @@ namespace EpicGames.OIDC
 
 			// if the token is valid for another few minutes we can use it
 			// we avoid using a token that is about to expire to make sure we can finish the call we expect to do with it before it expires
-			if (!String.IsNullOrEmpty(_accessToken) && _tokenExpiry.AddMinutes(2) > DateTime.Now)
+			// if a token has infinite lifetime its expiry is set to 0 and is thus always valid
+			bool tokenValid = _tokenExpiry.AddMinutes(2) > DateTime.Now || _tokenExpiry == DateTimeOffset.MinValue;
+			if (!String.IsNullOrEmpty(_accessToken) && tokenValid)
 			{
 				return new OidcTokenInfo
 				{
