Horde: Allow overriding the minimum output level for server-side logging. Increase verbosity of remote execution logging to "trace" verbosity.

diff --git a/Engine/Source/Programs/Horde/Horde.Agent/Leases/Handlers/ComputeHandler.cs b/Engine/Source/Programs/Horde/Horde.Agent/Leases/Handlers/ComputeHandler.cs
--- a/Engine/Source/Programs/Horde/Horde.Agent/Leases/Handlers/ComputeHandler.cs
+++ b/Engine/Source/Programs/Horde/Horde.Agent/Leases/Handlers/ComputeHandler.cs
@@ -81,7 +81,7 @@ namespace Horde.Agent.Leases.Handlers
 		/// <inheritdoc/>
 		public override async Task<LeaseResult> ExecuteAsync(ISession session, string leaseId, ComputeTask computeTask, CancellationToken cancellationToken)
 		{
-			await using IServerLogger? serverLogger = (computeTask.LogId != null)? _serverLoggerFactory.CreateLogger(session, computeTask.LogId, null, true) : null;
+			await using IServerLogger? serverLogger = (computeTask.LogId != null)? _serverLoggerFactory.CreateLogger(session, computeTask.LogId, null, true, LogLevel.Trace) : null;
 			ILogger logger = serverLogger ?? _logger;
 
 			if (!String.IsNullOrEmpty(computeTask.ParentLeaseId))
diff --git a/Engine/Source/Programs/Horde/Horde.Agent/Services/ServerLoggerFactory.cs b/Engine/Source/Programs/Horde/Horde.Agent/Services/ServerLoggerFactory.cs
--- a/Engine/Source/Programs/Horde/Horde.Agent/Services/ServerLoggerFactory.cs
+++ b/Engine/Source/Programs/Horde/Horde.Agent/Services/ServerLoggerFactory.cs
@@ -21,8 +21,9 @@ namespace Horde.Agent.Services
 		/// <param name="stepId">The job step id</param>
 		/// <param name="warnings">Whether to suppress warnings</param>
 		/// <param name="useNewLogger">Whether to enable the new logger backend</param>
+		/// <param name="outputLevel">Minimum output level for messages</param>
 		/// <returns>New logger instance</returns>
-		IServerLogger CreateLogger(ISession session, string logId, string? jobId, string? batchId, string? stepId, bool? warnings, bool? useNewLogger);
+		IServerLogger CreateLogger(ISession session, string logId, string? jobId, string? batchId, string? stepId, bool? warnings, bool? useNewLogger, LogLevel outputLevel = LogLevel.Information);
 	}
 
 	/// <summary>
@@ -38,10 +39,11 @@ namespace Horde.Agent.Services
 		/// <param name="logId">The log identifier</param>
 		/// <param name="warnings">Whether to suppress warnings</param>
 		/// <param name="useNewLogger">Whether to enable the new logger backend</param>
+		/// <param name="outputLevel">Minimum output level for messages</param>
 		/// <returns>New logger instance</returns>
-		public static IServerLogger CreateLogger(this IServerLoggerFactory service, ISession session, string logId, bool? warnings, bool? useNewLogger)
+		public static IServerLogger CreateLogger(this IServerLoggerFactory service, ISession session, string logId, bool? warnings, bool? useNewLogger, LogLevel outputLevel = LogLevel.Information)
 		{
-			return service.CreateLogger(session, logId, null, null, null, warnings, useNewLogger);
+			return service.CreateLogger(session, logId, null, null, null, warnings, useNewLogger, outputLevel);
 		}
 	}
 
@@ -63,7 +65,7 @@ namespace Horde.Agent.Services
 		}
 
 		/// <inheritdoc/>
-		public IServerLogger CreateLogger(ISession session, string logId, string? jobId, string? batchId, string? stepId, bool? warnings, bool? useNewLogger)
+		public IServerLogger CreateLogger(ISession session, string logId, string? jobId, string? batchId, string? stepId, bool? warnings, bool? useNewLogger, LogLevel outputLevel)
 		{
 #pragma warning disable CA2000 // Dispose objects before losing scope
 			IJsonRpcLogSink sink = new JsonRpcLogSink(session.RpcConnection, jobId, batchId, stepId, _logger);
@@ -72,7 +74,7 @@ namespace Horde.Agent.Services
 				IStorageClient storageClient = _storageClientFactory.CreateStorageClient(session, $"/api/v1/logs/{logId}/");
 				sink = new JsonRpcAndStorageLogSink(session.RpcConnection, logId, sink, storageClient, _logger);
 			}
-			return new JsonRpcLogger(sink, logId, warnings, _logger);
+			return new JsonRpcLogger(sink, logId, warnings, outputLevel, _logger);
 #pragma warning restore CA2000 // Dispose objects before losing scope
 		}
 	}
diff --git a/Engine/Source/Programs/Horde/Horde.Agent/Utility/JsonRpcLogger.cs b/Engine/Source/Programs/Horde/Horde.Agent/Utility/JsonRpcLogger.cs
--- a/Engine/Source/Programs/Horde/Horde.Agent/Utility/JsonRpcLogger.cs
+++ b/Engine/Source/Programs/Horde/Horde.Agent/Utility/JsonRpcLogger.cs
@@ -47,6 +47,7 @@ namespace Horde.Agent.Utility
 		internal readonly IJsonRpcLogSink _sink;
 		internal readonly string _logId;
 		internal readonly bool _warnings;
+		internal readonly LogLevel _outputLevel;
 		internal readonly ILogger _inner;
 		readonly Channel<JsonLogEvent> _dataChannel;
 		Task? _dataWriter;
@@ -66,12 +67,14 @@ namespace Horde.Agent.Utility
 		/// <param name="sink">Sink for log events</param>
 		/// <param name="logId">The log id to write to</param>
 		/// <param name="warnings">Whether to include warnings in the output</param>
+		/// <param name="outputLevel">Minimum level for output</param>
 		/// <param name="inner">Additional logger to write to</param>
-		public JsonRpcLogger(IJsonRpcLogSink sink, string logId, bool? warnings, ILogger inner)
+		public JsonRpcLogger(IJsonRpcLogSink sink, string logId, bool? warnings, LogLevel outputLevel, ILogger inner)
 		{
 			_sink = sink;
 			_logId = logId;
 			_warnings = warnings ?? true;
+			_outputLevel = outputLevel;
 			_inner = inner;
 			_dataChannel = Channel.CreateUnbounded<JsonLogEvent>();
 			_dataWriter = Task.Run(() => RunDataWriterAsync());
@@ -93,7 +96,7 @@ namespace Horde.Agent.Utility
 		}
 
 		/// <inheritdoc/>
-		public bool IsEnabled(LogLevel logLevel) => _inner.IsEnabled(logLevel);
+		public bool IsEnabled(LogLevel logLevel) => logLevel >= _outputLevel;
 
 		/// <inheritdoc/>
 		public IDisposable BeginScope<TState>(TState state) => _inner.BeginScope(state);
diff --git a/Engine/Source/Programs/Shared/EpicGames.Horde/Compute/AgentMessageHandler.cs b/Engine/Source/Programs/Shared/EpicGames.Horde/Compute/AgentMessageHandler.cs
--- a/Engine/Source/Programs/Shared/EpicGames.Horde/Compute/AgentMessageHandler.cs
+++ b/Engine/Source/Programs/Shared/EpicGames.Horde/Compute/AgentMessageHandler.cs
@@ -68,7 +68,7 @@ namespace EpicGames.Horde.Compute
 				for (; ; )
 				{
 					using AgentMessage message = await channel.ReceiveAsync(cancellationToken);
-					_logger.LogTrace("Compute Channel {ChannelId}: {MessageType}", channelId, message.Type);
+					_logger.LogDebug("Compute Channel {ChannelId}: {MessageType}", channelId, message.Type);
 
 					switch (message.Type)
 					{
diff --git a/Engine/Source/Programs/Shared/EpicGames.Horde/Compute/ComputeSocket.cs b/Engine/Source/Programs/Shared/EpicGames.Horde/Compute/ComputeSocket.cs
--- a/Engine/Source/Programs/Shared/EpicGames.Horde/Compute/ComputeSocket.cs
+++ b/Engine/Source/Programs/Shared/EpicGames.Horde/Compute/ComputeSocket.cs
@@ -287,7 +287,7 @@ namespace EpicGames.Horde.Compute
 
 		async Task RunRecvTaskAsync(ComputeTransport transport, CancellationToken cancellationToken)
 		{
-			_logger.LogTrace("[{Tag}] Started socket reader", Tag);
+			_logger.LogDebug("[{Tag}] Started socket reader", Tag);
 
 			List<Task> detachTasks = new List<Task>();
 
@@ -304,7 +304,7 @@ namespace EpicGames.Horde.Compute
 					// Read the next packet header
 					if (!await transport.RecvOptionalAsync(header, cancellationToken))
 					{
-						_logger.LogTrace("[{Tag}] End of socket", Tag);
+						_logger.LogDebug("[{Tag}] End of socket", Tag);
 						break;
 					}
 
@@ -323,7 +323,7 @@ namespace EpicGames.Horde.Compute
 					}
 					else
 					{
-						_logger.LogWarning("[{Tag}] Unrecognized control message: {Message}", Tag, size);
+						_logger.LogDebug("[{Tag}] Unrecognized control message: {Message}", Tag, size);
 					}
 				}
 			}
@@ -344,11 +344,11 @@ namespace EpicGames.Horde.Compute
 			// Wait for all the detach tasks to finish
 			if (detachTasks.Count > 0)
 			{
-				_logger.LogTrace("[{Tag}] Waiting for detach tasks to complete...", Tag);
+				_logger.LogDebug("[{Tag}] Waiting for detach tasks to complete...", Tag);
 				await Task.WhenAll(detachTasks).WaitAsync(cancellationToken);
 			}
 
-			_logger.LogTrace("[{Tag}] Closing reader", Tag);
+			_logger.LogDebug("[{Tag}] Closing reader", Tag);
 		}
 
 		async Task ReadPacketAsync(ComputeTransport transport, int id, int size, CancellationToken cancellationToken)
@@ -399,7 +399,7 @@ namespace EpicGames.Horde.Compute
 						Memory<byte> memory = writer.GetWriteBuffer();
 						while (memory.Length < size)
 						{
-							_logger.LogTrace("[{Tag}] No space in buffer {Id}, flushing", Tag, id);
+							_logger.LogDebug("[{Tag}] No space in buffer {Id}, flushing", Tag, id);
 							await writer.WaitToWriteAsync(size, cancellationToken);
 							memory = writer.GetWriteBuffer();
 						}
@@ -474,7 +474,7 @@ namespace EpicGames.Horde.Compute
 		/// <inheritdoc/>
 		public override void AttachRecvBuffer(int channelId, ComputeBuffer recvBuffer)
 		{
-			_logger.LogTrace("[{Tag}] Attaching recv buffer {Id}", Tag, channelId);
+			_logger.LogDebug("[{Tag}] Attaching recv buffer {Id}", Tag, channelId);
 			lock (_lockObject)
 			{
 				if (_complete)
@@ -494,7 +494,7 @@ namespace EpicGames.Horde.Compute
 		[SuppressMessage("Reliability", "CA2000:Dispose objects before losing scope")]
 		async Task DetachRecvBufferAsync(int id, CancellationToken cancellationToken)
 		{
-			_logger.LogTrace("[{Tag}] Detaching recv buffer {Id}", Tag, id);
+			_logger.LogDebug("[{Tag}] Detaching recv buffer {Id}", Tag, id);
 
 			// Get the current receive buffer
 			RecvBuffer? recvBuffer;
@@ -502,7 +502,7 @@ namespace EpicGames.Horde.Compute
 			{
 				if (!_recvBuffers.TryGetValue(id, out recvBuffer))
 				{
-					_logger.LogTrace("[{Tag}] Buffer {Id} has already been detached", Tag, id);
+					_logger.LogDebug("[{Tag}] Buffer {Id} has already been detached", Tag, id);
 					return;
 				}
 				recvBuffer.AddRef(); // Note: adding extra ref here
@@ -563,7 +563,7 @@ namespace EpicGames.Horde.Compute
 		/// <inheritdoc/>
 		public override void AttachSendBuffer(int channelId, ComputeBuffer sendBuffer)
 		{
-			_logger.LogTrace("[{Tag}] Attaching send buffer {Id}", Tag, channelId);
+			_logger.LogDebug("[{Tag}] Attaching send buffer {Id}", Tag, channelId);
 			lock (_lockObject)
 			{
 				if (_sendBuffers.ContainsKey(channelId))
