Horde: Automatically insert the CLIENT/REMOTE prefix to log messages output by the local compute client, and have the agent channel output use the same logger.

diff --git a/Engine/Source/Programs/Horde/Horde.Agent/Commands/Compute/ComputeCommand.cs b/Engine/Source/Programs/Horde/Horde.Agent/Commands/Compute/ComputeCommand.cs
--- a/Engine/Source/Programs/Horde/Horde.Agent/Commands/Compute/ComputeCommand.cs
+++ b/Engine/Source/Programs/Horde/Horde.Agent/Commands/Compute/ComputeCommand.cs
@@ -130,7 +130,7 @@ namespace Horde.Agent.Commands.Compute
 			using Socket tcpSocket = new Socket(SocketType.Stream, ProtocolType.IP);
 			await tcpSocket.ConnectAsync(IPAddress.Loopback, Port);
 
-			await using (RemoteComputeSocket socket = new RemoteComputeSocket(new TcpTransport(tcpSocket), ComputeSocketEndpoint.Remote, logger))
+			await using (RemoteComputeSocket socket = new RemoteComputeSocket(new TcpTransport(tcpSocket), logger))
 			{
 				logger.LogInformation("Running worker...");
 				await RunWorkerAsync(socket, _memoryCache, logger, CancellationToken.None);
diff --git a/Engine/Source/Programs/Horde/Horde.Agent/Commands/Compute/CppComputeCommand.cs b/Engine/Source/Programs/Horde/Horde.Agent/Commands/Compute/CppComputeCommand.cs
--- a/Engine/Source/Programs/Horde/Horde.Agent/Commands/Compute/CppComputeCommand.cs
+++ b/Engine/Source/Programs/Horde/Horde.Agent/Commands/Compute/CppComputeCommand.cs
@@ -59,7 +59,7 @@ namespace Horde.Agent.Commands.Compute
 			BundleNodeLocator sandbox = await CreateSandboxAsync(TaskFile, storage, cancellationToken);
 
 			// Open a socket and upload the sandbox
-			using (AgentMessageChannel channel = lease.Socket.CreateAgentMessageChannel(ControlChannelId, 4 * 1024 * 1024, _logger))
+			using (AgentMessageChannel channel = lease.Socket.CreateAgentMessageChannel(ControlChannelId, 4 * 1024 * 1024))
 			{
 				await channel.UploadFilesAsync("", sandbox, storage, cancellationToken);
 
diff --git a/Engine/Source/Programs/Horde/Horde.Agent/Commands/Compute/DdcComputeCommand.cs b/Engine/Source/Programs/Horde/Horde.Agent/Commands/Compute/DdcComputeCommand.cs
--- a/Engine/Source/Programs/Horde/Horde.Agent/Commands/Compute/DdcComputeCommand.cs
+++ b/Engine/Source/Programs/Horde/Horde.Agent/Commands/Compute/DdcComputeCommand.cs
@@ -31,7 +31,7 @@ namespace Horde.Agent.Commands.Compute
 		protected override async Task<bool> HandleRequestAsync(IComputeLease lease, CancellationToken cancellationToken)
 		{
 			ComputeSocket socket = lease.Socket;
-			using (AgentMessageChannel channel = socket.CreateAgentMessageChannel(0, _logger))
+			using (AgentMessageChannel channel = socket.CreateAgentMessageChannel(0))
 			{
 				_logger.LogInformation("Sending XOR request");
 				await channel.SendXorRequestAsync(new byte[] { 1, 2, 3, 4, 5 }, (byte)123, cancellationToken);
diff --git a/Engine/Source/Programs/Horde/Horde.Agent/Commands/Compute/XorComputeCommand.cs b/Engine/Source/Programs/Horde/Horde.Agent/Commands/Compute/XorComputeCommand.cs
--- a/Engine/Source/Programs/Horde/Horde.Agent/Commands/Compute/XorComputeCommand.cs
+++ b/Engine/Source/Programs/Horde/Horde.Agent/Commands/Compute/XorComputeCommand.cs
@@ -31,7 +31,7 @@ namespace Horde.Agent.Commands.Compute
 		protected override async Task<bool> HandleRequestAsync(IComputeLease lease, CancellationToken cancellationToken)
 		{
 			ComputeSocket socket = lease.Socket;
-			using (AgentMessageChannel mainChannel = socket.CreateAgentMessageChannel(0, _logger))
+			using (AgentMessageChannel mainChannel = socket.CreateAgentMessageChannel(0))
 			{
 				_logger.LogInformation("Forking compute channel...");
 
diff --git a/Engine/Source/Programs/Horde/Horde.Agent/Leases/Handlers/ComputeHandler.cs b/Engine/Source/Programs/Horde/Horde.Agent/Leases/Handlers/ComputeHandler.cs
--- a/Engine/Source/Programs/Horde/Horde.Agent/Leases/Handlers/ComputeHandler.cs
+++ b/Engine/Source/Programs/Horde/Horde.Agent/Leases/Handlers/ComputeHandler.cs
@@ -111,7 +111,7 @@ namespace Horde.Agent.Leases.Handlers
 					await using BackgroundTask timeoutTask = BackgroundTask.StartNew(ctx => TickTimeoutAsync(transport, cts, logger, ctx));
 					try
 					{
-						await using (RemoteComputeSocket socket = new RemoteComputeSocket(transport, ComputeSocketEndpoint.Local, logger))
+						await using (RemoteComputeSocket socket = new RemoteComputeSocket(transport, logger))
 						{
 							DirectoryReference sandboxDir = DirectoryReference.Combine(session.WorkingDir, "Sandbox", leaseId);
 							try
diff --git a/Engine/Source/Programs/Horde/Horde/Commands/Compute/ComputeRun.cs b/Engine/Source/Programs/Horde/Horde/Commands/Compute/ComputeRun.cs
--- a/Engine/Source/Programs/Horde/Horde/Commands/Compute/ComputeRun.cs
+++ b/Engine/Source/Programs/Horde/Horde/Commands/Compute/ComputeRun.cs
@@ -109,7 +109,7 @@ namespace Horde.Commands.Compute
 			BundleNodeLocator sandbox = await CreateSandboxAsync(TaskFile, storage, cancellationToken);
 
 			// Open a socket and upload the sandbox
-			using (AgentMessageChannel channel = lease.Socket.CreateAgentMessageChannel(ControlChannelId, 4 * 1024 * 1024, logger))
+			using (AgentMessageChannel channel = lease.Socket.CreateAgentMessageChannel(ControlChannelId, 4 * 1024 * 1024))
 			{
 				await channel.WaitForAttachAsync(cancellationToken);
 				await channel.UploadFilesAsync("", sandbox, storage, cancellationToken);
diff --git a/Engine/Source/Programs/Horde/Horde/Commands/Compute/ComputeWorker.cs b/Engine/Source/Programs/Horde/Horde/Commands/Compute/ComputeWorker.cs
--- a/Engine/Source/Programs/Horde/Horde/Commands/Compute/ComputeWorker.cs
+++ b/Engine/Source/Programs/Horde/Horde/Commands/Compute/ComputeWorker.cs
@@ -33,7 +33,7 @@ namespace Horde.Commands.Compute
 			using Socket tcpSocket = new Socket(SocketType.Stream, ProtocolType.IP);
 			await tcpSocket.ConnectAsync(IPAddress.Loopback, Port);
 
-			await using (RemoteComputeSocket socket = new RemoteComputeSocket(new TcpTransport(tcpSocket), ComputeSocketEndpoint.Remote, logger))
+			await using (RemoteComputeSocket socket = new RemoteComputeSocket(new TcpTransport(tcpSocket), logger))
 			{
 				logger.LogInformation("Running worker...");
 				await RunWorkerAsync(socket, _memoryCache, logger, CancellationToken.None);
diff --git a/Engine/Source/Programs/Horde/Samples/RemoteClient/RemoteClientApp.cs b/Engine/Source/Programs/Horde/Samples/RemoteClient/RemoteClientApp.cs
--- a/Engine/Source/Programs/Horde/Samples/RemoteClient/RemoteClientApp.cs
+++ b/Engine/Source/Programs/Horde/Samples/RemoteClient/RemoteClientApp.cs
@@ -71,12 +71,12 @@ namespace RemoteClient
 			if (options.UseCppWorker)
 			{
 				FileReference remoteServerFile = FileReference.Combine(ClientSourceDir, "../RemoteWorkerCpp/bin/RemoteServerCpp.exe");
-				await RunRemoteAsync(lease, remoteServerFile.Directory, "RemoteWorkerCpp.exe", new List<string>(), logger);
+				await RunRemoteAsync(lease, remoteServerFile.Directory, "RemoteWorkerCpp.exe", new List<string>());
 			}
 			else
 			{
 				FileReference remoteServerFile = FileReference.Combine(ClientSourceDir, "../RemoteWorker", CurrentAssemblyFile.Directory.MakeRelativeTo(ClientSourceDir), "RemoteWorker.dll");
-				await RunRemoteAsync(lease, remoteServerFile.Directory, @"C:\Program Files\dotnet\dotnet.exe", new List<string> { remoteServerFile.GetFileName() }, logger);
+				await RunRemoteAsync(lease, remoteServerFile.Directory, @"C:\Program Files\dotnet\dotnet.exe", new List<string> { remoteServerFile.GetFileName() });
 			}
 		}
 
@@ -84,15 +84,15 @@ namespace RemoteClient
 		const int BackgroundChannelId = 1;
 		const int ChildProcessChannelId = 100;
 
-		static async Task RunRemoteAsync(IComputeLease lease, DirectoryReference uploadDir, string executable, List<string> arguments, ILogger logger)
+		static async Task RunRemoteAsync(IComputeLease lease, DirectoryReference uploadDir, string executable, List<string> arguments)
 		{
 			// Create a message channel on channel id 0. The Horde Agent always listens on this channel for requests.
-			using (AgentMessageChannel channel = lease.Socket.CreateAgentMessageChannel(PrimaryChannelId, 4 * 1024 * 1024, logger))
+			using (AgentMessageChannel channel = lease.Socket.CreateAgentMessageChannel(PrimaryChannelId, 4 * 1024 * 1024))
 			{
 				await channel.WaitForAttachAsync();
 
 				// Fork another message loop. We'll use this to run an XOR task in the background.
-				using AgentMessageChannel backgroundChannel = lease.Socket.CreateAgentMessageChannel(BackgroundChannelId, 4 * 1024 * 1024, logger);
+				using AgentMessageChannel backgroundChannel = lease.Socket.CreateAgentMessageChannel(BackgroundChannelId, 4 * 1024 * 1024);
 				await using BackgroundTask otherChannelTask = BackgroundTask.StartNew(ctx => RunBackgroundXorAsync(backgroundChannel));
 				await channel.ForkAsync(BackgroundChannelId, 4 * 1024 * 1024, default);
 
@@ -108,14 +108,14 @@ namespace RemoteClient
 
 				// Run the task remotely in the background and echo the output to the console
 				using ComputeChannel childProcessChannel = lease.Socket.CreateChannel(ChildProcessChannelId);
-				await using BackgroundTask tickTask = BackgroundTask.StartNew(ctx => WriteNumbersAsync(childProcessChannel, logger, ctx));
+				await using BackgroundTask tickTask = BackgroundTask.StartNew(ctx => WriteNumbersAsync(childProcessChannel, lease.Socket.Logger, ctx));
 
 				await using (AgentManagedProcess process = await channel.ExecuteAsync(executable, arguments, null, null, ExecuteProcessFlags.None))
 				{
 					string? line;
 					while ((line = await process.ReadLineAsync()) != null)
 					{
-						logger.LogInformation("[REMOTE] {Line}", line);
+						lease.Socket.Logger.LogInformation("{Line}", line);
 					}
 				}
 			}
diff --git a/Engine/Source/Programs/Shared/EpicGames.Horde.Tests/BufferTests.cs b/Engine/Source/Programs/Shared/EpicGames.Horde.Tests/BufferTests.cs
--- a/Engine/Source/Programs/Shared/EpicGames.Horde.Tests/BufferTests.cs
+++ b/Engine/Source/Programs/Shared/EpicGames.Horde.Tests/BufferTests.cs
@@ -102,8 +102,8 @@ namespace EpicGames.Horde.Tests
 			Pipe sourceToTargetPipe = new Pipe();
 			Pipe targetToSourcePipe = new Pipe();
 
-			await using RemoteComputeSocket producerSocket = new RemoteComputeSocket(new PipeTransport(targetToSourcePipe.Reader, sourceToTargetPipe.Writer), ComputeSocketEndpoint.Local, NullLogger.Instance);
-			await using RemoteComputeSocket consumerSocket = new RemoteComputeSocket(new PipeTransport(sourceToTargetPipe.Reader, targetToSourcePipe.Writer), ComputeSocketEndpoint.Remote, NullLogger.Instance);
+			await using RemoteComputeSocket producerSocket = new RemoteComputeSocket(new PipeTransport(targetToSourcePipe.Reader, sourceToTargetPipe.Writer), NullLogger.Instance);
+			await using RemoteComputeSocket consumerSocket = new RemoteComputeSocket(new PipeTransport(sourceToTargetPipe.Reader, targetToSourcePipe.Writer), NullLogger.Instance);
 
 			using ComputeBuffer consumerBuffer = createBuffer(Length);
 			consumerSocket.AttachRecvBuffer(ChannelId, consumerBuffer);
@@ -163,8 +163,8 @@ namespace EpicGames.Horde.Tests
 		{
 			Pipe recvPipe = new Pipe();
 			Pipe sendPipe = new Pipe();
-			await using RemoteComputeSocket localSocket = new RemoteComputeSocket(new PipeTransport(sendPipe.Reader, recvPipe.Writer), ComputeSocketEndpoint.Local, NullLogger.Instance);
-			await using RemoteComputeSocket remoteSocket = new RemoteComputeSocket(new PipeTransport(recvPipe.Reader, sendPipe.Writer), ComputeSocketEndpoint.Remote, NullLogger.Instance);
+			await using RemoteComputeSocket localSocket = new RemoteComputeSocket(new PipeTransport(sendPipe.Reader, recvPipe.Writer), NullLogger.Instance);
+			await using RemoteComputeSocket remoteSocket = new RemoteComputeSocket(new PipeTransport(recvPipe.Reader, sendPipe.Writer), NullLogger.Instance);
 
 			using (PooledBuffer remoteBuffer = new PooledBuffer(1024))
 			{
diff --git a/Engine/Source/Programs/Shared/EpicGames.Horde.Tests/ComputeSocketTests.cs b/Engine/Source/Programs/Shared/EpicGames.Horde.Tests/ComputeSocketTests.cs
--- a/Engine/Source/Programs/Shared/EpicGames.Horde.Tests/ComputeSocketTests.cs
+++ b/Engine/Source/Programs/Shared/EpicGames.Horde.Tests/ComputeSocketTests.cs
@@ -34,6 +34,8 @@ namespace EpicGames.Horde.Tests
 			public Dictionary<int, ComputeBufferWriter> RecvBufferWriters { get; } = new Dictionary<int, ComputeBufferWriter>();
 			public Dictionary<int, ComputeBufferReader> SendBufferReaders { get; } = new Dictionary<int, ComputeBufferReader>();
 
+			public override ILogger Logger => NullLogger.Instance;
+
 			public void Dispose()
 			{
 				foreach (ComputeBufferWriter writer in RecvBufferWriters.Values)
@@ -76,8 +78,8 @@ namespace EpicGames.Horde.Tests
 		{
 			Pipe recvPipe = new Pipe();
 			Pipe sendPipe = new Pipe();
-			await using RemoteComputeSocket localSocket = new RemoteComputeSocket(new PipeTransport(sendPipe.Reader, recvPipe.Writer), ComputeSocketEndpoint.Local, new TestLogger());
-			await using RemoteComputeSocket agentSocket = new RemoteComputeSocket(new PipeTransport(recvPipe.Reader, sendPipe.Writer), ComputeSocketEndpoint.Remote, new TestLogger());
+			await using RemoteComputeSocket localSocket = new RemoteComputeSocket(new PipeTransport(sendPipe.Reader, recvPipe.Writer), new TestLogger());
+			await using RemoteComputeSocket agentSocket = new RemoteComputeSocket(new PipeTransport(recvPipe.Reader, sendPipe.Writer), new TestLogger());
 
 			await RunAgentTests(localSocket, agentSocket);
 		}
@@ -95,8 +97,8 @@ namespace EpicGames.Horde.Tests
 			using Socket serverSocket = await listener.AcceptSocketAsync(CancellationToken.None);
 			await clientConnectTask;
 
-			await using RemoteComputeSocket localSocket = new RemoteComputeSocket(new TcpTransport(clientSocket), ComputeSocketEndpoint.Local, new TestLogger());
-			await using RemoteComputeSocket agentSocket = new RemoteComputeSocket(new TcpTransport(serverSocket), ComputeSocketEndpoint.Remote, new TestLogger());
+			await using RemoteComputeSocket localSocket = new RemoteComputeSocket(new TcpTransport(clientSocket), new TestLogger());
+			await using RemoteComputeSocket agentSocket = new RemoteComputeSocket(new TcpTransport(serverSocket), new TestLogger());
 
 			await RunAgentTests(localSocket, agentSocket);
 		}
@@ -107,7 +109,7 @@ namespace EpicGames.Horde.Tests
 			await using (BackgroundTask agentTask = BackgroundTask.StartNew(ctx => RunAgent(agentSocket, tempDir, ctx)))
 			{
 				const int PrimaryChannelId = 0;
-				using (AgentMessageChannel channel = localSocket.CreateAgentMessageChannel(PrimaryChannelId, 4 * 1024 * 1024, NullLogger.Instance))
+				using (AgentMessageChannel channel = localSocket.CreateAgentMessageChannel(PrimaryChannelId, 4 * 1024 * 1024))
 				{
 					await channel.WaitForAttachAsync();
 
@@ -126,7 +128,7 @@ namespace EpicGames.Horde.Tests
 					}
 
 					const int SecondaryChannelId = 1;
-					using (AgentMessageChannel channel2 = localSocket.CreateAgentMessageChannel(SecondaryChannelId, 4 * 1024 * 1024, NullLogger.Instance))
+					using (AgentMessageChannel channel2 = localSocket.CreateAgentMessageChannel(SecondaryChannelId, 4 * 1024 * 1024))
 					{
 						await channel.ForkAsync(SecondaryChannelId, 4 * 1024 * 1024);
 
diff --git a/Engine/Source/Programs/Shared/EpicGames.Horde/Compute/AgentMessageChannel.cs b/Engine/Source/Programs/Shared/EpicGames.Horde/Compute/AgentMessageChannel.cs
--- a/Engine/Source/Programs/Shared/EpicGames.Horde/Compute/AgentMessageChannel.cs
+++ b/Engine/Source/Programs/Shared/EpicGames.Horde/Compute/AgentMessageChannel.cs
@@ -74,11 +74,11 @@ namespace EpicGames.Horde.Compute
 			public Span<byte> GetSpan(int sizeHint = 0) => GetMemory(sizeHint).Span;
 		}
 
-		readonly ComputeSocket _socket;
+		readonly int _channelId;
 		readonly ComputeBufferReader _recvBufferReader;
 		readonly ComputeBufferWriter _sendBufferWriter;
 
-		// Can lock chunked memory writer to acuqire pointer
+		// Can lock chunked memory writer to acquire pointer
 		readonly ILogger _logger;
 
 #pragma warning disable CA2213
@@ -88,12 +88,13 @@ namespace EpicGames.Horde.Compute
 		/// <summary>
 		/// Constructor
 		/// </summary>
+		/// <param name="channelId"></param>
 		/// <param name="recvBufferReader"></param>
 		/// <param name="sendBufferWriter"></param>
 		/// <param name="logger">Logger for diagnostic output</param>
-		public AgentMessageChannel(ComputeBufferReader recvBufferReader, ComputeBufferWriter sendBufferWriter, ILogger logger)
+		public AgentMessageChannel(int channelId, ComputeBufferReader recvBufferReader, ComputeBufferWriter sendBufferWriter, ILogger logger)
 		{
-			_socket = null!;
+			_channelId = channelId;
 			_recvBufferReader = recvBufferReader.AddRef();
 			_sendBufferWriter = sendBufferWriter.AddRef();
 			_logger = logger;
@@ -109,9 +110,9 @@ namespace EpicGames.Horde.Compute
 		/// <param name="logger">Logger for diagnostic output</param>
 		public AgentMessageChannel(ComputeSocket socket, int channelId, ComputeBuffer recvBuffer, ComputeBuffer sendBuffer, ILogger logger)
 		{
-			_socket = socket;
-			_socket.AttachRecvBuffer(channelId, recvBuffer);
-			_socket.AttachSendBuffer(channelId, sendBuffer);
+			socket.AttachRecvBuffer(channelId, recvBuffer);
+			socket.AttachSendBuffer(channelId, sendBuffer);
+			_channelId = channelId;
 			_recvBufferReader = recvBuffer.CreateReader();
 			_sendBufferWriter = sendBuffer.CreateWriter();
 			_logger = logger;
@@ -179,7 +180,7 @@ namespace EpicGames.Horde.Compute
 			{
 				bytes.Append("..");
 			}
-			_logger.LogTrace("{Verb} {Type,-22} [{Length,10:n0}] = {Bytes}", verb, type, data.Length, bytes.ToString());
+			_logger.LogTrace("{Verb} {ChannelId} {Type,-18} [{Length,10:n0}] = {Bytes}", verb, _channelId, type, data.Length, bytes.ToString());
 		}
 
 		/// <inheritdoc/>
@@ -207,9 +208,8 @@ namespace EpicGames.Horde.Compute
 		/// </summary>
 		/// <param name="socket">Socket to create a channel for</param>
 		/// <param name="channelId">Identifier for the channel</param>
-		/// <param name="logger">Logger for the channel</param>
-		public static AgentMessageChannel CreateAgentMessageChannel(this ComputeSocket socket, int channelId, ILogger logger)
-			=> socket.CreateAgentMessageChannel(channelId, 65536, logger);
+		public static AgentMessageChannel CreateAgentMessageChannel(this ComputeSocket socket, int channelId)
+			=> socket.CreateAgentMessageChannel(channelId, 65536);
 
 		/// <summary>
 		/// Creates a message channel with the given identifier
@@ -217,9 +217,8 @@ namespace EpicGames.Horde.Compute
 		/// <param name="socket">Socket to create a channel for</param>
 		/// <param name="channelId">Identifier for the channel</param>
 		/// <param name="bufferSize">Size of the send and receive buffer</param>
-		/// <param name="logger">Logger for the channel</param>
-		public static AgentMessageChannel CreateAgentMessageChannel(this ComputeSocket socket, int channelId, int bufferSize, ILogger logger)
-			=> socket.CreateAgentMessageChannel(channelId, bufferSize, bufferSize, logger);
+		public static AgentMessageChannel CreateAgentMessageChannel(this ComputeSocket socket, int channelId, int bufferSize)
+			=> socket.CreateAgentMessageChannel(channelId, bufferSize, bufferSize);
 
 		/// <summary>
 		/// Creates a message channel with the given identifier
@@ -228,12 +227,11 @@ namespace EpicGames.Horde.Compute
 		/// <param name="channelId">Identifier for the channel</param>
 		/// <param name="sendBufferSize">Size of the send buffer</param>
 		/// <param name="recvBufferSize">Size of the recieve buffer</param>
-		/// <param name="logger">Logger for the channel</param>
-		public static AgentMessageChannel CreateAgentMessageChannel(this ComputeSocket socket, int channelId, int sendBufferSize, int recvBufferSize, ILogger logger)
+		public static AgentMessageChannel CreateAgentMessageChannel(this ComputeSocket socket, int channelId, int sendBufferSize, int recvBufferSize)
 		{
 			using ComputeBuffer sendBuffer = new PooledBuffer(sendBufferSize);
 			using ComputeBuffer recvBuffer = new PooledBuffer(recvBufferSize);
-			return new AgentMessageChannel(socket, channelId, sendBuffer, recvBuffer, logger);
+			return new AgentMessageChannel(socket, channelId, sendBuffer, recvBuffer, socket.Logger);
 		}
 
 		/// <summary>
diff --git a/Engine/Source/Programs/Shared/EpicGames.Horde/Compute/AgentMessageHandler.cs b/Engine/Source/Programs/Shared/EpicGames.Horde/Compute/AgentMessageHandler.cs
--- a/Engine/Source/Programs/Shared/EpicGames.Horde/Compute/AgentMessageHandler.cs
+++ b/Engine/Source/Programs/Shared/EpicGames.Horde/Compute/AgentMessageHandler.cs
@@ -60,7 +60,7 @@ namespace EpicGames.Horde.Compute
 
 		async Task RunAsync(ComputeSocket socket, int channelId, int bufferSize, CancellationToken cancellationToken)
 		{
-			using (AgentMessageChannel channel = socket.CreateAgentMessageChannel(channelId, bufferSize, _logger))
+			using (AgentMessageChannel channel = socket.CreateAgentMessageChannel(channelId, bufferSize))
 			{
 				await channel.AttachAsync(cancellationToken);
 
diff --git a/Engine/Source/Programs/Shared/EpicGames.Horde/Compute/Clients/AgentComputeClient.cs b/Engine/Source/Programs/Shared/EpicGames.Horde/Compute/Clients/AgentComputeClient.cs
--- a/Engine/Source/Programs/Shared/EpicGames.Horde/Compute/Clients/AgentComputeClient.cs
+++ b/Engine/Source/Programs/Shared/EpicGames.Horde/Compute/Clients/AgentComputeClient.cs
@@ -92,7 +92,7 @@ namespace EpicGames.Horde.Compute.Clients
 			using BackgroundTask agentTask = BackgroundTask.StartNew(ctx => RunAgentAsync(_hordeAgentAssembly, _port, _logger, ctx));
 			using Socket tcpSocket = await listener.AcceptAsync(cancellationToken);
 
-			await using RemoteComputeSocket socket = new RemoteComputeSocket(new TcpTransport(tcpSocket), ComputeSocketEndpoint.Local, _logger);
+			await using RemoteComputeSocket socket = new RemoteComputeSocket(new TcpTransport(tcpSocket), _logger);
 			yield return socket;
 
 			await socket.CloseAsync(cancellationToken);
diff --git a/Engine/Source/Programs/Shared/EpicGames.Horde/Compute/Clients/LocalComputeClient.cs b/Engine/Source/Programs/Shared/EpicGames.Horde/Compute/Clients/LocalComputeClient.cs
--- a/Engine/Source/Programs/Shared/EpicGames.Horde/Compute/Clients/LocalComputeClient.cs
+++ b/Engine/Source/Programs/Shared/EpicGames.Horde/Compute/Clients/LocalComputeClient.cs
@@ -35,6 +35,39 @@ namespace EpicGames.Horde.Compute.Clients
 			public ValueTask CloseAsync(CancellationToken cancellationToken) => _socket.CloseAsync(cancellationToken);
 		}
 
+		class PrefixLogger : ILogger
+		{
+			readonly string _prefix;
+			readonly ILogger _inner;
+
+			public PrefixLogger(string prefix, ILogger inner)
+			{
+				_prefix = prefix;
+				_inner = inner;
+			}
+
+			public IDisposable BeginScope<TState>(TState state) => _inner.BeginScope<TState>(state);
+			public bool IsEnabled(LogLevel logLevel) => _inner.IsEnabled(logLevel);
+
+			public void Log<TState>(LogLevel logLevel, EventId eventId, TState state, Exception? exception, Func<TState, Exception?, string> formatter)
+			{
+				if (state is IEnumerable<KeyValuePair<string, object>> enumerable)
+				{
+					List<KeyValuePair<string, object>> copy = new List<KeyValuePair<string, object>>(enumerable);
+
+					int idx = copy.FindIndex(x => x.Key.Equals("{OriginalFormat}", StringComparison.OrdinalIgnoreCase));
+					if (idx != -1 && copy[idx].Value is string format)
+					{
+						copy[idx] = new KeyValuePair<string, object>(copy[idx].Key, "[{_tag}] " + format);
+						copy.Add(new KeyValuePair<string, object>("_tag", _prefix));
+						_inner.Log(logLevel, eventId, copy, exception, (s, e) => $"[{_prefix}] {formatter(state, exception)}");
+						return;
+					}
+				}
+				_inner.Log(logLevel, eventId, state, exception, formatter);
+			}
+		}
+
 		readonly BackgroundTask _listenerTask;
 		readonly Socket _listener;
 		readonly Socket _socket;
@@ -57,7 +90,7 @@ namespace EpicGames.Horde.Compute.Clients
 			_listener.Bind(new IPEndPoint(IPAddress.Loopback, port));
 			_listener.Listen();
 
-			_listenerTask = BackgroundTask.StartNew(ctx => RunListenerAsync(_listener, sandboxDir, _executeInProcess, logger, ctx));
+			_listenerTask = BackgroundTask.StartNew(ctx => RunListenerAsync(_listener, sandboxDir, _executeInProcess, new PrefixLogger("REMOTE", logger), ctx));
 
 			_socket = new Socket(SocketType.Stream, ProtocolType.IP);
 			_socket.Connect(IPAddress.Loopback, port);
@@ -80,7 +113,7 @@ namespace EpicGames.Horde.Compute.Clients
 
 			using MemoryCache memoryCache = new MemoryCache(new MemoryCacheOptions { SizeLimit = 10 * 1024 * 1024 });
 
-			await using (RemoteComputeSocket socket = new RemoteComputeSocket(new TcpTransport(tcpSocket), ComputeSocketEndpoint.Remote, logger))
+			await using (RemoteComputeSocket socket = new RemoteComputeSocket(new TcpTransport(tcpSocket), logger))
 			{
 				AgentMessageHandler worker = new AgentMessageHandler(sandboxDir, memoryCache, null, executeInProcess, null, logger);
 				await worker.RunAsync(socket, cancellationToken);
@@ -92,7 +125,7 @@ namespace EpicGames.Horde.Compute.Clients
 		public Task<IComputeLease?> TryAssignWorkerAsync(ClusterId clusterId, Requirements? requirements, CancellationToken cancellationToken)
 		{
 #pragma warning disable CA2000 // Dispose objects before losing scope
-			RemoteComputeSocket socket = new RemoteComputeSocket(new TcpTransport(_socket), ComputeSocketEndpoint.Local, _logger);
+			RemoteComputeSocket socket = new RemoteComputeSocket(new TcpTransport(_socket), new PrefixLogger("CLIENT", _logger));
 			return Task.FromResult<IComputeLease?>(new LeaseImpl(socket));
 #pragma warning restore CA2000 // Dispose objects before losing scope
 		}
diff --git a/Engine/Source/Programs/Shared/EpicGames.Horde/Compute/Clients/ServerComputeClient.cs b/Engine/Source/Programs/Shared/EpicGames.Horde/Compute/Clients/ServerComputeClient.cs
--- a/Engine/Source/Programs/Shared/EpicGames.Horde/Compute/Clients/ServerComputeClient.cs
+++ b/Engine/Source/Programs/Shared/EpicGames.Horde/Compute/Clients/ServerComputeClient.cs
@@ -175,7 +175,7 @@ namespace EpicGames.Horde.Compute.Clients
 			// Pass the rest of the call over to the handler
 			byte[] key = StringUtils.ParseHexString(responseMessage.Key);
 
-			await using RemoteComputeSocket computeSocket = new RemoteComputeSocket(new TcpTransport(socket), ComputeSocketEndpoint.Local, _logger);
+			await using RemoteComputeSocket computeSocket = new RemoteComputeSocket(new TcpTransport(socket), _logger);
 			yield return new LeaseInfo(responseMessage.Properties, responseMessage.AssignedResources, computeSocket);
 		}
 	}
diff --git a/Engine/Source/Programs/Shared/EpicGames.Horde/Compute/ComputeSocket.cs b/Engine/Source/Programs/Shared/EpicGames.Horde/Compute/ComputeSocket.cs
--- a/Engine/Source/Programs/Shared/EpicGames.Horde/Compute/ComputeSocket.cs
+++ b/Engine/Source/Programs/Shared/EpicGames.Horde/Compute/ComputeSocket.cs
@@ -6,13 +6,12 @@ using System.Buffers.Binary;
 using System.Collections.Generic;
 using System.Diagnostics.CodeAnalysis;
 using System.Linq;
-using System.Text;
 using System.Threading;
 using System.Threading.Tasks;
 using EpicGames.Core;
 using EpicGames.Horde.Compute.Buffers;
-using JetBrains.Annotations;
 using Microsoft.Extensions.Logging;
+using Microsoft.Extensions.Logging.Abstractions;
 
 namespace EpicGames.Horde.Compute
 {
@@ -21,6 +20,11 @@ namespace EpicGames.Horde.Compute
 	/// </summary>
 	public abstract class ComputeSocket
 	{
+		/// <summary>
+		/// Logger for diagnostic messages
+		/// </summary>
+		public abstract ILogger Logger { get; }
+
 		/// <summary>
 		/// Attaches a buffer to receive data.
 		/// </summary>
@@ -82,19 +86,30 @@ namespace EpicGames.Horde.Compute
 
 		readonly ComputeBufferWriter _commandBufferWriter;
 		readonly List<ComputeBuffer> _buffers = new List<ComputeBuffer>();
+		readonly ILogger _logger;
+
+		/// <inheritdoc/>
+		public override ILogger Logger => _logger;
 
 		/// <summary>
 		/// Creates a socket for a worker
 		/// </summary>
-		private WorkerComputeSocket(ComputeBufferWriter commandBufferWriter)
+		private WorkerComputeSocket(ComputeBufferWriter commandBufferWriter, ILogger logger)
 		{
 			_commandBufferWriter = commandBufferWriter;
+			_logger = logger;
 		}
 
 		/// <summary>
 		/// Opens a socket which allows a worker to communicate with the Horde Agent
 		/// </summary>
-		public static WorkerComputeSocket Open()
+		public static WorkerComputeSocket Open() => Open(NullLogger.Instance);
+
+		/// <summary>
+		/// Opens a socket which allows a worker to communicate with the Horde Agent
+		/// </summary>
+		/// <param name="logger">Logger for diagnostic messages</param>
+		public static WorkerComputeSocket Open(ILogger logger)
 		{
 			string? baseName = Environment.GetEnvironmentVariable(IpcEnvVar);
 			if (baseName == null)
@@ -102,17 +117,18 @@ namespace EpicGames.Horde.Compute
 				throw new InvalidOperationException($"Environment variable {IpcEnvVar} is not defined; cannot connect as worker.");
 			}
 
-			return Open(baseName);
+			return Open(baseName, logger);
 		}
 
 		/// <summary>
 		/// Opens a socket which allows a worker to communicate with the Horde Agent
 		/// </summary>
 		/// <param name="commandBufferName">Name of the command buffer</param>
-		public static WorkerComputeSocket Open(string commandBufferName)
+		/// <param name="logger">Logger for diagnostic messages</param>
+		public static WorkerComputeSocket Open(string commandBufferName, ILogger logger)
 		{
 			using SharedMemoryBuffer commandBuffer = SharedMemoryBuffer.OpenExisting(commandBufferName);
-			return new WorkerComputeSocket(commandBuffer.CreateWriter());
+			return new WorkerComputeSocket(commandBuffer.CreateWriter(), logger);
 		}
 
 		/// <inheritdoc/>
@@ -148,22 +164,6 @@ namespace EpicGames.Horde.Compute
 		}
 	}
 
-	/// <summary>
-	/// Enum identifying which end of the socket a particular machine is
-	/// </summary>
-	public enum ComputeSocketEndpoint
-	{
-		/// <summary>
-		/// The initiating machine
-		/// </summary>
-		Local,
-
-		/// <summary>
-		/// The remote machine
-		/// </summary>
-		Remote,
-	}
-
 	/// <summary>
 	/// Manages a set of readers and writers to buffers across a transport layer
 	/// </summary>
@@ -235,7 +235,6 @@ namespace EpicGames.Horde.Compute
 		bool _complete;
 
 		readonly ComputeTransport _transport;
-		readonly ComputeSocketEndpoint _endpoint;
 		readonly ILogger _logger;
 
 		BackgroundTask? _recvTask;
@@ -244,18 +243,17 @@ namespace EpicGames.Horde.Compute
 		readonly SemaphoreSlim _sendSemaphore = new SemaphoreSlim(1, 1);
 		readonly Dictionary<int, SendBuffer> _sendBuffers = new Dictionary<int, SendBuffer>();
 
-		string Tag => (_endpoint == ComputeSocketEndpoint.Local)? "LOCAL": "REMOTE";
+		/// <inheritdoc/>
+		public override ILogger Logger => _logger;
 
 		/// <summary>
 		/// Constructor
 		/// </summary>
 		/// <param name="transport">Transport to communicate with the remote</param>
-		/// <param name="endpoint">Tag for log messages</param>
 		/// <param name="logger">Logger for trace output</param>
-		public RemoteComputeSocket(ComputeTransport transport, ComputeSocketEndpoint endpoint, ILogger logger)
+		public RemoteComputeSocket(ComputeTransport transport, ILogger logger)
 		{
 			_transport = transport;
-			_endpoint = endpoint;
 			_logger = logger;
 		}
 
@@ -273,7 +271,7 @@ namespace EpicGames.Horde.Compute
 			// Wait for the reader to stop
 			if (_recvTask != null)
 			{
-				await _recvTask.StopAsync();
+				await _recvTask.DisposeAsync();
 			}
 		}
 
@@ -287,7 +285,7 @@ namespace EpicGames.Horde.Compute
 
 		async Task RunRecvTaskAsync(ComputeTransport transport, CancellationToken cancellationToken)
 		{
-			_logger.LogDebug("[{Tag}] Started socket reader", Tag);
+			_logger.LogDebug("Started socket reader");
 
 			List<Task> detachTasks = new List<Task>();
 
@@ -304,7 +302,7 @@ namespace EpicGames.Horde.Compute
 					// Read the next packet header
 					if (!await transport.RecvOptionalAsync(header, cancellationToken))
 					{
-						_logger.LogDebug("[{Tag}] End of socket", Tag);
+						_logger.LogDebug("End of socket");
 						break;
 					}
 
@@ -323,7 +321,7 @@ namespace EpicGames.Horde.Compute
 					}
 					else
 					{
-						_logger.LogDebug("[{Tag}] Unrecognized control message: {Message}", Tag, size);
+						_logger.LogDebug("Unrecognized control message: {Message}", size);
 					}
 				}
 			}
@@ -344,11 +342,11 @@ namespace EpicGames.Horde.Compute
 			// Wait for all the detach tasks to finish
 			if (detachTasks.Count > 0)
 			{
-				_logger.LogDebug("[{Tag}] Waiting for detach tasks to complete...", Tag);
+				_logger.LogDebug("Waiting for detach tasks to complete...");
 				await Task.WhenAll(detachTasks).WaitAsync(cancellationToken);
 			}
 
-			_logger.LogDebug("[{Tag}] Closing reader", Tag);
+			_logger.LogDebug("Closing reader");
 		}
 
 		async Task ReadPacketAsync(ComputeTransport transport, int id, int size, CancellationToken cancellationToken)
@@ -399,7 +397,7 @@ namespace EpicGames.Horde.Compute
 						Memory<byte> memory = writer.GetWriteBuffer();
 						while (memory.Length < size)
 						{
-							_logger.LogDebug("[{Tag}] No space in buffer {Id}, flushing", Tag, id);
+							_logger.LogDebug("No space in buffer {Id}, flushing", id);
 							await writer.WaitToWriteAsync(size, cancellationToken);
 							memory = writer.GetWriteBuffer();
 						}
@@ -474,7 +472,7 @@ namespace EpicGames.Horde.Compute
 		/// <inheritdoc/>
 		public override void AttachRecvBuffer(int channelId, ComputeBuffer recvBuffer)
 		{
-			_logger.LogDebug("[{Tag}] Attaching recv buffer {Id}", Tag, channelId);
+			_logger.LogDebug("Attaching recv buffer {Id}", channelId);
 			lock (_lockObject)
 			{
 				if (_complete)
@@ -494,7 +492,7 @@ namespace EpicGames.Horde.Compute
 		[SuppressMessage("Reliability", "CA2000:Dispose objects before losing scope")]
 		async Task DetachRecvBufferAsync(int id, CancellationToken cancellationToken)
 		{
-			_logger.LogDebug("[{Tag}] Detaching recv buffer {Id}", Tag, id);
+			_logger.LogDebug("Detaching recv buffer {Id}", id);
 
 			// Get the current receive buffer
 			RecvBuffer? recvBuffer;
@@ -502,7 +500,7 @@ namespace EpicGames.Horde.Compute
 			{
 				if (!_recvBuffers.TryGetValue(id, out recvBuffer))
 				{
-					_logger.LogDebug("[{Tag}] Buffer {Id} has already been detached", Tag, id);
+					_logger.LogDebug("Buffer {Id} has already been detached", id);
 					return;
 				}
 				recvBuffer.AddRef(); // Note: adding extra ref here
@@ -563,7 +561,7 @@ namespace EpicGames.Horde.Compute
 		/// <inheritdoc/>
 		public override void AttachSendBuffer(int channelId, ComputeBuffer sendBuffer)
 		{
-			_logger.LogDebug("[{Tag}] Attaching send buffer {Id}", Tag, channelId);
+			_logger.LogDebug("Attaching send buffer {Id}", channelId);
 			lock (_lockObject)
 			{
 				if (_sendBuffers.ContainsKey(channelId))
