Horde: Send keepalive messages to the remote machine every 5s. This should allow us to time out the remote more aggressively.

diff --git a/Engine/Source/Programs/Shared/EpicGames.Horde/Compute/Clients/ServerComputeClient.cs b/Engine/Source/Programs/Shared/EpicGames.Horde/Compute/Clients/ServerComputeClient.cs
--- a/Engine/Source/Programs/Shared/EpicGames.Horde/Compute/Clients/ServerComputeClient.cs
+++ b/Engine/Source/Programs/Shared/EpicGames.Horde/Compute/Clients/ServerComputeClient.cs
@@ -5,6 +5,7 @@ using EpicGames.Horde.Api;
 using EpicGames.Horde.Compute.Transports;
 using Microsoft.Extensions.Logging;
 using System;
+using System.Buffers;
 using System.Collections.Generic;
 using System.Net;
 using System.Net.Http;
@@ -33,6 +34,8 @@ namespace EpicGames.Horde.Compute.Clients
 		class LeaseImpl : IComputeLease
 		{
 			readonly IAsyncEnumerator<LeaseInfo> _source;
+			
+			BackgroundTask? _pingTask;
 
 			public IReadOnlyList<string> Properties => _source.Current.Properties;
 			public IReadOnlyDictionary<string, int> AssignedResources => _source.Current.AssignedResources;
@@ -41,17 +44,42 @@ namespace EpicGames.Horde.Compute.Clients
 			public LeaseImpl(IAsyncEnumerator<LeaseInfo> source)
 			{
 				_source = source;
+				_pingTask = BackgroundTask.StartNew(PingAsync);
 			}
 
 			/// <inheritdoc/>
 			public async ValueTask DisposeAsync()
 			{
+				if (_pingTask != null)
+				{
+					await _pingTask.DisposeAsync();
+					_pingTask = null;
+				}
+
 				await _source.MoveNextAsync();
 				await _source.DisposeAsync();
 			}
 
 			/// <inheritdoc/>
-			public ValueTask CloseAsync(CancellationToken cancellationToken) => Socket.CloseAsync(cancellationToken);
+			public async ValueTask CloseAsync(CancellationToken cancellationToken)
+			{
+				if (_pingTask != null)
+				{
+					await _pingTask.DisposeAsync();
+					_pingTask = null;
+				}
+
+				await Socket.CloseAsync(cancellationToken);
+			}
+
+			async Task PingAsync(CancellationToken cancellationToken)
+			{
+				while (!cancellationToken.IsCancellationRequested)
+				{
+					await Socket.SendKeepAliveMessageAsync(cancellationToken);
+					await Task.Delay(TimeSpan.FromSeconds(5.0), cancellationToken);
+				}
+			}
 		}
 
 		readonly HttpClient? _defaultHttpClient;
diff --git a/Engine/Source/Programs/Shared/EpicGames.Horde/Compute/ComputeSocket.cs b/Engine/Source/Programs/Shared/EpicGames.Horde/Compute/ComputeSocket.cs
--- a/Engine/Source/Programs/Shared/EpicGames.Horde/Compute/ComputeSocket.cs
+++ b/Engine/Source/Programs/Shared/EpicGames.Horde/Compute/ComputeSocket.cs
@@ -172,6 +172,7 @@ namespace EpicGames.Horde.Compute
 		enum ControlMessageType
 		{
 			Detach = -2,
+			KeepAlive = -3,
 		}
 
 		class RecvBuffer : IDisposable
@@ -237,7 +238,7 @@ namespace EpicGames.Horde.Compute
 		readonly ComputeTransport _transport;
 		readonly ILogger _logger;
 
-		BackgroundTask? _recvTask;
+		readonly BackgroundTask _recvTask;
 		readonly Dictionary<int, RecvBuffer> _recvBuffers = new Dictionary<int, RecvBuffer>();
 
 		readonly SemaphoreSlim _sendSemaphore = new SemaphoreSlim(1, 1);
@@ -255,6 +256,8 @@ namespace EpicGames.Horde.Compute
 		{
 			_transport = transport;
 			_logger = logger;
+
+			_recvTask = BackgroundTask.StartNew(ctx => RunRecvTaskAsync(_transport, ctx));
 		}
 
 		/// <summary>
@@ -283,6 +286,16 @@ namespace EpicGames.Horde.Compute
 			GC.SuppressFinalize(this);
 		}
 
+		/// <summary>
+		/// Sends a keep alive message to the remote machine. Does not wait for a response. Designed to keep a connection open when the remote is eagerly trying to close it.
+		/// </summary>
+		/// <param name="cancellationToken">Cancellation token for the operation</param>
+		public async Task SendKeepAliveMessageAsync(CancellationToken cancellationToken)
+		{
+			_logger.LogDebug("Sending ping message");
+			await SendInternalAsync(0, (int)ControlMessageType.KeepAlive, ReadOnlyMemory<byte>.Empty, cancellationToken);
+		}
+
 		async Task RunRecvTaskAsync(ComputeTransport transport, CancellationToken cancellationToken)
 		{
 			_logger.LogDebug("Started socket reader");
@@ -319,6 +332,10 @@ namespace EpicGames.Horde.Compute
 					{
 						detachTasks.Add(DetachRecvBufferAsync(id, cancellationToken));
 					}
+					else if (size == (int)ControlMessageType.KeepAlive)
+					{
+						_logger.LogDebug("Received ping message");
+					}
 					else
 					{
 						_logger.LogDebug("Unrecognized control message: {Message}", size);
@@ -485,7 +502,6 @@ namespace EpicGames.Horde.Compute
 				}
 
 				_recvBuffers.Add(channelId, new RecvBuffer(recvBuffer.CreateWriter()));
-				_recvTask ??= BackgroundTask.StartNew(ctx => RunRecvTaskAsync(_transport, ctx));
 			}
 		}
 
