Horde: Reduce memory usage when creating bundles.

* Maximum memory footprint is respected when queuing up bundles to write.
* All forked instances of a writer use the same write queue, and contribute towards the same resource limits.
* New writers are only allowed to start once previous data has been released.
* Default packet writer for bundles now uses a smaller (1mb vs 16mb) chunk size.

diff --git a/Engine/Source/Programs/Shared/EpicGames.Horde/Storage/Bundles/BundleWriter.cs b/Engine/Source/Programs/Shared/EpicGames.Horde/Storage/Bundles/BundleWriter.cs
--- a/Engine/Source/Programs/Shared/EpicGames.Horde/Storage/Bundles/BundleWriter.cs
+++ b/Engine/Source/Programs/Shared/EpicGames.Horde/Storage/Bundles/BundleWriter.cs
@@ -26,7 +26,7 @@ namespace EpicGames.Horde.Storage.Bundles
 		/// <summary>
 		/// Compression format to use
 		/// </summary>
-		public BundleCompressionFormat CompressionFormat { get; set; } = BundleCompressionFormat.LZ4;
+		public BundleCompressionFormat CompressionFormat { get; set; } = BundleCompressionFormat.Gzip;
 
 		/// <summary>
 		/// Minimum size of a block to be compressed
@@ -36,7 +36,7 @@ namespace EpicGames.Horde.Storage.Bundles
 		/// <summary>
 		/// Maximum amount of data to store in memory. This includes any background writes as well as bundles being built.
 		/// </summary>
-		public int MaxInMemoryDataLength { get; set; } = 128 * 1024 * 1024;
+		public long MaxWriteQueueLength { get; set; } = 256 * 1024 * 1024;
 
 		/// <summary>
 		/// Number of nodes to cache
@@ -289,6 +289,9 @@ namespace EpicGames.Horde.Storage.Bundles
 			// Length of the packet that has been written
 			int _currentPacketLength;
 
+			// Total size of compressed data in the current bundle
+			long _compressedLength;
+
 			// Map of keys to nodes in the queue
 			public readonly Dictionary<NodeKey, PendingNode> _nodeKeyToInfo = new Dictionary<NodeKey, PendingNode>();
 
@@ -304,17 +307,17 @@ namespace EpicGames.Horde.Storage.Bundles
 			// Set of all direct dependencies from this bundle
 			readonly HashSet<Task> _dependencies = new HashSet<Task>();
 
-			// Total size of uncompressed data in the current bundle
-			public int UncompressedLength { get; private set; }
+			// Total size of compressed data in the current bundle
+			public long CompressedLength => _compressedLength;
 
-			// Whether all nodes have been added to a bundle
-			public bool IsReadOnly { get; private set; }
+			// Total size of uncompressed data in the current bundle
+			public long UncompressedLength { get; private set; }
 
 			// List of post-write callbacks
 			BlobWriteCallback? _callbacks = null;
 
 			// Task used to compress data in the background
-			Task _writeTask = Task.CompletedTask;
+			Task _compressPacketsTask = Task.CompletedTask;
 
 			// Event which is signalled after the bundle is written to storage
 			readonly TaskCompletionSource<bool> _completeEvent = new TaskCompletionSource<bool>();
@@ -322,13 +325,13 @@ namespace EpicGames.Horde.Storage.Bundles
 			// Task signalled after the write is complete
 			public Task CompleteTask => _completeEvent.Task;
 
-			public PendingBundle(BundleReader treeReader, BundleWriter treeWriter, int maxPacketSize, int maxBlobSize, BundleCompressionFormat compressionFormat)
+			public PendingBundle(BundleReader treeReader, BundleWriter treeWriter, int maxPacketSize, BundleCompressionFormat compressionFormat)
 			{
 				_treeReader = treeReader;
 				_treeWriter = treeWriter;
 				_maxPacketSize = maxPacketSize;
 				_compressionFormat = compressionFormat;
-				_encodedPacketWriter = new ChunkedMemoryWriter(maxBlobSize);
+				_encodedPacketWriter = new ChunkedMemoryWriter(1024 * 1024);
 			}
 
 			/// <inheritdoc/>
@@ -352,9 +355,6 @@ namespace EpicGames.Horde.Storage.Bundles
 			// Whether this bundle has finished writing
 			public bool IsComplete() => CompleteTask.IsCompleted;
 
-			// Whether this bundle can be written
-			public bool CanComplete() => !IsReadOnly && _dependencies.Count == 0;
-
 			// Add a dependency onto another bundle
 			public void AddDependencyOn(PendingBundle bundle, ILogger? traceLogger)
 			{
@@ -479,7 +479,7 @@ namespace EpicGames.Horde.Storage.Bundles
 
 					if (currentPacketLength > 0)
 					{
-						_writeTask = _writeTask.ContinueWith(x => CompressPacket(currentPacket, currentPacketLength), TaskScheduler.Default);
+						_compressPacketsTask = _compressPacketsTask.ContinueWith(x => CompressPacket(currentPacket, currentPacketLength), TaskScheduler.Default);
 						_currentPacketIdx++;
 					}
 					else
@@ -501,41 +501,41 @@ namespace EpicGames.Horde.Storage.Bundles
 					int encodedOffset = (_packets.Count == 0) ? 0 : _packets[^1].EncodedOffset + _packets[^1].EncodedLength;
 					BundlePacket packet = new BundlePacket(_compressionFormat, encodedOffset, encodedLength, length);
 					_packets.Add(packet);
+					Interlocked.Add(ref _compressedLength, encodedLength);
 
 					packetData.Dispose();
 				}
 			}
 
-			// Mark the bundle as complete
-			public void MarkAsComplete(BundleStorageClient store, Utf8String prefix, ILogger? traceLogger)
+			// Wait for the bundle's dependencies to complete
+			public async Task FlushDependenciesAsync(ILogger? traceLogger, CancellationToken cancellationToken)
 			{
-				if (!IsReadOnly)
-				{
-					traceLogger?.LogInformation("Marking bundle {BundleId} as complete ({NumNodes} nodes); adding to write queue.", BundleId, _queue.Count);
-					FlushPacket();
-					Task prevWriteTask = _writeTask;
-					_writeTask = Task.Run(() => CompleteAsync(prevWriteTask, store, prefix, traceLogger));
-					IsReadOnly = true;
-				}
-			}
+				FlushPacket();
+				await _compressPacketsTask.WaitAsync(cancellationToken);
 
-			async Task CompleteAsync(Task prevWriteTask, BundleStorageClient store, Utf8String prefix, ILogger? traceLogger)
-			{
-				try
+				if (traceLogger != null)
 				{
-					await prevWriteTask;
-					if (traceLogger != null)
+					foreach (Task dependency in _dependencies)
 					{
-						foreach (Task dependency in _dependencies)
+						if (!dependency.IsCompleted)
 						{
-							if (!dependency.IsCompleted)
-							{
-								traceLogger.LogInformation("Bundle {BundleId} is stalling waiting for dependencies to flush first", BundleId);
-							}
+							traceLogger.LogInformation("Bundle {BundleId} is stalling waiting for dependencies to flush first", BundleId);
 						}
 					}
-					await Task.WhenAll(_dependencies);
+				}
 
+				await Task.WhenAll(_dependencies).WaitAsync(cancellationToken);
+			}
+
+			// Mark the bundle as complete
+			public async Task WriteAsync(BundleStorageClient store, Utf8String prefix, ILogger? traceLogger)
+			{
+				traceLogger?.LogInformation("Marking bundle {BundleId} as complete ({NumNodes} nodes); adding to write queue.", BundleId, _queue.Count);
+
+				await FlushDependenciesAsync(traceLogger, CancellationToken.None);
+
+				try
+				{
 					Bundle bundle = CreateBundle();
 					BundleLocator locator = await store.WriteBundleAsync(bundle, prefix);
 					traceLogger?.LogInformation("Written bundle {BundleId} as {Locator}", BundleId, locator);
@@ -653,6 +653,90 @@ namespace EpicGames.Horde.Storage.Bundles
 			}
 		}
 
+		class WriteQueue
+		{
+			long _memoryFootprint;
+			readonly BundleStorageClient _store;
+			readonly Utf8String _prefix;
+			readonly long _maxMemoryFootprint;
+			readonly ILogger? _traceLogger;
+			readonly AsyncEvent _completeEvent = new AsyncEvent();
+			int _refCount;
+			readonly List<Task> _writeTasks = new List<Task>();
+
+			public WriteQueue(BundleStorageClient store, Utf8String prefix, long maxMemoryFootprint, ILogger? traceLogger)
+			{
+				_store = store;
+				_prefix = prefix;
+				_maxMemoryFootprint = maxMemoryFootprint;
+				_refCount = 1;
+				_traceLogger = traceLogger;
+			}
+
+			public void AddRef()
+			{
+				Interlocked.Increment(ref _refCount);
+			}
+
+			public async ValueTask ReleaseAsync()
+			{
+				if (Interlocked.Decrement(ref _refCount) == 0)
+				{
+					await FlushAsync(CancellationToken.None);
+				}
+			}
+
+			public async Task AddAsync(PendingBundle pendingBundle, CancellationToken cancellationToken)
+			{
+				lock (_writeTasks)
+				{
+					_writeTasks.RemoveCompleteTasks();
+				}
+
+				await pendingBundle.FlushDependenciesAsync(_traceLogger, cancellationToken);
+
+				for (; ; )
+				{
+					Task completeTask = _completeEvent.Task;
+
+					long memoryFootprint = Interlocked.CompareExchange(ref _memoryFootprint, 0, 0);
+					long newMemoryFootprint = memoryFootprint + pendingBundle.CompressedLength;
+
+					if (memoryFootprint > 0 && newMemoryFootprint > _maxMemoryFootprint)
+					{
+						await completeTask;
+					}
+					else if(Interlocked.CompareExchange(ref _memoryFootprint, newMemoryFootprint, memoryFootprint) == memoryFootprint)
+					{
+						break;
+					}
+				}
+
+				lock (_writeTasks)
+				{
+					_writeTasks.Add(WriteAsync(pendingBundle));
+				}
+			}
+
+			public async Task FlushAsync(CancellationToken cancellationToken)
+			{
+				Task[] writeTasks;
+				lock (_writeTasks)
+				{
+					writeTasks = _writeTasks.ToArray();
+				}
+				await Task.WhenAll(writeTasks).WaitAsync(cancellationToken);
+			}
+
+			async Task WriteAsync(PendingBundle pendingBundle)
+			{
+				await pendingBundle.WriteAsync(_store, _prefix, _traceLogger);
+				Interlocked.Add(ref _memoryFootprint, -pendingBundle.CompressedLength);
+				pendingBundle.Dispose();
+				_completeEvent.Pulse();
+			}
+		}
+
 		static readonly BundleOptions s_defaultOptions = new BundleOptions();
 
 		readonly BundleStorageClient _store;
@@ -661,10 +745,9 @@ namespace EpicGames.Horde.Storage.Bundles
 		readonly RefName _refName;
 
 		readonly NodeCache _nodeCache;
-		readonly Queue<PendingBundle> _writeQueue = new Queue<PendingBundle>();
+		readonly WriteQueue _writeQueue;
 
 		PendingBundle? _currentBundle;
-		int _memoryFootprint;
 
 		bool _disposed;
 
@@ -685,13 +768,8 @@ namespace EpicGames.Horde.Storage.Bundles
 		/// <param name="nodeCache">Cache of nodes for deduplication</param>
 		/// <param name="traceLogger">Optional logger for trace information</param>
 		public BundleWriter(BundleStorageClient store, BundleReader reader, RefName refName, BundleOptions? options = null, NodeCache? nodeCache = null, ILogger? traceLogger = null)
+			: this(store, reader, refName, options, nodeCache, null, traceLogger)
 		{
-			_store = store;
-			_reader = reader;
-			_refName = refName;
-			_options = options ?? s_defaultOptions;
-			_nodeCache = nodeCache ?? new NodeCache(_options.NodeCacheSize);
-			_traceLogger = traceLogger;
 		}
 
 		/// <summary>
@@ -699,8 +777,23 @@ namespace EpicGames.Horde.Storage.Bundles
 		/// </summary>
 		/// <param name="other"></param>
 		public BundleWriter(BundleWriter other)
-			: this(other._store, other._reader, other._refName, other._options, other._nodeCache, other._traceLogger)
+			: this(other._store, other._reader, other._refName, other._options, other._nodeCache, other._writeQueue, other._traceLogger)
 		{
+			_writeQueue.AddRef();
+		}
+
+		/// <summary>
+		/// Internal constructor
+		/// </summary>
+		private BundleWriter(BundleStorageClient store, BundleReader reader, RefName refName, BundleOptions? options, NodeCache? nodeCache, WriteQueue? writeQueue, ILogger? traceLogger = null)
+		{
+			_store = store;
+			_reader = reader;
+			_refName = refName;
+			_options = options ?? s_defaultOptions;
+			_nodeCache = nodeCache ?? new NodeCache(_options.NodeCacheSize);
+			_writeQueue = writeQueue ?? new WriteQueue(store, refName.Text, _options.MaxWriteQueueLength, traceLogger);
+			_traceLogger = traceLogger;
 		}
 
 		/// <inheritdoc/>
@@ -712,6 +805,7 @@ namespace EpicGames.Horde.Storage.Bundles
 			if (!_disposed)
 			{
 				await FlushAsync();
+				await _writeQueue.ReleaseAsync();
 				_disposed = true;
 			}
 
@@ -725,12 +819,11 @@ namespace EpicGames.Horde.Storage.Bundles
 		/// <summary>
 		/// Mark this writer as complete, allowing its data to be serialized.
 		/// </summary>
-		public void Complete()
+		public async Task CompleteAsync(CancellationToken cancellationToken)
 		{
 			if (_currentBundle != null)
 			{
-				_currentBundle.MarkAsComplete(_store, _refName.Text, _traceLogger);
-				_writeQueue.Enqueue(_currentBundle);
+				await _writeQueue.AddAsync(_currentBundle, cancellationToken);
 				_currentBundle = null;
 			}
 		}
@@ -775,7 +868,6 @@ namespace EpicGames.Horde.Storage.Bundles
 
 			// Append this node data
 			PendingNode pendingNode = currentBundle.WriteNode(nodeKey, size, references);
-			_memoryFootprint += size;
 			_nodeCache.Add(nodeKey, pendingNode);
 			_traceLogger?.LogInformation("Added new node for {NodeKey} in bundle {BundleId}", nodeKey, currentBundle.BundleId);
 
@@ -792,13 +884,7 @@ namespace EpicGames.Horde.Storage.Bundles
 			// If the bundle is full, start the process of writing it to disk
 			if (currentBundle.UncompressedLength > _options.MaxBlobSize || currentBundle.IsFull())
 			{
-				Complete();
-			}
-
-			// Remove any complete bundle writes
-			while (_writeQueue.Count > 0 && (_writeQueue.Peek().IsComplete() || _memoryFootprint > _options.MaxInMemoryDataLength))
-			{
-				await WaitForWriteAsync(cancellationToken);
+				await CompleteAsync(cancellationToken);
 			}
 
 			return pendingNode;
@@ -806,23 +892,14 @@ namespace EpicGames.Horde.Storage.Bundles
 
 		PendingBundle GetCurrentBundle()
 		{
-			if (_currentBundle == null || _currentBundle.IsReadOnly)
+			if (_currentBundle == null)
 			{
 				int bufferSize = (int)(_options.MinCompressionPacketSize * 1.2);
-				_currentBundle = new PendingBundle(_reader, this, bufferSize, _options.MaxBlobSize, _options.CompressionFormat);
+				_currentBundle = new PendingBundle(_reader, this, bufferSize, _options.CompressionFormat);
 			}
 			return _currentBundle;
 		}
 
-		async Task WaitForWriteAsync(CancellationToken cancellationToken)
-		{
-			PendingBundle writtenBundle = _writeQueue.Dequeue();
-			_traceLogger?.LogInformation("Waiting for bundle {BundleId} to finish writing", writtenBundle.BundleId);
-			await await Task.WhenAny(writtenBundle.CompleteTask, Task.Delay(-1, cancellationToken));
-			_memoryFootprint -= writtenBundle.UncompressedLength;
-			writtenBundle.Dispose();
-		}
-
 		/// <summary>
 		/// Flushes all the current nodes to storage
 		/// </summary>
@@ -835,11 +912,8 @@ namespace EpicGames.Horde.Storage.Bundles
 				throw new ObjectDisposedException(GetType().Name);
 			}
 
-			Complete();
-			while (_writeQueue.Count > 0)
-			{
-				await WaitForWriteAsync(cancellationToken);
-			}
+			await CompleteAsync(cancellationToken);
+			await _writeQueue.FlushAsync(cancellationToken);
 		}
 
 		/// <summary>
