Horde: Revert change to use Gzip for bundles.

diff --git a/Engine/Source/Programs/Shared/EpicGames.Horde/Storage/Bundles/BundleWriter.cs b/Engine/Source/Programs/Shared/EpicGames.Horde/Storage/Bundles/BundleWriter.cs
--- a/Engine/Source/Programs/Shared/EpicGames.Horde/Storage/Bundles/BundleWriter.cs
+++ b/Engine/Source/Programs/Shared/EpicGames.Horde/Storage/Bundles/BundleWriter.cs
@@ -26,7 +26,7 @@ namespace EpicGames.Horde.Storage.Bundles
 		/// <summary>
 		/// Compression format to use
 		/// </summary>
-		public BundleCompressionFormat CompressionFormat { get; set; } = BundleCompressionFormat.Gzip;
+		public BundleCompressionFormat CompressionFormat { get; set; } = BundleCompressionFormat.LZ4;
 
 		/// <summary>
 		/// Minimum size of a block to be compressed
