Horde: Fix Gzip compression in bundles.

diff --git a/Engine/Source/Programs/Shared/EpicGames.Horde/Storage/Bundles/Bundle.cs b/Engine/Source/Programs/Shared/EpicGames.Horde/Storage/Bundles/Bundle.cs
--- a/Engine/Source/Programs/Shared/EpicGames.Horde/Storage/Bundles/Bundle.cs
+++ b/Engine/Source/Programs/Shared/EpicGames.Horde/Storage/Bundles/Bundle.cs
@@ -1232,8 +1232,11 @@ namespace EpicGames.Horde.Storage.Bundles
 				case BundleCompressionFormat.Gzip:
 					{
 						using MemoryStream outputStream = new MemoryStream(input.Length);
-						using GZipStream gzipStream = new GZipStream(outputStream, CompressionLevel.Fastest);
-						gzipStream.Write(input.Span);
+
+						using (GZipStream gzipStream = new GZipStream(outputStream, CompressionLevel.Fastest, true))
+						{
+							gzipStream.Write(input.Span);
+						}
 
 						writer.WriteFixedLengthBytes(outputStream.ToArray());
 						return (int)outputStream.Length;
@@ -1276,8 +1279,13 @@ namespace EpicGames.Horde.Storage.Bundles
 				case BundleCompressionFormat.Gzip:
 					{
 						using ReadOnlyMemoryStream inputStream = new ReadOnlyMemoryStream(input);
-						using GZipStream outputStream = new GZipStream(new MemoryWriterStream(new MemoryWriter(output)), CompressionMode.Decompress, false);
-						inputStream.CopyTo(outputStream);
+						using GZipStream inflatedStream = new GZipStream(inputStream, CompressionMode.Decompress, true);
+
+						int length = inflatedStream.Read(output.Span);
+						if (length != output.Length)
+						{
+							throw new InvalidDataException($"Decoded data is shorter than expected (expected {output.Length} bytes, got {length} bytes)");
+						}
 						break;
 					}
 				case BundleCompressionFormat.Oodle:
