Horde: Add Brotli as a bundle compression format.

diff --git a/Engine/Source/Programs/Shared/EpicGames.Horde/Storage/Bundles/Bundle.cs b/Engine/Source/Programs/Shared/EpicGames.Horde/Storage/Bundles/Bundle.cs
--- a/Engine/Source/Programs/Shared/EpicGames.Horde/Storage/Bundles/Bundle.cs
+++ b/Engine/Source/Programs/Shared/EpicGames.Horde/Storage/Bundles/Bundle.cs
@@ -178,6 +178,11 @@ namespace EpicGames.Horde.Storage.Bundles
 		/// Oodle compression (Kraken)
 		/// </summary>
 		Oodle = 3,
+
+		/// <summary>
+		/// Brotli compression
+		/// </summary>
+		Brotli = 4,
 	}
 
 	/// <summary>
@@ -1255,6 +1260,19 @@ namespace EpicGames.Horde.Storage.Bundles
 						throw new NotSupportedException("Oodle is not compiled into this build.");
 #endif
 					}
+				case BundleCompressionFormat.Brotli:
+					{
+						int maxSize = BrotliEncoder.GetMaxCompressedLength(input.Length);
+
+						Span<byte> buffer = writer.GetSpan(maxSize);
+						if (!BrotliEncoder.TryCompress(input.Span, buffer, out int encodedLength))
+						{
+							throw new InvalidOperationException("Unable to compress data using Brotli");
+						}
+
+						writer.Advance(encodedLength);
+						return encodedLength;
+					}
 				default:
 					throw new InvalidDataException($"Invalid compression format '{(int)format}'");
 			}
@@ -1295,6 +1313,15 @@ namespace EpicGames.Horde.Storage.Bundles
 #else
 					throw new NotSupportedException("Oodle is not compiled into this build.");
 #endif
+				case BundleCompressionFormat.Brotli:
+					{
+						int bytesWritten;
+						if (!BrotliDecoder.TryDecompress(input.Span, output.Span, out bytesWritten) || bytesWritten != output.Length)
+						{
+							throw new InvalidOperationException("Unable to decompress data using Brotli");
+						}
+						break;
+					}
 				default:
 					throw new InvalidDataException($"Invalid compression format '{(int)format}'");
 			}
