Horde: Output a log entry for reaching 100% in a copy operation.

diff --git a/Engine/Source/Programs/Shared/EpicGames.Horde/Storage/Nodes/DirectoryNode.cs b/Engine/Source/Programs/Shared/EpicGames.Horde/Storage/Nodes/DirectoryNode.cs
--- a/Engine/Source/Programs/Shared/EpicGames.Horde/Storage/Nodes/DirectoryNode.cs
+++ b/Engine/Source/Programs/Shared/EpicGames.Horde/Storage/Nodes/DirectoryNode.cs
@@ -435,7 +435,7 @@ namespace EpicGames.Horde.Storage.Nodes
 				{
 					CopiedCount += count;
 					CopiedSize += size;
-					if (_timer.Elapsed > TimeSpan.FromSeconds(10.0))
+					if (_timer.Elapsed > TimeSpan.FromSeconds(10.0) || CopiedCount == count)
 					{
 						_progress.Report(this);
 						_timer.Restart();
