Horde: Reduce default buffer size for writing new nodes. This has the effect of forcing data to be allocated in a new packet, which reduces compression performance and increases header size.

diff --git a/Engine/Source/Programs/Shared/EpicGames.Horde/Storage/Node.cs b/Engine/Source/Programs/Shared/EpicGames.Horde/Storage/Node.cs
--- a/Engine/Source/Programs/Shared/EpicGames.Horde/Storage/Node.cs
+++ b/Engine/Source/Programs/Shared/EpicGames.Horde/Storage/Node.cs
@@ -357,7 +357,7 @@ namespace EpicGames.Horde.Storage
 		public NodeWriter(IStorageWriter treeWriter)
 		{
 			_treeWriter = treeWriter;
-			_memory = treeWriter.GetOutputBuffer(0, 256 * 1024);
+			_memory = treeWriter.GetOutputBuffer(0, 2 * 1024);
 		}
 
 		/// <summary>
