Horde: Rename BundleStorageClient.TreeReader to BundleReader.

diff --git a/Engine/Source/Programs/Horde/Horde.Server/Storage/StorageService.cs b/Engine/Source/Programs/Horde/Horde.Server/Storage/StorageService.cs
--- a/Engine/Source/Programs/Horde/Horde.Server/Storage/StorageService.cs
+++ b/Engine/Source/Programs/Horde/Horde.Server/Storage/StorageService.cs
@@ -233,7 +233,7 @@ namespace Horde.Server.Storage
 			{
 				await foreach (BundleNodeLocator locator in _outer.FindAliasesAsync(NamespaceId, alias, cancellationToken))
 				{
-					yield return new FlushedNodeHandle(TreeReader, locator);
+					yield return new FlushedNodeHandle(BundleReader, locator);
 				}
 			}
 
@@ -249,7 +249,7 @@ namespace Horde.Server.Storage
 				{
 					return null;
 				}
-				return new FlushedNodeHandle(TreeReader, locator.Value);
+				return new FlushedNodeHandle(BundleReader, locator.Value);
 			}
 
 			/// <inheritdoc/>
diff --git a/Engine/Source/Programs/Shared/EpicGames.Horde/Storage/Backends/FileStorageClient.cs b/Engine/Source/Programs/Shared/EpicGames.Horde/Storage/Backends/FileStorageClient.cs
--- a/Engine/Source/Programs/Shared/EpicGames.Horde/Storage/Backends/FileStorageClient.cs
+++ b/Engine/Source/Programs/Shared/EpicGames.Horde/Storage/Backends/FileStorageClient.cs
@@ -46,7 +46,7 @@ namespace EpicGames.Horde.Storage.Backends
 		public async ValueTask<BlobHandle> ReadRefAsync(FileReference file)
 		{
 			string text = await FileReference.ReadAllTextAsync(file);
-			return new FlushedNodeHandle(TreeReader, BundleNodeLocator.Parse(text));
+			return new FlushedNodeHandle(BundleReader, BundleNodeLocator.Parse(text));
 		}
 
 		FileReference GetRefFile(RefName name) => FileReference.Combine(_rootDir, name.ToString() + ".ref");
@@ -146,7 +146,7 @@ namespace EpicGames.Horde.Storage.Backends
 
 			_logger.LogInformation("Reading {File}", file);
 			string text = await FileReference.ReadAllTextAsync(file, cancellationToken);
-			return new FlushedNodeHandle(TreeReader, BundleNodeLocator.Parse(text));
+			return new FlushedNodeHandle(BundleReader, BundleNodeLocator.Parse(text));
 		}
 
 		/// <inheritdoc/>
diff --git a/Engine/Source/Programs/Shared/EpicGames.Horde/Storage/Backends/HttpStorageClient.cs b/Engine/Source/Programs/Shared/EpicGames.Horde/Storage/Backends/HttpStorageClient.cs
--- a/Engine/Source/Programs/Shared/EpicGames.Horde/Storage/Backends/HttpStorageClient.cs
+++ b/Engine/Source/Programs/Shared/EpicGames.Horde/Storage/Backends/HttpStorageClient.cs
@@ -226,7 +226,7 @@ namespace EpicGames.Horde.Storage.Backends
 						FindNodesResponse? message = await response.Content.ReadFromJsonAsync<FindNodesResponse>(cancellationToken: cancellationToken);
 						foreach (FindNodeResponse node in message!.Nodes)
 						{
-							yield return new FlushedNodeHandle(TreeReader, new BundleNodeLocator(node.Hash, node.Blob, node.ExportIdx));
+							yield return new FlushedNodeHandle(BundleReader, new BundleNodeLocator(node.Hash, node.Blob, node.ExportIdx));
 						}
 					}
 				}
@@ -292,7 +292,7 @@ namespace EpicGames.Horde.Storage.Backends
 							response.EnsureSuccessStatusCode();
 							ReadRefResponse? data = await response.Content.ReadFromJsonAsync<ReadRefResponse>(cancellationToken: cancellationToken);
 							_logger.LogDebug("Read ref {RefName} -> {Blob}#{ExportIdx}", name, data!.Blob, data!.ExportIdx);
-							return new FlushedNodeHandle(TreeReader, new BundleNodeLocator(data.Hash, data!.Blob, data!.ExportIdx));
+							return new FlushedNodeHandle(BundleReader, new BundleNodeLocator(data.Hash, data!.Blob, data!.ExportIdx));
 						}
 					}
 				}
diff --git a/Engine/Source/Programs/Shared/EpicGames.Horde/Storage/Backends/MemoryStorageClient.cs b/Engine/Source/Programs/Shared/EpicGames.Horde/Storage/Backends/MemoryStorageClient.cs
--- a/Engine/Source/Programs/Shared/EpicGames.Horde/Storage/Backends/MemoryStorageClient.cs
+++ b/Engine/Source/Programs/Shared/EpicGames.Horde/Storage/Backends/MemoryStorageClient.cs
@@ -124,7 +124,7 @@ namespace EpicGames.Horde.Storage.Backends
 			BundleNodeLocator hashedLocator;
 			if (_refs.TryGetValue(name, out hashedLocator))
 			{
-				return Task.FromResult<BundleNodeHandle?>(new FlushedNodeHandle(TreeReader, hashedLocator)); 
+				return Task.FromResult<BundleNodeHandle?>(new FlushedNodeHandle(BundleReader, hashedLocator)); 
 			}
 			else
 			{
diff --git a/Engine/Source/Programs/Shared/EpicGames.Horde/Storage/Bundles/BundleStorageClient.cs b/Engine/Source/Programs/Shared/EpicGames.Horde/Storage/Bundles/BundleStorageClient.cs
--- a/Engine/Source/Programs/Shared/EpicGames.Horde/Storage/Bundles/BundleStorageClient.cs
+++ b/Engine/Source/Programs/Shared/EpicGames.Horde/Storage/Bundles/BundleStorageClient.cs
@@ -2,7 +2,6 @@
 
 using System;
 using System.Collections.Generic;
-using System.IO;
 using System.Runtime.CompilerServices;
 using System.Threading;
 using System.Threading.Tasks;
@@ -20,14 +19,14 @@ namespace EpicGames.Horde.Storage.Bundles
 		/// <summary>
 		/// Reader for node data
 		/// </summary>
-		protected BundleReader TreeReader { get; }
+		protected BundleReader BundleReader { get; }
 
 		/// <summary>
 		/// Constructor
 		/// </summary>
 		protected BundleStorageClient(IMemoryCache? memoryCache, ILogger logger)
 		{
-			TreeReader = new BundleReader(this, memoryCache, logger);
+			BundleReader = new BundleReader(this, memoryCache, logger);
 		}
 
 		#region Blobs
@@ -63,7 +62,7 @@ namespace EpicGames.Horde.Storage.Bundles
 		/// <inheritdoc/>
 		public BundleWriter CreateWriter(RefName refName = default, BundleOptions? options = null)
 		{
-			return new BundleWriter(this, TreeReader, refName, options);
+			return new BundleWriter(this, BundleReader, refName, options);
 		}
 
 		/// <inheritdoc/>
