Horde: Add an optional memory cache parameter to FileStorageClient.

diff --git a/Engine/Source/Programs/Horde/Horde.Agent/Utility/StorageClientFactory.cs b/Engine/Source/Programs/Horde/Horde.Agent/Utility/StorageClientFactory.cs
--- a/Engine/Source/Programs/Horde/Horde.Agent/Utility/StorageClientFactory.cs
+++ b/Engine/Source/Programs/Horde/Horde.Agent/Utility/StorageClientFactory.cs
@@ -39,7 +39,7 @@ namespace Horde.Agent.Utility
 			IStorageClient client;
 			if (_settings.Value.UseLocalStorageClient)
 			{
-				client = new FileStorageClient(DirectoryReference.Combine(Program.DataDir, "Storage", namespaceId.ToString()), _logger);
+				client = new FileStorageClient(DirectoryReference.Combine(Program.DataDir, "Storage", namespaceId.ToString()), _memoryCache, _logger);
 			}
 			else
 			{
diff --git a/Engine/Source/Programs/Horde/Horde.Server/Commands/Bundles/CreateCommand.cs b/Engine/Source/Programs/Horde/Horde.Server/Commands/Bundles/CreateCommand.cs
--- a/Engine/Source/Programs/Horde/Horde.Server/Commands/Bundles/CreateCommand.cs
+++ b/Engine/Source/Programs/Horde/Horde.Server/Commands/Bundles/CreateCommand.cs
@@ -8,6 +8,7 @@ using EpicGames.Horde.Storage.Backends;
 using EpicGames.Horde.Storage.Nodes;
 using Horde.Server.Perforce;
 using Horde.Server.Storage;
+using Microsoft.Extensions.Caching.Memory;
 using Microsoft.Extensions.Configuration;
 using Microsoft.Extensions.DependencyInjection;
 using Microsoft.Extensions.Logging;
@@ -50,7 +51,7 @@ namespace Horde.Server.Commands.Bundles
 			}
 			else
 			{
-				FileStorageClient fileStorageClient = new FileStorageClient(OutputDir, logger);
+				FileStorageClient fileStorageClient = new FileStorageClient(OutputDir, serviceProvider.GetRequiredService<IMemoryCache>(), logger);
 				storageClient = fileStorageClient;
 			}
 		
diff --git a/Engine/Source/Programs/Horde/Horde.Server/Commands/Install/SetupCommand.cs b/Engine/Source/Programs/Horde/Horde.Server/Commands/Install/SetupCommand.cs
--- a/Engine/Source/Programs/Horde/Horde.Server/Commands/Install/SetupCommand.cs
+++ b/Engine/Source/Programs/Horde/Horde.Server/Commands/Install/SetupCommand.cs
@@ -18,6 +18,7 @@ using EpicGames.Horde.Storage.Backends;
 using EpicGames.Horde.Storage.Nodes;
 using Horde.Server.Agents;
 using Horde.Server.Utilities;
+using Microsoft.Extensions.Caching.Memory;
 using Microsoft.Extensions.Logging;
 using Microsoft.IdentityModel.Tokens;
 
@@ -66,7 +67,9 @@ namespace Horde.Server.Commands.Install
 
 			RefName refName = new RefName("latest");
 
-			FileStorageClient client = new FileStorageClient(bundleDir, logger);
+			using IMemoryCache memoryCache = new MemoryCache(new MemoryCacheOptions());
+
+			FileStorageClient client = new FileStorageClient(bundleDir, memoryCache, logger);
 			await using (IStorageWriter writer = client.CreateWriter(refName))
 			{
 				DirectoryNode dirNode = new DirectoryNode();
diff --git a/Engine/Source/Programs/Horde/Horde.Server/Tools/ToolCollection.cs b/Engine/Source/Programs/Horde/Horde.Server/Tools/ToolCollection.cs
--- a/Engine/Source/Programs/Horde/Horde.Server/Tools/ToolCollection.cs
+++ b/Engine/Source/Programs/Horde/Horde.Server/Tools/ToolCollection.cs
@@ -418,7 +418,7 @@ namespace Horde.Server.Tools
 		{
 			if (tool.Config is BundledToolConfig bundledConfig)
 			{
-				return new FileStorageClient(DirectoryReference.Combine(Program.AppDir, bundledConfig.DataDir ?? $"tools/{tool.Id}"), _logger);
+				return new FileStorageClient(DirectoryReference.Combine(Program.AppDir, bundledConfig.DataDir ?? $"tools/{tool.Id}"), _cache, _logger);
 			}
 			else
 			{
diff --git a/Engine/Source/Programs/Horde/Horde/Commands/Bundles/BundleCreate.cs b/Engine/Source/Programs/Horde/Horde/Commands/Bundles/BundleCreate.cs
--- a/Engine/Source/Programs/Horde/Horde/Commands/Bundles/BundleCreate.cs
+++ b/Engine/Source/Programs/Horde/Horde/Commands/Bundles/BundleCreate.cs
@@ -5,6 +5,7 @@ using EpicGames.Core;
 using EpicGames.Horde.Storage;
 using EpicGames.Horde.Storage.Backends;
 using EpicGames.Horde.Storage.Nodes;
+using Microsoft.Extensions.Caching.Memory;
 using Microsoft.Extensions.Logging;
 
 namespace Horde.Commands.Bundles
@@ -26,11 +27,13 @@ namespace Horde.Commands.Bundles
 
 		public override async Task<int> ExecuteAsync(ILogger logger)
 		{
+			using IMemoryCache cache = new MemoryCache(new MemoryCacheOptions { SizeLimit = 256 * 1024 * 1024 });
+
 			// Create the storage client
 			IStorageClient store;
 			if (File != null)
 			{
-				store = new FileStorageClient(File.Directory, logger);
+				store = new FileStorageClient(File.Directory, cache, logger);
 			}
 			else if (!String.IsNullOrEmpty(Ref))
 			{
diff --git a/Engine/Source/Programs/Horde/Horde/Commands/Bundles/BundleExtract.cs b/Engine/Source/Programs/Horde/Horde/Commands/Bundles/BundleExtract.cs
--- a/Engine/Source/Programs/Horde/Horde/Commands/Bundles/BundleExtract.cs
+++ b/Engine/Source/Programs/Horde/Horde/Commands/Bundles/BundleExtract.cs
@@ -24,11 +24,13 @@ namespace Horde.Commands.Bundles
 
 		public override async Task<int> ExecuteAsync(ILogger logger)
 		{
+			using IMemoryCache cache = new MemoryCache(new MemoryCacheOptions { SizeLimit = 256 * 1024 * 1024 });
+
 			IStorageClient store;
 			BlobHandle handle;
 			if (File != null)
 			{
-				store = new FileStorageClient(File.Directory, logger);
+				store = new FileStorageClient(File.Directory, cache, logger);
 				handle = await ((FileStorageClient)store).ReadRefAsync(File);
 			}
 			else if (Ref != null)
@@ -41,8 +43,6 @@ namespace Horde.Commands.Bundles
 				throw new CommandLineArgumentException("Either -File=... or -Ref=... must be specified");
 			}
 
-			using MemoryCache cache = new MemoryCache(new MemoryCacheOptions());
-
 			Stopwatch timer = Stopwatch.StartNew();
 
 			DirectoryNode node = await handle.ReadNodeAsync<DirectoryNode>();
diff --git a/Engine/Source/Programs/Shared/EpicGames.Core/StreamUtils.cs b/Engine/Source/Programs/Shared/EpicGames.Core/StreamUtils.cs
--- a/Engine/Source/Programs/Shared/EpicGames.Core/StreamUtils.cs
+++ b/Engine/Source/Programs/Shared/EpicGames.Core/StreamUtils.cs
@@ -56,9 +56,9 @@ namespace EpicGames.Core
 		/// <summary>
 		/// Read the contents of a file async using double buffering
 		/// </summary>
-		/// <param name="stream">Data to compute the hash for</param>
+		/// <param name="stream">Stream to read from</param>
 		/// <param name="cancellationToken">Cancellation token used to terminate processing</param>
-		/// <returns>New content hash instance containing the hash of the data</returns>
+		/// <returns>Contents of the stream</returns>
 		public static async Task<byte[]> ReadAllBytesAsync(this Stream stream, CancellationToken cancellationToken = default)
 		{
 			using MemoryStream memoryStream = new MemoryStream();
@@ -69,13 +69,13 @@ namespace EpicGames.Core
 		/// <summary>
 		/// Read the contents of a file async using double buffering
 		/// </summary>
-		/// <param name="stream">Data to compute the hash for</param>
+		/// <param name="stream">Stream to read from</param>
 		/// <param name="fileSizeHint">If available, the file size so an appropriate buffer size can be used</param>
 		/// <param name="minBufferSize">Minimum size of the buffer</param>
 		/// <param name="maxBufferSize">Maximum size of the buffer</param>
 		/// <param name="callback">Callback used to send read data back to the caller</param>
 		/// <param name="cancellationToken">Cancellation token used to terminate processing</param>
-		/// <returns>New content hash instance containing the hash of the data</returns>
+		/// <returns>Contents of the stream</returns>
 		public static async Task ReadAllBytesAsync(this Stream stream, long fileSizeHint, int minBufferSize, int maxBufferSize, Func<ReadOnlyMemory<byte>, Task> callback, CancellationToken cancellationToken = default)
 		{
 			int bufferLength = minBufferSize;
diff --git a/Engine/Source/Programs/Shared/EpicGames.Horde/Storage/Backends/FileStorageClient.cs b/Engine/Source/Programs/Shared/EpicGames.Horde/Storage/Backends/FileStorageClient.cs
--- a/Engine/Source/Programs/Shared/EpicGames.Horde/Storage/Backends/FileStorageClient.cs
+++ b/Engine/Source/Programs/Shared/EpicGames.Horde/Storage/Backends/FileStorageClient.cs
@@ -14,7 +14,6 @@ using EpicGames.Core;
 using Microsoft.Extensions.Caching.Memory;
 using Microsoft.Extensions.Logging;
 using EpicGames.Horde.Storage.Bundles;
-using Microsoft.CodeAnalysis.FlowAnalysis;
 
 namespace EpicGames.Horde.Storage.Backends
 {
@@ -30,9 +29,10 @@ namespace EpicGames.Horde.Storage.Backends
 		/// Constructor
 		/// </summary>
 		/// <param name="rootDir">Root directory for storing blobs</param>
+		/// <param name="cache">Memory cache for read data</param>
 		/// <param name="logger">Logger interface</param>
-		public FileStorageClient(DirectoryReference rootDir, ILogger logger)
-			: base(null, logger)
+		public FileStorageClient(DirectoryReference rootDir, IMemoryCache? cache, ILogger logger)
+			: base(cache, logger)
 		{
 			_rootDir = rootDir;
 			_logger = logger;
