Horde: Expose stream API from BundleStorageClient. This avoids a lot of copying into temporary buffers when splitting data to store in the cache.

diff --git a/Engine/Source/Programs/Horde/Horde.Server.Tests/GcServiceTests.cs b/Engine/Source/Programs/Horde/Horde.Server.Tests/GcServiceTests.cs
--- a/Engine/Source/Programs/Horde/Horde.Server.Tests/GcServiceTests.cs
+++ b/Engine/Source/Programs/Horde/Horde.Server.Tests/GcServiceTests.cs
@@ -47,7 +47,7 @@ namespace Horde.Server.Tests
 				}
 			}
 
-			HashSet<BundleLocator> nodes = await FindNodes(store, roots);
+			HashSet<BundleLocator> nodes = await FindNodesAsync(store, roots);
 
 			await Clock.AdvanceAsync(TimeSpan.FromDays(1.0));
 
@@ -56,21 +56,21 @@ namespace Horde.Server.Tests
 			Assert.IsTrue(remaining.All(x => nodes.Contains(x)));
 		}
 
-		async Task<HashSet<BundleLocator>> FindNodes(BundleStorageClient store, IEnumerable<BundleLocator> roots)
+		async Task<HashSet<BundleLocator>> FindNodesAsync(BundleStorageClient store, IEnumerable<BundleLocator> roots)
 		{
 			HashSet<BundleLocator> nodes = new HashSet<BundleLocator>();
-			await FindNodes(store, roots, nodes);
+			await FindNodesAsync(store, roots, nodes);
 			return nodes;
 		}
 
-		async Task FindNodes(BundleStorageClient store, IEnumerable<BundleLocator> roots, HashSet<BundleLocator> nodes)
+		async Task FindNodesAsync(BundleStorageClient store, IEnumerable<BundleLocator> roots, HashSet<BundleLocator> nodes)
 		{
 			foreach (BundleLocator root in roots)
 			{
 				if (nodes.Add(root))
 				{
 					Bundle bundle = await store.ReadBundleAsync(root);
-					await FindNodes(store, bundle.Header.Imports, nodes);
+					await FindNodesAsync(store, bundle.Header.Imports, nodes);
 				}
 			}
 		}
diff --git a/Engine/Source/Programs/Horde/Horde.Server/Storage/StorageController.cs b/Engine/Source/Programs/Horde/Horde.Server/Storage/StorageController.cs
--- a/Engine/Source/Programs/Horde/Horde.Server/Storage/StorageController.cs
+++ b/Engine/Source/Programs/Horde/Horde.Server/Storage/StorageController.cs
@@ -285,7 +285,7 @@ namespace Horde.Server.Storage
 
 			// Parse the range header
 			int offset = 0;
-			int? length = null;
+			int length = 0;
 
 			if (headers.Range.Count > 0)
 			{
@@ -317,21 +317,16 @@ namespace Horde.Server.Storage
 					{
 						return new BadRequestObjectResult(LogEvent.Create(LogLevel.Error, "Unable to parse end for range: {Value}", value));
 					}
+
+					if (length == 0)
+					{
+						return new FileStreamResult(new MemoryStream(Array.Empty<byte>()), "application/octet-stream");
+					}
 				}
 			}
 
 #pragma warning disable CA2000 // Dispose objects before losing scope
-			Stream stream;
-			if (offset == 0 && length == null)
-			{
-				Bundle bundle = await storageClient.ReadBundleAsync(locator, cancellationToken);
-				stream = new ReadOnlySequenceStream(bundle.AsSequence());
-			}
-			else
-			{
-				ReadOnlyMemory<byte> memory = await storageClient.ReadBundleRangeAsync(locator, offset, length, cancellationToken);
-				stream = new ReadOnlyMemoryStream(memory);
-			}
+			Stream stream = await storageClient.OpenAsync(locator, offset, length, cancellationToken);
 			return new FileStreamResult(stream, "application/octet-stream");
 #pragma warning restore CA2000 // Dispose objects before losing scope
 		}
diff --git a/Engine/Source/Programs/Horde/Horde.Server/Storage/StorageService.cs b/Engine/Source/Programs/Horde/Horde.Server/Storage/StorageService.cs
--- a/Engine/Source/Programs/Horde/Horde.Server/Storage/StorageService.cs
+++ b/Engine/Source/Programs/Horde/Horde.Server/Storage/StorageService.cs
@@ -141,32 +141,19 @@ namespace Horde.Server.Storage
 
 			#region Blobs
 
-			/// <inheritdoc/>
-			public override async Task<Bundle> ReadBundleAsync(BundleLocator locator, CancellationToken cancellationToken = default)
-			{
-				using TelemetrySpan span = _tracer.StartActiveSpan($"{nameof(StorageService)}.{nameof(StorageClientImpl)}.{nameof(ReadBundleAsync)}");
-				span.SetAttribute("locator", locator.ToString());
-
-				string path = GetBlobPath(locator);
-				await using Stream stream = await Backend.ReadAsync(path, cancellationToken);
-				return await Bundle.FromStreamAsync(stream, cancellationToken);
-			}
-
 			/// <inheritdoc/>
 			public override ValueTask<Uri?> GetReadRedirectAsync(BundleLocator locator, CancellationToken cancellationToken = default) => Backend.TryGetReadRedirectAsync(GetBlobPath(locator), cancellationToken);
 
 			/// <inheritdoc/>
-			public override async Task<ReadOnlyMemory<byte>> ReadBundleRangeAsync(BundleLocator locator, int offset, int? length, CancellationToken cancellationToken = default)
+			public override async Task<Stream> OpenAsync(BundleLocator locator, int offset, int length, CancellationToken cancellationToken = default)
 			{
-				using TelemetrySpan span = _tracer.StartActiveSpan($"{nameof(StorageService)}.{nameof(StorageClientImpl)}.{nameof(ReadBundleRangeAsync)}");
+				using TelemetrySpan span = _tracer.StartActiveSpan($"{nameof(StorageService)}.{nameof(StorageClientImpl)}.{nameof(OpenAsync)}");
 				span.SetAttribute("locator", locator.ToString());
 				span.SetAttribute("offset", offset);
 				span.SetAttribute("length", length);
 
 				string path = GetBlobPath(locator);
-				await using Stream stream = await Backend.ReadAsync(path, offset, length, cancellationToken);
-
-				return await stream.ReadAllBytesAsync(cancellationToken);
+				return await Backend.ReadAsync(path, offset, length, cancellationToken);
 			}
 
 			/// <inheritdoc/>
@@ -1067,18 +1054,19 @@ namespace Horde.Server.Storage
 			for (; ; )
 			{
 				// Read the start of the blob
-				ReadOnlyMemory<byte> memory = await store.ReadBundleRangeAsync(locator, 0, fetchSize, cancellationToken);
-				if (memory.Length < BundleHeader.PreludeLength)
-				{
-					_logger.LogError("Blob {Blob} does not have a valid prelude", locator);
-					return null;
-				}
+				Stream stream = await store.OpenAsync(locator, 0, fetchSize, cancellationToken);
+
+				byte[] prelude = new byte[BundleHeader.PreludeLength];
+				await stream.ReadFixedLengthBytesAsync(prelude, cancellationToken);
 
 				// Make sure it's large enough to hold the header
-				int headerSize = BundleHeader.ReadPrelude(memory.Span);
+				int headerSize = BundleHeader.ReadPrelude(prelude);
 				if (headerSize <= fetchSize)
 				{
-					return BundleHeader.Read(memory);
+					byte[] header = new byte[headerSize];
+					prelude.AsSpan().CopyTo(header);
+					await stream.ReadFixedLengthBytesAsync(header.AsMemory(BundleHeader.PreludeLength), cancellationToken);
+					return BundleHeader.Read(header);
 				}
 
 				// Increase the fetch size and retry
diff --git a/Engine/Source/Programs/Shared/EpicGames.Horde/Compute/AgentMessage.cs b/Engine/Source/Programs/Shared/EpicGames.Horde/Compute/AgentMessage.cs
--- a/Engine/Source/Programs/Shared/EpicGames.Horde/Compute/AgentMessage.cs
+++ b/Engine/Source/Programs/Shared/EpicGames.Horde/Compute/AgentMessage.cs
@@ -578,13 +578,13 @@ namespace EpicGames.Horde.Compute
 		/// <param name="length">Length of data to return</param>
 		/// <param name="cancellationToken">Cancellation token for the operation</param>
 		/// <returns>Stream containing the blob data</returns>
-		public static async Task<ReadOnlyMemory<byte>> ReadBlobAsync(this AgentMessageChannel channel, BundleLocator locator, int offset, int? length, CancellationToken cancellationToken = default)
+		public static async Task<ReadOnlyMemory<byte>> ReadBlobAsync(this AgentMessageChannel channel, BundleLocator locator, int offset, int length, CancellationToken cancellationToken = default)
 		{
 			using (IAgentMessageBuilder request = await channel.CreateMessageAsync(AgentMessageType.ReadBlob, cancellationToken))
 			{
 				request.WriteBlobLocator(locator);
 				request.WriteUnsignedVarInt(offset);
-				request.WriteUnsignedVarInt(length ?? Int32.MaxValue);
+				request.WriteUnsignedVarInt(length);
 				request.Send();
 			}
 
@@ -645,27 +645,20 @@ namespace EpicGames.Horde.Compute
 		/// <param name="cancellationToken">Cancellation token for the operation</param>
 		public static async Task SendBlobDataAsync(this AgentMessageChannel channel, BundleLocator locator, int offset, int length, BundleStorageClient storage, CancellationToken cancellationToken = default)
 		{
-			ReadOnlySequence<byte> data;
-			if (offset == 0 && length == 0)
-			{
-				Bundle bundle = await storage.ReadBundleAsync(locator, cancellationToken);
-				data = bundle.AsSequence();
-			}
-			else
-			{
-				ReadOnlyMemory<byte> range = await storage.ReadBundleRangeAsync(locator, offset, length, cancellationToken);
-				data = new ReadOnlySequence<byte>(range);
-			}
+			using Stream stream = await storage.OpenAsync(locator, offset, length, cancellationToken);
 
 			const int MaxChunkSize = 512 * 1024;
-			for (int chunkOffset = 0; chunkOffset < data.Length;)
+			for (int chunkOffset = 0; chunkOffset < stream.Length;)
 			{
-				int chunkLength = (int)Math.Min(data.Length - chunkOffset, MaxChunkSize);
+				int chunkLength = (int)Math.Min(stream.Length - chunkOffset, MaxChunkSize);
 				using (IAgentMessageBuilder response = await channel.CreateMessageAsync(AgentMessageType.ReadBlobResponse, chunkLength + 128, cancellationToken))
 				{
 					response.WriteInt32(chunkOffset);
-					response.WriteInt32((int)data.Length);
-					response.WriteFixedLengthBytes(data.Slice(chunkOffset, chunkLength));
+					response.WriteInt32((int)stream.Length);
+
+					Memory<byte> memory = response.GetMemoryAndAdvance(chunkLength);
+					await stream.ReadFixedLengthBytesAsync(memory, cancellationToken);
+
 					response.Send();
 				}
 				chunkOffset += chunkLength;
diff --git a/Engine/Source/Programs/Shared/EpicGames.Horde/Compute/AgentStorageClient.cs b/Engine/Source/Programs/Shared/EpicGames.Horde/Compute/AgentStorageClient.cs
--- a/Engine/Source/Programs/Shared/EpicGames.Horde/Compute/AgentStorageClient.cs
+++ b/Engine/Source/Programs/Shared/EpicGames.Horde/Compute/AgentStorageClient.cs
@@ -40,19 +40,13 @@ namespace EpicGames.Horde.Compute
 		#region Blobs
 
 		/// <inheritdoc/>
-		public override async Task<Bundle> ReadBundleAsync(BundleLocator locator, CancellationToken cancellationToken = default)
-		{
-			ReadOnlyMemory<byte> data = await ReadBundleRangeAsync(locator, 0, 0, cancellationToken);
-			return new Bundle(data);
-		}
-
-		/// <inheritdoc/>
-		public override async Task<ReadOnlyMemory<byte>> ReadBundleRangeAsync(BundleLocator locator, int offset, int? length, CancellationToken cancellationToken = default)
+		public override async Task<Stream> OpenAsync(BundleLocator locator, int offset, int length, CancellationToken cancellationToken = default)
 		{
 			await _semaphore.WaitAsync(cancellationToken);
 			try
 			{
-				return await _channel.ReadBlobAsync(locator, offset, length, cancellationToken);
+				ReadOnlyMemory<byte> data = await _channel.ReadBlobAsync(locator, offset, length, cancellationToken);
+				return new ReadOnlyMemoryStream(data);
 			}
 			finally
 			{
diff --git a/Engine/Source/Programs/Shared/EpicGames.Horde/Storage/Backends/FileStorageClient.cs b/Engine/Source/Programs/Shared/EpicGames.Horde/Storage/Backends/FileStorageClient.cs
--- a/Engine/Source/Programs/Shared/EpicGames.Horde/Storage/Backends/FileStorageClient.cs
+++ b/Engine/Source/Programs/Shared/EpicGames.Horde/Storage/Backends/FileStorageClient.cs
@@ -55,27 +55,17 @@ namespace EpicGames.Horde.Storage.Backends
 		#region Blobs
 
 		/// <inheritdoc/>
-		public override async Task<Bundle> ReadBundleAsync(BundleLocator id, CancellationToken cancellationToken = default)
+		public override Task<Stream> OpenAsync(BundleLocator id, int offset, int length, CancellationToken cancellationToken = default)
 		{
 			FileReference file = GetBlobFile(id);
 			_logger.LogInformation("Reading {File}", file);
 
-			byte[] data = await FileReference.ReadAllBytesAsync(file, cancellationToken);
-			return new Bundle(data);
-		}
-
-		/// <inheritdoc/>
-		public override async Task<ReadOnlyMemory<byte>> ReadBundleRangeAsync(BundleLocator id, int offset, int? length, CancellationToken cancellationToken = default)
-		{
-			Bundle bundle = await ReadBundleAsync(id, cancellationToken);
-
-			ReadOnlySequence<byte> sequence = bundle.AsSequence().Slice(offset);
-			if (length.HasValue && sequence.Length > length)
+			Stream stream = FileReference.Open(file, FileMode.Open, FileAccess.Read, FileShare.Read | FileShare.Delete);
+			if (offset > 0)
 			{
-				sequence = sequence.Slice(0, length.Value);
+				stream.Seek(offset, SeekOrigin.Begin);
 			}
-
-			return sequence.AsSingleSegment();
+			return Task.FromResult(stream);
 		}
 
 		/// <inheritdoc/>
diff --git a/Engine/Source/Programs/Shared/EpicGames.Horde/Storage/Backends/HttpStorageClient.cs b/Engine/Source/Programs/Shared/EpicGames.Horde/Storage/Backends/HttpStorageClient.cs
--- a/Engine/Source/Programs/Shared/EpicGames.Horde/Storage/Backends/HttpStorageClient.cs
+++ b/Engine/Source/Programs/Shared/EpicGames.Horde/Storage/Backends/HttpStorageClient.cs
@@ -103,36 +103,33 @@ namespace EpicGames.Horde.Storage.Backends
 		#region Blobs
 
 		/// <inheritdoc/>
-		public override async Task<Bundle> ReadBundleAsync(BundleLocator locator, CancellationToken cancellationToken = default)
+		public override async Task<Stream> OpenAsync(BundleLocator locator, int offset, int length, CancellationToken cancellationToken = default)
 		{
-			_logger.LogDebug("Reading {Locator}", locator);
-			using (HttpClient httpClient = _createClient())
+			if (offset == 0 && length == 0)
 			{
-				using (HttpRequestMessage request = new HttpRequestMessage(HttpMethod.Get, $"blobs/{locator}"))
-				{
-					HttpResponseMessage response = await httpClient.SendAsync(request, cancellationToken);
-					response.EnsureSuccessStatusCode();
-
-					byte[] data = await response.Content.ReadAsByteArrayAsync(cancellationToken);
-					return new Bundle(data);
-				}
+				_logger.LogDebug("Reading {Locator}", locator);
+			}
+			else if (length == 0)
+			{
+				_logger.LogDebug("Reading {Locator} ({Offset}..)", locator, offset);
+			}
+			else
+			{
+				_logger.LogDebug("Reading {Locator} ({Offset}+{Length})", locator, offset, length);
 			}
-		}
 
-		/// <inheritdoc/>
-		public override async Task<ReadOnlyMemory<byte>> ReadBundleRangeAsync(BundleLocator locator, int offset, int? length, CancellationToken cancellationToken = default)
-		{
-			_logger.LogDebug("Reading {Locator} ({Offset}+{Length})", locator, offset, length);
 			using (HttpClient httpClient = _createClient())
 			{
 				using (HttpRequestMessage request = new HttpRequestMessage(HttpMethod.Get, $"blobs/{locator}"))
 				{
-					request.Headers.Range = new RangeHeaderValue(offset, (length == null)? null : (offset + (length.Value - 1)));
+					if (offset != 0 || length != 0)
+					{
+						request.Headers.Range = new RangeHeaderValue(offset, (length == 0) ? null : (offset + (length - 1)));
+					}
 
 					HttpResponseMessage response = await httpClient.SendAsync(request, cancellationToken);
 					response.EnsureSuccessStatusCode();
-
-					return await response.Content.ReadAsByteArrayAsync(cancellationToken);
+					return await response.Content.ReadAsStreamAsync(cancellationToken);
 				}
 			}
 		}
diff --git a/Engine/Source/Programs/Shared/EpicGames.Horde/Storage/Backends/MemoryStorageClient.cs b/Engine/Source/Programs/Shared/EpicGames.Horde/Storage/Backends/MemoryStorageClient.cs
--- a/Engine/Source/Programs/Shared/EpicGames.Horde/Storage/Backends/MemoryStorageClient.cs
+++ b/Engine/Source/Programs/Shared/EpicGames.Horde/Storage/Backends/MemoryStorageClient.cs
@@ -55,21 +55,14 @@ namespace EpicGames.Horde.Storage.Backends
 		#region Blobs
 
 		/// <inheritdoc/>
-		public override Task<Bundle> ReadBundleAsync(BundleLocator locator, CancellationToken cancellationToken = default)
-		{
-			Bundle bundle = _bundles[locator];
-			return Task.FromResult(bundle);
-		}
-
-		/// <inheritdoc/>
-		public override Task<ReadOnlyMemory<byte>> ReadBundleRangeAsync(BundleLocator locator, int offset, int? length, CancellationToken cancellationToken = default)
+		public override Task<Stream> OpenAsync(BundleLocator locator, int offset, int length, CancellationToken cancellationToken = default)
 		{
 			ReadOnlySequence<byte> sequence = _bundles[locator].AsSequence().Slice(offset);
-			if (length != null && sequence.Length > length)
+			if (length > 0 && sequence.Length > length)
 			{
-				sequence = sequence.Slice(0, length.Value);
+				sequence = sequence.Slice(0, length);
 			}
-			return Task.FromResult(sequence.AsSingleSegment());
+			return Task.FromResult<Stream>(new ReadOnlySequenceStream(sequence));
 		}
 
 		/// <inheritdoc/>
diff --git a/Engine/Source/Programs/Shared/EpicGames.Horde/Storage/Bundles/Bundle.cs b/Engine/Source/Programs/Shared/EpicGames.Horde/Storage/Bundles/Bundle.cs
--- a/Engine/Source/Programs/Shared/EpicGames.Horde/Storage/Bundles/Bundle.cs
+++ b/Engine/Source/Programs/Shared/EpicGames.Horde/Storage/Bundles/Bundle.cs
@@ -2,18 +2,15 @@
 
 using EpicGames.Core;
 using K4os.Compression.LZ4;
-using Microsoft.CodeAnalysis.CSharp.Syntax;
 using System;
 using System.Buffers;
 using System.Buffers.Binary;
 using System.Collections;
 using System.Collections.Generic;
 using System.Diagnostics;
-using System.Globalization;
 using System.IO;
 using System.IO.Compression;
 using System.Linq;
-using System.Runtime.CompilerServices;
 using System.Text;
 using System.Threading;
 using System.Threading.Tasks;
@@ -118,7 +115,7 @@ namespace EpicGames.Horde.Storage.Bundles
 		/// <param name="stream">Stream to read from</param>
 		/// <param name="cancellationToken">Cancellation token for the operation</param>
 		/// <returns>Bundle that was read</returns>
-		public static async Task<Bundle> FromStreamAsync(Stream stream, CancellationToken cancellationToken)
+		public static async Task<Bundle> FromStreamAsync(Stream stream, CancellationToken cancellationToken = default)
 		{
 			BundleHeader header = await BundleHeader.FromStreamAsync(stream, cancellationToken);
 
diff --git a/Engine/Source/Programs/Shared/EpicGames.Horde/Storage/Bundles/BundleReader.cs b/Engine/Source/Programs/Shared/EpicGames.Horde/Storage/Bundles/BundleReader.cs
--- a/Engine/Source/Programs/Shared/EpicGames.Horde/Storage/Bundles/BundleReader.cs
+++ b/Engine/Source/Programs/Shared/EpicGames.Horde/Storage/Bundles/BundleReader.cs
@@ -11,6 +11,7 @@ using Microsoft.Extensions.Caching.Memory;
 using Microsoft.Extensions.Logging;
 using Microsoft.CodeAnalysis;
 using System.Diagnostics;
+using System.IO;
 
 namespace EpicGames.Horde.Storage.Bundles
 {
@@ -228,14 +229,14 @@ namespace EpicGames.Horde.Storage.Bundles
 			int prefetchSize = _cache != null ? DefaultFetchSize : DefaultUncachedFetchSize;
 			for (; ; )
 			{
-				using (IMemoryOwner<byte> owner = MemoryPool<byte>.Shared.Rent(prefetchSize))
+				await using (Stream stream = await _store.OpenAsync(queuedHeader.Blob, 0, prefetchSize, cancellationToken))
 				{
-					// Read the prefetch size from the blob
-					Memory<byte> memory = owner.Memory.Slice(0, prefetchSize);
-					memory = await _store.ReadBundleRangeAsync(queuedHeader.Blob, 0, memory, cancellationToken);
+					// Read the header data
+					byte[] prelude = new byte[BundleHeader.PreludeLength];
+					await stream.ReadFixedLengthBytesAsync(prelude, cancellationToken);
 
-					// Make sure it's large enough to hold the header
-					int headerSize = BundleHeader.ReadPrelude(memory.Span);
+					// Make sure we've read enough to hold the header
+					int headerSize = BundleHeader.ReadPrelude(prelude);
 					if (headerSize > prefetchSize)
 					{
 						prefetchSize = headerSize;
@@ -243,7 +244,7 @@ namespace EpicGames.Horde.Storage.Bundles
 					}
 
 					// Parse the header and construct the bundle info from it
-					BundleHeader header = BundleHeader.Read(memory.ToArray());
+					BundleHeader header = await ReadHeaderAsync(prelude, stream, headerSize, cancellationToken);
 
 					// Construct the bundle info
 					BundleInfo bundleInfo = new BundleInfo(queuedHeader.Blob, header, headerSize);
@@ -251,20 +252,18 @@ namespace EpicGames.Horde.Storage.Bundles
 					if (_cache != null)
 					{
 						// Also add any encoded packets we prefetched
-						List<ReadOnlyMemory<byte>> packets = new List<ReadOnlyMemory<byte>>();
-						for (int packetOffset = headerSize; packets.Count < header.Packets.Count;)
+						int packetOffset = headerSize;
+						for (int packetIdx = 0; packetIdx < header.Packets.Count; packetIdx++)
 						{
-							int packetIdx = packets.Count;
-
 							int packetLength = header.Packets[packetIdx].EncodedLength;
-							if (packetOffset + packetLength > memory.Length)
+							if (packetOffset + packetLength > prefetchSize)
 							{
 								break;
 							}
 
-							ReadOnlyMemory<byte> packetData = memory.Slice(packetOffset, packetLength).ToArray();
+							byte[] packetData = await ReadPacketAsync(stream, packetLength, cancellationToken);
 							AddToCache(GetEncodedPacketCacheKey(queuedHeader.Blob, packetIdx), packetData, packetData.Length);
-							packets.Add(packetData);
+
 							packetOffset += packetLength;
 						}
 
@@ -327,19 +326,16 @@ namespace EpicGames.Horde.Storage.Bundles
 				readLength = nextReadLength;
 			}
 
-			using (IMemoryOwner<byte> owner = MemoryPool<byte>.Shared.Rent(readLength))
+			await using (Stream stream = await _store.OpenAsync(bundleInfo.Locator, bundleInfo.HeaderLength + bundleInfo.Header.Packets[minPacketIdx].EncodedOffset, readLength, cancellationToken))
 			{
-				Memory<byte> buffer = owner.Memory.Slice(0, readLength);
-				buffer = await _store.ReadBundleRangeAsync(bundleInfo.Locator, bundleInfo.HeaderLength + bundleInfo.Header.Packets[minPacketIdx].EncodedOffset, buffer, cancellationToken);
-
 				// Copy all the packets that have been read into separate buffers, so we can cache them individually.
-				ReadOnlyMemory<byte>?[] packets = new ReadOnlyMemory<byte>?[maxPacketIdx - minPacketIdx];
+				ReadOnlyMemory<byte>[] packets = new ReadOnlyMemory<byte>[maxPacketIdx - minPacketIdx];
 				if (_cache != null)
 				{
 					for (int idx = minPacketIdx; idx < maxPacketIdx; idx++)
 					{
 						string cacheKey = GetEncodedPacketCacheKey(bundleInfo.Locator, idx);
-						ReadOnlyMemory<byte> data = buffer.Slice(bundleInfo.Header.Packets[idx].EncodedOffset - bundleInfo.Header.Packets[minPacketIdx].EncodedOffset, bundleInfo.Header.Packets[idx].EncodedLength).ToArray();
+						byte[] data = await ReadPacketAsync(stream, bundleInfo.Header.Packets[idx].EncodedLength, cancellationToken);
 						packets[idx - minPacketIdx] = data;
 						AddToCache(cacheKey, data, data.Length);
 					}
@@ -355,13 +351,9 @@ namespace EpicGames.Horde.Storage.Bundles
 				// Mark them all as complete
 				foreach (QueuedPacket updatePacket in updatePackets)
 				{
-					ReadOnlyMemory<byte>? data = packets[updatePacket.PacketIdx - minPacketIdx];
-					if (data == null)
-					{
-						data = buffer.Slice(bundleInfo.Header.Packets[updatePacket.PacketIdx].EncodedOffset - bundleInfo.Header.Packets[minPacketIdx].EncodedOffset, bundleInfo.Header.Packets[updatePacket.PacketIdx].EncodedLength).ToArray();
+					ReadOnlyMemory<byte> data = packets[updatePacket.PacketIdx - minPacketIdx];
+					updatePacket.CompletionSource.SetResult(data);
 					}
-					updatePacket.CompletionSource.SetResult((ReadOnlyMemory<byte>)data);
-				}
 
 				// Remove all the completed packets from the queue
 				lock (_queueLock)
@@ -371,6 +363,21 @@ namespace EpicGames.Horde.Storage.Bundles
 			}
 		}
 
+		static async Task<BundleHeader> ReadHeaderAsync(byte[] prelude, Stream stream, int headerSize, CancellationToken cancellationToken)
+		{
+			byte[] header = new byte[headerSize];
+			prelude.CopyTo(header, 0);
+			await stream.ReadFixedLengthBytesAsync(header.AsMemory(prelude.Length), cancellationToken);
+			return BundleHeader.Read(header);
+		}
+
+		static async Task<byte[]> ReadPacketAsync(Stream stream, int packetSize, CancellationToken cancellationToken)
+		{
+			byte[] packet = new byte[packetSize];
+			await stream.ReadAsync(packet, cancellationToken);
+			return packet;
+		}
+
 		/// <summary>
 		/// Reads a bundle header from the given blob locator, or retrieves it from the cache
 		/// </summary>
diff --git a/Engine/Source/Programs/Shared/EpicGames.Horde/Storage/Bundles/BundleStorageClient.cs b/Engine/Source/Programs/Shared/EpicGames.Horde/Storage/Bundles/BundleStorageClient.cs
--- a/Engine/Source/Programs/Shared/EpicGames.Horde/Storage/Bundles/BundleStorageClient.cs
+++ b/Engine/Source/Programs/Shared/EpicGames.Horde/Storage/Bundles/BundleStorageClient.cs
@@ -2,6 +2,7 @@
 
 using System;
 using System.Collections.Generic;
+using System.IO;
 using System.Runtime.CompilerServices;
 using System.Threading;
 using System.Threading.Tasks;
@@ -31,25 +32,27 @@ namespace EpicGames.Horde.Storage.Bundles
 
 		#region Blobs
 
-		/// <inheritdoc/>
-		public abstract Task<Bundle> ReadBundleAsync(BundleLocator locator, CancellationToken cancellationToken = default);
+		/// <summary>
+		/// Opens a bundle stream for reading
+		/// </summary>
+		/// <param name="locator">Locator for the bundle</param>
+		/// <param name="cancellationToken">Cancellation token for the operation</param>
+		/// <returns>Stream for reading from the bundle</returns>
+		public Task<Stream> OpenAsync(BundleLocator locator, CancellationToken cancellationToken = default) => OpenAsync(locator, 0, 0, cancellationToken);
 
 		/// <inheritdoc/>
-		public abstract Task<ReadOnlyMemory<byte>> ReadBundleRangeAsync(BundleLocator locator, int offset, int? length, CancellationToken cancellationToken = default);
+		public abstract Task<Stream> OpenAsync(BundleLocator locator, int offset, int length = 0, CancellationToken cancellationToken = default);
 
 		/// <summary>
-		/// Utility method to read a blob into a buffer
+		/// Reads an entire bundle into memory
 		/// </summary>
-		/// <param name="locator">Blob location</param>
-		/// <param name="offset">Offset within the blob</param>
-		/// <param name="memory">Buffer to read into</param>
+		/// <param name="locator">Locator for the bundle</param>
 		/// <param name="cancellationToken">Cancellation token for the operation</param>
-		/// <returns>The data that was read</returns>
-		public async Task<Memory<byte>> ReadBundleRangeAsync(BundleLocator locator, int offset, Memory<byte> memory, CancellationToken cancellationToken = default)
+		/// <returns>Bundle that was read</returns>
+		public async Task<Bundle> ReadBundleAsync(BundleLocator locator, CancellationToken cancellationToken = default)
 		{
-			ReadOnlyMemory<byte> buffer = await ReadBundleRangeAsync(locator, offset, memory.Length, cancellationToken);
-			buffer.CopyTo(memory);
-			return memory.Slice(0, buffer.Length);
+			using Stream stream = await OpenAsync(locator, cancellationToken);
+			return await Bundle.FromStreamAsync(stream, cancellationToken);
 		}
 
 		/// <inheritdoc/>
