Horde: Use 'null' values for length when wanting to create an open-ended stream for bundles.

diff --git a/Engine/Source/Programs/Horde/Horde.Server/Storage/StorageController.cs b/Engine/Source/Programs/Horde/Horde.Server/Storage/StorageController.cs
--- a/Engine/Source/Programs/Horde/Horde.Server/Storage/StorageController.cs
+++ b/Engine/Source/Programs/Horde/Horde.Server/Storage/StorageController.cs
@@ -285,7 +285,7 @@ namespace Horde.Server.Storage
 
 			// Parse the range header
 			int offset = 0;
-			int length = 0;
+			int? length = null;
 
 			if (headers.Range.Count > 0)
 			{
@@ -317,11 +317,6 @@ namespace Horde.Server.Storage
 					{
 						return new BadRequestObjectResult(LogEvent.Create(LogLevel.Error, "Unable to parse end for range: {Value}", value));
 					}
-
-					if (length == 0)
-					{
-						return new FileStreamResult(new MemoryStream(Array.Empty<byte>()), "application/octet-stream");
-					}
 				}
 			}
 
diff --git a/Engine/Source/Programs/Horde/Horde.Server/Storage/StorageService.cs b/Engine/Source/Programs/Horde/Horde.Server/Storage/StorageService.cs
--- a/Engine/Source/Programs/Horde/Horde.Server/Storage/StorageService.cs
+++ b/Engine/Source/Programs/Horde/Horde.Server/Storage/StorageService.cs
@@ -145,13 +145,18 @@ namespace Horde.Server.Storage
 			public override ValueTask<Uri?> GetReadRedirectAsync(BundleLocator locator, CancellationToken cancellationToken = default) => Backend.TryGetReadRedirectAsync(GetBlobPath(locator), cancellationToken);
 
 			/// <inheritdoc/>
-			public override async Task<Stream> OpenAsync(BundleLocator locator, int offset, int length, CancellationToken cancellationToken = default)
+			public override async Task<Stream> OpenAsync(BundleLocator locator, int offset, int? length, CancellationToken cancellationToken = default)
 			{
 				using TelemetrySpan span = _tracer.StartActiveSpan($"{nameof(StorageService)}.{nameof(StorageClientImpl)}.{nameof(OpenAsync)}");
 				span.SetAttribute("locator", locator.ToString());
 				span.SetAttribute("offset", offset);
 				span.SetAttribute("length", length);
 
+				if (length.HasValue && length.Value == 0)
+				{
+					return new MemoryStream(Array.Empty<byte>());
+				}
+
 				string path = GetBlobPath(locator);
 				return await Backend.ReadAsync(path, offset, length, cancellationToken);
 			}
diff --git a/Engine/Source/Programs/Shared/EpicGames.Horde/Compute/AgentStorageClient.cs b/Engine/Source/Programs/Shared/EpicGames.Horde/Compute/AgentStorageClient.cs
--- a/Engine/Source/Programs/Shared/EpicGames.Horde/Compute/AgentStorageClient.cs
+++ b/Engine/Source/Programs/Shared/EpicGames.Horde/Compute/AgentStorageClient.cs
@@ -40,12 +40,20 @@ namespace EpicGames.Horde.Compute
 		#region Blobs
 
 		/// <inheritdoc/>
-		public override async Task<Stream> OpenAsync(BundleLocator locator, int offset, int length, CancellationToken cancellationToken = default)
+		public override async Task<Stream> OpenAsync(BundleLocator locator, int offset, int? length, CancellationToken cancellationToken = default)
 		{
 			await _semaphore.WaitAsync(cancellationToken);
 			try
 			{
-				ReadOnlyMemory<byte> data = await _channel.ReadBlobAsync(locator, offset, length, cancellationToken);
+				ReadOnlyMemory<byte> data;
+				if (length.HasValue && length.Value == 0)
+				{
+					data = ReadOnlyMemory<byte>.Empty;
+				}
+				else
+				{
+					data = await _channel.ReadBlobAsync(locator, offset, length ?? 0, cancellationToken);
+				}
 				return new ReadOnlyMemoryStream(data);
 			}
 			finally
diff --git a/Engine/Source/Programs/Shared/EpicGames.Horde/Storage/Backends/FileStorageClient.cs b/Engine/Source/Programs/Shared/EpicGames.Horde/Storage/Backends/FileStorageClient.cs
--- a/Engine/Source/Programs/Shared/EpicGames.Horde/Storage/Backends/FileStorageClient.cs
+++ b/Engine/Source/Programs/Shared/EpicGames.Horde/Storage/Backends/FileStorageClient.cs
@@ -55,7 +55,7 @@ namespace EpicGames.Horde.Storage.Backends
 		#region Blobs
 
 		/// <inheritdoc/>
-		public override Task<Stream> OpenAsync(BundleLocator id, int offset, int length, CancellationToken cancellationToken = default)
+		public override Task<Stream> OpenAsync(BundleLocator id, int offset, int? length, CancellationToken cancellationToken = default)
 		{
 			FileReference file = GetBlobFile(id);
 			_logger.LogInformation("Reading {File}", file);
diff --git a/Engine/Source/Programs/Shared/EpicGames.Horde/Storage/Backends/HttpStorageClient.cs b/Engine/Source/Programs/Shared/EpicGames.Horde/Storage/Backends/HttpStorageClient.cs
--- a/Engine/Source/Programs/Shared/EpicGames.Horde/Storage/Backends/HttpStorageClient.cs
+++ b/Engine/Source/Programs/Shared/EpicGames.Horde/Storage/Backends/HttpStorageClient.cs
@@ -103,13 +103,13 @@ namespace EpicGames.Horde.Storage.Backends
 		#region Blobs
 
 		/// <inheritdoc/>
-		public override async Task<Stream> OpenAsync(BundleLocator locator, int offset, int length, CancellationToken cancellationToken = default)
+		public override async Task<Stream> OpenAsync(BundleLocator locator, int offset, int? length, CancellationToken cancellationToken = default)
 		{
-			if (offset == 0 && length == 0)
+			if (offset == 0 && length == null)
 			{
 				_logger.LogDebug("Reading {Locator}", locator);
 			}
-			else if (length == 0)
+			else if (length == null)
 			{
 				_logger.LogDebug("Reading {Locator} ({Offset}..)", locator, offset);
 			}
@@ -122,7 +122,7 @@ namespace EpicGames.Horde.Storage.Backends
 			{
 				using (HttpRequestMessage request = new HttpRequestMessage(HttpMethod.Get, $"blobs/{locator}"))
 				{
-					if (offset != 0 || length != 0)
+					if (offset != 0 || length != null)
 					{
 						request.Headers.Range = new RangeHeaderValue(offset, (length == 0) ? null : (offset + (length - 1)));
 					}
diff --git a/Engine/Source/Programs/Shared/EpicGames.Horde/Storage/Backends/MemoryStorageClient.cs b/Engine/Source/Programs/Shared/EpicGames.Horde/Storage/Backends/MemoryStorageClient.cs
--- a/Engine/Source/Programs/Shared/EpicGames.Horde/Storage/Backends/MemoryStorageClient.cs
+++ b/Engine/Source/Programs/Shared/EpicGames.Horde/Storage/Backends/MemoryStorageClient.cs
@@ -55,12 +55,12 @@ namespace EpicGames.Horde.Storage.Backends
 		#region Blobs
 
 		/// <inheritdoc/>
-		public override Task<Stream> OpenAsync(BundleLocator locator, int offset, int length, CancellationToken cancellationToken = default)
+		public override Task<Stream> OpenAsync(BundleLocator locator, int offset, int? length, CancellationToken cancellationToken = default)
 		{
 			ReadOnlySequence<byte> sequence = _bundles[locator].AsSequence().Slice(offset);
-			if (length > 0 && sequence.Length > length)
+			if (length != null && sequence.Length > length)
 			{
-				sequence = sequence.Slice(0, length);
+				sequence = sequence.Slice(0, length.Value);
 			}
 			return Task.FromResult<Stream>(new ReadOnlySequenceStream(sequence));
 		}
diff --git a/Engine/Source/Programs/Shared/EpicGames.Horde/Storage/Bundles/BundleStorageClient.cs b/Engine/Source/Programs/Shared/EpicGames.Horde/Storage/Bundles/BundleStorageClient.cs
--- a/Engine/Source/Programs/Shared/EpicGames.Horde/Storage/Bundles/BundleStorageClient.cs
+++ b/Engine/Source/Programs/Shared/EpicGames.Horde/Storage/Bundles/BundleStorageClient.cs
@@ -38,10 +38,10 @@ namespace EpicGames.Horde.Storage.Bundles
 		/// <param name="locator">Locator for the bundle</param>
 		/// <param name="cancellationToken">Cancellation token for the operation</param>
 		/// <returns>Stream for reading from the bundle</returns>
-		public Task<Stream> OpenAsync(BundleLocator locator, CancellationToken cancellationToken = default) => OpenAsync(locator, 0, 0, cancellationToken);
+		public Task<Stream> OpenAsync(BundleLocator locator, CancellationToken cancellationToken = default) => OpenAsync(locator, 0, null, cancellationToken);
 
 		/// <inheritdoc/>
-		public abstract Task<Stream> OpenAsync(BundleLocator locator, int offset, int length = 0, CancellationToken cancellationToken = default);
+		public abstract Task<Stream> OpenAsync(BundleLocator locator, int offset, int? length = null, CancellationToken cancellationToken = default);
 
 		/// <summary>
 		/// Reads an entire bundle into memory
