Horde: Support for constructing collections in bundle headers directly.

diff --git a/Engine/Source/Programs/Shared/EpicGames.Horde.Tests/BundleTests.cs b/Engine/Source/Programs/Shared/EpicGames.Horde.Tests/BundleTests.cs
--- a/Engine/Source/Programs/Shared/EpicGames.Horde.Tests/BundleTests.cs
+++ b/Engine/Source/Programs/Shared/EpicGames.Horde.Tests/BundleTests.cs
@@ -156,7 +156,7 @@ namespace EpicGames.Horde.Tests
 				oldHeader = BundleHeader.Create(types, imports, exports, packets);
 			}
 
-			byte[] serializedData = oldHeader.Data.ToArray();
+			byte[] serializedData = oldHeader.ToByteArray();
 
 			BundleHeader newHeader = BundleHeader.Read(serializedData);
 				
diff --git a/Engine/Source/Programs/Shared/EpicGames.Horde/Storage/Bundles/Bundle.cs b/Engine/Source/Programs/Shared/EpicGames.Horde/Storage/Bundles/Bundle.cs
--- a/Engine/Source/Programs/Shared/EpicGames.Horde/Storage/Bundles/Bundle.cs
+++ b/Engine/Source/Programs/Shared/EpicGames.Horde/Storage/Bundles/Bundle.cs
@@ -140,7 +140,7 @@ namespace EpicGames.Horde.Storage.Bundles
 		public ReadOnlySequence<byte> AsSequence()
 		{
 			ReadOnlySequenceBuilder<byte> sequence = new ReadOnlySequenceBuilder<byte>();
-			sequence.Append(Header.Data);
+			Header.AppendTo(sequence);
 
 			foreach (ReadOnlyMemory<byte> packet in Packets)
 			{
@@ -228,11 +228,6 @@ namespace EpicGames.Horde.Storage.Bundles
 		/// </summary>
 		public const int PreludeLength = 8;
 
-		/// <summary>
-		/// Underlying data for the header
-		/// </summary>
-		public ReadOnlyMemory<byte> Data { get; }
-
 		/// <summary>
 		/// Types for exports within this bundle
 		/// </summary>
@@ -254,52 +249,25 @@ namespace EpicGames.Horde.Storage.Bundles
 		public BundlePacketCollection Packets { get; }
 
 		/// <summary>
-		/// Construct a header from the given data encoded in the latest format
+		/// Constructor
 		/// </summary>
-		/// <param name="data">Data for the header, including the prelude</param>
-		private BundleHeader(ReadOnlyMemory<byte> data)
+		public BundleHeader(BundleTypeCollection types, BundleImportCollection imports, BundleExportCollection exports, BundlePacketCollection packets)
 		{
-			Data = data;
-
-			ReadOnlySpan<byte> span = data.Span;
-			BundleVersion version = (BundleVersion)span[3];
-			int headerLength = BinaryPrimitives.ReadInt32LittleEndian(span.Slice(4));
-
-			ReadOnlyMemory<byte> exportData = ReadOnlyMemory<byte>.Empty;
-			ReadOnlyMemory<byte> exportRefData = ReadOnlyMemory<byte>.Empty;
-
-			for (int offset = PreludeLength; offset < headerLength;)
-			{
-				uint header = BinaryPrimitives.ReadUInt32LittleEndian(data.Slice(offset).Span);
-				offset += SectionHeaderLength;
-
-				int length = (int)(header >> 8);
-				ReadOnlyMemory<byte> sectionData = data.Slice(offset, length);
-
-				BundleSectionType type = (BundleSectionType)(header & 255);
-				switch (type)
-				{
-					case BundleSectionType.Types:
-						Types = new BundleTypeCollection(sectionData);
-						break;
-					case BundleSectionType.Imports:
-						Imports = new BundleImportCollection(sectionData);
-						break;
-					case BundleSectionType.Exports:
-						exportData = sectionData;
-						break;
-					case BundleSectionType.ExportRefs:
-						exportRefData = sectionData;
-						break;
-					case BundleSectionType.Packets:
-						Packets = new BundlePacketCollection(sectionData);
-						break;
-				}
-
-				offset += length;
-			}
+			Types = types;
+			Imports = imports;
+			Exports = exports;
+			Packets = packets;
+		}
 
-			Exports = new BundleExportCollection(exportData, exportRefData, version);
+		/// <summary>
+		/// Constructor
+		/// </summary>
+		public BundleHeader(IReadOnlyList<BlobType> types, IReadOnlyList<BundleLocator> imports, IReadOnlyList<BundleExport> exports, IReadOnlyList<BundlePacket> packets)
+		{
+			Types = new BundleTypeCollection(types);
+			Imports = new BundleImportCollection(imports);
+			Exports = new BundleExportCollection(exports);
+			Packets = new BundlePacketCollection(packets);
 		}
 
 		/// <summary>
@@ -311,112 +279,81 @@ namespace EpicGames.Horde.Storage.Bundles
 		/// <param name="packets">Compression packets within the bundle</param>
 		public static BundleHeader Create(IReadOnlyList<BlobType> types, IReadOnlyList<BundleLocator> imports, IReadOnlyList<BundleExport> exports, IReadOnlyList<BundlePacket> packets)
 		{
-			// Find the size of all the sections
-			int typesLength = BundleTypeCollection.Measure(types);
-			int importsLength = BundleImportCollection.Measure(imports);
-			int exportsLength = BundleExportCollection.Measure(exports);
-
-			int exportRefsLength = exports.Sum(x => x.References.Data.Length);
-			if (exportRefsLength > 0)
-			{
-				exportRefsLength += exports.Count * sizeof(int);
-			}
+			return new BundleHeader(types, imports, exports, packets);
+		}
 
-			int packetsLength = BundlePacketCollection.Measure(packets);
+		/// <summary>
+		/// Writes data for this bundle to a sequence builder
+		/// </summary>
+		public void AppendTo(ReadOnlySequenceBuilder<byte> builder)
+		{
+			long initialLength = builder.Length;
 
 			// Find the total size of the header
 			int length = PreludeLength;
-			if (typesLength > 0)
+			if (Types.Data.Length > 0)
 			{
-				length += SectionHeaderLength + typesLength;
+				length += SectionHeaderLength + Types.Data.Length;
 			}
-			if (importsLength > 0)
+			if (Imports.Data.Length > 0)
 			{
-				length += SectionHeaderLength + importsLength;
+				length += SectionHeaderLength + Imports.Data.Length;
 			}
-			if (exportsLength > 0)
+			if (Exports.Data.Length > 0)
 			{
-				length += SectionHeaderLength + exportsLength;
+				length += SectionHeaderLength + Exports.Data.Length;
 			}
-			if (exportRefsLength > 0)
+			if (Exports.Refs.Length > 0)
 			{
-				length += SectionHeaderLength + exportRefsLength;
+				length += SectionHeaderLength + Exports.Refs.Length;
 			}
-			if (packetsLength > 0)
+			if (Packets.Data.Length > 0)
 			{
-				length += SectionHeaderLength + packetsLength;
+				length += SectionHeaderLength + Packets.Data.Length;
 			}
 
-			// Allocate the header data
-			byte[] data = new byte[length];
-			data[0] = (byte)'U';
-			data[1] = (byte)'B';
-			data[2] = (byte)'N';
-			data[3] = (byte)BundleVersion.Latest;
-			BinaryPrimitives.WriteInt32LittleEndian(data.AsSpan(4), length);
+			// Allocate the prelude
+			byte[] prelude = new byte[PreludeLength];
+			prelude[0] = (byte)'U';
+			prelude[1] = (byte)'B';
+			prelude[2] = (byte)'N';
+			prelude[3] = (byte)BundleVersion.Latest;
+			BinaryPrimitives.WriteInt32LittleEndian(prelude.AsSpan(4), length);
+			builder.Append(prelude);
 
 			// Write all the sections
-			Span<byte> next = data.AsSpan(PreludeLength);
-			if (typesLength > 0)
-			{
-				WriteSectionHeader(next, BundleSectionType.Types, typesLength);
-				next = next.Slice(SectionHeaderLength);
-
-				BundleTypeCollection.Write(next.Slice(0, typesLength), types);
-				next = next.Slice(typesLength);
-			}
-			if (importsLength > 0)
-			{
-				WriteSectionHeader(next, BundleSectionType.Imports, importsLength);
-				next = next.Slice(SectionHeaderLength);
-
-				BundleImportCollection.Write(next.Slice(0, importsLength), imports);
-				next = next.Slice(importsLength);
-			}
-			if (exportsLength > 0)
-			{
-				WriteSectionHeader(next, BundleSectionType.Exports, exportsLength);
-				next = next.Slice(SectionHeaderLength);
-
-				BundleExportCollection.Write(next.Slice(0, exportsLength), exports);
-				next = next.Slice(exportsLength);
-			}
-			if (exportRefsLength > 0)
-			{
-				WriteSectionHeader(next, BundleSectionType.ExportRefs, exportRefsLength);
-				next = next.Slice(SectionHeaderLength);
-
-				int indexOffset = 0;
-				int dataOffset = sizeof(int) * exports.Count;
+			AppendSection(builder, BundleSectionType.Types, Types.Data);
+			AppendSection(builder, BundleSectionType.Imports, Imports.Data);
+			AppendSection(builder, BundleSectionType.Exports, Exports.Data);
+			AppendSection(builder, BundleSectionType.ExportRefs, Exports.Refs);
+			AppendSection(builder, BundleSectionType.Packets, Packets.Data);
+			Debug.Assert(builder.Length == initialLength + length);
+		}
 
-				foreach (BundleExport export in exports)
-				{
-					BinaryPrimitives.WriteInt32LittleEndian(next.Slice(indexOffset), dataOffset);
-					indexOffset += sizeof(int);
+		/// <summary>
+		/// Serialize the bundle to a byte array
+		/// </summary>
+		public byte[] ToByteArray()
+		{
+			ReadOnlySequenceBuilder<byte> builder = new ReadOnlySequenceBuilder<byte>();
+			AppendTo(builder);
+			return builder.Construct().ToArray();
+		}
 
-					export.References.Data.Span.CopyTo(next.Slice(dataOffset));
-					dataOffset += export.References.Data.Span.Length;
-				}
+		const int SectionHeaderLength = 4;
 
-				Debug.Assert(dataOffset == exportRefsLength);
-				next = next.Slice(exportRefsLength);
-			}
-			if (packetsLength > 0)
+		static void AppendSection(ReadOnlySequenceBuilder<byte> builder, BundleSectionType type, ReadOnlyMemory<byte> data)
+		{
+			if (data.Length > 0)
 			{
-				WriteSectionHeader(next, BundleSectionType.Packets, packetsLength);
-				next = next.Slice(SectionHeaderLength);
+				byte[] header = new byte[SectionHeaderLength];
+				WriteSectionHeader(header, type, data.Length);
+				builder.Append(header);
 
-				BundlePacketCollection.Write(next.Slice(0, packetsLength), packets);
-				next = next.Slice(packetsLength);
+				builder.Append(data);
 			}
-			Debug.Assert(next.Length == 0);
-
-			// Create the header instance
-			return new BundleHeader(data);
 		}
 
-		const int SectionHeaderLength = 4;
-
 		static void WriteSectionHeader(Span<byte> span, BundleSectionType type, int length)
 		{
 			BinaryPrimitives.WriteInt32LittleEndian(span, (int)type | (length << 8));
@@ -434,7 +371,7 @@ namespace EpicGames.Horde.Storage.Bundles
 			ReadOnlySpan<byte> span = memory.Span;
 			if (span[0] == (byte)'U' && span[1] == (byte)'B' && span[2] == (byte)'N')
 			{
-				return new BundleHeader(memory);
+				return ReadLatest(memory);
 			}
 			else if (span[0] == (byte)'U' && span[1] == (byte)'E' && span[2] == (byte)'B' && span[3] == (byte)'N')
 			{
@@ -446,6 +383,58 @@ namespace EpicGames.Horde.Storage.Bundles
 			}
 		}
 
+		/// <summary>
+		/// Construct a header from the given data encoded in the latest format
+		/// </summary>
+		/// <param name="data">Data for the header, including the prelude</param>
+		static BundleHeader ReadLatest(ReadOnlyMemory<byte> data)
+		{
+			ReadOnlySpan<byte> span = data.Span;
+			BundleVersion version = (BundleVersion)span[3];
+			int headerLength = BinaryPrimitives.ReadInt32LittleEndian(span.Slice(4));
+
+			ReadOnlyMemory<byte> exportData = ReadOnlyMemory<byte>.Empty;
+			ReadOnlyMemory<byte> exportRefData = ReadOnlyMemory<byte>.Empty;
+
+			BundleTypeCollection types = new BundleTypeCollection();
+			BundleImportCollection imports = new BundleImportCollection();
+			BundlePacketCollection packets = new BundlePacketCollection();
+
+			for (int offset = PreludeLength; offset < headerLength;)
+			{
+				uint header = BinaryPrimitives.ReadUInt32LittleEndian(data.Slice(offset).Span);
+				offset += SectionHeaderLength;
+
+				int length = (int)(header >> 8);
+				ReadOnlyMemory<byte> sectionData = data.Slice(offset, length);
+
+				BundleSectionType type = (BundleSectionType)(header & 255);
+				switch (type)
+				{
+					case BundleSectionType.Types:
+						types = new BundleTypeCollection(sectionData);
+						break;
+					case BundleSectionType.Imports:
+						imports = new BundleImportCollection(sectionData);
+						break;
+					case BundleSectionType.Exports:
+						exportData = sectionData;
+						break;
+					case BundleSectionType.ExportRefs:
+						exportRefData = sectionData;
+						break;
+					case BundleSectionType.Packets:
+						packets = new BundlePacketCollection(sectionData);
+						break;
+				}
+
+				offset += length;
+			}
+
+			BundleExportCollection exports = new BundleExportCollection(exportData, exportRefData, version);
+			return new BundleHeader(types, imports, exports, packets);
+		}
+
 		static BundleHeader ReadLegacy(ReadOnlyMemory<byte> memory)
 		{
 			int headerLength = BinaryPrimitives.ReadInt32BigEndian(memory.Span.Slice(4));
@@ -696,6 +685,11 @@ namespace EpicGames.Horde.Storage.Bundles
 	{
 		readonly ReadOnlyMemory<byte> _data;
 
+		/// <summary>
+		/// Data for the imports
+		/// </summary>
+		public ReadOnlyMemory<byte> Data => _data;
+
 		/// <summary>
 		/// Deserializing constructor
 		/// </summary>
@@ -863,6 +857,11 @@ namespace EpicGames.Horde.Storage.Bundles
 	{
 		readonly ReadOnlyMemory<byte> _data;
 
+		/// <summary>
+		/// Data for the collection
+		/// </summary>
+		public ReadOnlyMemory<byte> Data => _data;
+
 		/// <summary>
 		/// Constructor
 		/// </summary>
@@ -1018,6 +1017,25 @@ namespace EpicGames.Horde.Storage.Bundles
 		readonly ReadOnlyMemory<byte> _data;
 		readonly ReadOnlyMemory<byte> _refs;
 
+		internal ReadOnlyMemory<byte> Data => _data;
+		internal ReadOnlyMemory<byte> Refs => _refs;
+
+		/// <summary>
+		/// Constructor
+		/// </summary>
+		public BundleExportCollection(IReadOnlyList<BundleExport> exports)
+		{
+			_count = exports.Count;
+
+			byte[] data = new byte[Measure(exports)];
+			Write(data, exports);
+			_data = data;
+
+			byte[] refs = new byte[MeasureRefs(exports)];
+			WriteRefs(refs, exports);
+			_refs = refs;
+		}
+
 		/// <summary>
 		/// Constructor
 		/// </summary>
@@ -1088,17 +1106,17 @@ namespace EpicGames.Horde.Storage.Bundles
 		IEnumerator IEnumerable.GetEnumerator() => GetEnumerator();
 
 		/// <summary>
-		/// Measure the size of memory required to store a collection of import locators
+		/// Measure the size of memory required to store a collection of exports
 		/// </summary>
-		/// <param name="exports">Locators to write</param>
+		/// <param name="exports">Exports to write</param>
 		/// <returns>Size in bytes of the output buffer</returns>
 		public static int Measure(IReadOnlyCollection<BundleExport> exports) => BundleExport.NumBytes * exports.Count;
 
 		/// <summary>
-		/// Serialize a collection of packets to memory
+		/// Serialize a collection of exports to memory
 		/// </summary>
 		/// <param name="data">Output buffer for the serialized data</param>
-		/// <param name="exports">Packets to write</param>
+		/// <param name="exports">Exports to write</param>
 		public static void Write(Span<byte> data, IReadOnlyCollection<BundleExport> exports)
 		{
 			Span<byte> next = data;
@@ -1108,6 +1126,33 @@ namespace EpicGames.Horde.Storage.Bundles
 				next = next.Slice(export.Data.Length);
 			}
 		}
+
+		/// <summary>
+		/// Measure the size of memory required to store a collection of export refs
+		/// </summary>
+		/// <param name="exports">Exports to write</param>
+		/// <returns>Size in bytes of the output buffer</returns>
+		public static int MeasureRefs(IReadOnlyCollection<BundleExport> exports) => (sizeof(int) * exports.Count) + exports.Sum(x => x.References.Data.Length);
+
+		/// <summary>
+		/// Serialize a collection of export refs to memory
+		/// </summary>
+		/// <param name="data">Output buffer for the serialized data</param>
+		/// <param name="exports">Exports to write</param>
+		public static void WriteRefs(Span<byte> data, IReadOnlyCollection<BundleExport> exports)
+		{
+			int indexOffset = 0;
+			int dataOffset = sizeof(int) * exports.Count;
+
+			foreach (BundleExport export in exports)
+			{
+				BinaryPrimitives.WriteInt32LittleEndian(data.Slice(indexOffset), dataOffset);
+				indexOffset += sizeof(int);
+
+				export.References.Data.Span.CopyTo(data.Slice(dataOffset));
+				dataOffset += export.References.Data.Span.Length;
+			}
+		}
 	}
 
 	/// <summary>
