Removed the MODULENAME_API from the private copy constructor definition generated by UHT.

diff --git a/Engine/Source/Programs/Shared/EpicGames.UHT/Exporters/CodeGen/UhtHeaderCodeGeneratorHFile.cs b/Engine/Source/Programs/Shared/EpicGames.UHT/Exporters/CodeGen/UhtHeaderCodeGeneratorHFile.cs
--- a/Engine/Source/Programs/Shared/EpicGames.UHT/Exporters/CodeGen/UhtHeaderCodeGeneratorHFile.cs
+++ b/Engine/Source/Programs/Shared/EpicGames.UHT/Exporters/CodeGen/UhtHeaderCodeGeneratorHFile.cs
@@ -1135,7 +1135,7 @@ namespace EpicGames.UHT.Exporters.CodeGen
 				}
 
 				AppendVTableHelperCtorAndCaller(builder, classObj, api);
-				AppendCopyConstructorDefinition(builder, classObj, api);
+				AppendCopyConstructorDefinition(builder, classObj);
 				AppendDestructorDefinition(builder, classObj, api);
 			}
 			return builder;
@@ -1182,14 +1182,13 @@ namespace EpicGames.UHT.Exporters.CodeGen
 		/// </summary>
 		/// <param name="builder">Output builder</param>
 		/// <param name="classObj">Class being exported</param>
-		/// <param name="api">API text to be used</param>
 		/// <returns>Output builder</returns>
-		private static StringBuilder AppendCopyConstructorDefinition(StringBuilder builder, UhtClass classObj, string api)
+		private static StringBuilder AppendCopyConstructorDefinition(StringBuilder builder, UhtClass classObj)
 		{
 			builder.Append("private: \\\r\n");
 			builder.Append("\t/** Private move- and copy-constructors, should never be used */ \\\r\n");
-			builder.Append('\t').Append(api).Append(classObj.SourceName).Append('(').Append(classObj.SourceName).Append("&&); \\\r\n");
-			builder.Append('\t').Append(api).Append(classObj.SourceName).Append("(const ").Append(classObj.SourceName).Append("&); \\\r\n");
+			builder.Append('\t').Append(classObj.SourceName).Append('(').Append(classObj.SourceName).Append("&&); \\\r\n");
+			builder.Append('\t').Append(classObj.SourceName).Append("(const ").Append(classObj.SourceName).Append("&); \\\r\n");
 			builder.Append("public: \\\r\n");
 			return builder;
 		}
@@ -1233,7 +1232,7 @@ namespace EpicGames.UHT.Exporters.CodeGen
 						break;
 				}
 			}
-			AppendCopyConstructorDefinition(builder, classObj, api);
+			AppendCopyConstructorDefinition(builder, classObj);
 			return builder;
 		}
 
