Horde: Add support for reading bundle header collections directly from streams.

diff --git a/Engine/Source/Programs/Shared/EpicGames.Horde/Storage/Bundles/Bundle.cs b/Engine/Source/Programs/Shared/EpicGames.Horde/Storage/Bundles/Bundle.cs
--- a/Engine/Source/Programs/Shared/EpicGames.Horde/Storage/Bundles/Bundle.cs
+++ b/Engine/Source/Programs/Shared/EpicGames.Horde/Storage/Bundles/Bundle.cs
@@ -369,13 +369,15 @@ namespace EpicGames.Horde.Storage.Bundles
 		public static BundleHeader Read(ReadOnlyMemory<byte> memory)
 		{
 			ReadOnlySpan<byte> span = memory.Span;
-			if (span[0] == (byte)'U' && span[1] == (byte)'B' && span[2] == (byte)'N')
+			if (span.StartsWith(s_latestSignature))
 			{
-				return ReadLatest(memory);
+				int headerLength = BinaryPrimitives.ReadInt32LittleEndian(memory.Slice(4).Span);
+				return ReadLatest((BundleVersion)span[3], memory.Slice(PreludeLength, headerLength - PreludeLength));
 			}
-			else if (span[0] == (byte)'U' && span[1] == (byte)'E' && span[2] == (byte)'B' && span[3] == (byte)'N')
+			else if (span.StartsWith(s_legacySignature))
 			{
-				return ReadLegacy(memory);
+				int headerLength = BinaryPrimitives.ReadInt32BigEndian(memory.Slice(4).Span);
+				return ReadLegacy(memory.Slice(PreludeLength, headerLength - PreludeLength));
 			}
 			else
 			{
@@ -383,16 +385,60 @@ namespace EpicGames.Horde.Storage.Bundles
 			}
 		}
 
+		/// <summary>
+		/// Reads a bundle header from a stream
+		/// </summary>
+		/// <param name="stream">Stream to deserialize from</param>
+		/// <param name="cancellationToken">Cancellation token for the operation</param>
+		/// <returns>New header object</returns>
+		public static async Task<BundleHeader> ReadAsync(Stream stream, CancellationToken cancellationToken = default)
+		{
+			byte[] prelude = new byte[PreludeLength];
+			await stream.ReadFixedLengthBytesAsync(prelude, cancellationToken);
+
+			return await ReadAsync(prelude, stream, cancellationToken);
+		}
+
+		/// <summary>
+		/// Reads a bundle header from a stream
+		/// </summary>
+		/// <param name="prelude">Prelude data</param>
+		/// <param name="stream">Stream to deserialize from</param>
+		/// <param name="cancellationToken">Cancellation token for the operation</param>
+		/// <returns>New header object</returns>
+		public static async Task<BundleHeader> ReadAsync(ReadOnlyMemory<byte> prelude, Stream stream, CancellationToken cancellationToken = default)
+		{
+			if (prelude.Span.StartsWith(s_latestSignature))
+			{
+				BundleVersion version = (BundleVersion)prelude.Span[3];
+				int headerLength = BinaryPrimitives.ReadInt32LittleEndian(prelude.Span.Slice(4));
+
+				return await ReadLatestAsync(version, stream, headerLength - PreludeLength, cancellationToken);
+			}
+			else if (prelude.Span.StartsWith(s_legacySignature))
+			{
+				int headerLength = BinaryPrimitives.ReadInt32BigEndian(prelude.Span.Slice(4));
+
+				byte[] data = new byte[headerLength - PreludeLength];
+				await stream.ReadFixedLengthBytesAsync(data, cancellationToken);
+
+				return ReadLegacy(data);
+			}
+			else
+			{
+				throw new NotSupportedException();
+			}
+		}
+
+		static readonly byte[] s_latestSignature = new byte[] { (byte)'U', (byte)'B', (byte)'N' };
+
 		/// <summary>
 		/// Construct a header from the given data encoded in the latest format
 		/// </summary>
+		/// <param name="version">Version to serialize as</param>
 		/// <param name="data">Data for the header, including the prelude</param>
-		static BundleHeader ReadLatest(ReadOnlyMemory<byte> data)
+		static BundleHeader ReadLatest(BundleVersion version, ReadOnlyMemory<byte> data)
 		{
-			ReadOnlySpan<byte> span = data.Span;
-			BundleVersion version = (BundleVersion)span[3];
-			int headerLength = BinaryPrimitives.ReadInt32LittleEndian(span.Slice(4));
-
 			ReadOnlyMemory<byte> exportData = ReadOnlyMemory<byte>.Empty;
 			ReadOnlyMemory<byte> exportRefData = ReadOnlyMemory<byte>.Empty;
 
@@ -400,7 +446,7 @@ namespace EpicGames.Horde.Storage.Bundles
 			BundleImportCollection imports = new BundleImportCollection();
 			BundlePacketCollection packets = new BundlePacketCollection();
 
-			for (int offset = PreludeLength; offset < headerLength;)
+			for (int offset = 0; offset < data.Length;)
 			{
 				uint header = BinaryPrimitives.ReadUInt32LittleEndian(data.Slice(offset).Span);
 				offset += SectionHeaderLength;
@@ -435,12 +481,64 @@ namespace EpicGames.Horde.Storage.Bundles
 			return new BundleHeader(types, imports, exports, packets);
 		}
 
-		static BundleHeader ReadLegacy(ReadOnlyMemory<byte> memory)
+		/// <summary>
+		/// Construct a header from the given data encoded in the latest format
+		/// </summary>
+		/// <param name="version">Version to serialize as</param>
+		/// <param name="stream">Stream to read from</param>
+		/// <param name="length">Length of the header</param>
+		/// <param name="cancellationToken">Cancellation token for the operation</param>
+		static async Task<BundleHeader> ReadLatestAsync(BundleVersion version, Stream stream, int length, CancellationToken cancellationToken)
 		{
-			int headerLength = BinaryPrimitives.ReadInt32BigEndian(memory.Span.Slice(4));
+			byte[] exportData = Array.Empty<byte>();
+			byte[] exportRefData = Array.Empty<byte>();
+
+			BundleTypeCollection types = new BundleTypeCollection();
+			BundleImportCollection imports = new BundleImportCollection();
+			BundlePacketCollection packets = new BundlePacketCollection();
+
+			byte[] sectionHeader = new byte[SectionHeaderLength];
+			for (int offset = 0; offset < length;)
+			{
+				await stream.ReadFixedLengthBytesAsync(sectionHeader, cancellationToken);
+				offset += SectionHeaderLength;
+
+				BundleSectionType sectionType = (BundleSectionType)(sectionHeader[0] & 255);
+				int sectionLength = (int)(BinaryPrimitives.ReadUInt32LittleEndian(sectionHeader) >> 8);
+
+				switch (sectionType)
+				{
+					case BundleSectionType.Types:
+						types = await BundleTypeCollection.ReadAsync(stream, sectionLength, cancellationToken);
+						break;
+					case BundleSectionType.Imports:
+						imports = await BundleImportCollection.ReadAsync(stream, sectionLength, cancellationToken);
+						break;
+					case BundleSectionType.Exports:
+						exportData = new byte[sectionLength];
+						await stream.ReadFixedLengthBytesAsync(exportData, cancellationToken);
+						break;
+					case BundleSectionType.ExportRefs:
+						exportRefData = new byte[sectionLength];
+						await stream.ReadFixedLengthBytesAsync(exportRefData, cancellationToken);
+						break;
+					case BundleSectionType.Packets:
+						packets = await BundlePacketCollection.ReadAsync(stream, sectionLength, cancellationToken);
+						break;
+				}
+
+				offset += sectionLength;
+			}
+
+			BundleExportCollection exports = new BundleExportCollection(exportData, exportRefData, version);
+			return new BundleHeader(types, imports, exports, packets);
+		}
 
-			MemoryReader reader = new MemoryReader(memory.Slice(0, headerLength));
-			reader.Advance(PreludeLength);
+		static readonly byte[] s_legacySignature = new byte[] { (byte)'U', (byte)'E', (byte)'B', (byte)'N' };
+
+		static BundleHeader ReadLegacy(ReadOnlyMemory<byte> memory)
+		{
+			MemoryReader reader = new MemoryReader(memory);
 
 			BundleVersion version = (BundleVersion)reader.ReadUnsignedVarInt();
 			if (version > BundleVersion.RemoveAliases)
@@ -634,6 +732,16 @@ namespace EpicGames.Horde.Storage.Bundles
 			Data = data;
 		}
 
+		/// <summary>
+		/// Reads a type collection from a stream
+		/// </summary>
+		public static async Task<BundleTypeCollection> ReadAsync(Stream stream, int length, CancellationToken cancellationToken)
+		{
+			byte[] data = new byte[length];
+			await stream.ReadFixedLengthBytesAsync(data, cancellationToken);
+			return new BundleTypeCollection(data);
+		}
+
 		/// <inheritdoc/>
 		public int Count => Data.Length / NumBytesPerType;
 
@@ -706,6 +814,16 @@ namespace EpicGames.Horde.Storage.Bundles
 			_data = data;
 		}
 
+		/// <summary>
+		/// Reads a collection from a stream
+		/// </summary>
+		public static async Task<BundleImportCollection> ReadAsync(Stream stream, int length, CancellationToken cancellationToken)
+		{
+			byte[] data = new byte[length];
+			await stream.ReadFixedLengthBytesAsync(data, cancellationToken);
+			return new BundleImportCollection(data);
+		}
+
 		/// <summary>
 		/// Retrieve a single import from the collection
 		/// </summary>
@@ -879,6 +997,16 @@ namespace EpicGames.Horde.Storage.Bundles
 			_data = data;
 		}
 
+		/// <summary>
+		/// Reads a collection from a stream
+		/// </summary>
+		public static async Task<BundlePacketCollection> ReadAsync(Stream stream, int length, CancellationToken cancellationToken)
+		{
+			byte[] data = new byte[length];
+			await stream.ReadFixedLengthBytesAsync(data, cancellationToken);
+			return new BundlePacketCollection(data);
+		}
+
 		/// <inheritdoc/>
 		public int Count => _data.Length / BundlePacket.NumBytes;
 
diff --git a/Engine/Source/Programs/Shared/EpicGames.Horde/Storage/Bundles/BundleReader.cs b/Engine/Source/Programs/Shared/EpicGames.Horde/Storage/Bundles/BundleReader.cs
--- a/Engine/Source/Programs/Shared/EpicGames.Horde/Storage/Bundles/BundleReader.cs
+++ b/Engine/Source/Programs/Shared/EpicGames.Horde/Storage/Bundles/BundleReader.cs
@@ -244,7 +244,7 @@ namespace EpicGames.Horde.Storage.Bundles
 					}
 
 					// Parse the header and construct the bundle info from it
-					BundleHeader header = await ReadHeaderAsync(prelude, stream, headerSize, cancellationToken);
+					BundleHeader header = await BundleHeader.ReadAsync(prelude, stream, cancellationToken);
 
 					// Construct the bundle info
 					BundleInfo bundleInfo = new BundleInfo(queuedHeader.Blob, header, headerSize);
@@ -363,14 +363,6 @@ namespace EpicGames.Horde.Storage.Bundles
 			}
 		}
 
-		static async Task<BundleHeader> ReadHeaderAsync(byte[] prelude, Stream stream, int headerSize, CancellationToken cancellationToken)
-		{
-			byte[] header = new byte[headerSize];
-			prelude.CopyTo(header, 0);
-			await stream.ReadFixedLengthBytesAsync(header.AsMemory(prelude.Length), cancellationToken);
-			return BundleHeader.Read(header);
-		}
-
 		static async Task<byte[]> ReadPacketAsync(Stream stream, int packetSize, CancellationToken cancellationToken)
 		{
 			byte[] packet = new byte[packetSize];
