Horde: Fix incorrect range header when specifying null for length when retrieving bundle range.

diff --git a/Engine/Source/Programs/Shared/EpicGames.Horde/Storage/Backends/HttpStorageClient.cs b/Engine/Source/Programs/Shared/EpicGames.Horde/Storage/Backends/HttpStorageClient.cs
--- a/Engine/Source/Programs/Shared/EpicGames.Horde/Storage/Backends/HttpStorageClient.cs
+++ b/Engine/Source/Programs/Shared/EpicGames.Horde/Storage/Backends/HttpStorageClient.cs
@@ -118,13 +118,18 @@ namespace EpicGames.Horde.Storage.Backends
 				_logger.LogDebug("Reading {Locator} ({Offset}+{Length})", locator, offset, length);
 			}
 
+			if (length.HasValue && length.Value == 0)
+			{
+				return new MemoryStream(Array.Empty<byte>());
+			}
+
 			using (HttpClient httpClient = _createClient())
 			{
 				using (HttpRequestMessage request = new HttpRequestMessage(HttpMethod.Get, $"blobs/{locator}"))
 				{
 					if (offset != 0 || length != null)
 					{
-						request.Headers.Range = new RangeHeaderValue(offset, (length == 0) ? null : (offset + (length - 1)));
+						request.Headers.Range = new RangeHeaderValue(offset, (length == null) ? null : (offset + (length - 1)));
 					}
 
 					HttpResponseMessage response = await httpClient.SendAsync(request, cancellationToken);
