Horde: Better handling of exceptions when thrown reading a blob.

diff --git a/Engine/Source/Programs/Shared/EpicGames.Horde/Storage/Bundles/BundleReader.cs b/Engine/Source/Programs/Shared/EpicGames.Horde/Storage/Bundles/BundleReader.cs
--- a/Engine/Source/Programs/Shared/EpicGames.Horde/Storage/Bundles/BundleReader.cs
+++ b/Engine/Source/Programs/Shared/EpicGames.Horde/Storage/Bundles/BundleReader.cs
@@ -298,7 +298,10 @@ namespace EpicGames.Horde.Storage.Bundles
 			}
 			catch (Exception ex)
 			{
-				queuedPacket.CompletionSource.SetException(ex);
+				if (!queuedPacket.CompletionSource.TrySetException(ex))
+				{
+					_logger.LogWarning(ex, "Exception after setting completion source state; state: {Status}", queuedPacket.CompletionSource.Task.Status);
+				}
 			}
 		}
 
@@ -328,7 +331,7 @@ namespace EpicGames.Horde.Storage.Bundles
 
 			await using (Stream stream = await _store.OpenAsync(bundleInfo.Locator, bundleInfo.HeaderLength + bundleInfo.Header.Packets[minPacketIdx].EncodedOffset, readLength, cancellationToken))
 			{
-				// Copy all the packets that have been read into separate buffers, so we can cache them individually.
+				// Copy all the packets that have been read into separate buffers, so we can cache them indidually.
 				ReadOnlyMemory<byte>[] packets = new ReadOnlyMemory<byte>[maxPacketIdx - minPacketIdx];
 				if (_cache != null)
 				{
@@ -353,7 +356,7 @@ namespace EpicGames.Horde.Storage.Bundles
 				{
 					ReadOnlyMemory<byte> data = packets[updatePacket.PacketIdx - minPacketIdx];
 					updatePacket.CompletionSource.SetResult(data);
-					}
+				}
 
 				// Remove all the completed packets from the queue
 				lock (_queueLock)
