Horde: Fix caching of bundle packets read as part of the header.

diff --git a/Engine/Source/Programs/Shared/EpicGames.Horde/Storage/Bundles/BundleReader.cs b/Engine/Source/Programs/Shared/EpicGames.Horde/Storage/Bundles/BundleReader.cs
--- a/Engine/Source/Programs/Shared/EpicGames.Horde/Storage/Bundles/BundleReader.cs
+++ b/Engine/Source/Programs/Shared/EpicGames.Horde/Storage/Bundles/BundleReader.cs
@@ -252,6 +252,8 @@ namespace EpicGames.Horde.Storage.Bundles
 					if (_cache != null)
 					{
 						// Also add any encoded packets we prefetched
+						List<ReadOnlyMemory<byte>> packets = new List<ReadOnlyMemory<byte>>();
+
 						int packetOffset = headerSize;
 						for (int packetIdx = 0; packetIdx < header.Packets.Count; packetIdx++)
 						{
@@ -262,7 +264,8 @@ namespace EpicGames.Horde.Storage.Bundles
 							}
 
 							byte[] packetData = await ReadPacketAsync(stream, packetLength, cancellationToken);
-							AddToCache(GetEncodedPacketCacheKey(queuedHeader.Blob, packetIdx), packetData, packetData.Length);
+							AddToCache(GetEncodedPacketCacheKey(queuedHeader.Blob, packetIdx), (ReadOnlyMemory<byte>)packetData, packetData.Length);
+							packets.Add(packetData);
 
 							packetOffset += packetLength;
 						}
@@ -270,6 +273,19 @@ namespace EpicGames.Horde.Storage.Bundles
 						// Add the info to the cache
 						string cacheKey = GetBundleInfoCacheKey(queuedHeader.Blob);
 						AddToCache(cacheKey, bundleInfo, headerSize);
+
+						// Update any packets that now have a cached value
+						List<QueuedPacket> updatePackets = new List<QueuedPacket>();
+						lock (_queueLock)
+						{
+							updatePackets.AddRange(_queuedPackets.Where(x => x.Path == bundleInfo.Locator.Path && x.PacketIdx < packets.Count));
+						}
+
+						// Mark any packets that now have a valid cache entry as complete
+						foreach (QueuedPacket updatePacket in updatePackets)
+						{
+							updatePacket.CompletionSource.TrySetResult(packets[updatePacket.PacketIdx]);
+						}
 					}
 
 					// Update the task
@@ -300,7 +316,7 @@ namespace EpicGames.Horde.Storage.Bundles
 			{
 				if (!queuedPacket.CompletionSource.TrySetException(ex))
 				{
-					_logger.LogWarning(ex, "Exception after setting completion source state; state: {Status}", queuedPacket.CompletionSource.Task.Status);
+					_logger.LogWarning(ex, "Exception after setting completion source state; existing state: {Status}", queuedPacket.CompletionSource.Task.Status);
 				}
 			}
 		}
