Fixed issue with parsing LAYOUT_FIELD_INITITALIZED

diff --git a/Engine/Source/Programs/Shared/EpicGames.UHT/Parsers/UhtPropertyParser.cs b/Engine/Source/Programs/Shared/EpicGames.UHT/Parsers/UhtPropertyParser.cs
--- a/Engine/Source/Programs/Shared/EpicGames.UHT/Parsers/UhtPropertyParser.cs
+++ b/Engine/Source/Programs/Shared/EpicGames.UHT/Parsers/UhtPropertyParser.cs
@@ -898,17 +898,23 @@ namespace EpicGames.UHT.Parsers
 				{
 					tokenReader.Require(',');
 					RequireArray(tokenReader, propertySettings, ref nameToken, ')');
+					tokenReader.Require(')');
 				}
 				else if (layoutMacroType.IsBitfield())
 				{
 					tokenReader.Require(',');
 					RequireBitfield(tokenReader, propertySettings, ref nameToken);
+					tokenReader.Require(')');
 				}
 				else if (layoutMacroType.HasInitializer())
 				{
-					tokenReader.SkipBrackets('(', ')', 1);
+					tokenReader.Require(',');
+					tokenReader.SkipBrackets('(', ')', 1); // consumes ending ) too
+				}
+				else
+				{
+					tokenReader.Require(')');
 				}
-				tokenReader.Require(')');
 
 				Finalize(topScope, specifierContext, ref nameToken, new ReadOnlyMemory<UhtToken>(_currentTypeTokens.ToArray()), layoutMacroType, propertyDelegate);
 			}
