Horde: Fix type of object added to cache for bundle packets.

diff --git a/Engine/Source/Programs/Shared/EpicGames.Horde/Storage/Bundles/BundleReader.cs b/Engine/Source/Programs/Shared/EpicGames.Horde/Storage/Bundles/BundleReader.cs
--- a/Engine/Source/Programs/Shared/EpicGames.Horde/Storage/Bundles/BundleReader.cs
+++ b/Engine/Source/Programs/Shared/EpicGames.Horde/Storage/Bundles/BundleReader.cs
@@ -356,7 +356,7 @@ namespace EpicGames.Horde.Storage.Bundles
 						string cacheKey = GetEncodedPacketCacheKey(bundleInfo.Locator, idx);
 						byte[] data = await ReadPacketAsync(stream, bundleInfo.Header.Packets[idx].EncodedLength, cancellationToken);
 						packets[idx - minPacketIdx] = data;
-						AddToCache(cacheKey, data, data.Length);
+						AddToCache(cacheKey, (ReadOnlyMemory<byte>)data, data.Length);
 					}
 				}
 
