Horde: Fix partial reads from bundle streams.

diff --git a/Engine/Source/Programs/Shared/EpicGames.Horde/Storage/Bundles/BundleReader.cs b/Engine/Source/Programs/Shared/EpicGames.Horde/Storage/Bundles/BundleReader.cs
--- a/Engine/Source/Programs/Shared/EpicGames.Horde/Storage/Bundles/BundleReader.cs
+++ b/Engine/Source/Programs/Shared/EpicGames.Horde/Storage/Bundles/BundleReader.cs
@@ -385,7 +385,7 @@ namespace EpicGames.Horde.Storage.Bundles
 		static async Task<byte[]> ReadPacketAsync(Stream stream, int packetSize, CancellationToken cancellationToken)
 		{
 			byte[] packet = new byte[packetSize];
-			await stream.ReadAsync(packet, cancellationToken);
+			await stream.ReadFixedLengthBytesAsync(packet, cancellationToken);
 			return packet;
 		}
 
