Horde: Allow creating a bundle node handle directly from a locator.

diff --git a/Engine/Source/Programs/Shared/EpicGames.Horde/Storage/Bundles/BundleStorageClient.cs b/Engine/Source/Programs/Shared/EpicGames.Horde/Storage/Bundles/BundleStorageClient.cs
--- a/Engine/Source/Programs/Shared/EpicGames.Horde/Storage/Bundles/BundleStorageClient.cs
+++ b/Engine/Source/Programs/Shared/EpicGames.Horde/Storage/Bundles/BundleStorageClient.cs
@@ -62,6 +62,11 @@ namespace EpicGames.Horde.Storage.Bundles
 
 		#region Nodes
 
+		/// <summary>
+		/// Creates a handle to a node from its locator
+		/// </summary>
+		public BlobHandle CreateNodeHandle(BundleNodeLocator locator) => new FlushedNodeHandle(BundleReader, locator);
+
 		/// <inheritdoc/>
 		public BundleWriter CreateWriter(RefName refName = default, BundleOptions? options = null)
 		{
