Horde: Store bundle exports and references in a single block, so we can paginate reads more easily.

diff --git a/Engine/Source/Programs/Shared/EpicGames.Horde/Storage/Bundles/Bundle.cs b/Engine/Source/Programs/Shared/EpicGames.Horde/Storage/Bundles/Bundle.cs
--- a/Engine/Source/Programs/Shared/EpicGames.Horde/Storage/Bundles/Bundle.cs
+++ b/Engine/Source/Programs/Shared/EpicGames.Horde/Storage/Bundles/Bundle.cs
@@ -198,19 +198,24 @@ namespace EpicGames.Horde.Storage.Bundles
 		Imports = 1,
 
 		/// <summary>
-		/// List of exports
+		/// List of exports (deprecated; use Exports instead)
 		/// </summary>
-		Exports = 2,
+		ExportHeaders = 2,
 
 		/// <summary>
-		/// References to exports in other bundles
+		/// References to exports in other bundles (deprecated; use Exports instead)
 		/// </summary>
-		ExportRefs = 3,
+		ExportReferences = 3,
 
 		/// <summary>
 		/// Packet headers
 		/// </summary>
 		Packets = 4,
+
+		/// <summary>
+		/// List of combined exports and headers
+		/// </summary>
+		Exports = 5,
 	}
 
 	/// <summary>
@@ -303,10 +308,6 @@ namespace EpicGames.Horde.Storage.Bundles
 			{
 				length += SectionHeaderLength + Exports.Data.Length;
 			}
-			if (Exports.Refs.Length > 0)
-			{
-				length += SectionHeaderLength + Exports.Refs.Length;
-			}
 			if (Packets.Data.Length > 0)
 			{
 				length += SectionHeaderLength + Packets.Data.Length;
@@ -325,7 +326,6 @@ namespace EpicGames.Horde.Storage.Bundles
 			AppendSection(builder, BundleSectionType.Types, Types.Data);
 			AppendSection(builder, BundleSectionType.Imports, Imports.Data);
 			AppendSection(builder, BundleSectionType.Exports, Exports.Data);
-			AppendSection(builder, BundleSectionType.ExportRefs, Exports.Refs);
 			AppendSection(builder, BundleSectionType.Packets, Packets.Data);
 			Debug.Assert(builder.Length == initialLength + length);
 		}
@@ -439,11 +439,12 @@ namespace EpicGames.Horde.Storage.Bundles
 		/// <param name="data">Data for the header, including the prelude</param>
 		static BundleHeader ReadLatest(BundleVersion version, ReadOnlyMemory<byte> data)
 		{
-			ReadOnlyMemory<byte> exportData = ReadOnlyMemory<byte>.Empty;
-			ReadOnlyMemory<byte> exportRefData = ReadOnlyMemory<byte>.Empty;
+			ReadOnlyMemory<byte> exportHeaderData = ReadOnlyMemory<byte>.Empty;
+			ReadOnlyMemory<byte> exportReferenceData = ReadOnlyMemory<byte>.Empty;
 
 			BundleTypeCollection types = new BundleTypeCollection();
 			BundleImportCollection imports = new BundleImportCollection();
+			BundleExportCollection exports = new BundleExportCollection();
 			BundlePacketCollection packets = new BundlePacketCollection();
 
 			for (int offset = 0; offset < data.Length;)
@@ -463,21 +464,28 @@ namespace EpicGames.Horde.Storage.Bundles
 					case BundleSectionType.Imports:
 						imports = new BundleImportCollection(sectionData);
 						break;
-					case BundleSectionType.Exports:
-						exportData = sectionData;
+					case BundleSectionType.ExportHeaders:
+						exportHeaderData = sectionData;
 						break;
-					case BundleSectionType.ExportRefs:
-						exportRefData = sectionData;
+					case BundleSectionType.ExportReferences:
+						exportReferenceData = sectionData;
 						break;
 					case BundleSectionType.Packets:
 						packets = new BundlePacketCollection(sectionData);
 						break;
+					case BundleSectionType.Exports:
+						exports = new BundleExportCollection(sectionData);
+						break;
 				}
 
 				offset += length;
 			}
 
-			BundleExportCollection exports = new BundleExportCollection(exportData, exportRefData, version);
+			if (exportHeaderData.Length > 0)
+			{
+				exports = new BundleExportCollection(exportHeaderData, exportReferenceData, version);
+			}
+
 			return new BundleHeader(types, imports, exports, packets);
 		}
 
@@ -495,6 +503,7 @@ namespace EpicGames.Horde.Storage.Bundles
 
 			BundleTypeCollection types = new BundleTypeCollection();
 			BundleImportCollection imports = new BundleImportCollection();
+			BundleExportCollection exports = new BundleExportCollection();
 			BundlePacketCollection packets = new BundlePacketCollection();
 
 			byte[] sectionHeader = new byte[SectionHeaderLength];
@@ -514,23 +523,30 @@ namespace EpicGames.Horde.Storage.Bundles
 					case BundleSectionType.Imports:
 						imports = await BundleImportCollection.ReadAsync(stream, sectionLength, cancellationToken);
 						break;
-					case BundleSectionType.Exports:
+					case BundleSectionType.ExportHeaders:
 						exportData = new byte[sectionLength];
 						await stream.ReadFixedLengthBytesAsync(exportData, cancellationToken);
 						break;
-					case BundleSectionType.ExportRefs:
+					case BundleSectionType.ExportReferences:
 						exportRefData = new byte[sectionLength];
 						await stream.ReadFixedLengthBytesAsync(exportRefData, cancellationToken);
 						break;
 					case BundleSectionType.Packets:
 						packets = await BundlePacketCollection.ReadAsync(stream, sectionLength, cancellationToken);
 						break;
+					case BundleSectionType.Exports:
+						exports = await BundleExportCollection.ReadAsync(stream, sectionLength, cancellationToken);
+						break;
 				}
 
 				offset += sectionLength;
 			}
 
-			BundleExportCollection exports = new BundleExportCollection(exportData, exportRefData, version);
+			if (exportData.Length > 0)
+			{
+				exports = new BundleExportCollection(exportData, exportRefData, version);
+			}
+
 			return new BundleHeader(types, imports, exports, packets);
 		}
 
@@ -1064,7 +1080,7 @@ namespace EpicGames.Horde.Storage.Bundles
 	public struct BundleExport
 	{
 		/// <summary>
-		/// Number of bytes in a serialized export object
+		/// Number of bytes in a serialized export object, excluding any trailing BundleExportRef entries
 		/// </summary>
 		public const int NumBytes = 32;
 
@@ -1101,38 +1117,45 @@ namespace EpicGames.Horde.Storage.Bundles
 		/// <summary>
 		/// References to other nodes
 		/// </summary>
-		public BundleExportRefCollection References { get; }
+		public BundleExportRefCollection References => new BundleExportRefCollection(Data.Slice(32));
 
 		/// <summary>
 		/// Constructor
 		/// </summary>
-		public BundleExport(ReadOnlyMemory<byte> data, BundleExportRefCollection references)
-		{
-			Data = data;
-			References = references;
-		}
+		public BundleExport(ReadOnlyMemory<byte> data) => Data = data;
 
 		/// <summary>
 		/// Constructor
 		/// </summary>
 		public BundleExport(int typeIdx, IoHash hash, int packetIdx, int offset, int length, IReadOnlyList<BundleExportRef> references)
 		{
-			byte[] data = new byte[NumBytes];
-			Write(data, typeIdx, hash, packetIdx, offset, length);
+			byte[] data = new byte[Measure(references.Count)];
+			Write(data, typeIdx, hash, packetIdx, offset, length, references);
 			Data = data;
-			References = new BundleExportRefCollection(references);
 		}
 
+		/// <summary>
+		/// Measure the size of an exported reference
+		/// </summary>
+		public static int Measure(int numReferences) => NumBytes + (BundleExportRef.NumBytes * numReferences);
+
 		/// <summary>
 		/// Writes a new export to a block of memory
 		/// </summary>
-		public static void Write(Span<byte> data, int typeIdx, IoHash hash, int packet, int offset, int length)
+		public static void Write(Span<byte> data, int typeIdx, IoHash hash, int packet, int offset, int length, IReadOnlyList<BundleExportRef> references)
 		{
 			hash.CopyTo(data);
 			BinaryPrimitives.WriteUInt16LittleEndian(data.Slice(20), (ushort)typeIdx);
 			BinaryPrimitives.WriteUInt16LittleEndian(data.Slice(22), (ushort)packet);
 			BinaryPrimitives.WriteInt32LittleEndian(data.Slice(24), offset);
 			BinaryPrimitives.WriteInt32LittleEndian(data.Slice(28), length);
+
+			Span<byte> next = data.Slice(NumBytes);
+			for (int idx = 0; idx < references.Count; idx++)
+			{
+				references[idx].CopyTo(next);
+				next = next.Slice(BundleExportRef.NumBytes);
+			}
 		}
 	}
 
@@ -1143,10 +1166,25 @@ namespace EpicGames.Horde.Storage.Bundles
 	{
 		readonly int _count;
 		readonly ReadOnlyMemory<byte> _data;
-		readonly ReadOnlyMemory<byte> _refs;
 
 		internal ReadOnlyMemory<byte> Data => _data;
-		internal ReadOnlyMemory<byte> Refs => _refs;
+
+		/// <summary>
+		/// Constructor
+		/// </summary>
+		public BundleExportCollection()
+			: this(Array.Empty<BundleExport>())
+		{
+		}
+
+		/// <summary>
+		/// Constructor
+		/// </summary>
+		public BundleExportCollection(ReadOnlyMemory<byte> data)
+		{
+			_count = BinaryPrimitives.ReadInt32LittleEndian(data.Span) / sizeof(int);
+			_data = data;
+		}
 
 		/// <summary>
 		/// Constructor
@@ -1158,22 +1196,16 @@ namespace EpicGames.Horde.Storage.Bundles
 			byte[] data = new byte[Measure(exports)];
 			Write(data, exports);
 			_data = data;
-
-			byte[] refs = new byte[MeasureRefs(exports)];
-			WriteRefs(refs, exports);
-			_refs = refs;
 		}
 
 		/// <summary>
 		/// Constructor
 		/// </summary>
-		public BundleExportCollection(ReadOnlyMemory<byte> data, ReadOnlyMemory<byte> refs, BundleVersion version)
+		public BundleExportCollection(ReadOnlyMemory<byte> headers, ReadOnlyMemory<byte> refs, BundleVersion version)
 		{
-			_count = data.Length / BundleExport.NumBytes;
-			_data = data;
-			_refs = refs;
+			_count = headers.Length / BundleExport.NumBytes;
 
-			if (version < BundleVersion.ImportHashes && _refs.Length > 0)
+			if (version < BundleVersion.ImportHashes && refs.Length > 0)
 			{
 				const int OldRefNumBytes = 4;
 
@@ -1192,12 +1224,50 @@ namespace EpicGames.Horde.Storage.Bundles
 					source.CopyTo(newRefs.AsMemory(headerLength + (index * BundleExportRef.NumBytes)));
 				}
 
-				_refs = newRefs;
+				refs = newRefs;
+			}
+
+			int size = (sizeof(int) + BundleExport.NumBytes) * _count;
+			if (refs.Length > 0)
+			{
+				size += refs.Length - (_count * sizeof(int));
+			}
+
+			byte[] data = new byte[size];
+
+			int offset = _count * sizeof(int);
+			for (int idx = 0; idx < _count; idx++)
+			{
+				BinaryPrimitives.WriteInt32LittleEndian(data.AsSpan(idx * sizeof(int)), offset);
+
+				headers.Slice(idx * BundleExport.NumBytes, BundleExport.NumBytes).CopyTo(data.AsMemory(offset));
+				offset += BundleExport.NumBytes;
+
+				if (refs.Length > 0)
+				{
+					int minOffset = BinaryPrimitives.ReadInt32LittleEndian(refs.Span.Slice(idx * sizeof(int)));
+					int maxOffset = (idx == _count - 1) ? refs.Length : BinaryPrimitives.ReadInt32LittleEndian(refs.Span.Slice((idx + 1) * sizeof(int)));
+
+					refs.Span.Slice(minOffset, maxOffset - minOffset).CopyTo(data.AsSpan(offset));
+					offset += maxOffset - minOffset;
+				}
 			}
+
+			_data = data;
+		}
+
+		/// <summary>
+		/// Reads a collection from a stream
+		/// </summary>
+		public static async Task<BundleExportCollection> ReadAsync(Stream stream, int length, CancellationToken cancellationToken)
+		{
+			byte[] data = new byte[length];
+			await stream.ReadFixedLengthBytesAsync(data, cancellationToken);
+			return new BundleExportCollection(data);
 		}
 
 		/// <inheritdoc/>
-		public int Count => _data.Length / BundleExport.NumBytes;
+		public int Count => _count;
 
 		/// <inheritdoc/>
 		public BundleExport this[int index]
@@ -1206,17 +1276,11 @@ namespace EpicGames.Horde.Storage.Bundles
 			{
 				Debug.Assert(index < _count);
 
-				ReadOnlyMemory<byte> exportData = _data.Slice(index * BundleExport.NumBytes);
-
-				BundleExportRefCollection exportRefs = new BundleExportRefCollection();
-				if (_refs.Length > 0)
-				{
-					int minOffset = BinaryPrimitives.ReadInt32LittleEndian(_refs.Span.Slice(index * sizeof(int)));
-					int maxOffset = (index == Count - 1)? _refs.Length : BinaryPrimitives.ReadInt32LittleEndian(_refs.Span.Slice((index + 1) * sizeof(int)));
-					exportRefs = new BundleExportRefCollection(_refs.Slice(minOffset, maxOffset - minOffset));
-				}
+				int minOffset = BinaryPrimitives.ReadInt32LittleEndian(_data.Span.Slice(index * sizeof(int)));
+				int maxOffset = (index == _count - 1) ? _data.Length : BinaryPrimitives.ReadInt32LittleEndian(_data.Span.Slice((index + 1) * sizeof(int)));
 
-				return new BundleExport(exportData, exportRefs);
+				ReadOnlyMemory<byte> exportData = _data.Slice(minOffset, maxOffset - minOffset);
+				return new BundleExport(exportData);
 			}
 		}
 
@@ -1238,7 +1302,7 @@ namespace EpicGames.Horde.Storage.Bundles
 		/// </summary>
 		/// <param name="exports">Exports to write</param>
 		/// <returns>Size in bytes of the output buffer</returns>
-		public static int Measure(IReadOnlyCollection<BundleExport> exports) => BundleExport.NumBytes * exports.Count;
+		public static int Measure(IReadOnlyCollection<BundleExport> exports) => (exports.Count * sizeof(int)) + exports.Sum(x => x.Data.Length);
 
 		/// <summary>
 		/// Serialize a collection of exports to memory
@@ -1247,11 +1311,16 @@ namespace EpicGames.Horde.Storage.Bundles
 		/// <param name="exports">Exports to write</param>
 		public static void Write(Span<byte> data, IReadOnlyCollection<BundleExport> exports)
 		{
-			Span<byte> next = data;
+			Span<byte> nextOffset = data;
+
+			int offset = sizeof(int) * exports.Count;
 			foreach (BundleExport export in exports)
 			{
-				export.Data.Span.CopyTo(next);
-				next = next.Slice(export.Data.Length);
+				BinaryPrimitives.WriteInt32LittleEndian(nextOffset, offset);
+				nextOffset = nextOffset.Slice(sizeof(int));
+
+				export.Data.Span.CopyTo(data[offset..]);
+				offset += export.Data.Length;
 			}
 		}
 
