Horde: Store items within bundle headers as elements in an array, rather than reading from flat blocks of memory, allowing us to read into separate blocks of memory and reduce LOH fragmentation.

diff --git a/Engine/Source/Programs/Horde/Horde.Server.Tests/BlobStoreTests.cs b/Engine/Source/Programs/Horde/Horde.Server.Tests/BlobStoreTests.cs
--- a/Engine/Source/Programs/Horde/Horde.Server.Tests/BlobStoreTests.cs
+++ b/Engine/Source/Programs/Horde/Horde.Server.Tests/BlobStoreTests.cs
@@ -62,7 +62,7 @@ namespace Horde.Server.Tests
 			List<BundlePacket> packets = new List<BundlePacket>();
 			packets.Add(new BundlePacket(BundleCompressionFormat.None, 0, (int)data.Length, (int)data.Length));
 
-			BundleHeader header = BundleHeader.Create(types, imports, exports, packets);
+			BundleHeader header = new BundleHeader(types.ToArray(), imports.ToArray(), exports.ToArray(), packets.ToArray());
 			return new Bundle(header, new[] { data });
 		}
 
diff --git a/Engine/Source/Programs/Horde/Horde.Server.Tests/GcServiceTests.cs b/Engine/Source/Programs/Horde/Horde.Server.Tests/GcServiceTests.cs
--- a/Engine/Source/Programs/Horde/Horde.Server.Tests/GcServiceTests.cs
+++ b/Engine/Source/Programs/Horde/Horde.Server.Tests/GcServiceTests.cs
@@ -104,7 +104,7 @@ namespace Horde.Server.Tests
 			{
 				List<BlobType> types = new List<BlobType> { new BlobType(Guid.Parse("{AFDF76A7-5333-4DEE-B837-B5F5CA511245}"), 0) };
 				List<BundleLocator> imports = children[idx].ConvertAll(x => locators[x]);
-				BundleHeader header = BundleHeader.Create(types, imports, Array.Empty<BundleExport>(), Array.Empty<BundlePacket>());
+				BundleHeader header = new BundleHeader(types.ToArray(), imports.ToArray(), Array.Empty<BundleExport>(), Array.Empty<BundlePacket>());
 				Bundle bundle = new Bundle(header, Array.Empty<ReadOnlyMemory<byte>>());
 				locators[idx] = await store.WriteBundleAsync(bundle, prefix: "gctest");
 			}
diff --git a/Engine/Source/Programs/Horde/Horde.Server.Tests/StorageServiceTests.cs b/Engine/Source/Programs/Horde/Horde.Server.Tests/StorageServiceTests.cs
--- a/Engine/Source/Programs/Horde/Horde.Server.Tests/StorageServiceTests.cs
+++ b/Engine/Source/Programs/Horde/Horde.Server.Tests/StorageServiceTests.cs
@@ -20,7 +20,7 @@ namespace Horde.Server.Tests
 	public class StorageServiceTests : TestSetup
 	{
 		[TestMethod]
-		public async Task BlobCollectionTest()
+		public async Task BlobCollectionTestAsync()
 		{
 			StorageService storageService = ServiceProvider.GetRequiredService<StorageService>();
 			StorageClient client = await storageService.GetClientAsync(new NamespaceId("memory"), default);
@@ -41,7 +41,7 @@ namespace Horde.Server.Tests
 			exports.Add(new BundleExport(0, hash1, 0, 0, data1.Length, Array.Empty<BundleExportRef>()));
 			exports.Add(new BundleExport(0, hash2, 0, data1.Length, data2.Length, Array.Empty<BundleExportRef>()));
 
-			BundleHeader header = BundleHeader.Create(types, Array.Empty<BundleLocator>(), exports, new BundlePacket[1]);
+			BundleHeader header = new BundleHeader(types.ToArray(), Array.Empty<BundleLocator>(), exports.ToArray(), new BundlePacket[1]);
 			Bundle bundle = new Bundle(header, Array.Empty<ReadOnlyMemory<byte>>());
 			BundleLocator locator = await client.WriteBundleAsync(bundle);
 
diff --git a/Engine/Source/Programs/Horde/Horde/Commands/Bundles/BundleDump.cs b/Engine/Source/Programs/Horde/Horde/Commands/Bundles/BundleDump.cs
--- a/Engine/Source/Programs/Horde/Horde/Commands/Bundles/BundleDump.cs
+++ b/Engine/Source/Programs/Horde/Horde/Commands/Bundles/BundleDump.cs
@@ -22,8 +22,11 @@ namespace Horde.Commands.Bundles
 
 		public override async Task<int> ExecuteAsync(ILogger logger)
 		{
-			byte[] bytes = await FileReference.ReadAllBytesAsync(Input);
-			Bundle bundle = new Bundle(bytes);
+			Bundle bundle;
+			using (FileStream stream = FileReference.Open(Input, FileMode.Open, FileAccess.Read, FileShare.Read | FileShare.Delete))
+			{
+				bundle = await Bundle.FromStreamAsync(stream);
+			}
 			logger.LogInformation("Summary for blob {Location}", Input);
 
 			BundleHeader header = bundle.Header;
diff --git a/Engine/Source/Programs/Shared/EpicGames.Core/MemoryWriter.cs b/Engine/Source/Programs/Shared/EpicGames.Core/MemoryWriter.cs
--- a/Engine/Source/Programs/Shared/EpicGames.Core/MemoryWriter.cs
+++ b/Engine/Source/Programs/Shared/EpicGames.Core/MemoryWriter.cs
@@ -252,22 +252,22 @@ namespace EpicGames.Core
 	/// <summary>
 	/// Class for building byte sequences, similar to StringBuilder. Allocates memory in chunks to avoid copying data.
 	/// </summary>
-	public sealed class ChunkedMemoryWriter : IMemoryWriter, IDisposable
+	public abstract class ChunkedMemoryWriterBase : IMemoryWriter
 	{
-		class Chunk
+		protected class Chunk
 		{
-			public readonly int RunningIndex;
-			public readonly IMemoryOwner<byte> Owner;
-			public readonly Memory<byte> Data;
-			public int Length;
+			public int RunningIndex { get; }
+			public Memory<byte> Data { get; }
+			public int Length { get; set; }
 
-			public Chunk(int runningIndex, int size)
+			public Chunk(int runningIndex, Memory<byte> data)
 			{
 				RunningIndex = runningIndex;
-				Owner = MemoryPool<byte>.Shared.Rent(size);
-				Data = Owner.Memory;
+				Data = data;
 			}
 
+			public virtual void Release() { }
+
 			public ReadOnlySpan<byte> WrittenSpan => WrittenMemory.Span;
 			public ReadOnlyMemory<byte> WrittenMemory => Data.Slice(0, Length);
 		}
@@ -285,28 +285,17 @@ namespace EpicGames.Core
 		/// Constructor
 		/// </summary>
 		/// <param name="chunkSize"></param>
-		public ChunkedMemoryWriter(int chunkSize)
-			: this(chunkSize, chunkSize)
+		protected ChunkedMemoryWriterBase(Chunk initialChunk, int chunkSize)
 		{
+			_currentChunk = initialChunk;
+			_chunks.Add(initialChunk);
+			_chunkSize = chunkSize;
 		}
 
 		/// <summary>
-		/// Constructor
+		/// Creates a new chunk for the buffer
 		/// </summary>
-		/// <param name="initialSize">Size of the initial chunk</param>
-		/// <param name="chunkSize">Default size for subsequent chunks</param>
-		public ChunkedMemoryWriter(int initialSize = 4096, int chunkSize = 4096)
-		{
-			_currentChunk = new Chunk(0, initialSize);
-			_chunks.Add(_currentChunk);
-			_chunkSize = chunkSize;
-		}
-
-		/// <inheritdoc/>
-		public void Dispose()
-		{
-			Clear();
-		}
+		protected abstract Chunk CreateChunk(int runningIndex, int size);
 
 		/// <summary>
 		/// Clear the current builder
@@ -315,7 +304,7 @@ namespace EpicGames.Core
 		{
 			foreach (Chunk chunk in _chunks)
 			{
-				chunk.Owner.Dispose();
+				chunk.Release();
 			}
 
 			_chunks.Clear();
@@ -331,7 +320,7 @@ namespace EpicGames.Core
 			int requiredSize = _currentChunk.Length + Math.Max(sizeHint, 1);
 			if (requiredSize > _currentChunk.Data.Length)
 			{
-				_currentChunk = new Chunk(_currentChunk.RunningIndex + _currentChunk.Length, Math.Max(sizeHint, _chunkSize));
+				_currentChunk = CreateChunk(_currentChunk.RunningIndex + _currentChunk.Length, Math.Max(sizeHint, _chunkSize));
 				_chunks.Add(_currentChunk);
 			}
 			return _currentChunk.Data.Slice(_currentChunk.Length);
@@ -420,6 +409,83 @@ namespace EpicGames.Core
 		}
 	}
 
+	/// <summary>
+	/// Class for building byte sequences, similar to StringBuilder. Allocates memory in chunks to avoid copying data.
+	/// </summary>
+	public sealed class ChunkedArrayMemoryWriter : ChunkedMemoryWriterBase
+	{
+		/// <summary>
+		/// Constructor
+		/// </summary>
+		/// <param name="chunkSize">Size of each chunk</param>
+		public ChunkedArrayMemoryWriter(int chunkSize = 4096)
+			: this(chunkSize, chunkSize)
+		{ }
+
+		/// <summary>
+		/// Constructor
+		/// </summary>
+		/// <param name="initialChunkSize">Size of the initial chunk</param>
+		/// <param name="nextChunkSize">Size of subsequent chunks</param>
+		public ChunkedArrayMemoryWriter(int initialChunkSize, int nextChunkSize)
+			: base(new Chunk(0, new byte[initialChunkSize]), nextChunkSize)
+		{ }
+
+		/// <inheritdoc/>
+		protected override Chunk CreateChunk(int runningIndex, int size) => new Chunk(runningIndex, new byte[size]);
+	}
+
+	/// <summary>
+	/// Class for building byte sequences, similar to StringBuilder. Allocates memory in chunks using a pool allocator to avoid copying data.
+	/// </summary>
+	public sealed class ChunkedMemoryWriter : ChunkedMemoryWriterBase, IDisposable
+	{
+		class PooledChunk : Chunk
+		{
+			public readonly IMemoryOwner<byte> Owner;
+
+			public PooledChunk(int runningIndex, int size)
+				: this(runningIndex, MemoryPool<byte>.Shared.Rent(size))
+			{
+			}
+
+			private PooledChunk(int runningIndex, IMemoryOwner<byte> owner)
+				: base(runningIndex, owner.Memory)
+			{
+				Owner = owner;
+			}
+
+			public override void Release() => Owner.Dispose();
+		}
+
+		/// <summary>
+		/// Constructor
+		/// </summary>
+		/// <param name="chunkSize"></param>
+		public ChunkedMemoryWriter(int chunkSize)
+			: this(chunkSize, chunkSize)
+		{ }
+
+		/// <summary>
+		/// Constructor
+		/// </summary>
+		/// <param name="initialSize">Size of the initial chunk</param>
+		/// <param name="chunkSize">Default size for subsequent chunks</param>
+		public ChunkedMemoryWriter(int initialSize = 4096, int chunkSize = 4096)
+			: base(new PooledChunk(0, initialSize), chunkSize)
+		{
+		}
+
+		/// <inheritdoc/>
+		protected override Chunk CreateChunk(int runningIndex, int size) => new PooledChunk(runningIndex, size);
+
+		/// <inheritdoc/>
+		public void Dispose()
+		{
+			Clear();
+		}
+	}
+
 	/// <summary>
 	/// Extension methods for <see cref="IMemoryWriter"/>
 	/// </summary>
diff --git a/Engine/Source/Programs/Shared/EpicGames.Horde.Tests/BundleStoreTests.cs b/Engine/Source/Programs/Shared/EpicGames.Horde.Tests/BundleStoreTests.cs
--- a/Engine/Source/Programs/Shared/EpicGames.Horde.Tests/BundleStoreTests.cs
+++ b/Engine/Source/Programs/Shared/EpicGames.Horde.Tests/BundleStoreTests.cs
@@ -83,7 +83,7 @@ namespace EpicGames.Horde.Tests
 			List<BundlePacket> packets = new List<BundlePacket>();
 			packets.Add(new BundlePacket(BundleCompressionFormat.None, 0, payload.Length, payload.Length));
 
-			BundleHeader header = BundleHeader.Create(types, Array.Empty<BundleLocator>(), exports, packets);
+			BundleHeader header = new BundleHeader(types.ToArray(), Array.Empty<BundleLocator>(), exports.ToArray(), packets.ToArray());
 			return new Bundle(header, new List<ReadOnlyMemory<byte>> { payload });
 		}
 
diff --git a/Engine/Source/Programs/Shared/EpicGames.Horde.Tests/BundleTests.cs b/Engine/Source/Programs/Shared/EpicGames.Horde.Tests/BundleTests.cs
--- a/Engine/Source/Programs/Shared/EpicGames.Horde.Tests/BundleTests.cs
+++ b/Engine/Source/Programs/Shared/EpicGames.Horde.Tests/BundleTests.cs
@@ -153,7 +153,7 @@ namespace EpicGames.Horde.Tests
 				packets.Add(new BundlePacket(BundleCompressionFormat.LZ4, 0, 20, 40));
 				packets.Add(new BundlePacket(BundleCompressionFormat.LZ4, 20, 10, 20));
 
-				oldHeader = BundleHeader.Create(types, imports, exports, packets);
+				oldHeader = new BundleHeader(types.ToArray(), imports.ToArray(), exports.ToArray(), packets.ToArray());
 			}
 
 			byte[] serializedData = oldHeader.ToByteArray();
diff --git a/Engine/Source/Programs/Shared/EpicGames.Horde/Storage/Backends/MemoryStorageClient.cs b/Engine/Source/Programs/Shared/EpicGames.Horde/Storage/Backends/MemoryStorageClient.cs
--- a/Engine/Source/Programs/Shared/EpicGames.Horde/Storage/Backends/MemoryStorageClient.cs
+++ b/Engine/Source/Programs/Shared/EpicGames.Horde/Storage/Backends/MemoryStorageClient.cs
@@ -66,11 +66,12 @@ namespace EpicGames.Horde.Storage.Backends
 		}
 
 		/// <inheritdoc/>
-		public override Task<BundleLocator> WriteBundleAsync(Bundle bundle, Utf8String prefix = default, CancellationToken cancellationToken = default)
+		public override async Task<BundleLocator> WriteBundleAsync(Bundle bundle, Utf8String prefix = default, CancellationToken cancellationToken = default)
 		{
 			BundleLocator locator = BundleLocator.CreateUnique(prefix);
-			_bundles[locator] = new Bundle(bundle.AsSequence().ToArray());
-			return Task.FromResult(locator);
+			byte[] data = bundle.AsSequence().ToArray();
+			_bundles[locator] = await Bundle.FromStreamAsync(new MemoryStream(data), cancellationToken);
+			return locator;
 		}
 
 		#endregion
diff --git a/Engine/Source/Programs/Shared/EpicGames.Horde/Storage/BlobType.cs b/Engine/Source/Programs/Shared/EpicGames.Horde/Storage/BlobType.cs
--- a/Engine/Source/Programs/Shared/EpicGames.Horde/Storage/BlobType.cs
+++ b/Engine/Source/Programs/Shared/EpicGames.Horde/Storage/BlobType.cs
@@ -1,6 +1,8 @@
 // Copyright Epic Games, Inc. All Rights Reserved.
 
 using System;
+using System.Buffers.Binary;
+using EpicGames.Core;
 
 namespace EpicGames.Horde.Storage
 {
@@ -11,6 +13,31 @@ namespace EpicGames.Horde.Storage
 	/// <param name="Version">Version number for the serializer</param>
 	public record struct BlobType(Guid Guid, int Version)
 	{
+		/// <summary>
+		/// Number of bytes in a serialized blob type instance
+		/// </summary>
+		public const int NumBytes = 20;
+
+		/// <summary>
+		/// Deserialize a type from a byte span
+		/// </summary>
+		public static BlobType Read(ReadOnlySpan<byte> span)
+		{
+			Guid guid = new Guid(span.Slice(0, 16));
+			int version = BinaryPrimitives.ReadInt32LittleEndian(span.Slice(16));
+
+			return new BlobType(guid, version);
+		}
+
+		/// <summary>
+		/// Serialize to a byte span
+		/// </summary>
+		public void Write(Span<byte> data)
+		{
+			Guid.TryWriteBytes(data.Slice(0, 16));
+			BinaryPrimitives.WriteInt32LittleEndian(data.Slice(16), Version);
+		}
+
 		/// <inheritdoc/>
 		public override string ToString() => $"{Guid}#{Version}";
 	}
diff --git a/Engine/Source/Programs/Shared/EpicGames.Horde/Storage/Bundles/Bundle.cs b/Engine/Source/Programs/Shared/EpicGames.Horde/Storage/Bundles/Bundle.cs
--- a/Engine/Source/Programs/Shared/EpicGames.Horde/Storage/Bundles/Bundle.cs
+++ b/Engine/Source/Programs/Shared/EpicGames.Horde/Storage/Bundles/Bundle.cs
@@ -90,25 +90,6 @@ namespace EpicGames.Horde.Storage.Bundles
 			Packets = packets;
 		}
 
-		/// <summary>
-		/// Reads a bundle from a block of memory
-		/// </summary>
-		public Bundle(ReadOnlyMemory<byte> memory)
-		{
-			int offset = BundleHeader.ReadPrelude(memory.Span);
-			Header = BundleHeader.Read(memory.Slice(0, offset));
-
-			ReadOnlyMemory<byte>[] packets = new ReadOnlyMemory<byte>[Header.Packets.Count];
-			for (int idx = 0; idx < Header.Packets.Count; idx++)
-			{
-				BundlePacket packet = Header.Packets[idx];
-				packets[idx] = memory.Slice(offset, packet.EncodedLength);
-				offset += packet.EncodedLength;
-			}
-
-			Packets = packets;
-		}
-
 		/// <summary>
 		/// Reads a bundle from the given stream
 		/// </summary>
@@ -117,7 +98,7 @@ namespace EpicGames.Horde.Storage.Bundles
 		/// <returns>Bundle that was read</returns>
 		public static async Task<Bundle> FromStreamAsync(Stream stream, CancellationToken cancellationToken = default)
 		{
-			BundleHeader header = await BundleHeader.FromStreamAsync(stream, cancellationToken);
+			BundleHeader header = await BundleHeader.ReadAsync(stream, cancellationToken);
 
 			ReadOnlyMemory<byte>[] packets = new ReadOnlyMemory<byte>[header.Packets.Count];
 			for (int idx = 0; idx < header.Packets.Count; idx++)
@@ -262,7 +243,7 @@ namespace EpicGames.Horde.Storage.Bundles
 		/// <summary>
 		/// Constructor
 		/// </summary>
-		public BundleHeader(IReadOnlyList<BlobType> types, IReadOnlyList<BundleLocator> imports, IReadOnlyList<BundleExport> exports, IReadOnlyList<BundlePacket> packets)
+		public BundleHeader(BlobType[] types, BundleLocator[] imports, BundleExport[] exports, BundlePacket[] packets)
 		{
 			Types = new BundleTypeCollection(types);
 			Imports = new BundleImportCollection(imports);
@@ -270,64 +251,40 @@ namespace EpicGames.Horde.Storage.Bundles
 			Packets = new BundlePacketCollection(packets);
 		}
 
-		/// <summary>
-		/// Constructs a new bundle header
-		/// </summary>
-		/// <param name="types">Type array indexed by each export</param>
-		/// <param name="imports">Imported bundles</param>
-		/// <param name="exports">Exports for nodes</param>
-		/// <param name="packets">Compression packets within the bundle</param>
-		public static BundleHeader Create(IReadOnlyList<BlobType> types, IReadOnlyList<BundleLocator> imports, IReadOnlyList<BundleExport> exports, IReadOnlyList<BundlePacket> packets)
-		{
-			return new BundleHeader(types, imports, exports, packets);
-		}
-
 		/// <summary>
 		/// Writes data for this bundle to a sequence builder
 		/// </summary>
 		public void AppendTo(ReadOnlySequenceBuilder<byte> builder)
 		{
-			long initialLength = builder.Length;
+			// Create a writer for the output data
+			ChunkedArrayMemoryWriter writer = new ChunkedArrayMemoryWriter();
 
-			// Find the total size of the header
-			int length = PreludeLength;
-			if (Types.Data.Length > 0)
-			{
-				length += SectionHeaderLength + Types.Data.Length;
-			}
-			if (Imports.Data.Length > 0)
-			{
-				length += SectionHeaderLength + Imports.Data.Length;
-			}
-			if (Exports.Data.Length > 0)
-			{
-				length += SectionHeaderLength + Exports.Data.Length;
-			}
-			if (Exports.Refs.Length > 0)
-			{
-				length += SectionHeaderLength + Exports.Refs.Length;
-			}
-			if (Packets.Data.Length > 0)
-			{
-				length += SectionHeaderLength + Packets.Data.Length;
-			}
-
-			// Allocate the prelude
-			byte[] prelude = new byte[PreludeLength];
+			// Create the prelude
+			Span<byte> prelude = writer.GetSpanAndAdvance(PreludeLength);
 			prelude[0] = (byte)'U';
 			prelude[1] = (byte)'B';
 			prelude[2] = (byte)'N';
 			prelude[3] = (byte)BundleVersion.Latest;
-			BinaryPrimitives.WriteInt32LittleEndian(prelude.AsSpan(4), length);
-			builder.Append(prelude);
 
 			// Write all the sections
-			AppendSection(builder, BundleSectionType.Types, Types.Data);
-			AppendSection(builder, BundleSectionType.Imports, Imports.Data);
-			AppendSection(builder, BundleSectionType.Exports, Exports.Data);
-			AppendSection(builder, BundleSectionType.ExportRefs, Exports.Refs);
-			AppendSection(builder, BundleSectionType.Packets, Packets.Data);
-			Debug.Assert(builder.Length == initialLength + length);
+			WriteSectionHeader(writer, BundleSectionType.Types, Types.Measure());
+			Types.Write(writer);
+
+			WriteSectionHeader(writer, BundleSectionType.Imports, Imports.Measure());
+			Imports.Write(writer);
+
+			WriteSectionHeader(writer, BundleSectionType.Exports, Exports.MeasureHeaders());
+			Exports.WriteHeaders(writer);
+
+			WriteSectionHeader(writer, BundleSectionType.ExportRefs, Exports.MeasureRefs());
+			Exports.WriteRefs(writer);
+
+			WriteSectionHeader(writer, BundleSectionType.Packets, Packets.Measure());
+			Packets.Write(writer);
+
+			// Fill in the length of the header, and append it to the builder
+			BinaryPrimitives.WriteInt32LittleEndian(prelude[4..], writer.Length);
+			writer.AppendTo(builder);
 		}
 
 		/// <summary>
@@ -342,21 +299,9 @@ namespace EpicGames.Horde.Storage.Bundles
 
 		const int SectionHeaderLength = 4;
 
-		static void AppendSection(ReadOnlySequenceBuilder<byte> builder, BundleSectionType type, ReadOnlyMemory<byte> data)
-		{
-			if (data.Length > 0)
-			{
-				byte[] header = new byte[SectionHeaderLength];
-				WriteSectionHeader(header, type, data.Length);
-				builder.Append(header);
-
-				builder.Append(data);
-			}
-		}
-
-		static void WriteSectionHeader(Span<byte> span, BundleSectionType type, int length)
+		static void WriteSectionHeader(IMemoryWriter writer, BundleSectionType type, int length)
 		{
-			BinaryPrimitives.WriteInt32LittleEndian(span, (int)type | (length << 8));
+			writer.WriteInt32((int)type | (length << 8));
 		}
 
 		record class ExportInfo(int TypeIdx, IoHash Hash, int Length, List<BundleExportRef> References);
@@ -364,25 +309,12 @@ namespace EpicGames.Horde.Storage.Bundles
 		/// <summary>
 		/// Reads a bundle header from memory
 		/// </summary>
-		/// <param name="memory">Memory to deserialize from</param>
+		/// <param name="memory">Memory to read from</param>
 		/// <returns>New header object</returns>
 		public static BundleHeader Read(ReadOnlyMemory<byte> memory)
 		{
-			ReadOnlySpan<byte> span = memory.Span;
-			if (span.StartsWith(s_latestSignature))
-			{
-				int headerLength = BinaryPrimitives.ReadInt32LittleEndian(memory.Slice(4).Span);
-				return ReadLatest((BundleVersion)span[3], memory.Slice(PreludeLength, headerLength - PreludeLength));
-			}
-			else if (span.StartsWith(s_legacySignature))
-			{
-				int headerLength = BinaryPrimitives.ReadInt32BigEndian(memory.Slice(4).Span);
-				return ReadLegacy(memory.Slice(PreludeLength, headerLength - PreludeLength));
-			}
-			else
-			{
-				throw new NotSupportedException();
-			}
+			using ReadOnlyMemoryStream stream = new ReadOnlyMemoryStream(memory);
+			return ReadAsync(stream).Result;
 		}
 
 		/// <summary>
@@ -458,10 +390,10 @@ namespace EpicGames.Horde.Storage.Bundles
 				switch (type)
 				{
 					case BundleSectionType.Types:
-						types = new BundleTypeCollection(sectionData);
+						types = BundleTypeCollection.Read(sectionData);
 						break;
 					case BundleSectionType.Imports:
-						imports = new BundleImportCollection(sectionData);
+						imports = BundleImportCollection.Read(sectionData);
 						break;
 					case BundleSectionType.Exports:
 						exportData = sectionData;
@@ -470,14 +402,14 @@ namespace EpicGames.Horde.Storage.Bundles
 						exportRefData = sectionData;
 						break;
 					case BundleSectionType.Packets:
-						packets = new BundlePacketCollection(sectionData);
+						packets = BundlePacketCollection.Read(sectionData);
 						break;
 				}
 
 				offset += length;
 			}
 
-			BundleExportCollection exports = new BundleExportCollection(exportData, exportRefData, version);
+			BundleExportCollection exports = BundleExportCollection.Read(exportData, exportRefData, version);
 			return new BundleHeader(types, imports, exports, packets);
 		}
 
@@ -530,7 +462,7 @@ namespace EpicGames.Horde.Storage.Bundles
 				offset += sectionLength;
 			}
 
-			BundleExportCollection exports = new BundleExportCollection(exportData, exportRefData, version);
+			BundleExportCollection exports = BundleExportCollection.Read(exportData, exportRefData, version);
 			return new BundleHeader(types, imports, exports, packets);
 		}
 
@@ -657,7 +589,7 @@ namespace EpicGames.Horde.Storage.Bundles
 				}
 			}
 
-			return BundleHeader.Create(types, imports, exports, packets);
+			return new BundleHeader(types.ToArray(), imports.ToArray(), exports.ToArray(), packets.ToArray());
 		}
 
 		/// <summary>
@@ -680,56 +612,35 @@ namespace EpicGames.Horde.Storage.Bundles
 				throw new InvalidDataException("Invalid signature bytes for bundle. Corrupt data?");
 			}
 		}
-
-		/// <summary>
-		/// Reads a bundle header from a stream
-		/// </summary>
-		/// <param name="stream">Stream to read from</param>
-		/// <param name="cancellationToken">Cancellation token for the stream</param>
-		/// <returns>New header</returns>
-		public static async Task<BundleHeader> FromStreamAsync(Stream stream, CancellationToken cancellationToken)
-		{
-			byte[] prelude = new byte[PreludeLength];
-			await stream.ReadFixedLengthBytesAsync(prelude, cancellationToken);
-
-			int headerLength = ReadPrelude(prelude);
-
-			byte[] header = new byte[headerLength];
-			prelude.CopyTo(header.AsSpan());
-
-			await stream.ReadFixedLengthBytesAsync(header.AsMemory(prelude.Length), cancellationToken);
-			return BundleHeader.Read(header);
-		}
 	}
 
 	/// <summary>
 	/// Collection of node types in a bundle
 	/// </summary>
-	public struct BundleTypeCollection
+	public struct BundleTypeCollection : IReadOnlyList<BlobType>
 	{
-		/// <summary>
-		/// Number of bytes in a serialized <see cref="BlobType"/> instance
-		/// </summary>
-		const int NumBytesPerType = 20;
+		readonly BlobType[] _types;
 
-		/// <summary>
-		/// Data for this collection
-		/// </summary>
-		public ReadOnlyMemory<byte> Data { get; }
+		/// <inheritdoc/>
+		public int Count => _types.Length;
 
 		/// <summary>
 		/// Constructor
 		/// </summary>
-		public BundleTypeCollection(ReadOnlyMemory<byte> data) => Data = data;
+		public BundleTypeCollection(BlobType[] types) => _types = types;
 
 		/// <summary>
-		/// Constructor
+		/// Reads a type collection from a block of memory
 		/// </summary>
-		public BundleTypeCollection(IReadOnlyList<BlobType> types)
+		public static BundleTypeCollection Read(ReadOnlyMemory<byte> data)
 		{
-			byte[] data = new byte[Measure(types)];
-			Write(data, types);
-			Data = data;
+			BlobType[] types = new BlobType[data.Length / BlobType.NumBytes];
+			for (int idx = 0; idx < types.Length; idx++)
+			{
+				types[idx] = BlobType.Read(data.Span.Slice(idx * BlobType.NumBytes));
+			}
+
+			return new BundleTypeCollection(types);
 		}
 
 		/// <summary>
@@ -739,51 +650,37 @@ namespace EpicGames.Horde.Storage.Bundles
 		{
 			byte[] data = new byte[length];
 			await stream.ReadFixedLengthBytesAsync(data, cancellationToken);
-			return new BundleTypeCollection(data);
-		}
-
-		/// <inheritdoc/>
-		public int Count => Data.Length / NumBytesPerType;
-
-		/// <inheritdoc/>
-		public BlobType this[int index]
-		{
-			get
-			{
-				ReadOnlySpan<byte> span = Data.Slice(index * NumBytesPerType, NumBytesPerType).Span;
-
-				Guid guid = new Guid(span.Slice(0, 16));
-				int version = BinaryPrimitives.ReadInt32LittleEndian(span.Slice(16));
-
-				return new BlobType(guid, version);
-			}
+			return Read(data);
 		}
 
 		/// <summary>
 		/// Gets the size of memory required to serialize a collection of types
 		/// </summary>
-		/// <param name="types">Type collection</param>
 		/// <returns>Number of bytes required to serialize the types</returns>
-		public static int Measure(IReadOnlyCollection<BlobType> types) => types.Count * NumBytesPerType;
+		public int Measure() => _types.Length * BlobType.NumBytes;
 
 		/// <summary>
 		/// Serializes a set of types to a fixed block of memory
 		/// </summary>
-		/// <param name="span">Span to write the types to</param>
-		/// <param name="types">Collection of types to be written</param>
-		public static void Write(Span<byte> span, IReadOnlyCollection<BlobType> types)
+		/// <param name="writer">Writer to serialize to</param>
+		public void Write(IMemoryWriter writer)
 		{
-			Span<byte> next = span;
-			foreach (BlobType type in types)
+			Span<byte> span = writer.GetSpanAndAdvance(_types.Length * BlobType.NumBytes);
+			foreach (BlobType type in _types)
 			{
-				type.Guid.TryWriteBytes(next.Slice(0, 16));
-				next = next.Slice(16);
-
-				BinaryPrimitives.WriteInt32LittleEndian(next, type.Version);
-				next = next.Slice(4);
+				type.Write(span);
+				span = span.Slice(BlobType.NumBytes);
 			}
-			Debug.Assert(next.Length == 0);
 		}
+
+		/// <inheritdoc/>
+		public BlobType this[int index] => _types[index];
+
+		/// <inheritdoc/>
+		public IEnumerator<BlobType> GetEnumerator() => ((IEnumerable<BlobType>)_types).GetEnumerator();
+
+		/// <inheritdoc/>
+		IEnumerator IEnumerable.GetEnumerator() => GetEnumerator();
 	}
 
 	/// <summary>
@@ -791,27 +688,33 @@ namespace EpicGames.Horde.Storage.Bundles
 	/// </summary>
 	public struct BundleImportCollection : IReadOnlyList<BundleLocator>
 	{
-		readonly ReadOnlyMemory<byte> _data;
+		readonly BundleLocator[] _imports;
 
-		/// <summary>
-		/// Data for the imports
-		/// </summary>
-		public ReadOnlyMemory<byte> Data => _data;
+		/// <inheritdoc/>
+		public int Count => _imports.Length;
 
 		/// <summary>
-		/// Deserializing constructor
+		/// Constructor
 		/// </summary>
-		public BundleImportCollection(ReadOnlyMemory<byte> data) => _data = data;
+		public BundleImportCollection(BundleLocator[] imports) => _imports = imports;
 
 		/// <summary>
-		/// Constructor
+		/// Reads a collection from memory
 		/// </summary>
-		/// <param name="locators">Locators to write to the </param>
-		public BundleImportCollection(IReadOnlyCollection<BundleLocator> locators)
+		public static BundleImportCollection Read(ReadOnlyMemory<byte> data)
 		{
-			byte[] data = new byte[Measure(locators)];
-			Write(data, locators);
-			_data = data;
+			ReadOnlySpan<byte> span = data.Span;
+			int count = (span.Length == 0) ? 0 : BinaryPrimitives.ReadInt32LittleEndian(span) / sizeof(int);
+
+			BundleLocator[] imports = new BundleLocator[count];
+			for (int idx = 0; idx < count; idx++)
+			{
+				int offset = BinaryPrimitives.ReadInt32LittleEndian(span.Slice(idx * sizeof(int)));
+				int length = span.Slice(offset).IndexOf((byte)0);
+				imports[idx] = new BundleLocator(new Utf8String(data.Slice(offset, length)));
+			}
+
+			return new BundleImportCollection(imports);
 		}
 
 		/// <summary>
@@ -821,36 +724,14 @@ namespace EpicGames.Horde.Storage.Bundles
 		{
 			byte[] data = new byte[length];
 			await stream.ReadFixedLengthBytesAsync(data, cancellationToken);
-			return new BundleImportCollection(data);
-		}
-
-		/// <summary>
-		/// Retrieve a single import from the collection
-		/// </summary>
-		public BundleLocator this[int index]
-		{
-			get
-			{
-				Debug.Assert(index < Count);
-				ReadOnlySpan<byte> span = _data.Span;
-				int offset = BinaryPrimitives.ReadInt32LittleEndian(span.Slice(index * sizeof(int)));
-				int length = span.Slice(offset).IndexOf((byte)0);
-				return new BundleLocator(new Utf8String(_data.Slice(offset, length)));
-			}
+			return Read(data);
 		}
 
 		/// <inheritdoc/>
-		public int Count => (_data.Length == 0) ? 0 : BinaryPrimitives.ReadInt32LittleEndian(_data.Span) / sizeof(int);
+		public BundleLocator this[int index] => _imports[index];
 
 		/// <inheritdoc/>
-		public IEnumerator<BundleLocator> GetEnumerator()
-		{
-			int count = Count;
-			for (int idx = 0; idx < count; idx++)
-			{
-				yield return this[idx];
-			}
-		}
+		public IEnumerator<BundleLocator> GetEnumerator() => _imports.AsEnumerable().GetEnumerator();
 
 		/// <inheritdoc/>
 		IEnumerator IEnumerable.GetEnumerator() => GetEnumerator();
@@ -858,45 +739,25 @@ namespace EpicGames.Horde.Storage.Bundles
 		/// <summary>
 		/// Measure the size of memory required to store a collection of import locators
 		/// </summary>
-		/// <param name="locators">Locators to write</param>
 		/// <returns>Size in bytes of the output buffer</returns>
-		public static int Measure(IReadOnlyCollection<BundleLocator> locators)
-		{
-			int length = 0;
-			foreach (BundleLocator locator in locators)
-			{
-				length += sizeof(int) + locator.Path.Length + 1;
-			}
-			return length;
-		}
+		public int Measure() => _imports.Sum(x => sizeof(int) + x.Path.Length + 1);
 
 		/// <summary>
 		/// Serialize a collection of locators to memory
 		/// </summary>
-		/// <param name="data">Output buffer for the serialized data</param>
-		/// <param name="locators">Locators to write</param>
-		public static void Write(Span<byte> data, IReadOnlyCollection<BundleLocator> locators)
+		/// <param name="writer">Writer for output data</param>
+		public void Write(IMemoryWriter writer)
 		{
-			Span<byte> next = data;
-
-			int offset = locators.Count * sizeof(int);
-			foreach (BundleLocator locator in locators)
+			int offset = _imports.Length * sizeof(int);
+			foreach (BundleLocator import in _imports)
 			{
-				BinaryPrimitives.WriteInt32LittleEndian(next, offset);
-				offset += locator.Path.Length + 1;
-				next = next.Slice(sizeof(int));
+				writer.WriteInt32(offset);
+				offset += import.Path.Length + 1;
 			}
-
-			foreach (BundleLocator locator in locators)
+			foreach (BundleLocator import in _imports)
 			{
-				locator.Path.Span.CopyTo(next);
-				next = next.Slice(locator.Path.Length);
-
-				next[0] = 0;
-				next = next.Slice(1);
+				writer.WriteNullTerminatedUtf8String(import.Path);
 			}
-
-			Debug.Assert(next.Length == 0);
 		}
 	}
 
@@ -946,20 +807,21 @@ namespace EpicGames.Horde.Storage.Bundles
 		}
 
 		/// <summary>
-		/// Constructor
+		/// Read from a byte array
 		/// </summary>
-		public BundlePacket(ReadOnlySpan<byte> span)
+		public static BundlePacket Read(ReadOnlySpan<byte> span)
 		{
-			CompressionFormat = (BundleCompressionFormat)BinaryPrimitives.ReadInt32LittleEndian(span);
-			EncodedOffset = BinaryPrimitives.ReadInt32LittleEndian(span.Slice(4));
-			EncodedLength = BinaryPrimitives.ReadInt32LittleEndian(span.Slice(8));
-			DecodedLength = BinaryPrimitives.ReadInt32LittleEndian(span.Slice(12));
+			BundleCompressionFormat compressionFormat = (BundleCompressionFormat)BinaryPrimitives.ReadInt32LittleEndian(span);
+			int encodedOffset = BinaryPrimitives.ReadInt32LittleEndian(span.Slice(4));
+			int encodedLength = BinaryPrimitives.ReadInt32LittleEndian(span.Slice(8));
+			int decodedLength = BinaryPrimitives.ReadInt32LittleEndian(span.Slice(12));
+			return new BundlePacket(compressionFormat, encodedOffset, encodedLength, decodedLength);
 		}
 
 		/// <summary>
 		/// Serialize the struct to memory
 		/// </summary>
-		public void CopyTo(Span<byte> span)
+		public void Write(Span<byte> span)
 		{
 			BinaryPrimitives.WriteInt32LittleEndian(span, (int)CompressionFormat);
 			BinaryPrimitives.WriteInt32LittleEndian(span.Slice(4), EncodedOffset);
@@ -973,28 +835,30 @@ namespace EpicGames.Horde.Storage.Bundles
 	/// </summary>
 	public struct BundlePacketCollection : IReadOnlyCollection<BundlePacket>
 	{
-		readonly ReadOnlyMemory<byte> _data;
+		readonly BundlePacket[] _packets;
 
-		/// <summary>
-		/// Data for the collection
-		/// </summary>
-		public ReadOnlyMemory<byte> Data => _data;
+		/// <inheritdoc/>
+		public int Count => _packets.Length;
 
 		/// <summary>
 		/// Constructor
 		/// </summary>
-		/// <param name="data"></param>
-		public BundlePacketCollection(ReadOnlyMemory<byte> data) => _data = data;
+		public BundlePacketCollection(BundlePacket[] packets) => _packets = packets;
 
 		/// <summary>
-		/// Constructor
+		/// Reads a collection from memory
 		/// </summary>
-		/// <param name="packets">Packets to include in this collection</param>
-		public BundlePacketCollection(IReadOnlyCollection<BundlePacket> packets)
+		public static BundlePacketCollection Read(ReadOnlyMemory<byte> data)
 		{
-			byte[] data = new byte[Measure(packets)];
-			Write(data, packets);
-			_data = data;
+			int count = data.Length / BundlePacket.NumBytes;
+
+			BundlePacket[] packets = new BundlePacket[count];
+			for (int idx = 0; idx < count; idx++)
+			{
+				packets[idx] = BundlePacket.Read(data.Span.Slice(idx * BundlePacket.NumBytes));
+			}
+
+			return new BundlePacketCollection(packets);
 		}
 
 		/// <summary>
@@ -1004,57 +868,36 @@ namespace EpicGames.Horde.Storage.Bundles
 		{
 			byte[] data = new byte[length];
 			await stream.ReadFixedLengthBytesAsync(data, cancellationToken);
-			return new BundlePacketCollection(data);
+			return BundlePacketCollection.Read(data);
 		}
 
 		/// <inheritdoc/>
-		public int Count => _data.Length / BundlePacket.NumBytes;
+		public BundlePacket this[int index] => _packets[index];
 
 		/// <inheritdoc/>
-		public BundlePacket this[int index]
-		{
-			get
-			{
-				Debug.Assert(index < Count);
-				return new BundlePacket(_data.Slice(index * BundlePacket.NumBytes, BundlePacket.NumBytes).Span);
-			}
-		}
-
-		/// <inheritdoc/>
-		public IEnumerator<BundlePacket> GetEnumerator()
-		{
-			ReadOnlyMemory<byte> remaining = _data;
-			while (remaining.Length > 0)
-			{
-				yield return new BundlePacket(remaining.Slice(0, BundlePacket.NumBytes).Span);
-				remaining = remaining.Slice(BundlePacket.NumBytes);
-			}
-		}
+		public IEnumerator<BundlePacket> GetEnumerator() => _packets.AsEnumerable().GetEnumerator();
 
 		/// <inheritdoc/>
 		IEnumerator IEnumerable.GetEnumerator() => GetEnumerator();
 
 		/// <summary>
-		/// Measure the size of memory required to store a collection of import locators
+		/// Measure the size of memory required to store a collection of packets
 		/// </summary>
-		/// <param name="packets">Locators to write</param>
 		/// <returns>Size in bytes of the output buffer</returns>
-		public static int Measure(IReadOnlyCollection<BundlePacket> packets) => BundlePacket.NumBytes * packets.Count;
+		public int Measure() => BundlePacket.NumBytes * _packets.Length;
 
 		/// <summary>
 		/// Serialize a collection of packets to memory
 		/// </summary>
-		/// <param name="data">Output buffer for the serialized data</param>
-		/// <param name="packets">Packets to write</param>
-		public static void Write(Span<byte> data, IReadOnlyCollection<BundlePacket> packets)
+		/// <param name="writer">Writer to serialize to</param>
+		public void Write(IMemoryWriter writer)
 		{
-			Span<byte> next = data;
-			foreach (BundlePacket packet in packets)
+			Span<byte> span = writer.GetMemoryAndAdvance(_packets.Length * BundlePacket.NumBytes).Span;
+			foreach (BundlePacket packet in _packets)
 			{
-				packet.CopyTo(next);
-				next = next.Slice(BundlePacket.NumBytes);
+				packet.Write(span);
+				span = span.Slice(BundlePacket.NumBytes);
 			}
-			Debug.Assert(next.Length == 0);
 		}
 	}
 
@@ -1066,13 +909,18 @@ namespace EpicGames.Horde.Storage.Bundles
 		/// <summary>
 		/// Number of bytes in a serialized export object
 		/// </summary>
-		public const int NumBytes = 32;
+		public const int NumHeaderBytes = 32;
 
 		/// <summary>
 		/// Raw data for this export
 		/// </summary>
 		public ReadOnlyMemory<byte> Data { get; }
 
+		/// <summary>
+		/// Raw data for the header
+		/// </summary>
+		public ReadOnlyMemory<byte> HeaderData => Data.Slice(0, NumHeaderBytes);
+
 		/// <summary>
 		/// Hash of the node data
 		/// </summary>
@@ -1101,15 +949,14 @@ namespace EpicGames.Horde.Storage.Bundles
 		/// <summary>
 		/// References to other nodes
 		/// </summary>
-		public BundleExportRefCollection References { get; }
+		public BundleExportRefCollection References => new BundleExportRefCollection(Data.Slice(32));
 
 		/// <summary>
 		/// Constructor
 		/// </summary>
-		public BundleExport(ReadOnlyMemory<byte> data, BundleExportRefCollection references)
+		public BundleExport(ReadOnlyMemory<byte> data)
 		{
 			Data = data;
-			References = references;
 		}
 
 		/// <summary>
@@ -1117,16 +964,21 @@ namespace EpicGames.Horde.Storage.Bundles
 		/// </summary>
 		public BundleExport(int typeIdx, IoHash hash, int packetIdx, int offset, int length, IReadOnlyList<BundleExportRef> references)
 		{
-			byte[] data = new byte[NumBytes];
-			Write(data, typeIdx, hash, packetIdx, offset, length);
+			byte[] data = new byte[NumHeaderBytes + (BundleExportRef.NumBytes * references.Count)];
+			WriteHeader(data, typeIdx, hash, packetIdx, offset, length);
+
+			for (int idx = 0; idx < references.Count; idx++)
+			{
+				references[idx].CopyTo(data.AsSpan(NumHeaderBytes + (BundleExportRef.NumBytes * idx)));
+			}
+
 			Data = data;
-			References = new BundleExportRefCollection(references);
 		}
 
 		/// <summary>
 		/// Writes a new export to a block of memory
 		/// </summary>
-		public static void Write(Span<byte> data, int typeIdx, IoHash hash, int packet, int offset, int length)
+		public static void WriteHeader(Span<byte> data, int typeIdx, IoHash hash, int packet, int offset, int length)
 		{
 			hash.CopyTo(data);
 			BinaryPrimitives.WriteUInt16LittleEndian(data.Slice(20), (ushort)typeIdx);
@@ -1139,96 +991,75 @@ namespace EpicGames.Horde.Storage.Bundles
 	/// <summary>
 	/// Entry for a node exported from an object
 	/// </summary>
-	public class BundleExportCollection : IReadOnlyList<BundleExport>
+	public struct BundleExportCollection : IReadOnlyList<BundleExport>
 	{
-		readonly int _count;
-		readonly ReadOnlyMemory<byte> _data;
-		readonly ReadOnlyMemory<byte> _refs;
+		readonly BundleExport[] _exports;
 
-		internal ReadOnlyMemory<byte> Data => _data;
-		internal ReadOnlyMemory<byte> Refs => _refs;
+		/// <inheritdoc/>
+		public int Count => _exports.Length;
 
 		/// <summary>
 		/// Constructor
 		/// </summary>
-		public BundleExportCollection(IReadOnlyList<BundleExport> exports)
-		{
-			_count = exports.Count;
-
-			byte[] data = new byte[Measure(exports)];
-			Write(data, exports);
-			_data = data;
-
-			byte[] refs = new byte[MeasureRefs(exports)];
-			WriteRefs(refs, exports);
-			_refs = refs;
-		}
+		public BundleExportCollection(BundleExport[] exports) => _exports = exports;
 
 		/// <summary>
 		/// Constructor
 		/// </summary>
-		public BundleExportCollection(ReadOnlyMemory<byte> data, ReadOnlyMemory<byte> refs, BundleVersion version)
+		public static BundleExportCollection Read(ReadOnlyMemory<byte> data, ReadOnlyMemory<byte> refs, BundleVersion version)
 		{
-			_count = data.Length / BundleExport.NumBytes;
-			_data = data;
-			_refs = refs;
+			int count = data.Length / BundleExport.NumHeaderBytes;
 
-			if (version < BundleVersion.ImportHashes && _refs.Length > 0)
+			if (version < BundleVersion.ImportHashes && refs.Length > 0)
 			{
 				const int OldRefNumBytes = 4;
 
 				int headerLength = BinaryPrimitives.ReadInt32LittleEndian(refs.Span);
-				int count = (refs.Length - headerLength) / OldRefNumBytes;
+				int refCount = (refs.Length - headerLength) / OldRefNumBytes;
 
-				byte[] newRefs = new byte[headerLength + (count * BundleExportRef.NumBytes)];
+				byte[] newRefs = new byte[headerLength + (refCount * BundleExportRef.NumBytes)];
 				for (int headerOffset = 0; headerOffset < headerLength; headerOffset += sizeof(int))
 				{
 					int index = (BinaryPrimitives.ReadInt32LittleEndian(refs.Slice(headerOffset).Span) - headerLength) / OldRefNumBytes;
 					BinaryPrimitives.WriteInt32LittleEndian(newRefs.AsSpan(headerOffset), headerLength + (index * BundleExportRef.NumBytes));
 				}
-				for (int index = 0; index < count; index++)
+				for (int index = 0; index < refCount; index++)
 				{
 					ReadOnlyMemory<byte> source = refs.Slice(headerLength + (index * OldRefNumBytes), OldRefNumBytes);
 					source.CopyTo(newRefs.AsMemory(headerLength + (index * BundleExportRef.NumBytes)));
 				}
 
-				_refs = newRefs;
+				refs = newRefs;
 			}
-		}
 
-		/// <inheritdoc/>
-		public int Count => _data.Length / BundleExport.NumBytes;
-
-		/// <inheritdoc/>
-		public BundleExport this[int index]
-		{
-			get
+			BundleExport[] exports = new BundleExport[count];
+			for(int index = 0; index < count; index++)
 			{
-				Debug.Assert(index < _count);
-
-				ReadOnlyMemory<byte> exportData = _data.Slice(index * BundleExport.NumBytes);
+				ReadOnlyMemory<byte> exportData = data.Slice(index * BundleExport.NumHeaderBytes, BundleExport.NumHeaderBytes);
 
 				BundleExportRefCollection exportRefs = new BundleExportRefCollection();
-				if (_refs.Length > 0)
+				if (refs.Length > 0)
 				{
-					int minOffset = BinaryPrimitives.ReadInt32LittleEndian(_refs.Span.Slice(index * sizeof(int)));
-					int maxOffset = (index == Count - 1) ? _refs.Length : BinaryPrimitives.ReadInt32LittleEndian(_refs.Span.Slice((index + 1) * sizeof(int)));
-					exportRefs = new BundleExportRefCollection(_refs.Slice(minOffset, maxOffset - minOffset));
+					int minOffset = BinaryPrimitives.ReadInt32LittleEndian(refs.Span.Slice(index * sizeof(int)));
+					int maxOffset = (index == count - 1) ? refs.Length : BinaryPrimitives.ReadInt32LittleEndian(refs.Span.Slice((index + 1) * sizeof(int)));
+					exportRefs = new BundleExportRefCollection(refs.Slice(minOffset, maxOffset - minOffset));
 				}
 
-				return new BundleExport(exportData, exportRefs);
+				byte[] export = new byte[BundleExport.NumHeaderBytes + exportRefs.Data.Length];
+				exportData.CopyTo(export);
+				exportRefs.Data.CopyTo(export.AsMemory(BundleExport.NumHeaderBytes));
+
+				exports[index] = new BundleExport(export);
 			}
+
+			return new BundleExportCollection(exports);
 		}
 
 		/// <inheritdoc/>
-		public IEnumerator<BundleExport> GetEnumerator()
-		{
-			int count = Count;
-			for (int idx = 0; idx < count; idx++)
-			{
-				yield return this[idx];
-			}
-		}
+		public BundleExport this[int index] => _exports[index];
+
+		/// <inheritdoc/>
+		public IEnumerator<BundleExport> GetEnumerator() => _exports.AsEnumerable().GetEnumerator();
 
 		/// <inheritdoc/>
 		IEnumerator IEnumerable.GetEnumerator() => GetEnumerator();
@@ -1236,49 +1067,42 @@ namespace EpicGames.Horde.Storage.Bundles
 		/// <summary>
 		/// Measure the size of memory required to store a collection of exports
 		/// </summary>
-		/// <param name="exports">Exports to write</param>
 		/// <returns>Size in bytes of the output buffer</returns>
-		public static int Measure(IReadOnlyCollection<BundleExport> exports) => BundleExport.NumBytes * exports.Count;
+		public int MeasureHeaders() => BundleExport.NumHeaderBytes * _exports.Length;
 
 		/// <summary>
 		/// Serialize a collection of exports to memory
 		/// </summary>
-		/// <param name="data">Output buffer for the serialized data</param>
-		/// <param name="exports">Exports to write</param>
-		public static void Write(Span<byte> data, IReadOnlyCollection<BundleExport> exports)
+		/// <param name="writer">Writer to serialize to</param>
+		public void WriteHeaders(IMemoryWriter writer)
 		{
-			Span<byte> next = data;
-			foreach (BundleExport export in exports)
+			foreach (BundleExport export in _exports)
 			{
-				export.Data.Span.CopyTo(next);
-				next = next.Slice(export.Data.Length);
+				writer.WriteFixedLengthBytes(export.HeaderData.Span);
 			}
 		}
 
 		/// <summary>
 		/// Measure the size of memory required to store a collection of export refs
 		/// </summary>
-		/// <param name="exports">Exports to write</param>
 		/// <returns>Size in bytes of the output buffer</returns>
-		public static int MeasureRefs(IReadOnlyCollection<BundleExport> exports) => (sizeof(int) * exports.Count) + exports.Sum(x => x.References.Data.Length);
+		public int MeasureRefs() => (sizeof(int) * _exports.Length) + _exports.Sum(x => x.References.Data.Length);
 
 		/// <summary>
 		/// Serialize a collection of export refs to memory
 		/// </summary>
-		/// <param name="data">Output buffer for the serialized data</param>
-		/// <param name="exports">Exports to write</param>
-		public static void WriteRefs(Span<byte> data, IReadOnlyCollection<BundleExport> exports)
+		/// <param name="writer">Writer to serialize to</param>
+		public void WriteRefs(IMemoryWriter writer)
 		{
-			int indexOffset = 0;
-			int dataOffset = sizeof(int) * exports.Count;
-
-			foreach (BundleExport export in exports)
+			int offset = sizeof(int) * _exports.Length;
+			foreach(BundleExport export in _exports)
 			{
-				BinaryPrimitives.WriteInt32LittleEndian(data.Slice(indexOffset), dataOffset);
-				indexOffset += sizeof(int);
-
-				export.References.Data.Span.CopyTo(data.Slice(dataOffset));
-				dataOffset += export.References.Data.Span.Length;
+				writer.WriteInt32(offset);
+				offset += export.References.Data.Length;
+			}
+			foreach (BundleExport export in _exports)
+			{
+				writer.WriteFixedLengthBytes(export.References.Data.Span);
 			}
 		}
 	}
diff --git a/Engine/Source/Programs/Shared/EpicGames.Horde/Storage/Bundles/BundleWriter.cs b/Engine/Source/Programs/Shared/EpicGames.Horde/Storage/Bundles/BundleWriter.cs
--- a/Engine/Source/Programs/Shared/EpicGames.Horde/Storage/Bundles/BundleWriter.cs
+++ b/Engine/Source/Programs/Shared/EpicGames.Horde/Storage/Bundles/BundleWriter.cs
@@ -636,7 +636,7 @@ namespace EpicGames.Horde.Storage.Bundles
 				}
 
 				// Create the bundle
-				BundleHeader header = BundleHeader.Create(types, imports, exports, _packets);
+				BundleHeader header = new BundleHeader(types.ToArray(), imports.ToArray(), exports.ToArray(), _packets.ToArray());
 				return new Bundle(header, packetData);
 			}
 
