Horde: Fix CI warning due to stream not being disposed.

diff --git a/Engine/Source/Programs/Shared/EpicGames.Horde/Storage/Backends/MemoryStorageClient.cs b/Engine/Source/Programs/Shared/EpicGames.Horde/Storage/Backends/MemoryStorageClient.cs
--- a/Engine/Source/Programs/Shared/EpicGames.Horde/Storage/Backends/MemoryStorageClient.cs
+++ b/Engine/Source/Programs/Shared/EpicGames.Horde/Storage/Backends/MemoryStorageClient.cs
@@ -69,8 +69,8 @@ namespace EpicGames.Horde.Storage.Backends
 		public override async Task<BundleLocator> WriteBundleAsync(Bundle bundle, Utf8String prefix = default, CancellationToken cancellationToken = default)
 		{
 			BundleLocator locator = BundleLocator.CreateUnique(prefix);
-			byte[] data = bundle.AsSequence().ToArray();
-			_bundles[locator] = await Bundle.FromStreamAsync(new MemoryStream(data), cancellationToken);
+			using ReadOnlySequenceStream stream = new ReadOnlySequenceStream(bundle.AsSequence());
+			_bundles[locator] = await Bundle.FromStreamAsync(stream, cancellationToken);
 			return locator;
 		}
 
