Horde: Read bundle exports into separate buffers to reduce pressure on LOH.

diff --git a/Engine/Source/Programs/Shared/EpicGames.Horde/Storage/Bundles/Bundle.cs b/Engine/Source/Programs/Shared/EpicGames.Horde/Storage/Bundles/Bundle.cs
--- a/Engine/Source/Programs/Shared/EpicGames.Horde/Storage/Bundles/Bundle.cs
+++ b/Engine/Source/Programs/Shared/EpicGames.Horde/Storage/Bundles/Bundle.cs
@@ -179,19 +179,24 @@ namespace EpicGames.Horde.Storage.Bundles
 		Imports = 1,
 
 		/// <summary>
-		/// List of exports
+		/// List of export headers
 		/// </summary>
-		Exports = 2,
+		ExportHeaders = 2,
 
 		/// <summary>
 		/// References to exports in other bundles
 		/// </summary>
-		ExportRefs = 3,
+		ExportReferences = 3,
 
 		/// <summary>
 		/// Packet headers
 		/// </summary>
 		Packets = 4,
+
+		/// <summary>
+		/// Merged export headers and references
+		/// </summary>
+		Exports = 5,
 	}
 
 	/// <summary>
@@ -273,10 +278,10 @@ namespace EpicGames.Horde.Storage.Bundles
 			WriteSectionHeader(writer, BundleSectionType.Imports, Imports.Measure());
 			Imports.Write(writer);
 
-			WriteSectionHeader(writer, BundleSectionType.Exports, Exports.MeasureHeaders());
+			WriteSectionHeader(writer, BundleSectionType.ExportHeaders, Exports.MeasureHeaders());
 			Exports.WriteHeaders(writer);
 
-			WriteSectionHeader(writer, BundleSectionType.ExportRefs, Exports.MeasureRefs());
+			WriteSectionHeader(writer, BundleSectionType.ExportReferences, Exports.MeasureRefs());
 			Exports.WriteRefs(writer);
 
 			WriteSectionHeader(writer, BundleSectionType.Packets, Packets.Measure());
@@ -364,55 +369,6 @@ namespace EpicGames.Horde.Storage.Bundles
 
 		static readonly byte[] s_latestSignature = new byte[] { (byte)'U', (byte)'B', (byte)'N' };
 
-		/// <summary>
-		/// Construct a header from the given data encoded in the latest format
-		/// </summary>
-		/// <param name="version">Version to serialize as</param>
-		/// <param name="data">Data for the header, including the prelude</param>
-		static BundleHeader ReadLatest(BundleVersion version, ReadOnlyMemory<byte> data)
-		{
-			ReadOnlyMemory<byte> exportData = ReadOnlyMemory<byte>.Empty;
-			ReadOnlyMemory<byte> exportRefData = ReadOnlyMemory<byte>.Empty;
-
-			BundleTypeCollection types = new BundleTypeCollection();
-			BundleImportCollection imports = new BundleImportCollection();
-			BundlePacketCollection packets = new BundlePacketCollection();
-
-			for (int offset = 0; offset < data.Length;)
-			{
-				uint header = BinaryPrimitives.ReadUInt32LittleEndian(data.Slice(offset).Span);
-				offset += SectionHeaderLength;
-
-				int length = (int)(header >> 8);
-				ReadOnlyMemory<byte> sectionData = data.Slice(offset, length);
-
-				BundleSectionType type = (BundleSectionType)(header & 255);
-				switch (type)
-				{
-					case BundleSectionType.Types:
-						types = BundleTypeCollection.Read(sectionData);
-						break;
-					case BundleSectionType.Imports:
-						imports = BundleImportCollection.Read(sectionData);
-						break;
-					case BundleSectionType.Exports:
-						exportData = sectionData;
-						break;
-					case BundleSectionType.ExportRefs:
-						exportRefData = sectionData;
-						break;
-					case BundleSectionType.Packets:
-						packets = BundlePacketCollection.Read(sectionData);
-						break;
-				}
-
-				offset += length;
-			}
-
-			BundleExportCollection exports = BundleExportCollection.Read(exportData, exportRefData, version);
-			return new BundleHeader(types, imports, exports, packets);
-		}
-
 		/// <summary>
 		/// Construct a header from the given data encoded in the latest format
 		/// </summary>
@@ -422,11 +378,12 @@ namespace EpicGames.Horde.Storage.Bundles
 		/// <param name="cancellationToken">Cancellation token for the operation</param>
 		static async Task<BundleHeader> ReadLatestAsync(BundleVersion version, Stream stream, int length, CancellationToken cancellationToken)
 		{
-			byte[] exportData = Array.Empty<byte>();
-			byte[] exportRefData = Array.Empty<byte>();
+			byte[] exportHeaders = Array.Empty<byte>();
+			byte[] exportReferences = Array.Empty<byte>();
 
 			BundleTypeCollection types = new BundleTypeCollection();
 			BundleImportCollection imports = new BundleImportCollection();
+			BundleExportCollection exports = new BundleExportCollection();
 			BundlePacketCollection packets = new BundlePacketCollection();
 
 			byte[] sectionHeader = new byte[SectionHeaderLength];
@@ -447,12 +404,15 @@ namespace EpicGames.Horde.Storage.Bundles
 						imports = await BundleImportCollection.ReadAsync(stream, sectionLength, cancellationToken);
 						break;
 					case BundleSectionType.Exports:
-						exportData = new byte[sectionLength];
-						await stream.ReadFixedLengthBytesAsync(exportData, cancellationToken);
+						exports = await BundleExportCollection.ReadAsync(stream, sectionLength, cancellationToken);
 						break;
-					case BundleSectionType.ExportRefs:
-						exportRefData = new byte[sectionLength];
-						await stream.ReadFixedLengthBytesAsync(exportRefData, cancellationToken);
+					case BundleSectionType.ExportHeaders:
+						exportHeaders = new byte[sectionLength];
+						await stream.ReadFixedLengthBytesAsync(exportHeaders, cancellationToken);
+						break;
+					case BundleSectionType.ExportReferences:
+						exportReferences = new byte[sectionLength];
+						await stream.ReadFixedLengthBytesAsync(exportReferences, cancellationToken);
 						break;
 					case BundleSectionType.Packets:
 						packets = await BundlePacketCollection.ReadAsync(stream, sectionLength, cancellationToken);
@@ -462,7 +422,11 @@ namespace EpicGames.Horde.Storage.Bundles
 				offset += sectionLength;
 			}
 
-			BundleExportCollection exports = BundleExportCollection.Read(exportData, exportRefData, version);
+			if (exportHeaders.Length > 0)
+			{
+				exports = BundleExportCollection.Read(exportHeaders, exportReferences, version);
+			}
+
 			return new BundleHeader(types, imports, exports, packets);
 		}
 
@@ -636,7 +600,7 @@ namespace EpicGames.Horde.Storage.Bundles
 		{
 			BlobType[] types = new BlobType[data.Length / BlobType.NumBytes];
 			for (int idx = 0; idx < types.Length; idx++)
-			{
+		{
 				types[idx] = BlobType.Read(data.Span.Slice(idx * BlobType.NumBytes));
 			}
 
@@ -708,7 +672,7 @@ namespace EpicGames.Horde.Storage.Bundles
 
 			BundleLocator[] imports = new BundleLocator[count];
 			for (int idx = 0; idx < count; idx++)
-			{
+		{
 				int offset = BinaryPrimitives.ReadInt32LittleEndian(span.Slice(idx * sizeof(int)));
 				int length = span.Slice(offset).IndexOf((byte)0);
 				imports[idx] = new BundleLocator(new Utf8String(data.Slice(offset, length)));
@@ -1055,6 +1019,70 @@ namespace EpicGames.Horde.Storage.Bundles
 			return new BundleExportCollection(exports);
 		}
 
+		/// <summary>
+		/// Reads an export collection from a stream
+		/// </summary>
+		public static async Task<BundleExportCollection> ReadAsync(Stream stream, int sectionLength, CancellationToken cancellationToken)
+		{
+			BundleExport[] exports = Array.Empty<BundleExport>();
+			if (sectionLength > 0)
+			{
+				const int MaxPageLength = 16384;
+
+				// Read the offsets of each entry
+				byte[] headerBuffer = new byte[sizeof(int)];
+				await stream.ReadFixedLengthBytesAsync(headerBuffer, cancellationToken);
+
+				int headerLength = BinaryPrimitives.ReadInt32LittleEndian(headerBuffer);
+				Array.Resize(ref headerBuffer, headerLength);
+
+				await stream.ReadFixedLengthBytesAsync(headerBuffer.AsMemory(sizeof(int)), cancellationToken);
+
+				// Create the packet array
+				int exportCount = headerLength / sizeof(int);
+				exports = new BundleExport[exportCount];
+
+				// Read the exports in pages
+				int exportIdx = 0;
+				for (int exportOffset = headerLength; exportOffset < sectionLength;)
+				{
+					int pageOffset = exportOffset;
+
+					// Find the exports to read into this page
+					int nextPageOffset = sectionLength;
+					for (int nextExportIdx = exportIdx + 1; nextExportIdx < exportCount; nextExportIdx++)
+					{
+						int nextExportOffset = BinaryPrimitives.ReadInt32LittleEndian(headerBuffer.AsSpan(nextExportIdx * sizeof(int)));
+						if (nextExportOffset > pageOffset + MaxPageLength)
+						{
+							nextPageOffset = nextExportOffset;
+							break;
+						}
+					}
+
+					// Read the page data
+					byte[] pageBuffer = new byte[nextPageOffset - pageOffset];
+					await stream.ReadFixedLengthBytesAsync(pageBuffer, cancellationToken);
+
+					// Parse the individual exports
+					for (; exportOffset < nextPageOffset; exportIdx++)
+					{
+						int startExportOffset = exportOffset;
+						if (exportIdx + 1 < exportCount)
+						{
+							exportOffset = BinaryPrimitives.ReadInt32LittleEndian(headerBuffer.AsSpan((exportIdx + 1) * sizeof(int)));
+						}
+						else
+						{
+							exportOffset = nextPageOffset;
+						}
+						exports[exportIdx] = new BundleExport(pageBuffer.AsMemory(startExportOffset - pageOffset, exportOffset - startExportOffset));
+					}
+				}
+			}
+			return new BundleExportCollection(exports);
+		}
+
 		/// <inheritdoc/>
 		public BundleExport this[int index] => _exports[index];
 
@@ -1064,6 +1092,24 @@ namespace EpicGames.Horde.Storage.Bundles
 		/// <inheritdoc/>
 		IEnumerator IEnumerable.GetEnumerator() => GetEnumerator();
 
+		/// <summary>
+		/// Serializes this collection as a single section
+		/// </summary>
+		public void Write(IMemoryWriter writer)
+		{
+			int offset = _exports.Length * sizeof(int);
+			Span<byte> index = writer.GetSpanAndAdvance(offset);
+
+			foreach (BundleExport export in _exports)
+			{
+				BinaryPrimitives.WriteInt32LittleEndian(index, offset);
+				index = index.Slice(sizeof(int));
+
+				writer.WriteFixedLengthBytes(export.Data.Span);
+				offset += export.Data.Length;
+			}
+		}
+
 		/// <summary>
 		/// Measure the size of memory required to store a collection of exports
 		/// </summary>
