Horde: Fix bundle reader operation when caching is disabled.

diff --git a/Engine/Source/Programs/Shared/EpicGames.Horde/Storage/Bundles/BundleReader.cs b/Engine/Source/Programs/Shared/EpicGames.Horde/Storage/Bundles/BundleReader.cs
--- a/Engine/Source/Programs/Shared/EpicGames.Horde/Storage/Bundles/BundleReader.cs
+++ b/Engine/Source/Programs/Shared/EpicGames.Horde/Storage/Bundles/BundleReader.cs
@@ -43,7 +43,7 @@ namespace EpicGames.Horde.Storage.Bundles
 		class QueuedHeader
 		{
 			public readonly BundleLocator Blob;
-			public readonly TaskCompletionSource<BundleInfo> CompletionSource = new TaskCompletionSource<BundleInfo>();
+			public readonly TaskCompletionSource<BundleInfo> CompletionSource = new TaskCompletionSource<BundleInfo>(TaskCreationOptions.RunContinuationsAsynchronously);
 
 			public Utf8String Path => Blob.Path;
 
@@ -60,7 +60,7 @@ namespace EpicGames.Horde.Storage.Bundles
 		{
 			public readonly BundleInfo Bundle;
 			public readonly int PacketIdx;
-			public readonly TaskCompletionSource<ReadOnlyMemory<byte>> CompletionSource = new TaskCompletionSource<ReadOnlyMemory<byte>>();
+			public readonly TaskCompletionSource<ReadOnlyMemory<byte>> CompletionSource = new TaskCompletionSource<ReadOnlyMemory<byte>>(TaskCreationOptions.RunContinuationsAsynchronously);
 
 			public Utf8String Path => Bundle.Locator.Path;
 
@@ -249,11 +249,10 @@ namespace EpicGames.Horde.Storage.Bundles
 					// Construct the bundle info
 					BundleInfo bundleInfo = new BundleInfo(queuedHeader.Blob, header, headerSize);
 
+					List<ReadOnlyMemory<byte>> packets = new List<ReadOnlyMemory<byte>>();
 					if (_cache != null)
 					{
 						// Also add any encoded packets we prefetched
-						List<ReadOnlyMemory<byte>> packets = new List<ReadOnlyMemory<byte>>();
-
 						int packetOffset = headerSize;
 						for (int packetIdx = 0; packetIdx < header.Packets.Count; packetIdx++)
 						{
@@ -273,24 +272,22 @@ namespace EpicGames.Horde.Storage.Bundles
 						// Add the info to the cache
 						string cacheKey = GetBundleInfoCacheKey(queuedHeader.Blob);
 						AddToCache(cacheKey, bundleInfo, headerSize);
+					}
 
-						// Update any packets that now have a cached value
-						List<QueuedPacket> updatePackets = new List<QueuedPacket>();
-						lock (_queueLock)
-						{
-							updatePackets.AddRange(_queuedPackets.Where(x => x.Path == bundleInfo.Locator.Path && x.PacketIdx < packets.Count));
-						}
+					// Update any packets that now have a cached value
+					lock (_queueLock)
+					{
+						queuedHeader.CompletionSource.TrySetResult(bundleInfo);
 
-						// Mark any packets that now have a valid cache entry as complete
-						foreach (QueuedPacket updatePacket in updatePackets)
+						foreach (QueuedPacket queuedPacket in _queuedPackets)
 						{
-							updatePacket.CompletionSource.TrySetResult(packets[updatePacket.PacketIdx]);
+							if (queuedPacket.Path == bundleInfo.Locator.Path && queuedPacket.PacketIdx < packets.Count)
+							{
+								queuedPacket.CompletionSource.TrySetResult(packets[queuedPacket.PacketIdx]);
+							}
 						}
 					}
 
-					// Update the task
-					queuedHeader.CompletionSource.SetResult(bundleInfo);
-
 					// Remove it from the queue
 					lock (_queueLock)
 					{
@@ -316,7 +313,7 @@ namespace EpicGames.Horde.Storage.Bundles
 			{
 				if (!queuedPacket.CompletionSource.TrySetException(ex))
 				{
-					_logger.LogWarning(ex, "Exception after setting completion source state; existing state: {Status}", queuedPacket.CompletionSource.Task.Status);
+					_logger.LogWarning(ex, "Exception after setting completion source state; existing state: {Status}, new exception: {Ex}", queuedPacket.CompletionSource.Task.Status, ex);
 				}
 			}
 		}
@@ -335,49 +332,53 @@ namespace EpicGames.Horde.Storage.Bundles
 			int readLength = packet.EncodedLength;
 
 			int maxPacketIdx = minPacketIdx + 1;
-			for (; maxPacketIdx < bundleInfo.Header.Packets.Count; maxPacketIdx++)
+			if(_cache != null)
 			{
-				int nextReadLength = readLength + bundleInfo.Header.Packets[maxPacketIdx].EncodedLength;
-				if (nextReadLength > DefaultFetchSize)
+				for (; maxPacketIdx < bundleInfo.Header.Packets.Count; maxPacketIdx++)
 				{
-					break;
+					int nextReadLength = readLength + bundleInfo.Header.Packets[maxPacketIdx].EncodedLength;
+					if (nextReadLength > DefaultFetchSize)
+					{
+						break;
+					}
+					readLength = nextReadLength;
 				}
-				readLength = nextReadLength;
 			}
 
 			await using (Stream stream = await _store.OpenAsync(bundleInfo.Locator, bundleInfo.HeaderLength + bundleInfo.Header.Packets[minPacketIdx].EncodedOffset, readLength, cancellationToken))
 			{
 				// Copy all the packets that have been read into separate buffers, so we can cache them indidually.
 				ReadOnlyMemory<byte>[] packets = new ReadOnlyMemory<byte>[maxPacketIdx - minPacketIdx];
+				for (int idx = minPacketIdx; idx < maxPacketIdx; idx++)
+				{
+					byte[] data = await ReadPacketAsync(stream, bundleInfo.Header.Packets[idx].EncodedLength, cancellationToken);
+					packets[idx - minPacketIdx] = data;
+				}
+
+				// Add any complete packets
 				if (_cache != null)
 				{
 					for (int idx = minPacketIdx; idx < maxPacketIdx; idx++)
 					{
+						ReadOnlyMemory<byte> data = packets[idx - minPacketIdx];
 						string cacheKey = GetEncodedPacketCacheKey(bundleInfo.Locator, idx);
-						byte[] data = await ReadPacketAsync(stream, bundleInfo.Header.Packets[idx].EncodedLength, cancellationToken);
-						packets[idx - minPacketIdx] = data;
 						AddToCache(cacheKey, (ReadOnlyMemory<byte>)data, data.Length);
 					}
 				}
 
 				// Find all the packets we can mark as complete
-				List<QueuedPacket> updatePackets = new List<QueuedPacket>();
 				lock (_queueLock)
 				{
-					updatePackets.AddRange(_queuedPackets.Where(x => x.Path == bundleInfo.Locator.Path && (x.PacketIdx >= minPacketIdx && x.PacketIdx < maxPacketIdx)));
-				}
-
-				// Mark them all as complete
-				foreach (QueuedPacket updatePacket in updatePackets)
-				{
-					ReadOnlyMemory<byte> data = packets[updatePacket.PacketIdx - minPacketIdx];
-					updatePacket.CompletionSource.SetResult(data);
-				}
-
-				// Remove all the completed packets from the queue
-				lock (_queueLock)
-				{
-					_queuedPackets.RemoveAll(x => updatePackets.Contains(x));
+					for (int idx = 0; idx < _queuedPackets.Count; idx++)
+					{
+						QueuedPacket updatePacket = _queuedPackets[idx];
+						if (updatePacket.Path == bundleInfo.Locator.Path && (updatePacket.PacketIdx >= minPacketIdx && updatePacket.PacketIdx < maxPacketIdx))
+						{
+							ReadOnlyMemory<byte> data = packets[updatePacket.PacketIdx - minPacketIdx];
+							updatePacket.CompletionSource.TrySetResult(data);
+							_queuedPackets.RemoveAt(idx--);
+						}
+					}
 				}
 			}
 		}
