EpicGames.Core: Add missing comments to ChunkedMemoryWriterBase.

diff --git a/Engine/Source/Programs/Shared/EpicGames.Core/MemoryWriter.cs b/Engine/Source/Programs/Shared/EpicGames.Core/MemoryWriter.cs
--- a/Engine/Source/Programs/Shared/EpicGames.Core/MemoryWriter.cs
+++ b/Engine/Source/Programs/Shared/EpicGames.Core/MemoryWriter.cs
@@ -254,21 +254,48 @@ namespace EpicGames.Core
 	/// </summary>
 	public abstract class ChunkedMemoryWriterBase : IMemoryWriter
 	{
+		/// <summary>
+		/// Describes a chunk of data in the output stream
+		/// </summary>
 		protected class Chunk
 		{
+			/// <summary>
+			/// Start position in the sequence
+			/// </summary>
 			public int RunningIndex { get; }
+
+			/// <summary>
+			/// Data for this chunk
+			/// </summary>
 			public Memory<byte> Data { get; }
+
+			/// <summary>
+			/// Used length of the chunk
+			/// </summary>
 			public int Length { get; set; }
 
+			/// <summary>
+			/// Constructor
+			/// </summary>
 			public Chunk(int runningIndex, Memory<byte> data)
 			{
 				RunningIndex = runningIndex;
 				Data = data;
 			}
 
+			/// <summary>
+			/// Release any resources managed by this chunk. Note that ChunkedMemoryWriterBase does not implement IDisposable directly; derived classes should call Clear() on Dispose() if needed.
+			/// </summary>
 			public virtual void Release() { }
 
+			/// <summary>
+			/// Span that has been written to
+			/// </summary>
 			public ReadOnlySpan<byte> WrittenSpan => WrittenMemory.Span;
+
+			/// <summary>
+			/// Memory that has been written to
+			/// </summary>
 			public ReadOnlyMemory<byte> WrittenMemory => Data.Slice(0, Length);
 		}
 
