Horde: Add a stats option to bundle extract command, showing the number of headers and packets read. Also add strongly typed cache methods.

diff --git a/Engine/Source/Programs/Horde/Horde/Commands/Bundles/BundleExtract.cs b/Engine/Source/Programs/Horde/Horde/Commands/Bundles/BundleExtract.cs
--- a/Engine/Source/Programs/Horde/Horde/Commands/Bundles/BundleExtract.cs
+++ b/Engine/Source/Programs/Horde/Horde/Commands/Bundles/BundleExtract.cs
@@ -4,6 +4,7 @@ using System.Diagnostics;
 using EpicGames.Core;
 using EpicGames.Horde.Storage;
 using EpicGames.Horde.Storage.Backends;
+using EpicGames.Horde.Storage.Bundles;
 using EpicGames.Horde.Storage.Nodes;
 using Microsoft.Extensions.Caching.Memory;
 using Microsoft.Extensions.Logging;
@@ -19,6 +20,9 @@ namespace Horde.Commands.Bundles
 		[CommandLine("-Ref=")]
 		public string? Ref { get; set; }
 
+		[CommandLine("-Stats")]
+		public bool Stats { get; set; }
+
 		[CommandLine("-OutputDir=", Required = true)]
 		public DirectoryReference OutputDir { get; set; } = null!;
 
@@ -49,6 +53,17 @@ namespace Horde.Commands.Bundles
 			await node.CopyToDirectoryAsync(OutputDir.ToDirectoryInfo(), logger, CancellationToken.None);
 
 			logger.LogInformation("Elapsed: {Time}s", timer.Elapsed.TotalSeconds);
+
+			if (Stats)
+			{
+				BundleStorageClient? bundleStorageClient = store as BundleStorageClient;
+				if (bundleStorageClient != null)
+				{
+					logger.LogInformation("Num header reads: {NumReads:n0}", bundleStorageClient.BundleReader.NumHeaderReads);
+					logger.LogInformation("Num packet reads: {NumReads:n0}", bundleStorageClient.BundleReader.NumPacketReads);
+				}
+			}
+
 			return 0;
 		}
 	}
diff --git a/Engine/Source/Programs/Shared/EpicGames.Horde/Storage/Bundles/BundleReader.cs b/Engine/Source/Programs/Shared/EpicGames.Horde/Storage/Bundles/BundleReader.cs
--- a/Engine/Source/Programs/Shared/EpicGames.Horde/Storage/Bundles/BundleReader.cs
+++ b/Engine/Source/Programs/Shared/EpicGames.Horde/Storage/Bundles/BundleReader.cs
@@ -1,7 +1,6 @@
 // Copyright Epic Games, Inc. All Rights Reserved.
 
 using System;
-using System.Buffers;
 using System.Collections.Generic;
 using System.Linq;
 using System.Threading;
@@ -12,6 +11,7 @@ using Microsoft.Extensions.Logging;
 using Microsoft.CodeAnalysis;
 using System.Diagnostics;
 using System.IO;
+using System.Diagnostics.CodeAnalysis;
 
 namespace EpicGames.Horde.Storage.Bundles
 {
@@ -87,6 +87,19 @@ namespace EpicGames.Horde.Storage.Bundles
 		readonly Dictionary<string, Task<ReadOnlyMemory<byte>>> _decodeTasks = new Dictionary<string, Task<ReadOnlyMemory<byte>>>(StringComparer.Ordinal);
 		Task? _readTask;
 
+		int _numHeaderReads;
+		int _numPacketReads;
+
+		/// <summary>
+		/// Total number of header reads
+		/// </summary>
+		public int NumHeaderReads => _numHeaderReads;
+
+		/// <summary>
+		/// Total number of packet reads
+		/// </summary>
+		public int NumPacketReads => _numPacketReads;
+
 		/// <summary>
 		/// Constructor
 		/// </summary>
@@ -100,23 +113,13 @@ namespace EpicGames.Horde.Storage.Bundles
 			_logger = logger;
 		}
 
-		#region Bundles
+		#region Cache
 
-		static string GetBundleInfoCacheKey(BundleLocator locator) => $"bundle:{locator}";
-		static string GetEncodedPacketCacheKey(BundleLocator locator, int packetIdx) => $"encoded-packet:{locator}#{packetIdx}";
-		static string GetDecodedPacketCacheKey(BundleLocator locator, int packetIdx) => $"decoded-packet:{locator}#{packetIdx}";
-
-		/// <summary>
-		/// Adds an object to the storage cache
-		/// </summary>
-		/// <param name="cacheKey">Key for the item</param>
-		/// <param name="value">Value to add</param>
-		/// <param name="size">Size of the value</param>
-		void AddToCache(string cacheKey, object value, int size)
+		static void AddCachedValue(IMemoryCache? cache, string cacheKey, object value, int size)
 		{
-			if (_cache != null)
+			if (cache != null)
 			{
-				using (ICacheEntry entry = _cache.CreateEntry(cacheKey))
+				using (ICacheEntry entry = cache.CreateEntry(cacheKey))
 				{
 					entry.SetValue(value);
 					entry.SetSize(size);
@@ -124,6 +127,33 @@ namespace EpicGames.Horde.Storage.Bundles
 			}
 		}
 
+		static bool TryGetCachedValue<T>(IMemoryCache? cache, string cacheKey, [MaybeNull, NotNullWhen(true)] out T value)
+		{
+			if (cache == null)
+			{
+				value = default;
+				return false;
+			}
+			return cache.TryGetValue(cacheKey, out value);
+		}
+
+		static string GetBundleInfoCacheKey(BundleLocator locator) => $"bundle:{locator}";
+		static string GetEncodedPacketCacheKey(BundleLocator locator, int packetIdx) => $"encoded-packet:{locator}#{packetIdx}";
+		static string GetDecodedPacketCacheKey(BundleLocator locator, int packetIdx) => $"decoded-packet:{locator}#{packetIdx}";
+
+		void AddCachedBundleInfo(BundleLocator locator, BundleInfo bundleInfo) => AddCachedValue(_cache, GetBundleInfoCacheKey(locator), bundleInfo, bundleInfo.HeaderLength);
+		bool TryGetCachedBundleInfo(BundleLocator locator, [NotNullWhen(true)] out BundleInfo? bundleInfo) => TryGetCachedValue(_cache, GetBundleInfoCacheKey(locator), out bundleInfo);
+
+		void AddCachedEncodedPacket(BundleLocator locator, int packetIdx, ReadOnlyMemory<byte> data) => AddCachedValue(_cache, GetEncodedPacketCacheKey(locator, packetIdx), data, data.Length);
+		bool TryGetCachedEncodedPacket(BundleLocator locator, int packetIdx, out ReadOnlyMemory<byte> data) => TryGetCachedValue(_cache, GetEncodedPacketCacheKey(locator, packetIdx), out data);
+
+		void AddCachedDecodedPacket(BundleLocator locator, int packetIdx, ReadOnlyMemory<byte> data) => AddCachedValue(_cache, GetDecodedPacketCacheKey(locator, packetIdx), data, data.Length);
+		bool TryGetCachedDecodedPacket(BundleLocator locator, int packetIdx, out ReadOnlyMemory<byte> data) => TryGetCachedValue(_cache, GetDecodedPacketCacheKey(locator, packetIdx), out data);
+
+		#endregion
+
+		#region Bundles
+
 		/// <summary>
 		/// Starts the background task for reading data from the store
 		/// </summary>
@@ -226,6 +256,8 @@ namespace EpicGames.Horde.Storage.Bundles
 		/// <param name="cancellationToken">Cancellation token for the operation</param>
 		async Task PerformHeaderReadAsync(QueuedHeader queuedHeader, CancellationToken cancellationToken)
 		{
+			Interlocked.Increment(ref _numHeaderReads);
+
 			int prefetchSize = _cache != null ? DefaultFetchSize : DefaultUncachedFetchSize;
 			for (; ; )
 			{
@@ -263,15 +295,14 @@ namespace EpicGames.Horde.Storage.Bundles
 							}
 
 							byte[] packetData = await ReadPacketAsync(stream, packetLength, cancellationToken);
-							AddToCache(GetEncodedPacketCacheKey(queuedHeader.Blob, packetIdx), (ReadOnlyMemory<byte>)packetData, packetData.Length);
+							AddCachedEncodedPacket(queuedHeader.Blob, packetIdx, packetData);
 							packets.Add(packetData);
 
 							packetOffset += packetLength;
 						}
 
 						// Add the info to the cache
-						string cacheKey = GetBundleInfoCacheKey(queuedHeader.Blob);
-						AddToCache(cacheKey, bundleInfo, headerSize);
+						AddCachedBundleInfo(queuedHeader.Blob, bundleInfo);
 					}
 
 					// Update any packets that now have a cached value
@@ -325,6 +356,8 @@ namespace EpicGames.Horde.Storage.Bundles
 		/// <param name="cancellationToken">Cancellation token for the operation</param>
 		async Task PerformPacketReadAsync(QueuedPacket queuedPacket, CancellationToken cancellationToken)
 		{
+			Interlocked.Increment(ref _numPacketReads);
+
 			BundleInfo bundleInfo = queuedPacket.Bundle;
 			int minPacketIdx = queuedPacket.PacketIdx;
 
@@ -361,8 +394,7 @@ namespace EpicGames.Horde.Storage.Bundles
 					for (int idx = minPacketIdx; idx < maxPacketIdx; idx++)
 					{
 						ReadOnlyMemory<byte> data = packets[idx - minPacketIdx];
-						string cacheKey = GetEncodedPacketCacheKey(bundleInfo.Locator, idx);
-						AddToCache(cacheKey, (ReadOnlyMemory<byte>)data, data.Length);
+						AddCachedEncodedPacket(bundleInfo.Locator, idx, data);
 					}
 				}
 
@@ -412,8 +444,8 @@ namespace EpicGames.Horde.Storage.Bundles
 		{
 			Debug.Assert(locator.IsValid());
 
-			string cacheKey = GetBundleInfoCacheKey(locator);
-			if (_cache != null && _cache.TryGetValue(cacheKey, out BundleInfo bundleInfo))
+			BundleInfo? bundleInfo;
+			if (TryGetCachedBundleInfo(locator, out bundleInfo))
 			{
 				return bundleInfo;
 			}
@@ -422,7 +454,7 @@ namespace EpicGames.Horde.Storage.Bundles
 			lock (_queueLock)
 			{
 				// Check the cache again inside lock scope to avoid races
-				if (_cache != null && _cache.TryGetValue(cacheKey, out bundleInfo))
+				if (TryGetCachedBundleInfo(locator, out bundleInfo))
 				{
 					return bundleInfo;
 				}
@@ -453,8 +485,8 @@ namespace EpicGames.Horde.Storage.Bundles
 				throw new ArgumentException("Packet index is out of range", nameof(packetIdx));
 			}
 
-			string decodedCacheKey = GetDecodedPacketCacheKey(bundleInfo.Locator, packetIdx);
-			if (_cache != null && _cache.TryGetValue(decodedCacheKey, out ReadOnlyMemory<byte> decodedPacket))
+			ReadOnlyMemory<byte> decodedPacket;
+			if (TryGetCachedDecodedPacket(bundleInfo.Locator, packetIdx, out decodedPacket))
 			{
 				return decodedPacket;
 			}
@@ -463,12 +495,13 @@ namespace EpicGames.Horde.Storage.Bundles
 			lock (_queueLock)
 			{
 				// Query the cache again, to eliminate races between cache checks and decode tasks finishing.
-				if (_cache != null && _cache.TryGetValue(decodedCacheKey, out decodedPacket))
+				if (TryGetCachedDecodedPacket(bundleInfo.Locator, packetIdx, out decodedPacket))
 				{
 					return decodedPacket;
 				}
 
 				// Create an async task to read the data
+				string decodedCacheKey = GetDecodedPacketCacheKey(bundleInfo.Locator, packetIdx);
 				if (!_decodeTasks.TryGetValue(decodedCacheKey, out decodeTask))
 				{
 					decodeTask = Task.Run(() => ReadAndDecodePacketAsync(bundleInfo, packetIdx), CancellationToken.None);
@@ -492,12 +525,11 @@ namespace EpicGames.Horde.Storage.Bundles
 			byte[] decodedPacket = new byte[packet.DecodedLength];
 
 			BundleData.Decompress(packet.CompressionFormat, encodedPacket, decodedPacket);
-
-			string decodedCacheKey = GetDecodedPacketCacheKey(bundleInfo.Locator, packetIdx);
-			AddToCache(decodedCacheKey, (ReadOnlyMemory<byte>)decodedPacket, decodedPacket.Length);
+			AddCachedDecodedPacket(bundleInfo.Locator, packetIdx, decodedPacket);
 
 			lock (_queueLock)
 			{
+				string decodedCacheKey = GetDecodedPacketCacheKey(bundleInfo.Locator, packetIdx);
 				_decodeTasks.Remove(decodedCacheKey);
 			}
 
@@ -517,8 +549,8 @@ namespace EpicGames.Horde.Storage.Bundles
 				throw new ArgumentException("Packet index is out of range", nameof(packetIdx));
 			}
 
-			string encodedCacheKey = GetEncodedPacketCacheKey(bundleInfo.Locator, packetIdx);
-			if (_cache != null && _cache.TryGetValue(encodedCacheKey, out ReadOnlyMemory<byte> encodedPacket))
+			ReadOnlyMemory<byte> encodedPacket;
+			if (TryGetCachedEncodedPacket(bundleInfo.Locator, packetIdx, out encodedPacket))
 			{
 				return encodedPacket;
 			}
@@ -527,7 +559,7 @@ namespace EpicGames.Horde.Storage.Bundles
 			lock (_queueLock)
 			{
 				// Query the cache again, to eliminate races between cache checks and decode tasks finishing.
-				if (_cache != null && _cache.TryGetValue(encodedCacheKey, out encodedPacket))
+				if (TryGetCachedEncodedPacket(bundleInfo.Locator, packetIdx, out encodedPacket))
 				{
 					return encodedPacket;
 				}
diff --git a/Engine/Source/Programs/Shared/EpicGames.Horde/Storage/Bundles/BundleStorageClient.cs b/Engine/Source/Programs/Shared/EpicGames.Horde/Storage/Bundles/BundleStorageClient.cs
--- a/Engine/Source/Programs/Shared/EpicGames.Horde/Storage/Bundles/BundleStorageClient.cs
+++ b/Engine/Source/Programs/Shared/EpicGames.Horde/Storage/Bundles/BundleStorageClient.cs
@@ -20,7 +20,7 @@ namespace EpicGames.Horde.Storage.Bundles
 		/// <summary>
 		/// Reader for node data
 		/// </summary>
-		protected BundleReader BundleReader { get; }
+		public BundleReader BundleReader { get; }
 
 		/// <summary>
 		/// Constructor
