Horde: Add a custom StorageCache implementation rather than using IMemoryCache.

This creates separate arenas for different types of content, improving read performance due to header data being discarded, and also allows cache instances to be globally configured and injected via DI.

diff --git a/Engine/Source/Programs/AutomationTool/BuildGraph/Tasks/DeployToolTask.cs b/Engine/Source/Programs/AutomationTool/BuildGraph/Tasks/DeployToolTask.cs
--- a/Engine/Source/Programs/AutomationTool/BuildGraph/Tasks/DeployToolTask.cs
+++ b/Engine/Source/Programs/AutomationTool/BuildGraph/Tasks/DeployToolTask.cs
@@ -151,7 +151,7 @@ namespace AutomationTool.Tasks
 
 			BundleNodeHandle handle;
 
-			HttpStorageClient storageClient = new HttpStorageClient(CreateHttpClient, () => new HttpClient(), null, Logger);
+			HttpStorageClient storageClient = new HttpStorageClient(CreateHttpClient, () => new HttpClient(), StorageCache.None, Logger);
 			await using (BundleWriter treeWriter = storageClient.CreateWriter())
 			{
 				DirectoryNode sandbox = new DirectoryNode();
diff --git a/Engine/Source/Programs/Horde/Horde.Agent.Tests/JsonRpcLoggerTests.cs b/Engine/Source/Programs/Horde/Horde.Agent.Tests/JsonRpcLoggerTests.cs
--- a/Engine/Source/Programs/Horde/Horde.Agent.Tests/JsonRpcLoggerTests.cs
+++ b/Engine/Source/Programs/Horde/Horde.Agent.Tests/JsonRpcLoggerTests.cs
@@ -140,9 +140,9 @@ namespace Horde.Agent.Tests
 		}
 
 		[TestMethod]
-		public async Task StorageLoggerTest()
+		public async Task StorageLoggerTestAsync()
 		{
-			using MemoryCache cache = new MemoryCache(new MemoryCacheOptions());
+			using StorageCache cache = new StorageCache();
 			MemoryStorageClient store = new MemoryStorageClient();
 
 			BundleReader reader = new BundleReader(store, cache, NullLogger.Instance);
diff --git a/Engine/Source/Programs/Horde/Horde.Agent/Commands/Compute/ComputeCommand.cs b/Engine/Source/Programs/Horde/Horde.Agent/Commands/Compute/ComputeCommand.cs
--- a/Engine/Source/Programs/Horde/Horde.Agent/Commands/Compute/ComputeCommand.cs
+++ b/Engine/Source/Programs/Horde/Horde.Agent/Commands/Compute/ComputeCommand.cs
@@ -13,6 +13,8 @@ using EpicGames.Horde.Common;
 using EpicGames.Horde.Compute;
 using EpicGames.Horde.Compute.Clients;
 using EpicGames.Horde.Compute.Transports;
+using EpicGames.Horde.Storage;
+using EpicGames.Horde.Storage.Bundles;
 using Microsoft.Extensions.Caching.Memory;
 using Microsoft.Extensions.DependencyInjection;
 using Microsoft.Extensions.Logging;
@@ -113,14 +115,14 @@ namespace Horde.Agent.Commands.Compute
 	[Command("computeworker", "Runs the agent as a local compute host, accepting incoming connections on the loopback adapter with a given port")]
 	class ComputeWorkerCommand : Command
 	{
-		readonly IMemoryCache _memoryCache;
+		readonly StorageCache _storageCache;
 
 		[CommandLine("-Port=")]
 		int Port { get; set; } = 2000;
 
-		public ComputeWorkerCommand(IMemoryCache memoryCache)
+		public ComputeWorkerCommand(StorageCache storageCache)
 		{
-			_memoryCache = memoryCache;
+			_storageCache = storageCache;
 		}
 
 		public override async Task<int> ExecuteAsync(ILogger logger)
@@ -133,7 +135,7 @@ namespace Horde.Agent.Commands.Compute
 			await using (RemoteComputeSocket socket = new RemoteComputeSocket(new TcpTransport(tcpSocket), logger))
 			{
 				logger.LogInformation("Running worker...");
-				await RunWorkerAsync(socket, _memoryCache, logger, CancellationToken.None);
+				await RunWorkerAsync(socket, _storageCache, logger, CancellationToken.None);
 				logger.LogInformation("Worker complete");
 				await socket.CloseAsync(CancellationToken.None);
 			}
@@ -142,11 +144,11 @@ namespace Horde.Agent.Commands.Compute
 			return 0;
 		}
 
-		public static async Task RunWorkerAsync(ComputeSocket socket, IMemoryCache memoryCache, ILogger logger, CancellationToken cancellationToken)
+		public static async Task RunWorkerAsync(ComputeSocket socket, StorageCache storageCache, ILogger logger, CancellationToken cancellationToken)
 		{
 			DirectoryReference sandboxDir = DirectoryReference.Combine(Program.DataDir, "Sandbox");
 
-			AgentMessageHandler worker = new AgentMessageHandler(sandboxDir, memoryCache, null, false, null, logger);
+			AgentMessageHandler worker = new AgentMessageHandler(sandboxDir, storageCache, null, false, null, logger);
 			await worker.RunAsync(socket, cancellationToken);
 		}
 	}
diff --git a/Engine/Source/Programs/Horde/Horde.Agent/Leases/Handlers/ComputeHandler.cs b/Engine/Source/Programs/Horde/Horde.Agent/Leases/Handlers/ComputeHandler.cs
--- a/Engine/Source/Programs/Horde/Horde.Agent/Leases/Handlers/ComputeHandler.cs
+++ b/Engine/Source/Programs/Horde/Horde.Agent/Leases/Handlers/ComputeHandler.cs
@@ -10,6 +10,8 @@ using System.Threading.Tasks;
 using EpicGames.Core;
 using EpicGames.Horde.Compute;
 using EpicGames.Horde.Compute.Transports;
+using EpicGames.Horde.Storage;
+using EpicGames.Horde.Storage.Bundles;
 using Horde.Agent.Services;
 using Horde.Agent.Utility;
 using HordeCommon.Rpc.Tasks;
@@ -61,7 +63,7 @@ namespace Horde.Agent.Leases.Handlers
 		static TimeSpan NoDataTimeout { get; } = TimeSpan.FromSeconds(20);
 
 		readonly ComputeListenerService _listenerService;
-		readonly IMemoryCache _memoryCache;
+		readonly StorageCache _storageCache;
 		readonly IServerLoggerFactory _serverLoggerFactory;
 		readonly AgentSettings _settings;
 		readonly ILogger _logger;
@@ -69,10 +71,10 @@ namespace Horde.Agent.Leases.Handlers
 		/// <summary>
 		/// Constructor
 		/// </summary>
-		public ComputeHandler(ComputeListenerService listenerService, IMemoryCache memoryCache, IServerLoggerFactory serverLoggerFactory, IOptions<AgentSettings> settings, ILogger<ComputeHandler> logger)
+		public ComputeHandler(ComputeListenerService listenerService, StorageCache storageCache, IServerLoggerFactory serverLoggerFactory, IOptions<AgentSettings> settings, ILogger<ComputeHandler> logger)
 		{
 			_listenerService = listenerService;
-			_memoryCache = memoryCache;
+			_storageCache = storageCache;
 			_serverLoggerFactory = serverLoggerFactory;
 			_settings = settings.Value;
 			_logger = logger;
@@ -124,7 +126,7 @@ namespace Horde.Agent.Leases.Handlers
 								Dictionary<string, string?> newEnvVars = new Dictionary<string, string?>();
 								newEnvVars["UE_HORDE_SHARED_DIR"] = sharedDir.FullName;
 
-								AgentMessageHandler worker = new AgentMessageHandler(sandboxDir, _memoryCache, newEnvVars, false, _settings.WineExecutablePath, logger);
+								AgentMessageHandler worker = new AgentMessageHandler(sandboxDir, _storageCache, newEnvVars, false, _settings.WineExecutablePath, logger);
 								await worker.RunAsync(socket, cts.Token);
 								await socket.CloseAsync(cts.Token);
 								return LeaseResult.Success;
diff --git a/Engine/Source/Programs/Horde/Horde.Agent/Program.cs b/Engine/Source/Programs/Horde/Horde.Agent/Program.cs
--- a/Engine/Source/Programs/Horde/Horde.Agent/Program.cs
+++ b/Engine/Source/Programs/Horde/Horde.Agent/Program.cs
@@ -222,6 +222,7 @@ namespace Horde.Agent
 			services.AddHostedService(sp => sp.GetRequiredService<WorkerService>());
 
 			services.AddSingleton<IStorageClientFactory, StorageClientFactory>();
+			services.AddSingleton<StorageCache>();
 
 			services.AddSingleton<ComputeListenerService>();
 			services.AddHostedService(sp => sp.GetRequiredService<ComputeListenerService>());
@@ -232,7 +233,7 @@ namespace Horde.Agent
 			services.AddSingleton<DefaultServices>(x => new DefaultServices(configuration, services));
 
 			// Execute all the commands
-			IServiceProvider serviceProvider = services.BuildServiceProvider();
+			await using ServiceProvider serviceProvider = services.BuildServiceProvider();
 			return await CommandHost.RunAsync(arguments, serviceProvider, typeof(Horde.Agent.Modes.Service.RunCommand));
 		}
 
diff --git a/Engine/Source/Programs/Horde/Horde.Agent/Services/ServerStorageFactory.cs b/Engine/Source/Programs/Horde/Horde.Agent/Services/ServerStorageFactory.cs
--- a/Engine/Source/Programs/Horde/Horde.Agent/Services/ServerStorageFactory.cs
+++ b/Engine/Source/Programs/Horde/Horde.Agent/Services/ServerStorageFactory.cs
@@ -5,6 +5,7 @@ using System.Net.Http;
 using System.Net.Http.Headers;
 using EpicGames.Horde.Storage;
 using EpicGames.Horde.Storage.Backends;
+using EpicGames.Horde.Storage.Bundles;
 using Microsoft.Extensions.Caching.Memory;
 using Microsoft.Extensions.Logging;
 
@@ -75,23 +76,23 @@ namespace Horde.Agent.Services
 	class HttpServerStorageFactory : IServerStorageFactory
 	{
 		readonly IHttpClientFactory _httpClientFactory;
-		readonly IMemoryCache _memoryCache;
+		readonly StorageCache _storageCache;
 		readonly ILogger _logger;
 
 		/// <summary>
 		/// Constructor
 		/// </summary>
-		public HttpServerStorageFactory(IHttpClientFactory httpClientFactory, IMemoryCache memoryCache, ILogger<HttpStorageClient> logger)
+		public HttpServerStorageFactory(IHttpClientFactory httpClientFactory, StorageCache storageCache, ILogger<HttpStorageClient> logger)
 		{
 			_httpClientFactory = httpClientFactory;
-			_memoryCache = memoryCache;
+			_storageCache = storageCache;
 			_logger = logger;
 		}
 
 		/// <inheritdoc/>
 		public IStorageClient CreateStorageClient(Uri baseAddress, string token)
 		{
-			return new HttpStorageClient(() => CreateHttpClient(baseAddress, token), CreateHttpRedirectClient, _memoryCache, _logger);
+			return new HttpStorageClient(() => CreateHttpClient(baseAddress, token), CreateHttpRedirectClient, _storageCache, _logger);
 		}
 
 		HttpClient CreateHttpClient(Uri baseAddress, string token)
diff --git a/Engine/Source/Programs/Horde/Horde.Agent/Utility/StorageClientFactory.cs b/Engine/Source/Programs/Horde/Horde.Agent/Utility/StorageClientFactory.cs
--- a/Engine/Source/Programs/Horde/Horde.Agent/Utility/StorageClientFactory.cs
+++ b/Engine/Source/Programs/Horde/Horde.Agent/Utility/StorageClientFactory.cs
@@ -10,7 +10,7 @@ using Microsoft.Extensions.Options;
 using EpicGames.Core;
 using System.Threading;
 using System.Threading.Tasks;
-using Microsoft.Extensions.Caching.Memory;
+using EpicGames.Horde.Storage.Bundles;
 
 namespace Horde.Agent.Utility
 {
@@ -20,13 +20,13 @@ namespace Horde.Agent.Utility
 	class StorageClientFactory : IStorageClientFactory
 	{
 		readonly IOptions<AgentSettings> _settings;
-		readonly IMemoryCache _memoryCache;
+		readonly StorageCache _memoryCache;
 		readonly ILogger _logger;
 
 		/// <summary>
 		/// Constructor
 		/// </summary>
-		public StorageClientFactory(IOptions<AgentSettings> settings, IMemoryCache memoryCache, ILogger<IStorageClient> logger)
+		public StorageClientFactory(IOptions<AgentSettings> settings, StorageCache memoryCache, ILogger<IStorageClient> logger)
 		{
 			_settings = settings;
 			_memoryCache = memoryCache;
diff --git a/Engine/Source/Programs/Horde/Horde.Server.Tests/TestSetup.cs b/Engine/Source/Programs/Horde/Horde.Server.Tests/TestSetup.cs
--- a/Engine/Source/Programs/Horde/Horde.Server.Tests/TestSetup.cs
+++ b/Engine/Source/Programs/Horde/Horde.Server.Tests/TestSetup.cs
@@ -276,6 +276,7 @@ namespace Horde.Server.Tests
 			services.AddSingleton<IStorageBackend<ArtifactCollectionV1>>(sp => new MemoryStorageBackend().ForType<ArtifactCollectionV1>());
 
 			services.AddSingleton<StorageService>();
+			services.AddSingleton<StorageCache>();
 
 			services.AddSingleton<ISingletonDocument<DevicePlatformMapV1>>(new SingletonDocumentStub<DevicePlatformMapV1>());
 		}
diff --git a/Engine/Source/Programs/Horde/Horde.Server/Commands/Bundles/CreateCommand.cs b/Engine/Source/Programs/Horde/Horde.Server/Commands/Bundles/CreateCommand.cs
--- a/Engine/Source/Programs/Horde/Horde.Server/Commands/Bundles/CreateCommand.cs
+++ b/Engine/Source/Programs/Horde/Horde.Server/Commands/Bundles/CreateCommand.cs
@@ -51,7 +51,7 @@ namespace Horde.Server.Commands.Bundles
 			}
 			else
 			{
-				FileStorageClient fileStorageClient = new FileStorageClient(OutputDir, serviceProvider.GetRequiredService<IMemoryCache>(), logger);
+				FileStorageClient fileStorageClient = new FileStorageClient(OutputDir, serviceProvider.GetRequiredService<StorageCache>(), logger);
 				storageClient = fileStorageClient;
 			}
 		
diff --git a/Engine/Source/Programs/Horde/Horde.Server/Commands/Install/SetupCommand.cs b/Engine/Source/Programs/Horde/Horde.Server/Commands/Install/SetupCommand.cs
--- a/Engine/Source/Programs/Horde/Horde.Server/Commands/Install/SetupCommand.cs
+++ b/Engine/Source/Programs/Horde/Horde.Server/Commands/Install/SetupCommand.cs
@@ -67,9 +67,9 @@ namespace Horde.Server.Commands.Install
 
 			RefName refName = new RefName("latest");
 
-			using IMemoryCache memoryCache = new MemoryCache(new MemoryCacheOptions());
+			using StorageCache storageCache = new StorageCache();
 
-			FileStorageClient client = new FileStorageClient(bundleDir, memoryCache, logger);
+			FileStorageClient client = new FileStorageClient(bundleDir, storageCache, logger);
 			await using (IStorageWriter writer = client.CreateWriter(refName))
 			{
 				DirectoryNode dirNode = new DirectoryNode();
diff --git a/Engine/Source/Programs/Horde/Horde.Server/Startup.cs b/Engine/Source/Programs/Horde/Horde.Server/Startup.cs
--- a/Engine/Source/Programs/Horde/Horde.Server/Startup.cs
+++ b/Engine/Source/Programs/Horde/Horde.Server/Startup.cs
@@ -523,11 +523,12 @@ namespace Horde.Server
 				services.AddSingleton<IAvatarService, NullAvatarService>();
 			}
 
-			services.AddSingleton<DeviceService>();						
+			services.AddSingleton<DeviceService>();
 			services.AddSingleton<NoticeService>();
 			services.AddSingleton<StorageService>();
 			services.AddSingleton<IStorageClientFactory>(sp => sp.GetRequiredService<StorageService>());
 			services.AddSingleton<TestDataService>();
+			services.AddSingleton<StorageCache>();
 
 			if (settings.JiraUrl != null)
 			{
@@ -1256,7 +1257,7 @@ namespace Horde.Server
 			serviceCollection.AddSingleton<IHostApplicationLifetime, HostApplicationLifetime>();
 		}
 
-		public static IServiceProvider CreateServiceProvider(IConfiguration configuration)
+		public static ServiceProvider CreateServiceProvider(IConfiguration configuration)
 		{
 			IServiceCollection serviceCollection = new ServiceCollection();
 			AddServices(serviceCollection, configuration);
diff --git a/Engine/Source/Programs/Horde/Horde.Server/Storage/StorageController.cs b/Engine/Source/Programs/Horde/Horde.Server/Storage/StorageController.cs
--- a/Engine/Source/Programs/Horde/Horde.Server/Storage/StorageController.cs
+++ b/Engine/Source/Programs/Horde/Horde.Server/Storage/StorageController.cs
@@ -164,17 +164,17 @@ namespace Horde.Server.Storage
 	public class StorageController : HordeControllerBase
 	{
 		readonly StorageService _storageService;
-		readonly IMemoryCache _memoryCache;
+		readonly StorageCache _storageCache;
 		readonly IOptionsSnapshot<GlobalConfig> _globalConfig;
 		readonly ILogger _logger;
 
 		/// <summary>
 		/// Constructor
 		/// </summary>
-		public StorageController(StorageService storageService, IMemoryCache memoryCache, IOptionsSnapshot<GlobalConfig> globalConfig, ILogger<StorageController> logger)
+		public StorageController(StorageService storageService, StorageCache storageCache, IOptionsSnapshot<GlobalConfig> globalConfig, ILogger<StorageController> logger)
 		{
 			_storageService = storageService;
-			_memoryCache = memoryCache;
+			_storageCache = storageCache;
 			_globalConfig = globalConfig;
 			_logger = logger;
 		}
@@ -500,7 +500,7 @@ namespace Horde.Server.Storage
 				return Forbid(StorageAclAction.ReadBlobs, namespaceId);
 			}
 
-			BundleReader reader = new BundleReader(storageClient, _memoryCache, _logger);
+			BundleReader reader = new BundleReader(storageClient, _storageCache, _logger);
 
 			BundleHeader header = await reader.ReadBundleHeaderAsync(locator, cancellationToken);
 
@@ -580,7 +580,7 @@ namespace Horde.Server.Storage
 			}
 
 			StorageClient storageClient = await _storageService.GetClientAsync(namespaceId, cancellationToken);
-			BundleReader reader = new BundleReader(storageClient, _memoryCache, _logger);
+			BundleReader reader = new BundleReader(storageClient, _storageCache, _logger);
 
 			BundleHeader header = await reader.ReadBundleHeaderAsync(locator, cancellationToken);
 			BundleExport export = header.Exports[exportIdx];
diff --git a/Engine/Source/Programs/Horde/Horde.Server/Storage/StorageService.cs b/Engine/Source/Programs/Horde/Horde.Server/Storage/StorageService.cs
--- a/Engine/Source/Programs/Horde/Horde.Server/Storage/StorageService.cs
+++ b/Engine/Source/Programs/Horde/Horde.Server/Storage/StorageService.cs
@@ -75,10 +75,10 @@ namespace Horde.Server.Storage
 		/// </summary>
 		/// <param name="config">Namespace configuration</param>
 		/// <param name="backend">Backend store</param>
-		/// <param name="memoryCache">Memory cache</param>
+		/// <param name="storageCache">Storage cache</param>
 		/// <param name="logger">Logger instance</param>
-		protected StorageClient(NamespaceConfig config, IStorageBackend backend, IMemoryCache? memoryCache, ILogger logger)
-			: base(memoryCache, logger)
+		protected StorageClient(NamespaceConfig config, IStorageBackend backend, StorageCache storageCache, ILogger logger)
+			: base(storageCache, logger)
 		{
 			Config = config;
 			Backend = backend;
@@ -123,8 +123,8 @@ namespace Horde.Server.Storage
 			readonly Tracer _tracer;
 			readonly ILogger _logger;
 
-			public StorageClientImpl(StorageService outer, NamespaceConfig config, IStorageBackend backend, IMemoryCache? memoryCache, Tracer tracer, ILogger logger)
-				: base(config, backend, memoryCache, logger)
+			public StorageClientImpl(StorageService outer, NamespaceConfig config, IStorageBackend backend, StorageCache storageCache, Tracer tracer, ILogger logger)
+				: base(config, backend, storageCache, logger)
 			{
 				_outer = outer;
 				_tracer = tracer;
@@ -433,7 +433,8 @@ namespace Horde.Server.Storage
 
 		readonly RedisService _redisService;
 		readonly IClock _clock;
-		readonly IMemoryCache _cache;
+		readonly StorageCache _storageCache;
+		readonly IMemoryCache _memoryCache;
 		readonly IStorageBackendProvider _storageBackendProvider;
 		readonly IOptionsMonitor<GlobalConfig> _globalConfig;
 		readonly Tracer _tracer;
@@ -456,11 +457,12 @@ namespace Horde.Server.Storage
 		/// <summary>
 		/// Constructor
 		/// </summary>
-		public StorageService(MongoService mongoService, RedisService redisService, IClock clock, IMemoryCache cache, IStorageBackendProvider storageBackendProvider, IOptionsMonitor<GlobalConfig> globalConfig, Tracer tracer, ILogger<StorageService> logger)
+		public StorageService(MongoService mongoService, RedisService redisService, IClock clock, StorageCache storageCache, IMemoryCache memoryCache, IStorageBackendProvider storageBackendProvider, IOptionsMonitor<GlobalConfig> globalConfig, Tracer tracer, ILogger<StorageService> logger)
 		{
 			_redisService = redisService;
 			_clock = clock;
-			_cache = cache;
+			_storageCache = storageCache;
+			_memoryCache = memoryCache;
 			_storageBackendProvider = storageBackendProvider;
 			_cachedState = new AsyncCachedValue<State>(() => Task.FromResult(GetNextState()), TimeSpan.FromMinutes(1.0));
 			_globalConfig = globalConfig;
@@ -586,7 +588,7 @@ namespace Horde.Server.Storage
 					foreach (NamespaceConfig namespaceConfig in storageConfig.Namespaces)
 					{
 						IStorageBackend backend = _storageBackendProvider.CreateBackend(namespaceConfig.BackendConfig);
-						StorageClientImpl client = new StorageClientImpl(this, namespaceConfig, backend, _cache, _tracer, _logger);
+						StorageClientImpl client = new StorageClientImpl(this, namespaceConfig, backend, _storageCache, _tracer, _logger);
 						nextState.Namespaces.Add(namespaceConfig.Id, new NamespaceInfo(namespaceConfig, client, backend));
 					}
 				}
@@ -787,7 +789,7 @@ namespace Horde.Server.Storage
 		RefCacheValue AddRefToCache(NamespaceId namespaceId, RefName name, RefInfo? value)
 		{
 			RefCacheValue cacheValue = new RefCacheValue(value, DateTime.UtcNow);
-			using (ICacheEntry newEntry = _cache.CreateEntry(new RefCacheKey(namespaceId, name)))
+			using (ICacheEntry newEntry = _memoryCache.CreateEntry(new RefCacheKey(namespaceId, name)))
 			{
 				newEntry.Value = cacheValue;
 				newEntry.SetSize(name.Text.Length);
@@ -862,7 +864,7 @@ namespace Horde.Server.Storage
 		async Task<BundleNodeLocator?> TryReadRefTargetAsync(NamespaceId namespaceId, RefName name, RefCacheTime cacheTime = default, CancellationToken cancellationToken = default)
 		{
 			RefCacheValue entry;
-			if (!_cache.TryGetValue(name, out entry) || RefCacheTime.IsStaleCacheEntry(entry.Time, cacheTime))
+			if (!_memoryCache.TryGetValue(name, out entry) || RefCacheTime.IsStaleCacheEntry(entry.Time, cacheTime))
 			{
 				RefInfo? refDocument = await _refCollection.Find(x => x.NamespaceId == namespaceId && x.Name == name).FirstOrDefaultAsync(cancellationToken);
 				entry = AddRefToCache(namespaceId, name, refDocument);
diff --git a/Engine/Source/Programs/Horde/Horde.Server/Tools/ToolCollection.cs b/Engine/Source/Programs/Horde/Horde.Server/Tools/ToolCollection.cs
--- a/Engine/Source/Programs/Horde/Horde.Server/Tools/ToolCollection.cs
+++ b/Engine/Source/Programs/Horde/Horde.Server/Tools/ToolCollection.cs
@@ -158,7 +158,7 @@ namespace Horde.Server.Tools
 		private readonly VersionedCollection<ToolId, Tool> _tools;
 		private readonly StorageService _storageService;
 		private readonly IClock _clock;
-		private readonly IMemoryCache _cache;
+		private readonly StorageCache _cache;
 		private readonly ILogger _logger;
 
 		private static readonly RedisKey s_baseKey = "tools/v1/";
@@ -168,7 +168,7 @@ namespace Horde.Server.Tools
 		/// <summary>
 		/// Constructor
 		/// </summary>
-		public ToolCollection(MongoService mongoService, RedisService redisService, StorageService storageService, IMemoryCache cache, IClock clock, ILogger<ToolCollection> logger)
+		public ToolCollection(MongoService mongoService, RedisService redisService, StorageService storageService, StorageCache cache, IClock clock, ILogger<ToolCollection> logger)
 		{
 			_tools = new VersionedCollection<ToolId, Tool>(mongoService, "Tools", redisService, s_baseKey, s_types);
 			_storageService = storageService;
diff --git a/Engine/Source/Programs/Horde/Horde/Commands/Bundles/BundleCreate.cs b/Engine/Source/Programs/Horde/Horde/Commands/Bundles/BundleCreate.cs
--- a/Engine/Source/Programs/Horde/Horde/Commands/Bundles/BundleCreate.cs
+++ b/Engine/Source/Programs/Horde/Horde/Commands/Bundles/BundleCreate.cs
@@ -4,8 +4,8 @@ using System.Diagnostics;
 using EpicGames.Core;
 using EpicGames.Horde.Storage;
 using EpicGames.Horde.Storage.Backends;
+using EpicGames.Horde.Storage.Bundles;
 using EpicGames.Horde.Storage.Nodes;
-using Microsoft.Extensions.Caching.Memory;
 using Microsoft.Extensions.Logging;
 
 namespace Horde.Commands.Bundles
@@ -25,15 +25,18 @@ namespace Horde.Commands.Bundles
 		[CommandLine("-Filter=", Description = "Filter for files to include, in P4 syntax (eg. Foo/...).")]
 		public string Filter { get; set; } = "...";
 
-		public override async Task<int> ExecuteAsync(ILogger logger)
+		public BundleCreate(StorageCache storageCache)
+			: base(storageCache)
 		{
-			using IMemoryCache cache = new MemoryCache(new MemoryCacheOptions { SizeLimit = 256 * 1024 * 1024 });
+		}
 
+		public override async Task<int> ExecuteAsync(ILogger logger)
+		{
 			// Create the storage client
 			IStorageClient store;
 			if (File != null)
 			{
-				store = new FileStorageClient(File.Directory, cache, logger);
+				store = new FileStorageClient(File.Directory, StorageCache, logger);
 			}
 			else if (!String.IsNullOrEmpty(Ref))
 			{
diff --git a/Engine/Source/Programs/Horde/Horde/Commands/Bundles/BundleExtract.cs b/Engine/Source/Programs/Horde/Horde/Commands/Bundles/BundleExtract.cs
--- a/Engine/Source/Programs/Horde/Horde/Commands/Bundles/BundleExtract.cs
+++ b/Engine/Source/Programs/Horde/Horde/Commands/Bundles/BundleExtract.cs
@@ -6,7 +6,6 @@ using EpicGames.Horde.Storage;
 using EpicGames.Horde.Storage.Backends;
 using EpicGames.Horde.Storage.Bundles;
 using EpicGames.Horde.Storage.Nodes;
-using Microsoft.Extensions.Caching.Memory;
 using Microsoft.Extensions.Logging;
 
 namespace Horde.Commands.Bundles
@@ -26,15 +25,18 @@ namespace Horde.Commands.Bundles
 		[CommandLine("-OutputDir=", Required = true)]
 		public DirectoryReference OutputDir { get; set; } = null!;
 
-		public override async Task<int> ExecuteAsync(ILogger logger)
+		public BundleExtract(StorageCache storageCache)
+			: base(storageCache)
 		{
-			using IMemoryCache cache = new MemoryCache(new MemoryCacheOptions { SizeLimit = 256 * 1024 * 1024 });
+		}
 
+		public override async Task<int> ExecuteAsync(ILogger logger)
+		{
 			IStorageClient store;
 			BlobHandle handle;
 			if (File != null)
 			{
-				store = new FileStorageClient(File.Directory, cache, logger);
+				store = new FileStorageClient(File.Directory, StorageCache, logger);
 				handle = await ((FileStorageClient)store).ReadRefAsync(File);
 			}
 			else if (Ref != null)
diff --git a/Engine/Source/Programs/Horde/Horde/Commands/Compute/ComputeWorker.cs b/Engine/Source/Programs/Horde/Horde/Commands/Compute/ComputeWorker.cs
--- a/Engine/Source/Programs/Horde/Horde/Commands/Compute/ComputeWorker.cs
+++ b/Engine/Source/Programs/Horde/Horde/Commands/Compute/ComputeWorker.cs
@@ -5,7 +5,7 @@ using System.Net.Sockets;
 using EpicGames.Core;
 using EpicGames.Horde.Compute;
 using EpicGames.Horde.Compute.Transports;
-using Microsoft.Extensions.Caching.Memory;
+using EpicGames.Horde.Storage;
 using Microsoft.Extensions.Logging;
 
 namespace Horde.Commands.Compute
@@ -16,16 +16,9 @@ namespace Horde.Commands.Compute
 	[Command("compute", "worker", "Runs the agent as a local compute host, accepting incoming connections on the loopback adapter with a given port")]
 	class ComputeWorkerCommand : Command
 	{
-		readonly IMemoryCache _memoryCache;
-
 		[CommandLine("-Port=")]
 		int Port { get; set; } = 2000;
 
-		public ComputeWorkerCommand(IMemoryCache memoryCache)
-		{
-			_memoryCache = memoryCache;
-		}
-
 		public override async Task<int> ExecuteAsync(ILogger logger)
 		{
 			logger.LogInformation("** WORKER **");
@@ -33,10 +26,11 @@ namespace Horde.Commands.Compute
 			using Socket tcpSocket = new Socket(SocketType.Stream, ProtocolType.IP);
 			await tcpSocket.ConnectAsync(IPAddress.Loopback, Port);
 
+			using StorageCache storageCache = new StorageCache();
 			await using (RemoteComputeSocket socket = new RemoteComputeSocket(new TcpTransport(tcpSocket), logger))
 			{
 				logger.LogInformation("Running worker...");
-				await RunWorkerAsync(socket, _memoryCache, logger, CancellationToken.None);
+				await RunWorkerAsync(socket, storageCache, logger, CancellationToken.None);
 				logger.LogInformation("Worker complete");
 				await socket.CloseAsync(CancellationToken.None);
 			}
@@ -45,11 +39,11 @@ namespace Horde.Commands.Compute
 			return 0;
 		}
 
-		public static async Task RunWorkerAsync(ComputeSocket socket, IMemoryCache memoryCache, ILogger logger, CancellationToken cancellationToken)
+		public static async Task RunWorkerAsync(ComputeSocket socket, StorageCache storageCache, ILogger logger, CancellationToken cancellationToken)
 		{
 			DirectoryReference sandboxDir = DirectoryReference.Combine(DirectoryReference.GetSpecialFolder(Environment.SpecialFolder.LocalApplicationData)!, "Horde", "Sandbox");
 
-			AgentMessageHandler worker = new AgentMessageHandler(sandboxDir, memoryCache, null, false, null, logger);
+			AgentMessageHandler worker = new AgentMessageHandler(sandboxDir, storageCache, null, false, null, logger);
 			await worker.RunAsync(socket, cancellationToken);
 		}
 	}
diff --git a/Engine/Source/Programs/Horde/Horde/Commands/StorageCommandBase.cs b/Engine/Source/Programs/Horde/Horde/Commands/StorageCommandBase.cs
--- a/Engine/Source/Programs/Horde/Horde/Commands/StorageCommandBase.cs
+++ b/Engine/Source/Programs/Horde/Horde/Commands/StorageCommandBase.cs
@@ -4,6 +4,7 @@ using System.Net.Http.Headers;
 using EpicGames.Core;
 using EpicGames.Horde.Storage;
 using EpicGames.Horde.Storage.Backends;
+using EpicGames.Horde.Storage.Bundles;
 using Microsoft.Extensions.Logging;
 
 namespace Horde.Commands
@@ -25,6 +26,19 @@ namespace Horde.Commands
 		[CommandLine("-Path=", Description = "Relative path on the server for the store to write to/from (eg. api/v1/storage/default)")]
 		public string? Path { get; set; }
 
+		/// <summary>
+		/// Cache for storage
+		/// </summary>
+		public StorageCache StorageCache { get; }
+
+		/// <summary>
+		/// Constructor
+		/// </summary>
+		public StorageCommandBase(StorageCache storageCache)
+		{
+			StorageCache = storageCache;
+		}
+
 		/// <summary>
 		/// Creates a new client instance
 		/// </summary>
@@ -54,7 +68,7 @@ namespace Horde.Commands
 			}
 
 			server = new Uri(server, Path);
-			return new HttpStorageClient(() => CreateDefaultHttpClient(server, token), () => new HttpClient(), null, logger);
+			return new HttpStorageClient(() => CreateDefaultHttpClient(server, token), () => new HttpClient(), StorageCache, logger);
 		}
 
 		static HttpClient CreateDefaultHttpClient(Uri server, string token)
diff --git a/Engine/Source/Programs/Horde/Horde/Program.cs b/Engine/Source/Programs/Horde/Horde/Program.cs
--- a/Engine/Source/Programs/Horde/Horde/Program.cs
+++ b/Engine/Source/Programs/Horde/Horde/Program.cs
@@ -9,6 +9,7 @@ using System.Runtime.InteropServices;
 using Serilog.Sinks.SystemConsole.Themes;
 using Serilog;
 using Serilog.Formatting.Json;
+using EpicGames.Horde.Storage;
 
 namespace Horde
 {
@@ -32,9 +33,10 @@ namespace Horde
 			services.AddSingleton(loggerFactory);
 			services.AddLogging();
 			services.AddMemoryCache();
+			services.AddSingleton<StorageCache>();
 
 			// Execute all the commands
-			IServiceProvider serviceProvider = services.BuildServiceProvider();
+			await using ServiceProvider serviceProvider = services.BuildServiceProvider();
 			return await CommandHost.RunAsync(arguments, serviceProvider, null);
 		}
 
@@ -50,7 +52,7 @@ namespace Horde
 				DirectoryReference? programDataDir = DirectoryReference.GetSpecialFolder(Environment.SpecialFolder.CommonApplicationData);
 				if (programDataDir != null)
 				{
-					return DirectoryReference.Combine(programDataDir, "HordeAgent");
+					return DirectoryReference.Combine(programDataDir, "Horde");
 				}
 			}
 			return GetAppDir();
diff --git a/Engine/Source/Programs/Shared/EpicGames.Horde.Tests/BundleStoreTests.cs b/Engine/Source/Programs/Shared/EpicGames.Horde.Tests/BundleStoreTests.cs
--- a/Engine/Source/Programs/Shared/EpicGames.Horde.Tests/BundleStoreTests.cs
+++ b/Engine/Source/Programs/Shared/EpicGames.Horde.Tests/BundleStoreTests.cs
@@ -151,7 +151,7 @@ namespace EpicGames.Horde.Tests
 
 				await store.WriteRefTargetAsync(new RefName("test"), await writer.WriteNodeAsync(root));
 
-				BundleReader reader = new BundleReader(store, null, NullLogger.Instance);
+				BundleReader reader = new BundleReader(store, StorageCache.None, NullLogger.Instance);
 				await CheckTree(root);
 			}
 
@@ -259,7 +259,7 @@ namespace EpicGames.Horde.Tests
 		{
 			using IMemoryCache cache = new MemoryCache(new MemoryCacheOptions());
 			MemoryStorageClient store = new MemoryStorageClient();
-			BundleReader reader = new BundleReader(store, null, NullLogger.Instance);
+			BundleReader reader = new BundleReader(store, StorageCache.None, NullLogger.Instance);
 
 			// Generate a tree
 			{
diff --git a/Engine/Source/Programs/Shared/EpicGames.Horde.Tests/BundleTests.cs b/Engine/Source/Programs/Shared/EpicGames.Horde.Tests/BundleTests.cs
--- a/Engine/Source/Programs/Shared/EpicGames.Horde.Tests/BundleTests.cs
+++ b/Engine/Source/Programs/Shared/EpicGames.Horde.Tests/BundleTests.cs
@@ -69,7 +69,7 @@ namespace EpicGames.Horde.Tests
 		public async Task BasicChunkingTests()
 		{
 			RefName refName = new RefName("test");
-			BundleReader reader = new BundleReader(_storage, null, NullLogger.Instance);
+			BundleReader reader = new BundleReader(_storage, StorageCache.None, NullLogger.Instance);
 			await using IStorageWriter writer = _storage.CreateWriter(refName);
 
 			ChunkingOptions options = new ChunkingOptions();
diff --git a/Engine/Source/Programs/Shared/EpicGames.Horde.Tests/ComputeSocketTests.cs b/Engine/Source/Programs/Shared/EpicGames.Horde.Tests/ComputeSocketTests.cs
--- a/Engine/Source/Programs/Shared/EpicGames.Horde.Tests/ComputeSocketTests.cs
+++ b/Engine/Source/Programs/Shared/EpicGames.Horde.Tests/ComputeSocketTests.cs
@@ -194,8 +194,8 @@ namespace EpicGames.Horde.Tests
 
 		static async Task RunAgent(ComputeSocket socket, DirectoryReference tempDir, CancellationToken cancellationToken)
 		{
-			using MemoryCache memoryCache = new MemoryCache(new MemoryCacheOptions());
-			AgentMessageHandler handler = new AgentMessageHandler(tempDir, memoryCache, null, true, null, NullLogger.Instance);
+			using StorageCache storageCache = new StorageCache();
+			AgentMessageHandler handler = new AgentMessageHandler(tempDir, storageCache, null, true, null, NullLogger.Instance);
 			await handler.RunAsync(socket, cancellationToken);
 		}
 	}
diff --git a/Engine/Source/Programs/Shared/EpicGames.Horde/Compute/AgentMessageHandler.cs b/Engine/Source/Programs/Shared/EpicGames.Horde/Compute/AgentMessageHandler.cs
--- a/Engine/Source/Programs/Shared/EpicGames.Horde/Compute/AgentMessageHandler.cs
+++ b/Engine/Source/Programs/Shared/EpicGames.Horde/Compute/AgentMessageHandler.cs
@@ -9,11 +9,11 @@ using System.Threading.Tasks;
 using EpicGames.Core;
 using EpicGames.Horde.Compute.Buffers;
 using EpicGames.Horde.Storage.Nodes;
-using Microsoft.Extensions.Caching.Memory;
 using Microsoft.Extensions.Logging;
-using EpicGames.Horde.Storage.Bundles;
+using EpicGames.Horde.Storage;
 using System.IO;
 using System.Linq;
+using EpicGames.Horde.Storage.Bundles;
 
 namespace EpicGames.Horde.Compute
 {
@@ -23,7 +23,7 @@ namespace EpicGames.Horde.Compute
 	public class AgentMessageHandler
 	{
 		readonly DirectoryReference _sandboxDir;
-		readonly IMemoryCache _memoryCache;
+		readonly StorageCache _cache;
 		readonly Dictionary<string, string?> _envVars;
 		readonly bool _executeInProcess;
 		readonly string? _wineExecutablePath;
@@ -33,15 +33,15 @@ namespace EpicGames.Horde.Compute
 		/// Constructor
 		/// </summary>
 		/// <param name="sandboxDir">Directory to use for reading/writing files</param>
-		/// <param name="memoryCache">Cache for nodes read from storage</param>
+		/// <param name="cache">Cache for nodes read from storage</param>
 		/// <param name="envVars">Environment variables to set for any child processes</param>
 		/// <param name="executeInProcess">Whether to execute any external assemblies in the current process</param>
 		/// <param name="wineExecutablePath">Path to Wine executable. If null, execution under Wine is disabled</param>
 		/// <param name="logger">Logger for diagnostics</param>
-		public AgentMessageHandler(DirectoryReference sandboxDir, IMemoryCache memoryCache, Dictionary<string, string?>? envVars, bool executeInProcess, string? wineExecutablePath, ILogger logger)
+		public AgentMessageHandler(DirectoryReference sandboxDir, StorageCache cache, Dictionary<string, string?>? envVars, bool executeInProcess, string? wineExecutablePath, ILogger logger)
 		{
 			_sandboxDir = sandboxDir;
-			_memoryCache = memoryCache;
+			_cache = cache;
 			_envVars = envVars ?? new Dictionary<string, string?>();
 			_executeInProcess = executeInProcess;
 			_wineExecutablePath = wineExecutablePath;
@@ -139,7 +139,7 @@ namespace EpicGames.Horde.Compute
 		async Task WriteFilesAsync(AgentMessageChannel channel, string path, BundleNodeLocator locator, CancellationToken cancellationToken)
 		{
 			using AgentStorageClient store = new AgentStorageClient(channel);
-			BundleReader reader = new BundleReader(store, _memoryCache, _logger);
+			BundleReader reader = new BundleReader(store, _cache, _logger);
 
 			DirectoryNode directoryNode = await reader.ReadNodeAsync<DirectoryNode>(locator, cancellationToken);
 
diff --git a/Engine/Source/Programs/Shared/EpicGames.Horde/Compute/AgentStorageClient.cs b/Engine/Source/Programs/Shared/EpicGames.Horde/Compute/AgentStorageClient.cs
--- a/Engine/Source/Programs/Shared/EpicGames.Horde/Compute/AgentStorageClient.cs
+++ b/Engine/Source/Programs/Shared/EpicGames.Horde/Compute/AgentStorageClient.cs
@@ -25,7 +25,7 @@ namespace EpicGames.Horde.Compute
 		/// </summary>
 		/// <param name="channel"></param>
 		public AgentStorageClient(AgentMessageChannel channel)
-			: base(null, NullLogger.Instance)
+			: base(StorageCache.None, NullLogger.Instance)
 		{
 			_channel = channel;
 			_semaphore = new SemaphoreSlim(1);
diff --git a/Engine/Source/Programs/Shared/EpicGames.Horde/Compute/Clients/LocalComputeClient.cs b/Engine/Source/Programs/Shared/EpicGames.Horde/Compute/Clients/LocalComputeClient.cs
--- a/Engine/Source/Programs/Shared/EpicGames.Horde/Compute/Clients/LocalComputeClient.cs
+++ b/Engine/Source/Programs/Shared/EpicGames.Horde/Compute/Clients/LocalComputeClient.cs
@@ -8,7 +8,7 @@ using System.Threading;
 using System.Threading.Tasks;
 using EpicGames.Core;
 using EpicGames.Horde.Compute.Transports;
-using Microsoft.Extensions.Caching.Memory;
+using EpicGames.Horde.Storage;
 using Microsoft.Extensions.Logging;
 
 namespace EpicGames.Horde.Compute.Clients
@@ -111,11 +111,11 @@ namespace EpicGames.Horde.Compute.Clients
 		{
 			using Socket tcpSocket = await listener.AcceptAsync(cancellationToken);
 
-			using MemoryCache memoryCache = new MemoryCache(new MemoryCacheOptions { SizeLimit = 10 * 1024 * 1024 });
+			using StorageCache storageCache = new StorageCache();
 
 			await using (RemoteComputeSocket socket = new RemoteComputeSocket(new TcpTransport(tcpSocket), logger))
 			{
-				AgentMessageHandler worker = new AgentMessageHandler(sandboxDir, memoryCache, null, executeInProcess, null, logger);
+				AgentMessageHandler worker = new AgentMessageHandler(sandboxDir, storageCache, null, executeInProcess, null, logger);
 				await worker.RunAsync(socket, cancellationToken);
 				await socket.CloseAsync(cancellationToken);
 			}
diff --git a/Engine/Source/Programs/Shared/EpicGames.Horde/Storage/Backends/FileStorageClient.cs b/Engine/Source/Programs/Shared/EpicGames.Horde/Storage/Backends/FileStorageClient.cs
--- a/Engine/Source/Programs/Shared/EpicGames.Horde/Storage/Backends/FileStorageClient.cs
+++ b/Engine/Source/Programs/Shared/EpicGames.Horde/Storage/Backends/FileStorageClient.cs
@@ -31,7 +31,7 @@ namespace EpicGames.Horde.Storage.Backends
 		/// <param name="rootDir">Root directory for storing blobs</param>
 		/// <param name="cache">Memory cache for read data</param>
 		/// <param name="logger">Logger interface</param>
-		public FileStorageClient(DirectoryReference rootDir, IMemoryCache? cache, ILogger logger)
+		public FileStorageClient(DirectoryReference rootDir, StorageCache cache, ILogger logger)
 			: base(cache, logger)
 		{
 			_rootDir = rootDir;
@@ -58,7 +58,7 @@ namespace EpicGames.Horde.Storage.Backends
 		public override Task<Stream> OpenAsync(BundleLocator id, int offset, int? length, CancellationToken cancellationToken = default)
 		{
 			FileReference file = GetBlobFile(id);
-			_logger.LogInformation("Reading {File}", file);
+			_logger.LogInformation("Reading {File} ({Offset}-{Length})", file, offset, length);
 
 			Stream stream = FileReference.Open(file, FileMode.Open, FileAccess.Read, FileShare.Read | FileShare.Delete);
 			if (offset > 0)
diff --git a/Engine/Source/Programs/Shared/EpicGames.Horde/Storage/Backends/HttpStorageClient.cs b/Engine/Source/Programs/Shared/EpicGames.Horde/Storage/Backends/HttpStorageClient.cs
--- a/Engine/Source/Programs/Shared/EpicGames.Horde/Storage/Backends/HttpStorageClient.cs
+++ b/Engine/Source/Programs/Shared/EpicGames.Horde/Storage/Backends/HttpStorageClient.cs
@@ -62,8 +62,8 @@ namespace EpicGames.Horde.Storage.Backends
 		/// <summary>
 		/// Constructor
 		/// </summary>
-		public HttpStorageClient(Func<HttpClient> createClient, Func<HttpClient> createRedirectClient, IMemoryCache? memoryCache, ILogger logger) 
-			: base(memoryCache, logger)
+		public HttpStorageClient(Func<HttpClient> createClient, Func<HttpClient> createRedirectClient, StorageCache cache, ILogger logger) 
+			: base(cache, logger)
 		{
 			_createClient = createClient;
 			_createRedirectClient = createRedirectClient;
@@ -73,8 +73,8 @@ namespace EpicGames.Horde.Storage.Backends
 		/// <summary>
 		/// Constructor
 		/// </summary>
-		public HttpStorageClient(IHttpClientFactory httpClientFactory, Uri baseAddress, string? bearerToken, IMemoryCache? memoryCache, ILogger logger)
-			: this(() => CreateAuthenticatedClient(httpClientFactory, baseAddress, bearerToken), () => CreateClient(httpClientFactory), memoryCache, logger)
+		public HttpStorageClient(IHttpClientFactory httpClientFactory, Uri baseAddress, string? bearerToken, StorageCache cache, ILogger logger)
+			: this(() => CreateAuthenticatedClient(httpClientFactory, baseAddress, bearerToken), () => CreateClient(httpClientFactory), cache, logger)
 		{
 		}
 
diff --git a/Engine/Source/Programs/Shared/EpicGames.Horde/Storage/Backends/MemoryStorageClient.cs b/Engine/Source/Programs/Shared/EpicGames.Horde/Storage/Backends/MemoryStorageClient.cs
--- a/Engine/Source/Programs/Shared/EpicGames.Horde/Storage/Backends/MemoryStorageClient.cs
+++ b/Engine/Source/Programs/Shared/EpicGames.Horde/Storage/Backends/MemoryStorageClient.cs
@@ -48,7 +48,7 @@ namespace EpicGames.Horde.Storage.Backends
 		/// Constructor
 		/// </summary>
 		public MemoryStorageClient() 
-			: base(new MemoryCache(new MemoryCacheOptions()), NullLogger.Instance)
+			: base(StorageCache.None, NullLogger.Instance)
 		{
 		}
 
diff --git a/Engine/Source/Programs/Shared/EpicGames.Horde/Storage/Bundles/BundleReader.cs b/Engine/Source/Programs/Shared/EpicGames.Horde/Storage/Bundles/BundleReader.cs
--- a/Engine/Source/Programs/Shared/EpicGames.Horde/Storage/Bundles/BundleReader.cs
+++ b/Engine/Source/Programs/Shared/EpicGames.Horde/Storage/Bundles/BundleReader.cs
@@ -6,37 +6,49 @@ using System.Linq;
 using System.Threading;
 using System.Threading.Tasks;
 using EpicGames.Core;
-using Microsoft.Extensions.Caching.Memory;
 using Microsoft.Extensions.Logging;
 using Microsoft.CodeAnalysis;
 using System.Diagnostics;
 using System.IO;
-using System.Diagnostics.CodeAnalysis;
 
 namespace EpicGames.Horde.Storage.Bundles
 {
 	/// <summary>
-	/// Writes nodes from bundles in an <see cref="IStorageClient"/> instance.
+	/// Computed information about a bundle
 	/// </summary>
-	public class BundleReader
+	public class BundleInfo
 	{
 		/// <summary>
-		/// Computed information about a bundle
+		/// Locator for the bundle
 		/// </summary>
-		class BundleInfo
-		{
-			public readonly BundleLocator Locator;
-			public readonly BundleHeader Header;
-			public readonly int HeaderLength;
+		public BundleLocator Locator { get; }
 
-			public BundleInfo(BundleLocator locator, BundleHeader header, int headerLength)
-			{
-				Locator = locator;
-				Header = header;
-				HeaderLength = headerLength;
-			}
+		/// <summary>
+		/// Bundle header
+		/// </summary>
+		public BundleHeader Header { get; }
+		
+		/// <summary>
+		/// Length of the header. Required to offset packets from the start of the bundle
+		/// </summary>
+		public int HeaderLength { get; }
+
+		/// <summary>
+		/// Constructor
+		/// </summary>
+		public BundleInfo(BundleLocator locator, BundleHeader header, int headerLength)
+		{
+			Locator = locator;
+			Header = header;
+			HeaderLength = headerLength;
 		}
+	}
 
+	/// <summary>
+	/// Writes nodes from bundles in an <see cref="IStorageClient"/> instance.
+	/// </summary>
+	public class BundleReader
+	{
 		/// <summary>
 		/// Bundle header queued to be read
 		/// </summary>
@@ -78,7 +90,7 @@ namespace EpicGames.Horde.Storage.Bundles
 		const int DefaultUncachedFetchSize = 1 * 1024 * 1024;
 
 		readonly BundleStorageClient _store;
-		readonly IMemoryCache? _cache;
+		readonly StorageCache _cache;
 		readonly ILogger _logger;
 
 		readonly object _queueLock = new object();
@@ -106,52 +118,13 @@ namespace EpicGames.Horde.Storage.Bundles
 		/// <param name="store"></param>
 		/// <param name="cache">Cache for data</param>
 		/// <param name="logger">Logger for output</param>
-		public BundleReader(BundleStorageClient store, IMemoryCache? cache, ILogger logger)
+		public BundleReader(BundleStorageClient store, StorageCache cache, ILogger logger)
 		{
 			_store = store;
 			_cache = cache;
 			_logger = logger;
 		}
 
-		#region Cache
-
-		static void AddCachedValue(IMemoryCache? cache, string cacheKey, object value, int size)
-		{
-			if (cache != null)
-			{
-				using (ICacheEntry entry = cache.CreateEntry(cacheKey))
-				{
-					entry.SetValue(value);
-					entry.SetSize(size);
-				}
-			}
-		}
-
-		static bool TryGetCachedValue<T>(IMemoryCache? cache, string cacheKey, [MaybeNull, NotNullWhen(true)] out T value)
-		{
-			if (cache == null)
-			{
-				value = default;
-				return false;
-			}
-			return cache.TryGetValue(cacheKey, out value);
-		}
-
-		static string GetBundleInfoCacheKey(BundleLocator locator) => $"bundle:{locator}";
-		static string GetEncodedPacketCacheKey(BundleLocator locator, int packetIdx) => $"encoded-packet:{locator}#{packetIdx}";
-		static string GetDecodedPacketCacheKey(BundleLocator locator, int packetIdx) => $"decoded-packet:{locator}#{packetIdx}";
-
-		void AddCachedBundleInfo(BundleLocator locator, BundleInfo bundleInfo) => AddCachedValue(_cache, GetBundleInfoCacheKey(locator), bundleInfo, bundleInfo.HeaderLength);
-		bool TryGetCachedBundleInfo(BundleLocator locator, [NotNullWhen(true)] out BundleInfo? bundleInfo) => TryGetCachedValue(_cache, GetBundleInfoCacheKey(locator), out bundleInfo);
-
-		void AddCachedEncodedPacket(BundleLocator locator, int packetIdx, ReadOnlyMemory<byte> data) => AddCachedValue(_cache, GetEncodedPacketCacheKey(locator, packetIdx), data, data.Length);
-		bool TryGetCachedEncodedPacket(BundleLocator locator, int packetIdx, out ReadOnlyMemory<byte> data) => TryGetCachedValue(_cache, GetEncodedPacketCacheKey(locator, packetIdx), out data);
-
-		void AddCachedDecodedPacket(BundleLocator locator, int packetIdx, ReadOnlyMemory<byte> data) => AddCachedValue(_cache, GetDecodedPacketCacheKey(locator, packetIdx), data, data.Length);
-		bool TryGetCachedDecodedPacket(BundleLocator locator, int packetIdx, out ReadOnlyMemory<byte> data) => TryGetCachedValue(_cache, GetDecodedPacketCacheKey(locator, packetIdx), out data);
-
-		#endregion
-
 		#region Bundles
 
 		/// <summary>
@@ -258,7 +231,7 @@ namespace EpicGames.Horde.Storage.Bundles
 		{
 			Interlocked.Increment(ref _numHeaderReads);
 
-			int prefetchSize = _cache != null ? DefaultFetchSize : DefaultUncachedFetchSize;
+			int prefetchSize = _cache.HasPacketCache ? DefaultFetchSize : DefaultUncachedFetchSize;
 			for (; ; )
 			{
 				await using (Stream stream = await _store.OpenAsync(queuedHeader.Blob, 0, prefetchSize, cancellationToken))
@@ -282,7 +255,7 @@ namespace EpicGames.Horde.Storage.Bundles
 					BundleInfo bundleInfo = new BundleInfo(queuedHeader.Blob, header, headerSize);
 
 					List<ReadOnlyMemory<byte>> packets = new List<ReadOnlyMemory<byte>>();
-					if (_cache != null)
+					if (_cache.HasPacketCache)
 					{
 						// Also add any encoded packets we prefetched
 						int packetOffset = headerSize;
@@ -295,15 +268,13 @@ namespace EpicGames.Horde.Storage.Bundles
 							}
 
 							byte[] packetData = await ReadPacketAsync(stream, packetLength, cancellationToken);
-							AddCachedEncodedPacket(queuedHeader.Blob, packetIdx, packetData);
+							_cache.AddCachedEncodedPacket(queuedHeader.Blob, packetIdx, packetData);
 							packets.Add(packetData);
 
 							packetOffset += packetLength;
 						}
-
-						// Add the info to the cache
-						AddCachedBundleInfo(queuedHeader.Blob, bundleInfo);
 					}
+					_cache.AddCachedHeader(queuedHeader.Blob, bundleInfo);
 
 					// Update any packets that now have a cached value
 					lock (_queueLock)
@@ -365,7 +336,7 @@ namespace EpicGames.Horde.Storage.Bundles
 			int readLength = packet.EncodedLength;
 
 			int maxPacketIdx = minPacketIdx + 1;
-			if(_cache != null)
+			if(_cache.HasPacketCache)
 			{
 				for (; maxPacketIdx < bundleInfo.Header.Packets.Count; maxPacketIdx++)
 				{
@@ -388,13 +359,13 @@ namespace EpicGames.Horde.Storage.Bundles
 					packets[idx - minPacketIdx] = data;
 				}
 
-				// Add any complete packets
-				if (_cache != null)
+				// Add any complete packets to the cache
+				if (_cache.HasPacketCache)
 				{
 					for (int idx = minPacketIdx; idx < maxPacketIdx; idx++)
 					{
 						ReadOnlyMemory<byte> data = packets[idx - minPacketIdx];
-						AddCachedEncodedPacket(bundleInfo.Locator, idx, data);
+						_cache.AddCachedEncodedPacket(bundleInfo.Locator, idx, data);
 					}
 				}
 
@@ -430,22 +401,16 @@ namespace EpicGames.Horde.Storage.Bundles
 		/// <returns></returns>
 		public async Task<BundleHeader> ReadBundleHeaderAsync(BundleLocator locator, CancellationToken cancellationToken = default)
 		{
-			BundleInfo info = await GetBundleInfoAsync(locator, cancellationToken);
-			return info.Header;
+			BundleInfo bundleInfo = await GetBundleInfoAsync(locator, cancellationToken);
+			return bundleInfo.Header;
 		}
 
-		/// <summary>
-		/// Reads the header and structural metadata about the bundle
-		/// </summary>
-		/// <param name="locator">The bundle location</param>
-		/// <param name="cancellationToken">Cancellation token for the operation</param>
-		/// <returns>Information about the bundle</returns>
-		async ValueTask<BundleInfo> GetBundleInfoAsync(BundleLocator locator, CancellationToken cancellationToken = default)
+		async Task<BundleInfo> GetBundleInfoAsync(BundleLocator locator, CancellationToken cancellationToken = default)
 		{
 			Debug.Assert(locator.IsValid());
 
 			BundleInfo? bundleInfo;
-			if (TryGetCachedBundleInfo(locator, out bundleInfo))
+			if (_cache.TryGetCachedHeader(locator, out bundleInfo))
 			{
 				return bundleInfo;
 			}
@@ -454,7 +419,7 @@ namespace EpicGames.Horde.Storage.Bundles
 			lock (_queueLock)
 			{
 				// Check the cache again inside lock scope to avoid races
-				if (TryGetCachedBundleInfo(locator, out bundleInfo))
+				if (_cache.TryGetCachedHeader(locator, out bundleInfo))
 				{
 					return bundleInfo;
 				}
@@ -471,6 +436,8 @@ namespace EpicGames.Horde.Storage.Bundles
 			return await queuedHeader.CompletionSource.Task.WaitAsync(cancellationToken);
 		}
 
+		static string GetDecodeTaskKey(BundleLocator locator, int packetIdx) => $"{locator}:{packetIdx}";
+
 		/// <summary>
 		/// Gets a decoded block from the store
 		/// </summary>
@@ -486,7 +453,7 @@ namespace EpicGames.Horde.Storage.Bundles
 			}
 
 			ReadOnlyMemory<byte> decodedPacket;
-			if (TryGetCachedDecodedPacket(bundleInfo.Locator, packetIdx, out decodedPacket))
+			if (_cache.TryGetCachedDecodedPacket(bundleInfo.Locator, packetIdx, out decodedPacket))
 			{
 				return decodedPacket;
 			}
@@ -495,13 +462,13 @@ namespace EpicGames.Horde.Storage.Bundles
 			lock (_queueLock)
 			{
 				// Query the cache again, to eliminate races between cache checks and decode tasks finishing.
-				if (TryGetCachedDecodedPacket(bundleInfo.Locator, packetIdx, out decodedPacket))
+				if (_cache.TryGetCachedDecodedPacket(bundleInfo.Locator, packetIdx, out decodedPacket))
 				{
 					return decodedPacket;
 				}
 
 				// Create an async task to read the data
-				string decodedCacheKey = GetDecodedPacketCacheKey(bundleInfo.Locator, packetIdx);
+				string decodedCacheKey = GetDecodeTaskKey(bundleInfo.Locator, packetIdx);
 				if (!_decodeTasks.TryGetValue(decodedCacheKey, out decodeTask))
 				{
 					decodeTask = Task.Run(() => ReadAndDecodePacketAsync(bundleInfo, packetIdx), CancellationToken.None);
@@ -525,11 +492,11 @@ namespace EpicGames.Horde.Storage.Bundles
 			byte[] decodedPacket = new byte[packet.DecodedLength];
 
 			BundleData.Decompress(packet.CompressionFormat, encodedPacket, decodedPacket);
-			AddCachedDecodedPacket(bundleInfo.Locator, packetIdx, decodedPacket);
+			_cache.AddCachedDecodedPacket(bundleInfo.Locator, packetIdx, decodedPacket);
 
 			lock (_queueLock)
 			{
-				string decodedCacheKey = GetDecodedPacketCacheKey(bundleInfo.Locator, packetIdx);
+				string decodedCacheKey = GetDecodeTaskKey(bundleInfo.Locator, packetIdx);
 				_decodeTasks.Remove(decodedCacheKey);
 			}
 
@@ -550,7 +517,7 @@ namespace EpicGames.Horde.Storage.Bundles
 			}
 
 			ReadOnlyMemory<byte> encodedPacket;
-			if (TryGetCachedEncodedPacket(bundleInfo.Locator, packetIdx, out encodedPacket))
+			if (_cache.TryGetCachedEncodedPacket(bundleInfo.Locator, packetIdx, out encodedPacket))
 			{
 				return encodedPacket;
 			}
@@ -559,7 +526,7 @@ namespace EpicGames.Horde.Storage.Bundles
 			lock (_queueLock)
 			{
 				// Query the cache again, to eliminate races between cache checks and decode tasks finishing.
-				if (TryGetCachedEncodedPacket(bundleInfo.Locator, packetIdx, out encodedPacket))
+				if (_cache.TryGetCachedEncodedPacket(bundleInfo.Locator, packetIdx, out encodedPacket))
 				{
 					return encodedPacket;
 				}
diff --git a/Engine/Source/Programs/Shared/EpicGames.Horde/Storage/Bundles/BundleStorageClient.cs b/Engine/Source/Programs/Shared/EpicGames.Horde/Storage/Bundles/BundleStorageClient.cs
--- a/Engine/Source/Programs/Shared/EpicGames.Horde/Storage/Bundles/BundleStorageClient.cs
+++ b/Engine/Source/Programs/Shared/EpicGames.Horde/Storage/Bundles/BundleStorageClient.cs
@@ -25,9 +25,9 @@ namespace EpicGames.Horde.Storage.Bundles
 		/// <summary>
 		/// Constructor
 		/// </summary>
-		protected BundleStorageClient(IMemoryCache? memoryCache, ILogger logger)
+		protected BundleStorageClient(StorageCache cache, ILogger logger)
 		{
-			BundleReader = new BundleReader(this, memoryCache, logger);
+			BundleReader = new BundleReader(this, cache, logger);
 		}
 
 		#region Blobs
diff --git a/Engine/Source/Programs/Shared/EpicGames.Horde/Storage/StorageCache.cs b/Engine/Source/Programs/Shared/EpicGames.Horde/Storage/StorageCache.cs
new file mode 100644
--- /dev/null
+++ b/Engine/Source/Programs/Shared/EpicGames.Horde/Storage/StorageCache.cs
@@ -0,0 +1,147 @@
+// Copyright Epic Games, Inc. All Rights Reserved.
+
+using System;
+using System.Diagnostics.CodeAnalysis;
+using EpicGames.Horde.Storage.Bundles;
+using Microsoft.Extensions.Caching.Memory;
+
+namespace EpicGames.Horde.Storage
+{
+	/// <summary>
+	/// Options for creating a storage cache
+	/// </summary>
+	public class StorageCacheOptions
+	{
+		/// <summary>
+		/// Size of the header cache
+		/// </summary>
+		public long HeaderCacheSize { get; set; } = 64 * 1024 * 1024;
+
+		/// <summary>
+		/// Size of the packet cache
+		/// </summary>
+		public long PacketCacheSize { get; set; } = 256 * 1024 * 1024;
+	}
+
+	/// <summary>
+	/// Implements caching functionality for reading from a storage client
+	/// </summary>
+	public sealed class StorageCache : IDisposable
+	{
+		readonly MemoryCache? _headerCache;
+		readonly MemoryCache? _packetCache;
+
+		/// <summary>
+		/// Instance of an empty cache
+		/// </summary>
+		public static StorageCache None { get; } = new StorageCache();
+
+		/// <summary>
+		/// Whether there is a packet cache present
+		/// </summary>
+		public bool HasPacketCache => _packetCache != null;
+
+		/// <summary>
+		/// Constructor
+		/// </summary>
+		public StorageCache() : this(new StorageCacheOptions()) 
+		{ 
+		}
+
+		/// <summary>
+		/// Constructor
+		/// </summary>
+		public StorageCache(StorageCacheOptions options)
+		{
+			if (options.HeaderCacheSize > 0)
+			{
+				_headerCache = new MemoryCache(new MemoryCacheOptions { SizeLimit = options.HeaderCacheSize });
+			}
+			if (options.PacketCacheSize > 0)
+			{
+				_packetCache = new MemoryCache(new MemoryCacheOptions { SizeLimit = options.PacketCacheSize });
+			}
+		}
+
+		/// <inheritdoc/>
+		public void Dispose()
+		{
+			_headerCache?.Dispose();
+			_packetCache?.Dispose();
+		}
+
+		static void AddCachedValue(IMemoryCache? cache, string cacheKey, object value, int size)
+		{
+			if (cache != null)
+			{
+				using (ICacheEntry entry = cache.CreateEntry(cacheKey))
+				{
+					entry.SetValue(value);
+					entry.SetSize(size);
+				}
+			}
+		}
+
+		static bool TryGetCachedValue<T>(IMemoryCache? cache, string cacheKey, [MaybeNull, NotNullWhen(true)] out T value)
+		{
+			if (cache == null)
+			{
+				value = default;
+				return false;
+			}
+			return cache.TryGetValue(cacheKey, out value);
+		}
+
+		static string GetBundleInfoCacheKey(BundleLocator locator) => $"bundle:{locator}";
+		static string GetEncodedPacketCacheKey(BundleLocator locator, int packetIdx) => $"encoded-packet:{locator}#{packetIdx}";
+		static string GetDecodedPacketCacheKey(BundleLocator locator, int packetIdx) => $"decoded-packet:{locator}#{packetIdx}";
+
+		/// <summary>
+		/// Adds a bundle info object to the cache
+		/// </summary>
+		public void AddCachedHeader(BundleLocator locator, BundleInfo bundleInfo)
+		{
+			AddCachedValue(_headerCache, GetBundleInfoCacheKey(locator), bundleInfo, bundleInfo.HeaderLength);
+		}
+
+		/// <summary>
+		/// Try to read a bundle info object from the cache
+		/// </summary>
+		public bool TryGetCachedHeader(BundleLocator locator, [NotNullWhen(true)] out BundleInfo? bundleInfo)
+		{
+			return TryGetCachedValue(_headerCache, GetBundleInfoCacheKey(locator), out bundleInfo);
+		}
+
+		/// <summary>
+		/// Adds an encoded bundle packet to the cache
+		/// </summary>
+		public void AddCachedEncodedPacket(BundleLocator locator, int packetIdx, ReadOnlyMemory<byte> data)
+		{
+			AddCachedValue(_packetCache, GetEncodedPacketCacheKey(locator, packetIdx), data, data.Length);
+		}
+
+		/// <summary>
+		/// Try to read an encoded bundle packet from the cache
+		/// </summary>
+		public bool TryGetCachedEncodedPacket(BundleLocator locator, int packetIdx, out ReadOnlyMemory<byte> data)
+		{
+			return TryGetCachedValue(_packetCache, GetEncodedPacketCacheKey(locator, packetIdx), out data);
+		}
+
+		/// <summary>
+		/// Adds a decoded bundle packet to the cache
+		/// </summary>
+		public void AddCachedDecodedPacket(BundleLocator locator, int packetIdx, ReadOnlyMemory<byte> data)
+		{
+			AddCachedValue(_packetCache, GetDecodedPacketCacheKey(locator, packetIdx), data, data.Length);
+		}
+
+		/// <summary>
+		/// Try to read a decoded bundle packet from the cache
+		/// </summary>
+		public bool TryGetCachedDecodedPacket(BundleLocator locator, int packetIdx, out ReadOnlyMemory<byte> data)
+		{
+			return TryGetCachedValue(_packetCache, GetDecodedPacketCacheKey(locator, packetIdx), out data);
+		}
+	}
+}
diff --git a/Engine/Source/Programs/UnrealBuildTool/Artifacts/HordeStorageArtifactCache.cs b/Engine/Source/Programs/UnrealBuildTool/Artifacts/HordeStorageArtifactCache.cs
--- a/Engine/Source/Programs/UnrealBuildTool/Artifacts/HordeStorageArtifactCache.cs
+++ b/Engine/Source/Programs/UnrealBuildTool/Artifacts/HordeStorageArtifactCache.cs
@@ -525,7 +525,7 @@ namespace UnrealBuildTool.Artifacts
 				}
 				Directory.CreateDirectory(directory.FullName);
 
-				_store = new FileStorageClient(directory, null, logger);
+				_store = new FileStorageClient(directory, StorageCache.None, logger);
 
 				State = ArtifactCacheState.Available;
 				return State;
