Horde: Increase coherency between tasks spawned for extracting bundles; now partitions the files to extract along window boundaries.

diff --git a/Engine/Source/Programs/Shared/EpicGames.Horde/Storage/Nodes/DirectoryNode.cs b/Engine/Source/Programs/Shared/EpicGames.Horde/Storage/Nodes/DirectoryNode.cs
--- a/Engine/Source/Programs/Shared/EpicGames.Horde/Storage/Nodes/DirectoryNode.cs
+++ b/Engine/Source/Programs/Shared/EpicGames.Horde/Storage/Nodes/DirectoryNode.cs
@@ -1,12 +1,9 @@
 // Copyright Epic Games, Inc. All Rights Reserved.
 
 using EpicGames.Core;
-using EpicGames.Horde.Logs;
-using Microsoft.CodeAnalysis.Options;
 using Microsoft.Extensions.Logging;
 using System;
 using System.Buffers;
-using System.Buffers.Binary;
 using System.Collections.Generic;
 using System.Diagnostics;
 using System.Diagnostics.CodeAnalysis;
@@ -82,8 +79,8 @@ namespace EpicGames.Horde.Storage.Nodes
 	[NodeType("{0714EC11-291A-4D07-867F-E78AD6809979}", 1)]
 	public class DirectoryNode : Node
 	{
-		readonly Dictionary<Utf8String, FileEntry> _nameToFileEntry = new Dictionary<Utf8String, FileEntry>();
-		readonly Dictionary<Utf8String, DirectoryEntry> _nameToDirectoryEntry = new Dictionary<Utf8String, DirectoryEntry>();
+		readonly SortedDictionary<Utf8String, FileEntry> _nameToFileEntry = new SortedDictionary<Utf8String, FileEntry>();
+		readonly SortedDictionary<Utf8String, DirectoryEntry> _nameToDirectoryEntry = new SortedDictionary<Utf8String, DirectoryEntry>();
 
 		/// <summary>
 		/// Total size of this directory
@@ -133,8 +130,6 @@ namespace EpicGames.Horde.Storage.Nodes
 			Flags = (DirectoryFlags)reader.ReadUnsignedVarInt();
 
 			int fileCount = (int)reader.ReadUnsignedVarInt();
-			_nameToFileEntry.EnsureCapacity(fileCount);
-
 			for (int idx = 0; idx < fileCount; idx++)
 			{
 				FileEntry entry = new FileEntry(reader);
@@ -142,8 +137,6 @@ namespace EpicGames.Horde.Storage.Nodes
 			}
 
 			int directoryCount = (int)reader.ReadUnsignedVarInt();
-			_nameToDirectoryEntry.EnsureCapacity(directoryCount);
-
 			for (int idx = 0; idx < directoryCount; idx++)
 			{
 				DirectoryEntry entry = new DirectoryEntry(reader);
@@ -176,6 +169,7 @@ namespace EpicGames.Horde.Storage.Nodes
 		{
 			_nameToFileEntry.Clear();
 			_nameToDirectoryEntry.Clear();
+
 			MarkAsDirty();
 		}
 
@@ -721,22 +715,52 @@ namespace EpicGames.Horde.Storage.Nodes
 		/// <returns></returns>
 		public async Task CopyToDirectoryAsync(DirectoryInfo directoryInfo, ILogger logger, CancellationToken cancellationToken)
 		{
-			directoryInfo.Create();
+			int NumThreads = Math.Min(1 + (int)(Length / (10 * 1024 * 1024)), 4);
 
 			List<Task> tasks = new List<Task>();
+			try
+			{
+				long offset = 0;
+				for (int threadIdx = 0; threadIdx < NumThreads; threadIdx++)
+				{
+					long minOffset = offset;
+					long maxOffset = (Length * (threadIdx + 1)) / NumThreads;
+					tasks.Add(Task.Run(() => CopyToDirectoryInternalAsync(directoryInfo, minOffset, maxOffset - minOffset, logger, cancellationToken), cancellationToken));
+					offset = maxOffset;
+				}
+			}
+			finally
+			{
+				await Task.WhenAll(tasks);
+			}
+		}
+
+		async Task CopyToDirectoryInternalAsync(DirectoryInfo directoryInfo, long windowOffset, long windowLength, ILogger logger, CancellationToken cancellationToken)
+		{
+			directoryInfo.Create();
+
 			foreach (FileEntry fileEntry in _nameToFileEntry.Values)
 			{
-				FileInfo fileInfo = new FileInfo(Path.Combine(directoryInfo.FullName, fileEntry.Name.ToString()));
-				tasks.Add(Task.Run(() => fileEntry.CopyToFileAsync(fileInfo, cancellationToken), cancellationToken));
+				// Extract any file that starts within the window (window starts before this file, window ends after the start of the file)
+				if (windowOffset <= 0 && windowOffset + windowLength > 0)
+				{
+					FileInfo fileInfo = new FileInfo(Path.Combine(directoryInfo.FullName, fileEntry.Name.ToString()));
+					await fileEntry.CopyToFileAsync(fileInfo, cancellationToken);
+				}
+				windowOffset -= fileEntry.Length;
 			}
+
 			foreach (DirectoryEntry directoryEntry in _nameToDirectoryEntry.Values)
 			{
-				DirectoryInfo subDirectoryInfo = directoryInfo.CreateSubdirectory(directoryEntry.Name.ToString());
-				DirectoryNode subDirectoryNode = await directoryEntry.ExpandAsync(cancellationToken);
-				tasks.Add(Task.Run(() => subDirectoryNode.CopyToDirectoryAsync(subDirectoryInfo, logger, cancellationToken), cancellationToken));
+				// Traverse into any directory that overlaps with the window (window starts before end of the directory, and window ends at or beyond the start of the directory)
+				if (windowOffset < directoryEntry.Length && windowOffset + windowLength >= 0)
+				{
+					DirectoryInfo subDirectoryInfo = directoryInfo.CreateSubdirectory(directoryEntry.Name.ToString());
+					DirectoryNode subDirectoryNode = await directoryEntry.ExpandAsync(cancellationToken);
+					await subDirectoryNode.CopyToDirectoryInternalAsync(subDirectoryInfo, windowOffset, windowLength, logger, cancellationToken);
+				}
+				windowOffset -= directoryEntry.Length;
 			}
-
-			await Task.WhenAll(tasks);
 		}
 
 		/// <summary>
