fix for UHT allowing TObjectPtr<SomeInterface>
the dectection of using Native Pointer for member instead of TObjecPtr was swalling the error for interface

diff --git a/Engine/Source/Programs/Shared/EpicGames.UHT/Parsers/UhtPropertyParser.cs b/Engine/Source/Programs/Shared/EpicGames.UHT/Parsers/UhtPropertyParser.cs
--- a/Engine/Source/Programs/Shared/EpicGames.UHT/Parsers/UhtPropertyParser.cs
+++ b/Engine/Source/Programs/Shared/EpicGames.UHT/Parsers/UhtPropertyParser.cs
@@ -1383,10 +1383,6 @@ namespace EpicGames.UHT.Parsers
 				}
 				tokenReader.Require('*');
 
-				// Optionally emit messages about native pointer members and swallow trailing 'const' after pointer properties
-				UhtObjectPropertyBase.ConditionalLogPointerUsage(propertySettings, session.Config!.EngineNativePointerMemberBehavior,
-					session.Config!.EnginePluginNativePointerMemberBehavior, session.Config!.NonEngineNativePointerMemberBehavior, "Native pointer", tokenReader, typeStartPos, "TObjectPtr");
-
 				if (propertySettings.PropertyCategory == UhtPropertyCategory.Member)
 				{
 					tokenReader.TryOptional("const");
diff --git a/Engine/Source/Programs/Shared/EpicGames.UHT/Types/Properties/UhtObjectProperty.cs b/Engine/Source/Programs/Shared/EpicGames.UHT/Types/Properties/UhtObjectProperty.cs
--- a/Engine/Source/Programs/Shared/EpicGames.UHT/Types/Properties/UhtObjectProperty.cs
+++ b/Engine/Source/Programs/Shared/EpicGames.UHT/Types/Properties/UhtObjectProperty.cs
@@ -103,6 +103,17 @@ namespace EpicGames.UHT.Types
 			return builder;
 		}
 
+		/// <inheritdoc/>
+		protected override void ValidateMember(UhtStruct structObj, UhtValidationOptions options)
+		{
+			base.ValidateMember(structObj, options);
+
+			if(PointerType == UhtPointerType.Native)
+			{
+				ConditionalLogPointerUsage(Session.Config!.EngineNativePointerMemberBehavior, Session.Config!.EnginePluginNativePointerMemberBehavior, Session.Config!.NonEngineNativePointerMemberBehavior, "Native pointer", "TObjectPtr");
+			}	
+		}
+
 		/// <inheritdoc/>
 		public override bool IsSameType(UhtProperty other)
 		{
diff --git a/Engine/Source/Programs/Shared/EpicGames.UHT/Types/Properties/UhtObjectPropertyBase.cs b/Engine/Source/Programs/Shared/EpicGames.UHT/Types/Properties/UhtObjectPropertyBase.cs
--- a/Engine/Source/Programs/Shared/EpicGames.UHT/Types/Properties/UhtObjectPropertyBase.cs
+++ b/Engine/Source/Programs/Shared/EpicGames.UHT/Types/Properties/UhtObjectPropertyBase.cs
@@ -222,26 +222,23 @@ namespace EpicGames.UHT.Types
 		}
 
 		/// <summary>
-		/// Log point usage warning/error
+		/// Logs message for Object pointers to convert UObject* to TObjectPtr or the reverse
 		/// </summary>
-		/// <param name="propertySettings">Property settings</param>
 		/// <param name="engineBehavior">Expected behavior for engine types</param>
 		/// <param name="enginePluginBehavior">Expected behavior for engine plugin types</param>
 		/// <param name="nonEngineBehavior">Expected behavior for non-engine types</param>
 		/// <param name="pointerTypeDesc">Description of the pointer type</param>
-		/// <param name="tokenReader">Token reader for type being parsed</param>
-		/// <param name="typeStartPos">Starting character position of the type</param>
 		/// <param name="alternativeTypeDesc">Suggested alternate declaration</param>
 		/// <exception cref="UhtIceException">Thrown if the behavior type is unexpected</exception>
-		public static void ConditionalLogPointerUsage(UhtPropertySettings propertySettings, UhtIssueBehavior engineBehavior, UhtIssueBehavior enginePluginBehavior,
-			UhtIssueBehavior nonEngineBehavior, string pointerTypeDesc, IUhtTokenReader tokenReader, int typeStartPos, string? alternativeTypeDesc)
+		protected void ConditionalLogPointerUsage(UhtIssueBehavior engineBehavior, UhtIssueBehavior enginePluginBehavior,
+			UhtIssueBehavior nonEngineBehavior, string pointerTypeDesc, string? alternativeTypeDesc)
 		{
-			if (propertySettings.PropertyCategory != UhtPropertyCategory.Member)
+			UhtPackage? package = Outer?.Package;
+			if(Outer == null || package == null)
 			{
 				return;
 			}
 
-			UhtPackage package = propertySettings.Outer.Package;
 			UhtIssueBehavior behavior = nonEngineBehavior;
 			if (package.IsPartOfEngine)
 			{
@@ -260,32 +257,27 @@ namespace EpicGames.UHT.Types
 				return;
 			}
 
-			string type = tokenReader.GetStringView(typeStartPos, tokenReader.InputPos - typeStartPos).ToString();
-			type = type.Replace("\n", " ", StringComparison.Ordinal);
-			type = type.Replace("\r", "", StringComparison.Ordinal);
-			type = type.Replace("\t", " ", StringComparison.Ordinal);
-
 			switch (behavior)
 			{
 				case UhtIssueBehavior.Disallow:
 					if (!String.IsNullOrEmpty(alternativeTypeDesc))
 					{
-						tokenReader.LogError($"{pointerTypeDesc} usage in member declaration detected [[[{type}]]].  This is disallowed for the target/module, consider {alternativeTypeDesc} as an alternative.");
+						this.LogError($"{pointerTypeDesc} usage in member declaration detected [[[{Outer.EngineName}]]].  This is disallowed for the target/module, consider {alternativeTypeDesc} as an alternative.");
 					}
 					else
 					{
-						tokenReader.LogError($"{pointerTypeDesc} usage in member declaration detected [[[{type}]]].");
+						this.LogError($"{pointerTypeDesc} usage in member declaration detected [[[{Outer.EngineName}]]].");
 					}
 					break;
 
 				case UhtIssueBehavior.AllowAndLog:
 					if (!String.IsNullOrEmpty(alternativeTypeDesc))
 					{
-						tokenReader.LogTrace($"{pointerTypeDesc} usage in member declaration detected [[[{type}]]].  Consider {alternativeTypeDesc} as an alternative.");
+						this.LogTrace($"{pointerTypeDesc} usage in member declaration detected [[[{Outer.EngineName}]]].  Consider {alternativeTypeDesc} as an alternative.");
 					}
 					else
 					{
-						tokenReader.LogTrace("{PointerTypeDesc} usage in member declaration detected [[[{Type}]]].");
+						this.LogTrace($"{pointerTypeDesc} usage in member declaration detected [[[{Outer.EngineName}]]].");
 					}
 					break;
 
diff --git a/Engine/Source/Programs/Shared/EpicGames.UHT/Types/Properties/UhtObjectPtrProperty.cs b/Engine/Source/Programs/Shared/EpicGames.UHT/Types/Properties/UhtObjectPtrProperty.cs
--- a/Engine/Source/Programs/Shared/EpicGames.UHT/Types/Properties/UhtObjectPtrProperty.cs
+++ b/Engine/Source/Programs/Shared/EpicGames.UHT/Types/Properties/UhtObjectPtrProperty.cs
@@ -108,8 +108,24 @@ namespace EpicGames.UHT.Types
 			}
 		}
 
-		#region Keyword
-		[UhtPropertyType(Keyword = "TObjectPtr")]
+		/// <inheritdoc/>
+		protected override void ValidateMember(UhtStruct structObj, UhtValidationOptions options)
+		{
+			base.ValidateMember(structObj, options);
+			if(Class.NativeInterface != null)
+			{
+				this.LogError($"UPROPERTY pointers cannot be interfaces - did you mean TScriptInterface<{Class.SourceName}>?");
+			}
+			else
+			{
+				// Use to reverse migration from TObjectPtr<UObject> Member; back to UObject* Member;
+				ConditionalLogPointerUsage(Session.Config!.EngineObjectPtrMemberBehavior, Session.Config!.EnginePluginObjectPtrMemberBehavior,
+				Session.Config!.NonEngineObjectPtrMemberBehavior, "ObjectPtr", null);
+			}
+		}
+
+			#region Keyword
+			[UhtPropertyType(Keyword = "TObjectPtr")]
 		[SuppressMessage("CodeQuality", "IDE0051:Remove unused private members", Justification = "Attribute accessed method")]
 		[SuppressMessage("Style", "IDE0060:Remove unused parameter", Justification = "Attribute accessed method")]
 		private static UhtProperty? ObjectPtrProperty(UhtPropertyResolvePhase resolvePhase, UhtPropertySettings propertySettings, IUhtTokenReader tokenReader, UhtToken matchedToken)
@@ -123,9 +139,6 @@ namespace EpicGames.UHT.Types
 				return null;
 			}
 
-			ConditionalLogPointerUsage(propertySettings, session.Config!.EngineObjectPtrMemberBehavior, session.Config!.EnginePluginObjectPtrMemberBehavior,
-				session.Config!.NonEngineObjectPtrMemberBehavior, "ObjectPtr", tokenReader, typeStartPos, null);
-
 			if (propertyClass.IsChildOf(propertyClass.Session.UClass))
 			{
 				// UObject specifies that there is no limiter
