Horde: Reduce memory usage when reading bundles.

diff --git a/Engine/Source/Programs/Horde/Horde.Server/Storage/StorageController.cs b/Engine/Source/Programs/Horde/Horde.Server/Storage/StorageController.cs
--- a/Engine/Source/Programs/Horde/Horde.Server/Storage/StorageController.cs
+++ b/Engine/Source/Programs/Horde/Horde.Server/Storage/StorageController.cs
@@ -502,7 +502,7 @@ namespace Horde.Server.Storage
 
 			BundleReader reader = new BundleReader(storageClient, _storageCache, _logger);
 
-			BundleHeader header = await reader.ReadBundleHeaderAsync(locator, cancellationToken);
+			BundleHeader header = await reader.ReadHeaderAsync(locator, cancellationToken);
 
 			string linkBase = $"/api/v1/storage/{namespaceId}";
 
@@ -582,7 +582,7 @@ namespace Horde.Server.Storage
 			StorageClient storageClient = await _storageService.GetClientAsync(namespaceId, cancellationToken);
 			BundleReader reader = new BundleReader(storageClient, _storageCache, _logger);
 
-			BundleHeader header = await reader.ReadBundleHeaderAsync(locator, cancellationToken);
+			BundleHeader header = await reader.ReadHeaderAsync(locator, cancellationToken);
 			BundleExport export = header.Exports[exportIdx];
 
 			string linkBase = $"/api/v1/storage/{namespaceId}";
diff --git a/Engine/Source/Programs/Horde/Horde/Commands/Bundles/BundleExtract.cs b/Engine/Source/Programs/Horde/Horde/Commands/Bundles/BundleExtract.cs
--- a/Engine/Source/Programs/Horde/Horde/Commands/Bundles/BundleExtract.cs
+++ b/Engine/Source/Programs/Horde/Horde/Commands/Bundles/BundleExtract.cs
@@ -61,6 +61,7 @@ namespace Horde.Commands.Bundles
 				BundleStorageClient? bundleStorageClient = store as BundleStorageClient;
 				if (bundleStorageClient != null)
 				{
+					logger.LogInformation("Num bytes read: {NumBytes:n0}", bundleStorageClient.BundleReader.NumBytesRead);
 					logger.LogInformation("Num header reads: {NumReads:n0}", bundleStorageClient.BundleReader.NumHeaderReads);
 					logger.LogInformation("Num packet reads: {NumReads:n0}", bundleStorageClient.BundleReader.NumPacketReads);
 				}
diff --git a/Engine/Source/Programs/Shared/EpicGames.Horde/Storage/Bundles/BundleReader.cs b/Engine/Source/Programs/Shared/EpicGames.Horde/Storage/Bundles/BundleReader.cs
--- a/Engine/Source/Programs/Shared/EpicGames.Horde/Storage/Bundles/BundleReader.cs
+++ b/Engine/Source/Programs/Shared/EpicGames.Horde/Storage/Bundles/BundleReader.cs
@@ -10,6 +10,7 @@ using Microsoft.Extensions.Logging;
 using Microsoft.CodeAnalysis;
 using System.Diagnostics;
 using System.IO;
+using System.Diagnostics.CodeAnalysis;
 
 namespace EpicGames.Horde.Storage.Bundles
 {
@@ -50,37 +51,74 @@ namespace EpicGames.Horde.Storage.Bundles
 	public class BundleReader
 	{
 		/// <summary>
-		/// Bundle header queued to be read
+		/// Queued set of requests from a particular bundle
 		/// </summary>
-		class QueuedHeader
+		[DebuggerDisplay("{Locator}")]
+		class QueuedBundle
 		{
-			public readonly BundleLocator Blob;
-			public readonly TaskCompletionSource<BundleInfo> CompletionSource = new TaskCompletionSource<BundleInfo>(TaskCreationOptions.RunContinuationsAsynchronously);
-
-			public Utf8String Path => Blob.Path;
+			public BundleLocator Locator { get; }
+			public TaskCompletionSource<BundleInfo> BundleInfo = new TaskCompletionSource<BundleInfo>();
+			public List<QueuedPacket> QueuedPackets { get; } = new List<QueuedPacket>(); // Sorted by index
+			public int InfoRefCount { get; set; } // Ref count for reading the header
+			public PendingRead? PendingRead { get; set; }
+			public Task? WorkerTask { get; set; }
+			public CancellationTokenSource CancellationSource { get; set; } = new CancellationTokenSource();
+			public bool Complete { get; set; }
+
+			public QueuedBundle(BundleLocator locator)
+			{
+				Locator = locator;
+			}
 
-			public QueuedHeader(BundleLocator blob)
+			public QueuedPacket AddPacket(int packetIdx)
 			{
-				Blob = blob;
+				QueuedPacket packet = new QueuedPacket(packetIdx);
+
+				int insertIdx = QueuedPackets.BinarySearch(packet);
+				if (insertIdx < 0)
+				{
+					insertIdx = ~insertIdx;
 			}
+				QueuedPackets.Insert(insertIdx, packet);
+
+				return packet;
+		}
+
+			public void RemovePacket(QueuedPacket packet) => QueuedPackets.Remove(packet);
 		}
 
 		/// <summary>
 		/// Encoded bundle packet queued to be read
 		/// </summary>
-		class QueuedPacket
+		class QueuedPacket : IComparable<QueuedPacket>
 		{
-			public readonly BundleInfo Bundle;
 			public readonly int PacketIdx;
 			public readonly TaskCompletionSource<ReadOnlyMemory<byte>> CompletionSource = new TaskCompletionSource<ReadOnlyMemory<byte>>(TaskCreationOptions.RunContinuationsAsynchronously);
 
-			public Utf8String Path => Bundle.Locator.Path;
+			public QueuedPacket(int packetIdx) => PacketIdx = packetIdx;
+			public int CompareTo(QueuedPacket? other) => PacketIdx - other?.PacketIdx ?? 0;
+		}
+
+		/// <summary>
+		/// Information about a pending read
+		/// </summary>
+		sealed class PendingRead : IDisposable
+			{
+			public int MinPacketIdx { get; }
+			public int MaxPacketIdx { get; }
+			public CancellationTokenSource CancellationSource { get; }
 
-			public QueuedPacket(BundleInfo bundle, int packetIdx)
+			public PendingRead(int minPacketIdx, int maxPacketIdx, CancellationToken cancellationToken)
 			{
-				Bundle = bundle;
-				PacketIdx = packetIdx;
+				MinPacketIdx = minPacketIdx;
+				MaxPacketIdx = maxPacketIdx;
+				CancellationSource = CancellationTokenSource.CreateLinkedTokenSource(cancellationToken);
 			}
+
+			public void Dispose()
+			{
+				CancellationSource.Dispose();
+		}
 		}
 
 		// Size of data to fetch by default. This is larger than the minimum request size to reduce number of reads.
@@ -94,13 +132,17 @@ namespace EpicGames.Horde.Storage.Bundles
 		readonly ILogger _logger;
 
 		readonly object _queueLock = new object();
-		readonly List<QueuedHeader> _queuedHeaders = new List<QueuedHeader>();
-		readonly List<QueuedPacket> _queuedPackets = new List<QueuedPacket>();
+		readonly List<QueuedBundle> _queuedBundles = new List<QueuedBundle>();
 		readonly Dictionary<string, Task<ReadOnlyMemory<byte>>> _decodeTasks = new Dictionary<string, Task<ReadOnlyMemory<byte>>>(StringComparer.Ordinal);
-		Task? _readTask;
 
 		int _numHeaderReads;
 		int _numPacketReads;
+		long _numBytesRead;
+
+		/// <summary>
+		/// Total number of header reads
+		/// </summary>
+		public long NumBytesRead => _numBytesRead;
 
 		/// <summary>
 		/// Total number of header reads
@@ -128,419 +170,372 @@ namespace EpicGames.Horde.Storage.Bundles
 		#region Bundles
 
 		/// <summary>
-		/// Starts the background task for reading data from the store
+		/// Reads a bundle header
 		/// </summary>
-		void StartReadTask()
+		/// <param name="locator">Locator for the bundle</param>
+		/// <param name="cancellationToken">Cancellation token for the operation</param>
+		/// <returns>Information about the bundle</returns>
+		public async Task<BundleHeader> ReadHeaderAsync(BundleLocator locator, CancellationToken cancellationToken)
 		{
-			if (_readTask == null)
+			// Check for a cached value first
+			if (_cache.TryGetCachedHeader(locator, out BundleInfo? cachedBundleInfo))
 			{
-				_readTask = Task.Run(() => ServiceReadQueueAsync(CancellationToken.None));
-			}
+				return cachedBundleInfo.Header;
 		}
 
-		/// <summary>
-		/// Dispatches requests in the read queue
-		/// </summary>
-		/// <param name="cancellationToken">Cancellation token for the background task</param>
-		async Task ServiceReadQueueAsync(CancellationToken cancellationToken)
-		{
-			const int MaxConcurrentReads = 4;
-
-			List<(Utf8String Path, Task Task)> currentTasks = new List<(Utf8String, Task)>();
-			for (; ; )
-			{
-				// Start any new reads
+			// Find a registered bundle info with the given locator
+			QueuedBundle? bundle = null;
 				lock (_queueLock)
 				{
-					while (currentTasks.Count < MaxConcurrentReads)
-					{
-						HashSet<Utf8String> currentPaths = new HashSet<Utf8String>(currentTasks.Select(x => x.Path));
-
-						// Try to start another header read
-						QueuedHeader? queuedHeader = _queuedHeaders.FirstOrDefault(x => !currentPaths.Contains(x.Path));
-						if (queuedHeader != null)
-						{
-							Task task = Task.Run(() => PerformHeaderReadGuardedAsync(queuedHeader, cancellationToken), cancellationToken);
-							currentTasks.Add((queuedHeader.Path, task));
-							continue;
+				bundle = FindOrAddBundle(locator);
+				bundle.InfoRefCount++;
 						}
 
-						// Try to start another packet read
-						QueuedPacket? queuedPacket = _queuedPackets.FirstOrDefault(x => !currentPaths.Contains(x.Path));
-						if (queuedPacket != null)
+			// Wait for the read to complete
+			try
 						{
-							Task task = Task.Run(() => PerformPacketReadGuardedAsync(queuedPacket, cancellationToken), cancellationToken);
-							currentTasks.Add((queuedPacket.Path, task));
-							continue;
+				BundleInfo bundleInfo = await bundle.BundleInfo.Task.WaitAsync(cancellationToken);
+				return bundleInfo.Header;
 						}
-
-						// If we're not waiting for anything else and there are no more requests, end the task thread.
-						if (currentTasks.Count == 0)
+			finally
 						{
-							_readTask = null;
-							return;
+				await CancelBundleRequestAsync(bundle, () => bundle.InfoRefCount--);
 						}
-
-						// Break out of the loop
-						break;
 					}
-				}
 
-				// Wait for any read task to complete
-				await Task.WhenAny(currentTasks.Select(x => x.Task));
+		static string GetDecodeTaskKey(BundleLocator locator, int packetIdx) => $"{locator}#{packetIdx}";
 
-				// Remove any tasks which are complete
-				for (int idx = 0; idx < currentTasks.Count; idx++)
+		/// <summary>
+		/// Reads decoded packet data from a bundle
+		/// </summary>
+		/// <param name="locator">Locator for the bundle to read</param>
+		/// <param name="packetIdx">Index of the bundle packet</param>
+		/// <param name="cancellationToken">Cancellation token for the operation</param>
+		/// <returns>Data for the packet</returns>
+		public async Task<ReadOnlyMemory<byte>> ReadPacketAsync(BundleLocator locator, int packetIdx, CancellationToken cancellationToken)
 				{
-					Task task = currentTasks[idx].Task;
-					if (task.IsCompleted)
+			ReadOnlyMemory<byte> cachedDecodedPacket;
+			if (_cache.TryGetCachedDecodedPacket(locator, packetIdx, out cachedDecodedPacket))
 					{
-						if (task.Exception != null)
-						{
-							_logger.LogError(task.Exception, "Exception while reading from blob {BlobId}.", currentTasks[idx].Path);
+				return cachedDecodedPacket;
 						}
-						currentTasks.RemoveAt(idx--);
-					}
-				}
-			}
-		}
 
-		/// <summary>
-		/// Reads a bundle header from the queue
-		/// </summary>
-		/// <param name="queuedHeader">The header to read</param>
-		/// <param name="cancellationToken">Cancellation token for the operation</param>
-		async Task PerformHeaderReadGuardedAsync(QueuedHeader queuedHeader, CancellationToken cancellationToken)
+			// Create an async task to read the data
+			Task<ReadOnlyMemory<byte>>? decodeTask;
+			lock (_queueLock)
 		{
-			try
+				if (_cache.TryGetCachedDecodedPacket(locator, packetIdx, out cachedDecodedPacket))
 			{
-				await PerformHeaderReadAsync(queuedHeader, cancellationToken);
+					return cachedDecodedPacket;
 			}
-			catch (Exception ex)
+
+				string decodedCacheKey = GetDecodeTaskKey(locator, packetIdx);
+				if (!_decodeTasks.TryGetValue(decodedCacheKey, out decodeTask))
 			{
-				queuedHeader.CompletionSource.TrySetException(ex);
+					decodeTask = Task.Run(() => ReadAndDecodePacketAsync(locator, packetIdx, CancellationToken.None), CancellationToken.None);
+					_decodeTasks.Add(decodedCacheKey, decodeTask);
 			}
 		}
 
-		/// <summary>
-		/// Reads a bundle header from the queue
-		/// </summary>
-		/// <param name="queuedHeader">The header to read</param>
-		/// <param name="cancellationToken">Cancellation token for the operation</param>
-		async Task PerformHeaderReadAsync(QueuedHeader queuedHeader, CancellationToken cancellationToken)
+			return await decodeTask.WaitAsync(cancellationToken);
+		}
+
+		async Task<ReadOnlyMemory<byte>> ReadAndDecodePacketAsync(BundleLocator locator, int packetIdx, CancellationToken cancellationToken)
 		{
-			Interlocked.Increment(ref _numHeaderReads);
+			BundleHeader header = await ReadHeaderAsync(locator, cancellationToken);
+			ReadOnlyMemory<byte> encodedPacket = await ReadEncodedPacketAsync(locator, packetIdx, cancellationToken);
 
-			int prefetchSize = _cache.HasPacketCache ? DefaultFetchSize : DefaultUncachedFetchSize;
-			for (; ; )
-			{
-				await using (Stream stream = await _store.OpenAsync(queuedHeader.Blob, 0, prefetchSize, cancellationToken))
-				{
-					// Read the header data
-					byte[] prelude = new byte[BundleHeader.PreludeLength];
-					await stream.ReadFixedLengthBytesAsync(prelude, cancellationToken);
+			BundlePacket packet = header.Packets[packetIdx];
 
-					// Make sure we've read enough to hold the header
-					int headerSize = BundleHeader.ReadPrelude(prelude);
-					if (headerSize > prefetchSize)
+			byte[] decodedPacket = new byte[packet.DecodedLength];
+			BundleData.Decompress(packet.CompressionFormat, encodedPacket, decodedPacket);
+			_cache.AddCachedDecodedPacket(locator, packetIdx, decodedPacket);
+
+			lock (_queueLock)
 					{
-						prefetchSize = headerSize;
-						continue;
+				string decodedCacheKey = GetDecodeTaskKey(locator, packetIdx);
+				_decodeTasks.Remove(decodedCacheKey);
 					}
 
-					// Parse the header and construct the bundle info from it
-					BundleHeader header = await BundleHeader.ReadAsync(prelude, stream, cancellationToken);
-
-					// Construct the bundle info
-					BundleInfo bundleInfo = new BundleInfo(queuedHeader.Blob, header, headerSize);
+			return decodedPacket;
+		}
 
-					List<ReadOnlyMemory<byte>> packets = new List<ReadOnlyMemory<byte>>();
-					if (_cache.HasPacketCache)
+		/// <summary>
+		/// Reads encoded packet data from a bundle
+		/// </summary>
+		/// <param name="locator">Locator for the bundle to read</param>
+		/// <param name="packetIdx">Index of the bundle packet</param>
+		/// <param name="cancellationToken">Cancellation token for the operation</param>
+		/// <returns>Dtaa for the packet</returns>
+		public async Task<ReadOnlyMemory<byte>> ReadEncodedPacketAsync(BundleLocator locator, int packetIdx, CancellationToken cancellationToken)
 					{
-						// Also add any encoded packets we prefetched
-						int packetOffset = headerSize;
-						for (int packetIdx = 0; packetIdx < header.Packets.Count; packetIdx++)
+			// Register a read for the packet
+			QueuedBundle? bundle = null;
+			QueuedPacket? packet = null;
+			lock (_queueLock)
 						{
-							int packetLength = header.Packets[packetIdx].EncodedLength;
-							if (packetOffset + packetLength > prefetchSize)
+				if (_cache.TryGetCachedEncodedPacket(locator, packetIdx, out ReadOnlyMemory<byte> cachedEncodedPacket))
 							{
-								break;
+					return cachedEncodedPacket;
 							}
 
-							byte[] packetData = await ReadPacketAsync(stream, packetLength, cancellationToken);
-							_cache.AddCachedEncodedPacket(queuedHeader.Blob, packetIdx, packetData);
-							packets.Add(packetData);
-
-							packetOffset += packetLength;
-						}
+				bundle = FindOrAddBundle(locator);
+				packet = bundle.AddPacket(packetIdx);
 					}
-					_cache.AddCachedHeader(queuedHeader.Blob, bundleInfo);
 
-					// Update any packets that now have a cached value
-					lock (_queueLock)
+			// Wait for the read to complete
+			using (CancellationTokenRegistration registration = cancellationToken.Register(() => packet.CompletionSource.TrySetCanceled()))
 					{
-						queuedHeader.CompletionSource.TrySetResult(bundleInfo);
-
-						foreach (QueuedPacket queuedPacket in _queuedPackets)
+				try
 						{
-							if (queuedPacket.Path == bundleInfo.Locator.Path && queuedPacket.PacketIdx < packets.Count)
+					return await packet.CompletionSource.Task;
+				}
+				finally
 							{
-								queuedPacket.CompletionSource.TrySetResult(packets[queuedPacket.PacketIdx]);
+					await CancelBundleRequestAsync(bundle, () => bundle.QueuedPackets.Remove(packet));
 							}
 						}
 					}
 
-					// Remove it from the queue
-					lock (_queueLock)
+		QueuedBundle FindOrAddBundle(BundleLocator locator)
 					{
-						_queuedHeaders.Remove(queuedHeader);
+			QueuedBundle? bundle = _queuedBundles.FirstOrDefault(x => x.Locator == locator);
+			if (bundle == null)
+			{
+				bundle = new QueuedBundle(locator);
+				bundle.WorkerTask = Task.Run(() => HandleBundleRequestsAsync(bundle), bundle.CancellationSource.Token);
+				_queuedBundles.Add(bundle);
 					}
-					break;
-				}
-			}
+			return bundle;
 		}
 
-		/// <summary>
-		/// Reads a packet from storage
-		/// </summary>
-		/// <param name="queuedPacket">The packet to read</param>
-		/// <param name="cancellationToken">Cancellation token for the operation</param>
-		async Task PerformPacketReadGuardedAsync(QueuedPacket queuedPacket, CancellationToken cancellationToken)
+		async Task HandleBundleRequestsAsync(QueuedBundle bundle)
 		{
+			// Read the bundle header
+			BundleInfo? bundleInfo;
+			if (!_cache.TryGetCachedHeader(bundle.Locator, out bundleInfo))
+			{
 			try
 			{
-				await PerformPacketReadAsync(queuedPacket, cancellationToken);
+					bundleInfo = await ReadBundleInfoAsync(bundle, bundle.CancellationSource.Token);
+					_cache.AddCachedHeader(bundle.Locator, bundleInfo);
 			}
 			catch (Exception ex)
 			{
-				if (!queuedPacket.CompletionSource.TrySetException(ex))
-				{
-					_logger.LogWarning(ex, "Exception after setting completion source state; existing state: {Status}, new exception: {Ex}", queuedPacket.CompletionSource.Task.Status, ex);
+					bundle.BundleInfo.SetException(ex);
+					return;
 				}
 			}
+			bundle.BundleInfo.SetResult(bundleInfo);
+
+			// Serve any packet read requests
+			for (; ; )
+			{
+				// Create the next read
+				PendingRead? pendingRead;
+				lock (_queueLock)
+				{
+					// Dispose of the previous pending read
+					if (bundle.PendingRead != null)
+					{
+						bundle.PendingRead.Dispose();
+						bundle.PendingRead = null;
 		}
 
-		/// <summary>
-		/// Reads a packet from storage
-		/// </summary>
-		/// <param name="queuedPacket">The packet to read</param>
-		/// <param name="cancellationToken">Cancellation token for the operation</param>
-		async Task PerformPacketReadAsync(QueuedPacket queuedPacket, CancellationToken cancellationToken)
+					// If there's nothing left to read, dispose of it
+					if (bundle.QueuedPackets.Count == 0)
 		{
-			Interlocked.Increment(ref _numPacketReads);
-
-			BundleInfo bundleInfo = queuedPacket.Bundle;
-			int minPacketIdx = queuedPacket.PacketIdx;
+						bundle.Complete = true;
+						_queuedBundles.Remove(bundle);
+						break;
+					}
 
-			BundlePacket packet = bundleInfo.Header.Packets[minPacketIdx];
-			int readLength = packet.EncodedLength;
+					// Figure out the range of packets to read
+					int minPacketIdx = bundle.QueuedPackets[0].PacketIdx;
+					int maxPacketIdx = minPacketIdx;
 
-			int maxPacketIdx = minPacketIdx + 1;
-			if(_cache.HasPacketCache)
+					long length = 0;
+					for (int packetIdx = minPacketIdx; packetIdx < bundleInfo.Header.Packets.Count; packetIdx++)
 			{
-				for (; maxPacketIdx < bundleInfo.Header.Packets.Count; maxPacketIdx++)
+						length += bundleInfo.Header.Packets[packetIdx].EncodedLength;
+						if (length > DefaultFetchSize)
 				{
-					int nextReadLength = readLength + bundleInfo.Header.Packets[maxPacketIdx].EncodedLength;
-					if (nextReadLength > DefaultFetchSize)
-					{
 						break;
 					}
-					readLength = nextReadLength;
+						maxPacketIdx = packetIdx;
 				}
+
+					// Create the new read
+					pendingRead = new PendingRead(minPacketIdx, maxPacketIdx, bundle.CancellationSource.Token);
+					bundle.PendingRead = pendingRead;
 			}
 
-			await using (Stream stream = await _store.OpenAsync(bundleInfo.Locator, bundleInfo.HeaderLength + bundleInfo.Header.Packets[minPacketIdx].EncodedOffset, readLength, cancellationToken))
+				// Execute the read
+				try
 			{
-				// Copy all the packets that have been read into separate buffers, so we can cache them indidually.
-				ReadOnlyMemory<byte>[] packets = new ReadOnlyMemory<byte>[maxPacketIdx - minPacketIdx];
-				for (int idx = minPacketIdx; idx < maxPacketIdx; idx++)
-				{
-					byte[] data = await ReadPacketAsync(stream, bundleInfo.Header.Packets[idx].EncodedLength, cancellationToken);
-					packets[idx - minPacketIdx] = data;
+					await ReadEncodedPacketsAsync(bundle, bundle.PendingRead.MinPacketIdx, bundle.PendingRead.MaxPacketIdx, pendingRead.CancellationSource.Token);
 				}
-
-				// Add any complete packets to the cache
-				if (_cache.HasPacketCache)
+				catch (OperationCanceledException ex)
 				{
-					for (int idx = minPacketIdx; idx < maxPacketIdx; idx++)
-					{
-						ReadOnlyMemory<byte> data = packets[idx - minPacketIdx];
-						_cache.AddCachedEncodedPacket(bundleInfo.Locator, idx, data);
-					}
+					_logger.LogTrace(ex, "Read from bundle was cancelled");
 				}
+				catch (Exception ex)
+				{
+					_logger.LogError(ex, "Error reading from bundle: {Message}", ex.Message);
 
-				// Find all the packets we can mark as complete
 				lock (_queueLock)
 				{
-					for (int idx = 0; idx < _queuedPackets.Count; idx++)
+						List<QueuedPacket> packets = bundle.QueuedPackets;
+						for (int idx = 0; idx < packets.Count && packets[idx].PacketIdx <= pendingRead.MaxPacketIdx; idx++)
 					{
-						QueuedPacket updatePacket = _queuedPackets[idx];
-						if (updatePacket.Path == bundleInfo.Locator.Path && (updatePacket.PacketIdx >= minPacketIdx && updatePacket.PacketIdx < maxPacketIdx))
+							if (packets[idx].PacketIdx >= pendingRead.MinPacketIdx)
 						{
-							ReadOnlyMemory<byte> data = packets[updatePacket.PacketIdx - minPacketIdx];
-							updatePacket.CompletionSource.TrySetResult(data);
-							_queuedPackets.RemoveAt(idx--);
+								packets[idx].CompletionSource.TrySetException(ex);
+								packets.RemoveAt(idx--);
 						}
 					}
 				}
 			}
 		}
 
-		static async Task<byte[]> ReadPacketAsync(Stream stream, int packetSize, CancellationToken cancellationToken)
+			// Dispose the cancellation source used for this bundle
+			lock (_queueLock)
 		{
-			byte[] packet = new byte[packetSize];
-			await stream.ReadFixedLengthBytesAsync(packet, cancellationToken);
-			return packet;
+				bundle.CancellationSource.Dispose();
+				bundle.CancellationSource = null!;
 		}
-
-		/// <summary>
-		/// Reads a bundle header from the given blob locator, or retrieves it from the cache
-		/// </summary>
-		/// <param name="locator"></param>
-		/// <param name="cancellationToken"></param>
-		/// <returns></returns>
-		public async Task<BundleHeader> ReadBundleHeaderAsync(BundleLocator locator, CancellationToken cancellationToken = default)
-		{
-			BundleInfo bundleInfo = await GetBundleInfoAsync(locator, cancellationToken);
-			return bundleInfo.Header;
 		}
 
-		async Task<BundleInfo> GetBundleInfoAsync(BundleLocator locator, CancellationToken cancellationToken = default)
+		async Task CancelBundleRequestAsync(QueuedBundle bundle, Action updateAction)
 		{
-			Debug.Assert(locator.IsValid());
+			// Remove this read request
+			Task? workerTask = null;
+			lock (_queueLock)
+		{
+				updateAction();
 
-			BundleInfo? bundleInfo;
-			if (_cache.TryGetCachedHeader(locator, out bundleInfo))
+				// If there's nothing required any more, remove the bundle from the queue
+				if (bundle.InfoRefCount == 0 && bundle.QueuedPackets.Count == 0)
 			{
-				return bundleInfo;
+					workerTask = bundle.WorkerTask;
+					_queuedBundles.Remove(bundle);
+				}
 			}
 
-			QueuedHeader? queuedHeader;
-			lock (_queueLock)
+			// If this was the last thing using the queued bundle, wait for the worker task to finish
+			if (workerTask != null)
 			{
-				// Check the cache again inside lock scope to avoid races
-				if (_cache.TryGetCachedHeader(locator, out bundleInfo))
-				{
-					return bundleInfo;
-				}
+				await workerTask;
 
-				// Find or start the read
-				queuedHeader = _queuedHeaders.FirstOrDefault(x => x.Blob == locator);
-				if (queuedHeader == null)
+				try
 				{
-					queuedHeader = new QueuedHeader(locator);
-					_queuedHeaders.Add(queuedHeader);
-					StartReadTask();
+					await bundle.BundleInfo.Task; // Avoid unobserved cancellation exceptions
 				}
+				catch { }
 			}
-			return await queuedHeader.CompletionSource.Task.WaitAsync(cancellationToken);
 		}
 
-		static string GetDecodeTaskKey(BundleLocator locator, int packetIdx) => $"{locator}:{packetIdx}";
+		#endregion
+
+		#region Reading
 
 		/// <summary>
-		/// Gets a decoded block from the store
+		/// Reads a bundle header from the queue
 		/// </summary>
-		/// <param name="bundleInfo">Information about the bundle</param>
-		/// <param name="packetIdx">Index of the packet</param>
+		/// <param name="bundle">Bundle header to read</param>
 		/// <param name="cancellationToken">Cancellation token for the operation</param>
-		/// <returns>The decoded data</returns>
-		async ValueTask<ReadOnlyMemory<byte>> ReadBundlePacketAsync(BundleInfo bundleInfo, int packetIdx, CancellationToken cancellationToken)
+		async Task<BundleInfo> ReadBundleInfoAsync(QueuedBundle bundle, CancellationToken cancellationToken)
 		{
-			if (packetIdx < 0 || packetIdx >= bundleInfo.Header.Packets.Count)
+			Interlocked.Increment(ref _numHeaderReads);
+
+			int prefetchSize = _cache != null ? DefaultFetchSize : DefaultUncachedFetchSize;
+			for (; ; )
 			{
-				throw new ArgumentException("Packet index is out of range", nameof(packetIdx));
-			}
+				await using (Stream stream = await _store.OpenAsync(bundle.Locator, 0, prefetchSize, cancellationToken))
+				{
+					// Read the header data
+					byte[] prelude = new byte[BundleHeader.PreludeLength];
+					await stream.ReadFixedLengthBytesAsync(prelude, cancellationToken);
 
-			ReadOnlyMemory<byte> decodedPacket;
-			if (_cache.TryGetCachedDecodedPacket(bundleInfo.Locator, packetIdx, out decodedPacket))
+					// Make sure we've read enough to hold the header
+					int headerSize = BundleHeader.ReadPrelude(prelude);
+					if (headerSize > prefetchSize)
 			{
-				return decodedPacket;
+						prefetchSize = headerSize;
+						continue;
 			}
 
-			Task<ReadOnlyMemory<byte>>? decodeTask;
-			lock (_queueLock)
+					// Parse the header and construct the bundle info from it
+					BundleHeader header = await BundleHeader.ReadAsync(prelude, stream, cancellationToken);
+
+					// Construct the bundle info
+					BundleInfo bundleInfo = new BundleInfo(bundle.Locator, header, headerSize);
+
+					// Also add any encoded packets we prefetched
+					int packetOffset = headerSize;
+					for (int packetIdx = 0; packetIdx < header.Packets.Count; packetIdx++)
 			{
-				// Query the cache again, to eliminate races between cache checks and decode tasks finishing.
-				if (_cache.TryGetCachedDecodedPacket(bundleInfo.Locator, packetIdx, out decodedPacket))
+						BundlePacket packet = header.Packets[packetIdx];
+
+						int packetLength = packet.EncodedLength;
+						if (packetOffset + packetLength > prefetchSize)
 				{
-					return decodedPacket;
+							break;
 				}
+						packetOffset += packetLength;
 
-				// Create an async task to read the data
-				string decodedCacheKey = GetDecodeTaskKey(bundleInfo.Locator, packetIdx);
-				if (!_decodeTasks.TryGetValue(decodedCacheKey, out decodeTask))
-				{
-					decodeTask = Task.Run(() => ReadAndDecodePacketAsync(bundleInfo, packetIdx), CancellationToken.None);
-					_decodeTasks.Add(decodedCacheKey, decodeTask);
+						await ReadEncodedPacketFromStreamAsync(bundle, packetIdx, packet, stream, cancellationToken);
 				}
+
+					Interlocked.Add(ref _numBytesRead, packetOffset);
+					return bundleInfo;
+			}
 			}
-			return await decodeTask.WaitAsync(cancellationToken);
 		}
 
-		/// <summary>
-		/// Reads and decodes a packet from a bundle
-		/// </summary>
-		/// <param name="bundleInfo">Bundle to read from</param>
-		/// <param name="packetIdx">Index of the packet to return</param>
-		/// <returns>The decoded packet data</returns>
-		async Task<ReadOnlyMemory<byte>> ReadAndDecodePacketAsync(BundleInfo bundleInfo, int packetIdx)
+		async Task ReadEncodedPacketsAsync(QueuedBundle bundle, int minPacketIdx, int maxPacketIdx, CancellationToken cancellationToken)
 		{
-			ReadOnlyMemory<byte> encodedPacket = await ReadEncodedPacketAsync(bundleInfo, packetIdx);
+			BundleInfo bundleInfo = await bundle.BundleInfo.Task;
 
-			BundlePacket packet = bundleInfo.Header.Packets[packetIdx];
-			byte[] decodedPacket = new byte[packet.DecodedLength];
+			BundlePacket minPacket = bundleInfo.Header.Packets[minPacketIdx];
+			int minOffset = minPacket.EncodedOffset;
 
-			BundleData.Decompress(packet.CompressionFormat, encodedPacket, decodedPacket);
-			_cache.AddCachedDecodedPacket(bundleInfo.Locator, packetIdx, decodedPacket);
+			BundlePacket maxPacket = bundleInfo.Header.Packets[maxPacketIdx];
+			int maxOffset = maxPacket.EncodedOffset + maxPacket.EncodedLength;
 
-			lock (_queueLock)
+			Interlocked.Increment(ref _numPacketReads);
+			Interlocked.Add(ref _numBytesRead, maxOffset - minOffset);
+
+			await using (Stream stream = await _store.OpenAsync(bundleInfo.Locator, bundleInfo.HeaderLength + minOffset, maxOffset - minOffset, cancellationToken))
 			{
-				string decodedCacheKey = GetDecodeTaskKey(bundleInfo.Locator, packetIdx);
-				_decodeTasks.Remove(decodedCacheKey);
+				// Copy all the packets that have been read into separate buffers, so we can cache them indidually.
+				for (int packetIdx = minPacketIdx; packetIdx <= maxPacketIdx; packetIdx++)
+			{
+					BundlePacket packet = bundleInfo.Header.Packets[packetIdx];
+					await ReadEncodedPacketFromStreamAsync(bundle, packetIdx, packet, stream, cancellationToken);
+				}
 			}
-
-			return decodedPacket;
 		}
 
-		/// <summary>
-		/// Reads an encoded packet from a bundle
-		/// </summary>
-		/// <param name="bundleInfo">Bundle to read from</param>
-		/// <param name="packetIdx">Index of the packet to return</param>
-		/// <returns>The encoded packet data</returns>
-		async ValueTask<ReadOnlyMemory<byte>> ReadEncodedPacketAsync(BundleInfo bundleInfo, int packetIdx)
+		async Task ReadEncodedPacketFromStreamAsync(QueuedBundle bundle, int packetIdx, BundlePacket packet, Stream stream, CancellationToken cancellationToken)
 		{
-			if (packetIdx < 0 || packetIdx >= bundleInfo.Header.Packets.Count)
-			{
-				throw new ArgumentException("Packet index is out of range", nameof(packetIdx));
-			}
+			byte[] data = new byte[packet.EncodedLength];
+			await stream.ReadFixedLengthBytesAsync(data, cancellationToken);
 
-			ReadOnlyMemory<byte> encodedPacket;
-			if (_cache.TryGetCachedEncodedPacket(bundleInfo.Locator, packetIdx, out encodedPacket))
+			lock (_queueLock)
 			{
-				return encodedPacket;
-			}
+				_cache.AddCachedEncodedPacket(bundle.Locator, packetIdx, data); // Note: do this while holding the lock, to avoid races with new encodes being added
 
-			QueuedPacket? queuedPacket;
-			lock (_queueLock)
+				List<QueuedPacket> packets = bundle.QueuedPackets;
+				for (int idx = 0; idx < packets.Count && packets[idx].PacketIdx <= packetIdx; idx++)
 			{
-				// Query the cache again, to eliminate races between cache checks and decode tasks finishing.
-				if (_cache.TryGetCachedEncodedPacket(bundleInfo.Locator, packetIdx, out encodedPacket))
+					if (packets[idx].PacketIdx == packetIdx)
 				{
-					return encodedPacket;
+						packets[idx].CompletionSource.TrySetResult(data);
+						packets.RemoveAt(idx--);
 				}
-
-				// Add a read to the queue
-				queuedPacket = _queuedPackets.FirstOrDefault(x => x.Bundle.Locator == bundleInfo.Locator && x.PacketIdx == packetIdx);
-				if (queuedPacket == null)
-				{
-					queuedPacket = new QueuedPacket(bundleInfo, packetIdx);
-					_queuedPackets.Add(queuedPacket);
-					StartReadTask();
 				}
 			}
-			return await queuedPacket.CompletionSource.Task;
 		}
 
 		#endregion
@@ -553,8 +548,8 @@ namespace EpicGames.Horde.Storage.Bundles
 		/// <returns>Node data read from the given bundle</returns>
 		public async ValueTask<BlobData> ReadNodeDataAsync(BundleNodeLocator locator, CancellationToken cancellationToken = default)
 		{
-			BundleInfo bundleInfo = await GetBundleInfoAsync(locator.Blob, cancellationToken);
-			BundleExport export = bundleInfo.Header.Exports[locator.ExportIdx];
+			BundleHeader header = await ReadHeaderAsync(locator.Blob, cancellationToken);
+			BundleExport export = header.Exports[locator.ExportIdx];
 
 			List<BlobHandle> refs = new List<BlobHandle>(export.References.Count);
 			foreach (BundleExportRef reference in export.References)
@@ -566,7 +561,7 @@ namespace EpicGames.Horde.Storage.Bundles
 				}
 				else
 				{
-					importBlob = bundleInfo.Header.Imports[reference.ImportIdx];
+					importBlob = header.Imports[reference.ImportIdx];
 				}
 				Debug.Assert(importBlob.IsValid());
 				refs.Add(new FlushedNodeHandle(this, new BundleNodeLocator(reference.Hash, importBlob, reference.NodeIdx)));
@@ -575,11 +570,11 @@ namespace EpicGames.Horde.Storage.Bundles
 			ReadOnlyMemory<byte> nodeData = ReadOnlyMemory<byte>.Empty;
 			if (export.Length > 0)
 			{
-				ReadOnlyMemory<byte> packetData = await ReadBundlePacketAsync(bundleInfo, export.Packet, cancellationToken);
+				ReadOnlyMemory<byte> packetData = await ReadPacketAsync(locator.Blob, export.Packet, cancellationToken);
 				nodeData = packetData.Slice(export.Offset, export.Length);
 			}
 
-			BlobType nodeType = bundleInfo.Header.Types[export.TypeIdx];
+			BlobType nodeType = header.Types[export.TypeIdx];
 			return new BlobData(nodeType, export.Hash, nodeData, refs);
 		}
 
diff --git a/Engine/Source/Programs/Shared/EpicGames.Horde/Storage/StorageCache.cs b/Engine/Source/Programs/Shared/EpicGames.Horde/Storage/StorageCache.cs
--- a/Engine/Source/Programs/Shared/EpicGames.Horde/Storage/StorageCache.cs
+++ b/Engine/Source/Programs/Shared/EpicGames.Horde/Storage/StorageCache.cs
@@ -20,7 +20,7 @@ namespace EpicGames.Horde.Storage
 		/// <summary>
 		/// Size of the packet cache
 		/// </summary>
-		public long PacketCacheSize { get; set; } = 256 * 1024 * 1024;
+		public long PacketCacheSize { get; set; } = 192 * 1024 * 1024;
 	}
 
 	/// <summary>
