EpicGames.Oodle: Add support for Windows on Arm64.

diff --git a/Engine/Source/Programs/Shared/EpicGames.Oodle/EpicGames.Oodle.csproj b/Engine/Source/Programs/Shared/EpicGames.Oodle/EpicGames.Oodle.csproj
--- a/Engine/Source/Programs/Shared/EpicGames.Oodle/EpicGames.Oodle.csproj
+++ b/Engine/Source/Programs/Shared/EpicGames.Oodle/EpicGames.Oodle.csproj
@@ -48,5 +48,6 @@
     <None Include="Sdk\2.9.3\linux\lib\liboo2corelinux.so.9" Link="liboo2corelinux.so.9" CopyToOutputDirectory="PreserveNewest" Visible="false" />
     <None Include="Sdk\2.9.3\win\redist\oo2core_9_win64.dll" Link="oo2core_9_win64.dll" CopyToOutputDirectory="PreserveNewest" Visible="false" />
     <None Include="Sdk\2.9.3\win\redist\oo2core_9_win32.dll" Link="oo2core_9_win32.dll" CopyToOutputDirectory="PreserveNewest" Visible="false" />
+    <None Include="Sdk\2.9.10\winuwp\redist\oo2core_9_winuwparm64.dll" Link="oo2core_9_winuwparm64.dll" CopyToOutputDirectory="PreserveNewest" Visible="false" />
   </ItemGroup>
 </Project>
\ No newline at end of file
diff --git a/Engine/Source/Programs/Shared/EpicGames.Oodle/Oodle.cs b/Engine/Source/Programs/Shared/EpicGames.Oodle/Oodle.cs
--- a/Engine/Source/Programs/Shared/EpicGames.Oodle/Oodle.cs
+++ b/Engine/Source/Programs/Shared/EpicGames.Oodle/Oodle.cs
@@ -252,19 +252,40 @@ namespace EpicGames.Compression
 			{
 				// we manually load the assembly as the file name is not consistent between platforms thus the automatic searching will not work
 				string assemblyFilename;
-				bool is64Bit = Environment.Is64BitProcess;
-
 				if (RuntimeInformation.IsOSPlatform(OSPlatform.Windows))
 				{
-					assemblyFilename = $"oo2core_9_{(is64Bit ? "win64" : "win32")}.dll";
+					switch (RuntimeInformation.ProcessArchitecture)
+					{
+						case Architecture.X86:
+							assemblyFilename = "oo2core_9_win32.dll";
+							break;
+						case Architecture.X64:
+							assemblyFilename = "oo2core_9_win64.dll";
+							break;
+						case Architecture.Arm64:
+							assemblyFilename = "oo2core_9_winuwparm64.dll";
+							break;
+						default:
+							throw new PlatformNotSupportedException($"Oodle support is not currently implemented for {RuntimeInformation.RuntimeIdentifier}");
+					}
 				}
 				else if (RuntimeInformation.IsOSPlatform(OSPlatform.Linux))
 				{
-					assemblyFilename = $"liboo2core{(is64Bit ? "linux64" : "linux")}.so.9";
+					switch (RuntimeInformation.ProcessArchitecture)
+					{
+						case Architecture.X86:
+							assemblyFilename = "liboo2corelinux.so.9";
+							break;
+						case Architecture.X64:
+							assemblyFilename = "liboo2corelinux64.so.9";
+							break;
+						default:
+							throw new PlatformNotSupportedException($"Oodle support is not currently implemented for {RuntimeInformation.RuntimeIdentifier}");
+					}
 				}
 				else
 				{
-					throw new NotImplementedException();
+					throw new PlatformNotSupportedException($"Oodle support is not currently implemented for {RuntimeInformation.RuntimeIdentifier}");
 				}
 				IntPtr handle = NativeLibrary.Load(assemblyFilename, assembly, searchpath);
 				libHandle = handle;
