Horde: Enable Oodle support for bundle compression.

diff --git a/Engine/Source/Programs/Shared/EpicGames.Horde/EpicGames.Horde.csproj b/Engine/Source/Programs/Shared/EpicGames.Horde/EpicGames.Horde.csproj
--- a/Engine/Source/Programs/Shared/EpicGames.Horde/EpicGames.Horde.csproj
+++ b/Engine/Source/Programs/Shared/EpicGames.Horde/EpicGames.Horde.csproj
@@ -48,6 +48,7 @@
   <ItemGroup>
     <ProjectReference Include="..\EpicGames.Core\EpicGames.Core.csproj" />
     <ProjectReference Include="..\EpicGames.IoHash\EpicGames.IoHash.csproj" />
+    <ProjectReference Include="..\EpicGames.Oodle\EpicGames.Oodle.csproj" />
     <ProjectReference Include="..\EpicGames.Serialization\EpicGames.Serialization.csproj" />
   </ItemGroup>
 </Project>
\ No newline at end of file
diff --git a/Engine/Source/Programs/Shared/EpicGames.Horde/Storage/Bundles/Bundle.cs b/Engine/Source/Programs/Shared/EpicGames.Horde/Storage/Bundles/Bundle.cs
--- a/Engine/Source/Programs/Shared/EpicGames.Horde/Storage/Bundles/Bundle.cs
+++ b/Engine/Source/Programs/Shared/EpicGames.Horde/Storage/Bundles/Bundle.cs
@@ -1,5 +1,6 @@
 // Copyright Epic Games, Inc. All Rights Reserved.
 
+using EpicGames.Compression;
 using EpicGames.Core;
 using K4os.Compression.LZ4;
 using System;
@@ -1288,17 +1289,13 @@ namespace EpicGames.Horde.Storage.Bundles
 					}
 				case BundleCompressionFormat.Oodle:
 					{
-#if WITH_OODLE
 						int maxSize = Oodle.MaximumOutputSize(OodleCompressorType.Selkie, input.Length);
 
-						Span<byte> outputSpan = builder.GetSpan(maxSize);
-						int encodedLength = Oodle.Compress(OodleCompressorType.Selkie, input.Span, outputSpan, OodleCompressionLevel.HyperFast);
-						builder.Advance(encodedLength);
+						Span<byte> outputSpan = writer.GetSpan(maxSize);
+						int encodedLength = Oodle.Compress(OodleCompressorType.Selkie, input.Span, outputSpan, OodleCompressionLevel.Normal);
+						writer.Advance(encodedLength);
 
 						return encodedLength;
-#else
-						throw new NotSupportedException("Oodle is not compiled into this build.");
-#endif
 					}
 				case BundleCompressionFormat.Brotli:
 					{
