EpicGames.Oodle: Upgrade Oodle to 2.9.10 on all platforms. Add Linux arm32/arm64 and Mac support.

diff --git a/Engine/Source/Programs/Shared/EpicGames.Oodle/EpicGames.Oodle.csproj b/Engine/Source/Programs/Shared/EpicGames.Oodle/EpicGames.Oodle.csproj
--- a/Engine/Source/Programs/Shared/EpicGames.Oodle/EpicGames.Oodle.csproj
+++ b/Engine/Source/Programs/Shared/EpicGames.Oodle/EpicGames.Oodle.csproj
@@ -44,10 +44,12 @@
     <PackageReference Include="Microsoft.Extensions.Logging" Version="6.0.0" />
   </ItemGroup>
   <ItemGroup>
-    <None Include="Sdk\2.9.3\linux\lib\liboo2corelinux64.so.9" Link="liboo2corelinux64.so.9" CopyToOutputDirectory="PreserveNewest" Visible="false" />
-    <None Include="Sdk\2.9.3\linux\lib\liboo2corelinux.so.9" Link="liboo2corelinux.so.9" CopyToOutputDirectory="PreserveNewest" Visible="false" />
-    <None Include="Sdk\2.9.3\win\redist\oo2core_9_win64.dll" Link="oo2core_9_win64.dll" CopyToOutputDirectory="PreserveNewest" Visible="false" />
-    <None Include="Sdk\2.9.3\win\redist\oo2core_9_win32.dll" Link="oo2core_9_win32.dll" CopyToOutputDirectory="PreserveNewest" Visible="false" />
+    <None Include="Sdk\2.9.10\linux\lib\liboo2corelinux64.so.9" Link="liboo2corelinux64.so.9" CopyToOutputDirectory="PreserveNewest" Visible="false" />
+    <None Include="Sdk\2.9.10\linuxarm\lib\liboo2corelinuxarm32.so.9" Link="liboo2corelinuxarm32.so.9" CopyToOutputDirectory="PreserveNewest" Visible="false" />
+    <None Include="Sdk\2.9.10\linuxarm\lib\liboo2corelinuxarm64.so.9" Link="liboo2corelinuxarm64.so.9" CopyToOutputDirectory="PreserveNewest" Visible="false" />
+    <None Include="Sdk\2.9.10\mac\lib\liboo2coremac64.2.9.10.dylib" Link="liboo2coremac64.2.9.10.dylib" CopyToOutputDirectory="PreserveNewest" Visible="false" />
+    <None Include="Sdk\2.9.10\win\redist\oo2core_9_win64.dll" Link="oo2core_9_win64.dll" CopyToOutputDirectory="PreserveNewest" Visible="false" />
+    <None Include="Sdk\2.9.10\win\redist\oo2core_9_win32.dll" Link="oo2core_9_win32.dll" CopyToOutputDirectory="PreserveNewest" Visible="false" />
     <None Include="Sdk\2.9.10\winuwp\redist\oo2core_9_winuwparm64.dll" Link="oo2core_9_winuwparm64.dll" CopyToOutputDirectory="PreserveNewest" Visible="false" />
   </ItemGroup>
 </Project>
\ No newline at end of file
diff --git a/Engine/Source/Programs/Shared/EpicGames.Oodle/Oodle.cs b/Engine/Source/Programs/Shared/EpicGames.Oodle/Oodle.cs
--- a/Engine/Source/Programs/Shared/EpicGames.Oodle/Oodle.cs
+++ b/Engine/Source/Programs/Shared/EpicGames.Oodle/Oodle.cs
@@ -269,16 +269,33 @@ namespace EpicGames.Compression
 							throw new PlatformNotSupportedException($"Oodle support is not currently implemented for {RuntimeInformation.RuntimeIdentifier}");
 					}
 				}
-				else if (RuntimeInformation.IsOSPlatform(OSPlatform.Linux))
+				else if (RuntimeInformation.IsOSPlatform(OSPlatform.OSX))
 				{
 					switch (RuntimeInformation.ProcessArchitecture)
 					{
-						case Architecture.X86:
-							assemblyFilename = "liboo2corelinux.so.9";
+						case Architecture.X64:
+							assemblyFilename = "liboo2coremac64.2.9.10.dylib";
+							break;
+						case Architecture.Arm64:
+							assemblyFilename = "liboo2coremac64.2.9.10.dylib";
 							break;
+						default:
+							throw new PlatformNotSupportedException($"Oodle support is not currently implemented for {RuntimeInformation.RuntimeIdentifier}");
+					}
+				}
+				else if (RuntimeInformation.IsOSPlatform(OSPlatform.Linux))
+				{
+					switch (RuntimeInformation.ProcessArchitecture)
+					{
 						case Architecture.X64:
 							assemblyFilename = "liboo2corelinux64.so.9";
 							break;
+						case Architecture.Arm:
+							assemblyFilename = "liboo2corelinuxarm32.so.9";
+							break;
+						case Architecture.Arm64:
+							assemblyFilename = "liboo2corelinuxarm64.so.9";
+							break;
 						default:
 							throw new PlatformNotSupportedException($"Oodle support is not currently implemented for {RuntimeInformation.RuntimeIdentifier}");
 					}
