Horde: Move functionality for running the IPC server for compute sockets into WorkerComputeSocketBridge. Allows host applications to communicate via sockets as well as clients.

diff --git a/Engine/Source/Programs/Shared/EpicGames.Horde/Compute/AgentMessageHandler.cs b/Engine/Source/Programs/Shared/EpicGames.Horde/Compute/AgentMessageHandler.cs
--- a/Engine/Source/Programs/Shared/EpicGames.Horde/Compute/AgentMessageHandler.cs
+++ b/Engine/Source/Programs/Shared/EpicGames.Horde/Compute/AgentMessageHandler.cs
@@ -198,26 +198,14 @@ namespace EpicGames.Horde.Compute
 				}
 			}
 
-			using (SharedMemoryBuffer ipcBuffer = SharedMemoryBuffer.CreateNew(null, 1, 64 * 1024))
+			await using (WorkerComputeSocketBridge server = await WorkerComputeSocketBridge.CreateAsync(socket, _logger))
 			{
-				newEnvVars[WorkerComputeSocket.IpcEnvVar] = ipcBuffer.Name;
+				newEnvVars[WorkerComputeSocket.IpcEnvVar] = server.BufferName;
 
-				using ComputeBufferReader ipcBufferReader = ipcBuffer.CreateReader();
-				await using (BackgroundTask backgroundTask = BackgroundTask.StartNew(ctx => ProcessIpcMessagesAsync(socket, ipcBufferReader, new[] { cancellationToken, ctx }, _logger)))
-				{
-					_logger.LogInformation("Launching {Executable} {Arguments}", CommandLineArguments.Quote(executable), CommandLineArguments.Join(arguments));
-					try
-					{
-						await ExecuteProcessInternalAsync(channel, executable, arguments, workingDir, newEnvVars, flags, cancellationToken);
-					}
-					finally
-					{
-						ipcBufferReader.Detach();
-					}
-					_logger.LogInformation("Finished executing process");
-				}
+				_logger.LogInformation("Launching {Executable} {Arguments}", CommandLineArguments.Quote(executable), CommandLineArguments.Join(arguments));
 
-				_logger.LogInformation("Ipc message loop is complete");
+				await ExecuteProcessInternalAsync(channel, executable, arguments, workingDir, newEnvVars, flags, cancellationToken);
+				_logger.LogInformation("Finished executing process");
 			}
 
 			_logger.LogInformation("Child process has shut down");
diff --git a/Engine/Source/Programs/Shared/EpicGames.Horde/Compute/ComputeSocket.cs b/Engine/Source/Programs/Shared/EpicGames.Horde/Compute/ComputeSocket.cs
--- a/Engine/Source/Programs/Shared/EpicGames.Horde/Compute/ComputeSocket.cs
+++ b/Engine/Source/Programs/Shared/EpicGames.Horde/Compute/ComputeSocket.cs
@@ -164,6 +164,141 @@ namespace EpicGames.Horde.Compute
 		}
 	}
 
+	/// <summary>
+	/// Operates a server that a child process can open a <see cref="WorkerComputeSocket"/> to.
+	/// </summary>
+	public sealed class WorkerComputeSocketBridge : IAsyncDisposable
+	{
+		readonly SharedMemoryBuffer _ipcBuffer;
+		readonly ComputeBufferReader _ipcBufferReader;
+		readonly BackgroundTask _backgroundTask;
+		readonly ILogger _logger;
+
+		/// <summary>
+		/// Name of the buffer to pass via <see cref="WorkerComputeSocket.IpcEnvVar"/>
+		/// </summary>
+		public string BufferName => _ipcBuffer.Name;
+
+		/// <summary>
+		/// Constructor
+		/// </summary>
+		private WorkerComputeSocketBridge(SharedMemoryBuffer ipcBuffer, ComputeBufferReader ipcBufferReader, BackgroundTask backgroundTask, ILogger logger)
+		{
+			_ipcBuffer = ipcBuffer;
+			_ipcBufferReader = ipcBufferReader;
+			_backgroundTask = backgroundTask;
+			_logger = logger;
+		}
+
+		/// <summary>
+		/// Creates a new server for <see cref="WorkerComputeSocket"/>
+		/// </summary>
+		/// <param name="socket">Socket to connect to</param>
+		/// <param name="logger">Logger for errors</param>
+		/// <returns>New server instance</returns>
+		public static async Task<WorkerComputeSocketBridge> CreateAsync(ComputeSocket socket, ILogger logger)
+		{
+			SharedMemoryBuffer? ipcBuffer = null;
+			ComputeBufferReader? ipcBufferReader = null;
+			BackgroundTask? backgroundTask = null;
+			try
+			{
+				ipcBuffer = SharedMemoryBuffer.CreateNew(null, 1, 64 * 1024);
+				ipcBufferReader = ipcBuffer.CreateReader();
+				backgroundTask = BackgroundTask.StartNew(ctx => ProcessIpcMessagesAsync(socket, ipcBufferReader, logger, ctx));
+				return new WorkerComputeSocketBridge(ipcBuffer, ipcBufferReader, backgroundTask, logger);
+			}
+			catch
+			{
+				if (backgroundTask != null)
+				{
+					await backgroundTask.DisposeAsync();
+				}
+				ipcBufferReader?.Dispose();
+				ipcBuffer?.Dispose();
+				throw;
+			}
+		}
+
+		/// <inheritdoc/>
+		public async ValueTask DisposeAsync()
+		{
+			await _backgroundTask.DisposeAsync();
+			_logger.LogDebug("Ipc message loop is complete");
+
+			_ipcBufferReader.Dispose();
+			_logger.LogDebug("Closed IPC reader");
+
+			_ipcBuffer.Dispose();
+			_logger.LogDebug("Closed IPC buffer");
+		}
+
+		static async Task ProcessIpcMessagesAsync(ComputeSocket socket, ComputeBufferReader ipcReader, ILogger logger, CancellationToken cancellationToken)
+		{
+			List<SharedMemoryBuffer> buffers = new();
+			try
+			{
+				List<(int, ComputeBufferWriter)> writers = new List<(int, ComputeBufferWriter)>();
+				while (await ipcReader.WaitToReadAsync(1, cancellationToken))
+				{
+					ReadOnlyMemory<byte> memory = ipcReader.GetReadBuffer();
+					MemoryReader reader = new MemoryReader(memory);
+
+					IpcMessage message = (IpcMessage)reader.ReadUnsignedVarInt();
+					try
+					{
+						switch (message)
+						{
+							case IpcMessage.AttachSendBuffer:
+								{
+									int channelId = (int)reader.ReadUnsignedVarInt();
+									string name = reader.ReadString();
+									logger.LogDebug("Attaching send buffer for channel {ChannelId} to {Name}", channelId, name);
+
+									SharedMemoryBuffer buffer = SharedMemoryBuffer.OpenExisting(name);
+									buffers.Add(buffer);
+
+									socket.AttachSendBuffer(channelId, buffer);
+								}
+								break;
+							case IpcMessage.AttachRecvBuffer:
+								{
+									int channelId = (int)reader.ReadUnsignedVarInt();
+									string name = reader.ReadString();
+									logger.LogDebug("Attaching recv buffer for channel {ChannelId} to {Name}", channelId, name);
+
+									SharedMemoryBuffer buffer = SharedMemoryBuffer.OpenExisting(name);
+									buffers.Add(buffer);
+
+									socket.AttachRecvBuffer(channelId, buffer);
+								}
+								break;
+							default:
+								throw new InvalidOperationException($"Invalid IPC message: {message}");
+						}
+					}
+					catch (Exception ex)
+					{
+						logger.LogError(ex, "Exception while processing messages from child process: {Message}", ex.Message);
+					}
+
+					ipcReader.AdvanceReadPosition(memory.Length - reader.RemainingMemory.Length);
+				}
+			}
+			catch (OperationCanceledException)
+			{
+				logger.LogDebug("Ipc message loop cancelled");
+			}
+			finally
+			{
+				foreach (SharedMemoryBuffer buffer in buffers)
+				{
+					buffer.Dispose();
+				}
+			}
+		}
+	}
+
 	/// <summary>
 	/// Manages a set of readers and writers to buffers across a transport layer
 	/// </summary>
