Horde: Add helper methods for opening file streams and subdirectories from bundles.

diff --git a/Engine/Source/Programs/Horde/Horde.Agent/Execution/JobExecutor.cs b/Engine/Source/Programs/Horde/Horde.Agent/Execution/JobExecutor.cs
--- a/Engine/Source/Programs/Horde/Horde.Agent/Execution/JobExecutor.cs
+++ b/Engine/Source/Programs/Horde/Horde.Agent/Execution/JobExecutor.cs
@@ -771,7 +771,7 @@ namespace Horde.Agent.Execution
 					RefName refName = TempStorage.GetRefNameForNode(_storagePrefix, SetupStepName);
 
 					DirectoryNode node = await storage.ReadRefAsync<DirectoryNode>(refName, cancellationToken: cancellationToken);
-					DirectoryNode? buildGraphDir = await node.FindDirectoryAsync(BuildGraphTempStorageDir, cancellationToken);
+					DirectoryNode? buildGraphDir = await node.TryOpenDirectoryAsync(BuildGraphTempStorageDir, cancellationToken);
 					if (buildGraphDir != null)
 					{
 						await buildGraphDir.CopyToDirectoryAsync(new DirectoryInfo(workspaceDir.FullName), logger, cancellationToken);
diff --git a/Engine/Source/Programs/Horde/Horde.Server/Artifacts/ArtifactsController.cs b/Engine/Source/Programs/Horde/Horde.Server/Artifacts/ArtifactsController.cs
--- a/Engine/Source/Programs/Horde/Horde.Server/Artifacts/ArtifactsController.cs
+++ b/Engine/Source/Programs/Horde/Horde.Server/Artifacts/ArtifactsController.cs
@@ -253,7 +253,7 @@ namespace Horde.Server.Artifacts
 				contentType = "application/octet-stream";
 			}
 
-			Stream stream = fileEntry.AsStream();
+			Stream stream = fileEntry.OpenAsStream();
 			if (inline)
 			{
 				return new InlineFileStreamResult(stream, contentType, Path.GetFileName(path));
diff --git a/Engine/Source/Programs/Horde/Horde.Server/Tools/ToolsController.cs b/Engine/Source/Programs/Horde/Horde.Server/Tools/ToolsController.cs
--- a/Engine/Source/Programs/Horde/Horde.Server/Tools/ToolsController.cs
+++ b/Engine/Source/Programs/Horde/Horde.Server/Tools/ToolsController.cs
@@ -522,7 +522,7 @@ namespace Horde.Server.Tools
 					contentType = "application/octet-stream";
 				}
 
-				return new FileStreamResult(entry.AsStream(), contentType) { FileDownloadName = entry.Name.ToString() };
+				return new FileStreamResult(entry.OpenAsStream(), contentType) { FileDownloadName = entry.Name.ToString() };
 			}
 
 			Stream stream = node.AsZipStream();
diff --git a/Engine/Source/Programs/Shared/EpicGames.Horde.Tests/BundleTests.cs b/Engine/Source/Programs/Shared/EpicGames.Horde.Tests/BundleTests.cs
--- a/Engine/Source/Programs/Shared/EpicGames.Horde.Tests/BundleTests.cs
+++ b/Engine/Source/Programs/Shared/EpicGames.Horde.Tests/BundleTests.cs
@@ -227,12 +227,12 @@ namespace EpicGames.Horde.Tests
 
 			Assert.AreEqual(0, newRoot.Files.Count);
 			Assert.AreEqual(1, newRoot.Directories.Count);
-			DirectoryNode? outputNode = await newRoot.FindDirectoryAsync("hello", CancellationToken.None);
+			DirectoryNode? outputNode = await newRoot.TryOpenDirectoryAsync("hello", CancellationToken.None);
 			Assert.IsNotNull(outputNode);
 
 			Assert.AreEqual(0, outputNode!.Files.Count);
 			Assert.AreEqual(1, outputNode!.Directories.Count);
-			DirectoryNode? outputNode2 = await outputNode.FindDirectoryAsync("world", CancellationToken.None);
+			DirectoryNode? outputNode2 = await outputNode.TryOpenDirectoryAsync("world", CancellationToken.None);
 			Assert.IsNotNull(outputNode2);
 
 			Assert.AreEqual(0, outputNode2!.Files.Count);
@@ -295,16 +295,16 @@ namespace EpicGames.Horde.Tests
 			{
 				DirectoryNode root = await _storage.ReadRefAsync<DirectoryNode>(refName);
 
-				DirectoryNode? newNode1 = await root.FindDirectoryAsync("node1", CancellationToken.None);
+				DirectoryNode? newNode1 = await root.TryOpenDirectoryAsync("node1", CancellationToken.None);
 				Assert.IsNotNull(newNode1);
 
-				DirectoryNode? newNode2 = await newNode1!.FindDirectoryAsync("node2", CancellationToken.None);
+				DirectoryNode? newNode2 = await newNode1!.TryOpenDirectoryAsync("node2", CancellationToken.None);
 				Assert.IsNotNull(newNode2);
 
-				DirectoryNode? newNode3 = await newNode2!.FindDirectoryAsync("node3", CancellationToken.None);
+				DirectoryNode? newNode3 = await newNode2!.TryOpenDirectoryAsync("node3", CancellationToken.None);
 				Assert.IsNotNull(newNode3);
 
-				DirectoryNode? newNode4 = await newNode3!.FindDirectoryAsync("node4", CancellationToken.None);
+				DirectoryNode? newNode4 = await newNode3!.TryOpenDirectoryAsync("node4", CancellationToken.None);
 				Assert.IsNotNull(newNode4);
 			}
 		}
diff --git a/Engine/Source/Programs/Shared/EpicGames.Horde/Storage/Nodes/DirectoryNode.cs b/Engine/Source/Programs/Shared/EpicGames.Horde/Storage/Nodes/DirectoryNode.cs
--- a/Engine/Source/Programs/Shared/EpicGames.Horde/Storage/Nodes/DirectoryNode.cs
+++ b/Engine/Source/Programs/Shared/EpicGames.Horde/Storage/Nodes/DirectoryNode.cs
@@ -222,6 +222,31 @@ namespace EpicGames.Horde.Storage.Nodes
 		/// <returns>True if the file was found</returns>
 		public bool TryGetFileEntry(Utf8String name, [NotNullWhen(true)] out FileEntry? entry) => _nameToFileEntry.TryGetValue(name, out entry);
 
+		/// <summary>
+		/// Opens a file for reading
+		/// </summary>
+		/// <param name="name">Name of the file to open</param>
+		/// <returns>Stream for the file</returns>
+		public Stream OpenFile(Utf8String name) => GetFileEntry(name).OpenAsStream();
+
+		/// <summary>
+		/// Attempts to open a file for reading
+		/// </summary>
+		/// <param name="name">Name of the file</param>
+		/// <returns>File stream, or null if the file does not exist</returns>
+		public Stream? TryOpenFile(Utf8String name)
+		{
+			FileEntry? entry;
+			if (TryGetFileEntry(name, out entry))
+			{
+				return entry.OpenAsStream();
+			}
+			else
+			{
+				return null;
+			}
+		}
+
 		/// <summary>
 		/// Deletes the file entry with the given name
 		/// </summary>
@@ -243,7 +268,7 @@ namespace EpicGames.Horde.Storage.Nodes
 		/// <param name="path">Path to the directory</param>
 		/// <param name="cancellationToken">Cancellation token</param>
 		/// <returns>The directory with the given path, or null if it was not found</returns>
-		public async ValueTask<FileEntry?> GetFileEntryByPathAsync(Utf8String path, CancellationToken cancellationToken)
+		public async ValueTask<FileEntry?> GetFileEntryByPathAsync(Utf8String path, CancellationToken cancellationToken = default)
 		{
 			FileEntry? fileEntry;
 
@@ -277,9 +302,9 @@ namespace EpicGames.Horde.Storage.Nodes
 		/// <param name="path">Path to the directory</param>
 		/// <param name="cancellationToken">Cancellation token</param>
 		/// <returns>The directory with the given path, or null if it was not found</returns>
-		public ValueTask<DirectoryNode?> GetDirectoryByPathAsync(Utf8String path, CancellationToken cancellationToken) => GetDirectoryByPathAsync(this, path, cancellationToken);
+		public ValueTask<DirectoryNode?> GetDirectoryByPathAsync(Utf8String path, CancellationToken cancellationToken = default) => GetDirectoryByPathAsync(this, path, cancellationToken);
 
-		static async ValueTask<DirectoryNode?> GetDirectoryByPathAsync(DirectoryNode directoryNode, Utf8String path, CancellationToken cancellationToken)
+		static async ValueTask<DirectoryNode?> GetDirectoryByPathAsync(DirectoryNode directoryNode, Utf8String path, CancellationToken cancellationToken = default)
 		{
 			while (path.Length > 0)
 			{
@@ -314,7 +339,7 @@ namespace EpicGames.Horde.Storage.Nodes
 		/// <param name="path"></param>
 		/// <param name="cancellationToken"></param>
 		/// <returns></returns>
-		public async ValueTask<bool> DeleteFileByPathAsync(Utf8String path, CancellationToken cancellationToken)
+		public async ValueTask<bool> DeleteFileByPathAsync(Utf8String path, CancellationToken cancellationToken = default)
 		{
 			Utf8String remainingPath = path;
 			for (DirectoryNode? directory = this; directory != null;)
@@ -326,7 +351,7 @@ namespace EpicGames.Horde.Storage.Nodes
 				}
 				if (length > 0)
 				{
-					directory = await directory.FindDirectoryAsync(remainingPath.Slice(0, length), cancellationToken);
+					directory = await directory.TryOpenDirectoryAsync(remainingPath.Slice(0, length), cancellationToken);
 				}
 				remainingPath = remainingPath.Slice(length + 1);
 			}
@@ -373,7 +398,23 @@ namespace EpicGames.Horde.Storage.Nodes
 		/// <param name="name">Name of the new directory</param>
 		/// <param name="cancellationToken">Cancellation token for the operation</param>
 		/// <returns>The new directory object</returns>
-		public async ValueTask<DirectoryNode?> FindDirectoryAsync(Utf8String name, CancellationToken cancellationToken)
+		public async ValueTask<DirectoryNode> OpenDirectoryAsync(Utf8String name, CancellationToken cancellationToken = default)
+		{
+			DirectoryNode? directoryNode = await TryOpenDirectoryAsync(name, cancellationToken);
+			if (directoryNode == null)
+			{
+				throw new DirectoryNotFoundException();
+			}
+			return directoryNode;
+		}
+
+		/// <summary>
+		/// Tries to get a directory with the given name
+		/// </summary>
+		/// <param name="name">Name of the new directory</param>
+		/// <param name="cancellationToken">Cancellation token for the operation</param>
+		/// <returns>The new directory object</returns>
+		public async ValueTask<DirectoryNode?> TryOpenDirectoryAsync(Utf8String name, CancellationToken cancellationToken = default)
 		{
 			if (TryGetDirectoryEntry(name, out DirectoryEntry? entry))
 			{
@@ -629,7 +670,7 @@ namespace EpicGames.Horde.Storage.Nodes
 				}
 				else
 				{
-					DirectoryNode? childNode = await FindDirectoryAsync(name, cancellationToken);
+					DirectoryNode? childNode = await TryOpenDirectoryAsync(name, cancellationToken);
 					childNode ??= new DirectoryNode();
 					await childNode.UpdateAsync(directory, writer, cancellationToken);
 					NodeRef<DirectoryNode> nodeRef = await writer.WriteNodeAsync(childNode, cancellationToken);
diff --git a/Engine/Source/Programs/Shared/EpicGames.Horde/Storage/Nodes/FileEntry.cs b/Engine/Source/Programs/Shared/EpicGames.Horde/Storage/Nodes/FileEntry.cs
--- a/Engine/Source/Programs/Shared/EpicGames.Horde/Storage/Nodes/FileEntry.cs
+++ b/Engine/Source/Programs/Shared/EpicGames.Horde/Storage/Nodes/FileEntry.cs
@@ -133,7 +133,7 @@ namespace EpicGames.Horde.Storage.Nodes
 		/// Creates a stream that returns the contents of this file
 		/// </summary>
 		/// <returns>The content stream</returns>
-		public Stream AsStream() => new FileEntryContentStream(this);
+		public Stream OpenAsStream() => new FileEntryContentStream(this);
 
 		/// <summary>
 		/// Copies the contents of this node and its children to the given output stream
