Horde: Formatting fixes.

diff --git a/Engine/Source/Programs/Shared/EpicGames.Horde/Storage/Bundles/BundleReader.cs b/Engine/Source/Programs/Shared/EpicGames.Horde/Storage/Bundles/BundleReader.cs
--- a/Engine/Source/Programs/Shared/EpicGames.Horde/Storage/Bundles/BundleReader.cs
+++ b/Engine/Source/Programs/Shared/EpicGames.Horde/Storage/Bundles/BundleReader.cs
@@ -28,7 +28,7 @@ namespace EpicGames.Horde.Storage.Bundles
 		/// Bundle header
 		/// </summary>
 		public BundleHeader Header { get; }
-		
+
 		/// <summary>
 		/// Length of the header. Required to offset packets from the start of the bundle
 		/// </summary>
@@ -78,11 +78,11 @@ namespace EpicGames.Horde.Storage.Bundles
 				if (insertIdx < 0)
 				{
 					insertIdx = ~insertIdx;
-			}
+				}
 				QueuedPackets.Insert(insertIdx, packet);
 
 				return packet;
-		}
+			}
 
 			public void RemovePacket(QueuedPacket packet) => QueuedPackets.Remove(packet);
 		}
@@ -103,7 +103,7 @@ namespace EpicGames.Horde.Storage.Bundles
 		/// Information about a pending read
 		/// </summary>
 		sealed class PendingRead : IDisposable
-			{
+		{
 			public int MinPacketIdx { get; }
 			public int MaxPacketIdx { get; }
 			public CancellationTokenSource CancellationSource { get; }
@@ -118,7 +118,7 @@ namespace EpicGames.Horde.Storage.Bundles
 			public void Dispose()
 			{
 				CancellationSource.Dispose();
-		}
+			}
 		}
 
 		// Size of data to fetch by default. This is larger than the minimum request size to reduce number of reads.
@@ -181,27 +181,27 @@ namespace EpicGames.Horde.Storage.Bundles
 			if (_cache.TryGetCachedHeader(locator, out BundleInfo? cachedBundleInfo))
 			{
 				return cachedBundleInfo.Header;
-		}
+			}
 
 			// Find a registered bundle info with the given locator
 			QueuedBundle? bundle = null;
-				lock (_queueLock)
-				{
+			lock (_queueLock)
+			{
 				bundle = FindOrAddBundle(locator);
 				bundle.InfoRefCount++;
-						}
+			}
 
 			// Wait for the read to complete
 			try
-						{
+			{
 				BundleInfo bundleInfo = await bundle.BundleInfo.Task.WaitAsync(cancellationToken);
 				return bundleInfo.Header;
-						}
+			}
 			finally
-						{
+			{
 				await CancelBundleRequestAsync(bundle, () => bundle.InfoRefCount--);
-						}
-					}
+			}
+		}
 
 		static string GetDecodeTaskKey(BundleLocator locator, int packetIdx) => $"{locator}#{packetIdx}";
 
@@ -213,29 +213,29 @@ namespace EpicGames.Horde.Storage.Bundles
 		/// <param name="cancellationToken">Cancellation token for the operation</param>
 		/// <returns>Data for the packet</returns>
 		public async Task<ReadOnlyMemory<byte>> ReadPacketAsync(BundleLocator locator, int packetIdx, CancellationToken cancellationToken)
-				{
+		{
 			ReadOnlyMemory<byte> cachedDecodedPacket;
 			if (_cache.TryGetCachedDecodedPacket(locator, packetIdx, out cachedDecodedPacket))
-					{
+			{
 				return cachedDecodedPacket;
-						}
+			}
 
 			// Create an async task to read the data
 			Task<ReadOnlyMemory<byte>>? decodeTask;
 			lock (_queueLock)
-		{
-				if (_cache.TryGetCachedDecodedPacket(locator, packetIdx, out cachedDecodedPacket))
 			{
+				if (_cache.TryGetCachedDecodedPacket(locator, packetIdx, out cachedDecodedPacket))
+				{
 					return cachedDecodedPacket;
-			}
+				}
 
 				string decodedCacheKey = GetDecodeTaskKey(locator, packetIdx);
 				if (!_decodeTasks.TryGetValue(decodedCacheKey, out decodeTask))
-			{
+				{
 					decodeTask = Task.Run(() => ReadAndDecodePacketAsync(locator, packetIdx, CancellationToken.None), CancellationToken.None);
 					_decodeTasks.Add(decodedCacheKey, decodeTask);
+				}
 			}
-		}
 
 			return await decodeTask.WaitAsync(cancellationToken);
 		}
@@ -252,10 +252,10 @@ namespace EpicGames.Horde.Storage.Bundles
 			_cache.AddCachedDecodedPacket(locator, packetIdx, decodedPacket);
 
 			lock (_queueLock)
-					{
+			{
 				string decodedCacheKey = GetDecodeTaskKey(locator, packetIdx);
 				_decodeTasks.Remove(decodedCacheKey);
-					}
+			}
 
 			return decodedPacket;
 		}
@@ -268,44 +268,44 @@ namespace EpicGames.Horde.Storage.Bundles
 		/// <param name="cancellationToken">Cancellation token for the operation</param>
 		/// <returns>Dtaa for the packet</returns>
 		public async Task<ReadOnlyMemory<byte>> ReadEncodedPacketAsync(BundleLocator locator, int packetIdx, CancellationToken cancellationToken)
-					{
+		{
 			// Register a read for the packet
 			QueuedBundle? bundle = null;
 			QueuedPacket? packet = null;
 			lock (_queueLock)
-						{
+			{
 				if (_cache.TryGetCachedEncodedPacket(locator, packetIdx, out ReadOnlyMemory<byte> cachedEncodedPacket))
-							{
+				{
 					return cachedEncodedPacket;
-							}
+				}
 
 				bundle = FindOrAddBundle(locator);
 				packet = bundle.AddPacket(packetIdx);
-					}
+			}
 
 			// Wait for the read to complete
 			using (CancellationTokenRegistration registration = cancellationToken.Register(() => packet.CompletionSource.TrySetCanceled()))
-					{
+			{
 				try
-						{
+				{
 					return await packet.CompletionSource.Task;
 				}
 				finally
-							{
+				{
 					await CancelBundleRequestAsync(bundle, () => bundle.QueuedPackets.Remove(packet));
-							}
-						}
-					}
+				}
+			}
+		}
 
 		QueuedBundle FindOrAddBundle(BundleLocator locator)
-					{
+		{
 			QueuedBundle? bundle = _queuedBundles.FirstOrDefault(x => x.Locator == locator);
 			if (bundle == null)
 			{
 				bundle = new QueuedBundle(locator);
 				bundle.WorkerTask = Task.Run(() => HandleBundleRequestsAsync(bundle), bundle.CancellationSource.Token);
 				_queuedBundles.Add(bundle);
-					}
+			}
 			return bundle;
 		}
 
@@ -315,13 +315,13 @@ namespace EpicGames.Horde.Storage.Bundles
 			BundleInfo? bundleInfo;
 			if (!_cache.TryGetCachedHeader(bundle.Locator, out bundleInfo))
 			{
-			try
-			{
+				try
+				{
 					bundleInfo = await ReadBundleInfoAsync(bundle, bundle.CancellationSource.Token);
 					_cache.AddCachedHeader(bundle.Locator, bundleInfo);
-			}
-			catch (Exception ex)
-			{
+				}
+				catch (Exception ex)
+				{
 					bundle.BundleInfo.SetException(ex);
 					return;
 				}
@@ -340,11 +340,11 @@ namespace EpicGames.Horde.Storage.Bundles
 					{
 						bundle.PendingRead.Dispose();
 						bundle.PendingRead = null;
-		}
+					}
 
 					// If there's nothing left to read, dispose of it
 					if (bundle.QueuedPackets.Count == 0)
-		{
+					{
 						bundle.Complete = true;
 						_queuedBundles.Remove(bundle);
 						break;
@@ -356,23 +356,23 @@ namespace EpicGames.Horde.Storage.Bundles
 
 					long length = 0;
 					for (int packetIdx = minPacketIdx; packetIdx < bundleInfo.Header.Packets.Count; packetIdx++)
-			{
+					{
 						length += bundleInfo.Header.Packets[packetIdx].EncodedLength;
 						if (length > DefaultFetchSize)
-				{
-						break;
-					}
+						{
+							break;
+						}
 						maxPacketIdx = packetIdx;
-				}
+					}
 
 					// Create the new read
 					pendingRead = new PendingRead(minPacketIdx, maxPacketIdx, bundle.CancellationSource.Token);
 					bundle.PendingRead = pendingRead;
-			}
+				}
 
 				// Execute the read
 				try
-			{
+				{
 					await ReadEncodedPacketsAsync(bundle, bundle.PendingRead.MinPacketIdx, bundle.PendingRead.MaxPacketIdx, pendingRead.CancellationSource.Token);
 				}
 				catch (OperationCanceledException ex)
@@ -383,27 +383,27 @@ namespace EpicGames.Horde.Storage.Bundles
 				{
 					_logger.LogError(ex, "Error reading from bundle: {Message}", ex.Message);
 
-				lock (_queueLock)
-				{
+					lock (_queueLock)
+					{
 						List<QueuedPacket> packets = bundle.QueuedPackets;
 						for (int idx = 0; idx < packets.Count && packets[idx].PacketIdx <= pendingRead.MaxPacketIdx; idx++)
-					{
-							if (packets[idx].PacketIdx >= pendingRead.MinPacketIdx)
 						{
+							if (packets[idx].PacketIdx >= pendingRead.MinPacketIdx)
+							{
 								packets[idx].CompletionSource.TrySetException(ex);
 								packets.RemoveAt(idx--);
+							}
 						}
 					}
 				}
 			}
-		}
 
 			// Dispose the cancellation source used for this bundle
 			lock (_queueLock)
-		{
+			{
 				bundle.CancellationSource.Dispose();
 				bundle.CancellationSource = null!;
-		}
+			}
 		}
 
 		async Task CancelBundleRequestAsync(QueuedBundle bundle, Action updateAction)
@@ -411,12 +411,12 @@ namespace EpicGames.Horde.Storage.Bundles
 			// Remove this read request
 			Task? workerTask = null;
 			lock (_queueLock)
-		{
+			{
 				updateAction();
 
 				// If there's nothing required any more, remove the bundle from the queue
 				if (bundle.InfoRefCount == 0 && bundle.QueuedPackets.Count == 0)
-			{
+				{
 					workerTask = bundle.WorkerTask;
 					_queuedBundles.Remove(bundle);
 				}
@@ -460,10 +460,10 @@ namespace EpicGames.Horde.Storage.Bundles
 					// Make sure we've read enough to hold the header
 					int headerSize = BundleHeader.ReadPrelude(prelude);
 					if (headerSize > prefetchSize)
-			{
+					{
 						prefetchSize = headerSize;
 						continue;
-			}
+					}
 
 					// Parse the header and construct the bundle info from it
 					BundleHeader header = await BundleHeader.ReadAsync(prelude, stream, cancellationToken);
@@ -474,22 +474,22 @@ namespace EpicGames.Horde.Storage.Bundles
 					// Also add any encoded packets we prefetched
 					int packetOffset = headerSize;
 					for (int packetIdx = 0; packetIdx < header.Packets.Count; packetIdx++)
-			{
+					{
 						BundlePacket packet = header.Packets[packetIdx];
 
 						int packetLength = packet.EncodedLength;
 						if (packetOffset + packetLength > prefetchSize)
-				{
+						{
 							break;
-				}
+						}
 						packetOffset += packetLength;
 
 						await ReadEncodedPacketFromStreamAsync(bundle, packetIdx, packet, stream, cancellationToken);
-				}
+					}
 
 					Interlocked.Add(ref _numBytesRead, packetOffset);
 					return bundleInfo;
-			}
+				}
 			}
 		}
 
@@ -510,7 +510,7 @@ namespace EpicGames.Horde.Storage.Bundles
 			{
 				// Copy all the packets that have been read into separate buffers, so we can cache them indidually.
 				for (int packetIdx = minPacketIdx; packetIdx <= maxPacketIdx; packetIdx++)
-			{
+				{
 					BundlePacket packet = bundleInfo.Header.Packets[packetIdx];
 					await ReadEncodedPacketFromStreamAsync(bundle, packetIdx, packet, stream, cancellationToken);
 				}
@@ -528,12 +528,12 @@ namespace EpicGames.Horde.Storage.Bundles
 
 				List<QueuedPacket> packets = bundle.QueuedPackets;
 				for (int idx = 0; idx < packets.Count && packets[idx].PacketIdx <= packetIdx; idx++)
-			{
-					if (packets[idx].PacketIdx == packetIdx)
 				{
+					if (packets[idx].PacketIdx == packetIdx)
+					{
 						packets[idx].CompletionSource.TrySetResult(data);
 						packets.RemoveAt(idx--);
-				}
+					}
 				}
 			}
 		}
