Horde: Automatically retry HTTP requests that return transient errors by default.

diff --git a/Engine/Source/Programs/Horde/Horde/Commands/StorageCommandBase.cs b/Engine/Source/Programs/Horde/Horde/Commands/StorageCommandBase.cs
--- a/Engine/Source/Programs/Horde/Horde/Commands/StorageCommandBase.cs
+++ b/Engine/Source/Programs/Horde/Horde/Commands/StorageCommandBase.cs
@@ -68,15 +68,7 @@ namespace Horde.Commands
 			}
 
 			server = new Uri(server, Path);
-			return new HttpStorageClient(() => CreateDefaultHttpClient(server, token), () => new HttpClient(), StorageCache, logger);
-		}
-
-		static HttpClient CreateDefaultHttpClient(Uri server, string token)
-		{
-			HttpClient client = new HttpClient();
-			client.BaseAddress = server;
-			client.DefaultRequestHeaders.Authorization = new AuthenticationHeaderValue("Bearer", token);
-			return client;
+			return new HttpStorageClient(server, token, StorageCache, logger);
 		}
 	}
 }
diff --git a/Engine/Source/Programs/Shared/EpicGames.Horde/EpicGames.Horde.csproj b/Engine/Source/Programs/Shared/EpicGames.Horde/EpicGames.Horde.csproj
--- a/Engine/Source/Programs/Shared/EpicGames.Horde/EpicGames.Horde.csproj
+++ b/Engine/Source/Programs/Shared/EpicGames.Horde/EpicGames.Horde.csproj
@@ -43,6 +43,7 @@
     <PackageReference Include="Microsoft.Extensions.Caching.Memory" Version="6.0.1" />
     <PackageReference Include="Microsoft.Extensions.DependencyInjection" Version="6.0.0" />
     <PackageReference Include="Microsoft.Extensions.Http" Version="6.0.0" />
+    <PackageReference Include="Microsoft.Extensions.Http.Polly" Version="6.0.16" />
     <PackageReference Include="Microsoft.Extensions.Logging" Version="6.0.0" />
   </ItemGroup>
   <ItemGroup>
diff --git a/Engine/Source/Programs/Shared/EpicGames.Horde/Storage/Backends/HttpStorageClient.cs b/Engine/Source/Programs/Shared/EpicGames.Horde/Storage/Backends/HttpStorageClient.cs
--- a/Engine/Source/Programs/Shared/EpicGames.Horde/Storage/Backends/HttpStorageClient.cs
+++ b/Engine/Source/Programs/Shared/EpicGames.Horde/Storage/Backends/HttpStorageClient.cs
@@ -15,6 +15,10 @@ using EpicGames.Core;
 using Microsoft.Extensions.Caching.Memory;
 using Microsoft.Extensions.Logging;
 using EpicGames.Horde.Storage.Bundles;
+using Microsoft.Extensions.Http;
+using Polly.Extensions.Http;
+using Polly;
+using Polly.Retry;
 
 namespace EpicGames.Horde.Storage.Backends
 {
@@ -59,6 +63,16 @@ namespace EpicGames.Horde.Storage.Backends
 		readonly ILogger _logger;
 		bool _supportsUploadRedirects = true;
 
+		static readonly HttpMessageHandler s_defaultHttpMessageHandler = CreateDefaultHttpMessageHandler();
+
+		/// <summary>
+		/// Constructor
+		/// </summary>
+		public HttpStorageClient(Uri baseAddress, string? bearerToken, StorageCache cache, ILogger logger)
+			: this(() => CreateAuthenticatedClient(null, baseAddress, bearerToken), () => CreateClient(null), cache, logger)
+		{
+		}
+
 		/// <summary>
 		/// Constructor
 		/// </summary>
@@ -78,18 +92,35 @@ namespace EpicGames.Horde.Storage.Backends
 		{
 		}
 
+		static HttpMessageHandler CreateDefaultHttpMessageHandler()
+		{
+			AsyncRetryPolicy<HttpResponseMessage> retryPolicy = HttpPolicyExtensions
+				.HandleTransientHttpError()
+				.WaitAndRetryAsync(new[] { TimeSpan.FromSeconds(2.0), TimeSpan.FromSeconds(5.0), TimeSpan.FromSeconds(10), TimeSpan.FromSeconds(30) });
+
+			SocketsHttpHandler socketsHandler = new SocketsHttpHandler { PooledConnectionLifetime = TimeSpan.FromMinutes(15) };
+			return new PolicyHttpMessageHandler(retryPolicy) { InnerHandler = socketsHandler };
+		}
+
 		/// <summary>
 		/// Helper method to create an HTTP client from the given factory
 		/// </summary>
-		static HttpClient CreateClient(IHttpClientFactory httpClientFactory)
+		static HttpClient CreateClient(IHttpClientFactory? httpClientFactory)
 		{
-			return httpClientFactory.CreateClient(HttpClientName);
+			if (httpClientFactory == null)
+			{
+				return new HttpClient(s_defaultHttpMessageHandler, disposeHandler: false);
+			}
+			else
+			{
+				return httpClientFactory.CreateClient(HttpClientName);
+			}
 		}
 
 		/// <summary>
 		/// Helper method to add the base address and auth header to an HTTP client
 		/// </summary>
-		static HttpClient CreateAuthenticatedClient(IHttpClientFactory httpClientFactory, Uri baseAddress, string? bearerToken)
+		static HttpClient CreateAuthenticatedClient(IHttpClientFactory? httpClientFactory, Uri baseAddress, string? bearerToken)
 		{
 			HttpClient httpClient = CreateClient(httpClientFactory);
 			httpClient.BaseAddress = baseAddress;
