Horde: Fix NuGet package versioning errors.

diff --git a/Engine/Source/Programs/Shared/EpicGames.Horde/EpicGames.Horde.csproj b/Engine/Source/Programs/Shared/EpicGames.Horde/EpicGames.Horde.csproj
--- a/Engine/Source/Programs/Shared/EpicGames.Horde/EpicGames.Horde.csproj
+++ b/Engine/Source/Programs/Shared/EpicGames.Horde/EpicGames.Horde.csproj
@@ -43,7 +43,7 @@
     <PackageReference Include="Microsoft.Extensions.Caching.Memory" Version="6.0.1" />
     <PackageReference Include="Microsoft.Extensions.DependencyInjection" Version="6.0.0" />
     <PackageReference Include="Microsoft.Extensions.Http" Version="6.0.0" />
-    <PackageReference Include="Microsoft.Extensions.Http.Polly" Version="6.0.16" />
+    <PackageReference Include="Microsoft.Extensions.Http.Polly" Version="6.0.10" />
     <PackageReference Include="Microsoft.Extensions.Logging" Version="6.0.0" />
   </ItemGroup>
   <ItemGroup>
