Horde: Add a utility method for reading a ref and throwing an exception on failure.

diff --git a/Engine/Source/Programs/Shared/EpicGames.Core/MemoryWriter.cs b/Engine/Source/Programs/Shared/EpicGames.Core/MemoryWriter.cs
--- a/Engine/Source/Programs/Shared/EpicGames.Core/MemoryWriter.cs
+++ b/Engine/Source/Programs/Shared/EpicGames.Core/MemoryWriter.cs
@@ -4,7 +4,10 @@ using System;
 using System.Buffers;
 using System.Buffers.Binary;
 using System.Collections.Generic;
+using System.IO;
 using System.Text;
+using System.Threading;
+using System.Threading.Tasks;
 
 namespace EpicGames.Core
 {
@@ -423,6 +426,18 @@ namespace EpicGames.Core
 			}
 		}
 
+		/// <summary>
+		/// Copies the data to a stream
+		/// </summary>
+		/// <param name="stream">Stream to write to</param>
+		public async Task CopyToAsync(Stream stream, CancellationToken cancellationToken = default)
+		{
+			foreach (Chunk chunk in _chunks)
+			{
+				await stream.WriteAsync(chunk.WrittenMemory, cancellationToken);
+			}
+		}
+
 		/// <summary>
 		/// Create a byte array from the sequence
 		/// </summary>
diff --git a/Engine/Source/Programs/Shared/EpicGames.Horde/Storage/NodeRef.cs b/Engine/Source/Programs/Shared/EpicGames.Horde/Storage/NodeRef.cs
--- a/Engine/Source/Programs/Shared/EpicGames.Horde/Storage/NodeRef.cs
+++ b/Engine/Source/Programs/Shared/EpicGames.Horde/Storage/NodeRef.cs
@@ -201,6 +201,23 @@ namespace EpicGames.Horde.Storage
 			return new NodeRef<TNode>(handle);
 		}
 
+		/// <summary>
+		/// Reads and deserializes a node for the given ref
+		/// </summary>
+		/// <param name="storageClient"></param>
+		/// <param name="refName">Name of the ref to write</param>
+		/// <param name="cacheTime"></param>
+		/// <param name="cancellationToken"></param>
+		public static async ValueTask<NodeRef<TNode>> ReadRefTargetAsync<TNode>(this IStorageClient storageClient, RefName refName, DateTime cacheTime = default, CancellationToken cancellationToken = default) where TNode : Node
+		{
+			NodeRef<TNode>? result = await TryReadRefTargetAsync<TNode>(storageClient, refName, cacheTime, cancellationToken);
+			if (result == null)
+			{
+				throw new RefNameNotFoundException(refName);
+			}
+			return result;
+		}
+
 		/// <summary>
 		/// Reads and deserializes a node for the given ref
 		/// </summary>
