Fix documentation warning.

diff --git a/Engine/Source/Programs/Shared/EpicGames.Core/MemoryWriter.cs b/Engine/Source/Programs/Shared/EpicGames.Core/MemoryWriter.cs
--- a/Engine/Source/Programs/Shared/EpicGames.Core/MemoryWriter.cs
+++ b/Engine/Source/Programs/Shared/EpicGames.Core/MemoryWriter.cs
@@ -430,6 +430,7 @@ namespace EpicGames.Core
 		/// Copies the data to a stream
 		/// </summary>
 		/// <param name="stream">Stream to write to</param>
+		/// <param name="cancellationToken">Cancellation token for the operation</param>
 		public async Task CopyToAsync(Stream stream, CancellationToken cancellationToken = default)
 		{
 			foreach (Chunk chunk in _chunks)
