EpicGames.Core: Ensure that managed processes are terminated if they are still in a suspended state.

diff --git a/Engine/Source/Programs/Shared/EpicGames.Core/ManagedProcess.cs b/Engine/Source/Programs/Shared/EpicGames.Core/ManagedProcess.cs
--- a/Engine/Source/Programs/Shared/EpicGames.Core/ManagedProcess.cs
+++ b/Engine/Source/Programs/Shared/EpicGames.Core/ManagedProcess.cs
@@ -824,7 +824,8 @@ namespace EpicGames.Core
 				{
 					if (processInfo.hProcess != IntPtr.Zero && _processHandle == null)
 					{
-						_ = CloseHandle(processInfo.hProcess);
+						using SafeFileHandle processHandle = new(processInfo.hProcess, true);
+						_ = TerminateProcess(processHandle, UInt32.MaxValue);
 					}
 					if (processInfo.hThread != IntPtr.Zero)
 					{
