Horde: Initial implementation of a generalized workspace file tracking system.

diff --git a/Engine/Source/Programs/Horde/Horde/Commands/Workspace/WorkspaceDrop.cs b/Engine/Source/Programs/Horde/Horde/Commands/Workspace/WorkspaceDrop.cs
new file mode 100644
--- /dev/null
+++ b/Engine/Source/Programs/Horde/Horde/Commands/Workspace/WorkspaceDrop.cs
@@ -0,0 +1,45 @@
+// Copyright Epic Games, Inc. All Rights Reserved.
+
+using System.Diagnostics;
+using EpicGames.Core;
+using EpicGames.Horde.Storage;
+using EpicGames.Horde.Storage.Backends;
+using EpicGames.Horde.Storage.Bundles;
+using EpicGames.Horde.Storage.Nodes;
+using Microsoft.Extensions.Logging;
+
+namespace Horde.Commands.Workspace
+{
+	using Workspace = EpicGames.Horde.Storage.Workspace;
+
+	[Command("workspace", "drop", "Drops a layer from a workspace")]
+	class WorkspaceDrop : Command
+	{
+		[CommandLine("-Root=")]
+		public DirectoryReference? RootDir { get; set; }
+
+		[CommandLine("-Layer=")]
+		public WorkspaceLayerId LayerId { get; set; } = WorkspaceLayerId.Default;
+
+		public override async Task<int> ExecuteAsync(ILogger logger)
+		{
+			CancellationToken cancellationToken = CancellationToken.None;
+			RootDir ??= DirectoryReference.GetCurrentDirectory();
+
+			Workspace? workspace = await Workspace.TryOpenAsync(RootDir, logger, cancellationToken);
+			if (workspace == null)
+			{
+				logger.LogError("No workspace has been initialized in {RootDir}. Use 'workspace init' to create a new workspace.", RootDir);
+				return 1;
+			}
+
+			await workspace.SyncAsync(LayerId, new DirectoryNode(), cancellationToken);
+			workspace.RemoveLayer(LayerId);
+			await workspace.SaveAsync(cancellationToken);
+
+			logger.LogInformation("Removed layer {LayerId}", LayerId);
+
+			return 0;
+		}
+	}
+}
diff --git a/Engine/Source/Programs/Horde/Horde/Commands/Workspace/WorkspaceInit.cs b/Engine/Source/Programs/Horde/Horde/Commands/Workspace/WorkspaceInit.cs
new file mode 100644
--- /dev/null
+++ b/Engine/Source/Programs/Horde/Horde/Commands/Workspace/WorkspaceInit.cs
@@ -0,0 +1,26 @@
+// Copyright Epic Games, Inc. All Rights Reserved.
+
+using EpicGames.Core;
+using Microsoft.Extensions.Logging;
+
+namespace Horde.Commands.Workspace
+{
+	using Workspace = EpicGames.Horde.Storage.Workspace;
+
+	[Command("workspace", "init", "Initialize a managed workspace in the current folder")]
+	class WorkspaceInit : Command
+	{
+		[CommandLine("-Root=")]
+		public DirectoryReference? RootDir { get; set; }
+
+		public override async Task<int> ExecuteAsync(ILogger logger)
+		{
+			RootDir ??= DirectoryReference.GetCurrentDirectory();
+
+			Workspace workspace = await Workspace.CreateAsync(RootDir, logger);
+			logger.LogInformation("Created workspace in {RootDir}", workspace.RootDir);
+
+			return 0;
+		}
+	}
+}
diff --git a/Engine/Source/Programs/Horde/Horde/Commands/Workspace/WorkspaceSync.cs b/Engine/Source/Programs/Horde/Horde/Commands/Workspace/WorkspaceSync.cs
new file mode 100644
--- /dev/null
+++ b/Engine/Source/Programs/Horde/Horde/Commands/Workspace/WorkspaceSync.cs
@@ -0,0 +1,99 @@
+// Copyright Epic Games, Inc. All Rights Reserved.
+
+using System.Diagnostics;
+using EpicGames.Core;
+using EpicGames.Horde.Storage;
+using EpicGames.Horde.Storage.Backends;
+using EpicGames.Horde.Storage.Bundles;
+using EpicGames.Horde.Storage.Nodes;
+using Microsoft.Extensions.Logging;
+
+namespace Horde.Commands.Workspace
+{
+	using Workspace = EpicGames.Horde.Storage.Workspace;
+
+	[Command("workspace", "sync", "Extracts an archive into the workspace")]
+	class WorkspaceSync : StorageCommandBase
+	{
+		[CommandLine("-Root=")]
+		public DirectoryReference? RootDir { get; set; }
+
+		[CommandLine("-File=")]
+		public FileReference? File { get; set; }
+
+		[CommandLine("-Ref=")]
+		public string? Ref { get; set; }
+
+		[CommandLine("-Node=")]
+		public string? Node { get; set; }
+
+		[CommandLine("-Layer=")]
+		public WorkspaceLayerId LayerId { get; set; } = WorkspaceLayerId.Default;
+
+		[CommandLine("-Stats")]
+		public bool Stats { get; set; }
+
+		public WorkspaceSync(StorageCache storageCache)
+			: base(storageCache)
+		{
+		}
+
+		public override async Task<int> ExecuteAsync(ILogger logger)
+		{
+			CancellationToken cancellationToken = CancellationToken.None;
+			RootDir ??= DirectoryReference.GetCurrentDirectory();
+
+			Workspace? workspace = await Workspace.TryOpenAsync(RootDir, logger, cancellationToken);
+			if (workspace == null)
+			{
+				logger.LogError("No workspace has been initialized in {RootDir}. Use 'workspace init' to create a new workspace.", RootDir);
+				return 1;
+			}
+
+			IStorageClient store;
+
+			BlobHandle handle;
+			if (File != null)
+			{
+				store = new FileStorageClient(File.Directory, StorageCache, logger);
+				handle = await ((FileStorageClient)store).ReadRefAsync(File);
+			}
+			else if (Ref != null)
+			{
+				store = await CreateStorageClientAsync(logger);
+				handle = await store.ReadRefTargetAsync(new RefName(Ref));
+			}
+			else if (Node != null)
+			{
+				store = await CreateStorageClientAsync(logger);
+				handle = ((BundleStorageClient)store).CreateNodeHandle(BundleNodeLocator.Parse(Node));
+			}
+			else
+			{
+				throw new CommandLineArgumentException("Either -File=... or -Ref=... must be specified");
+			}
+
+			Stopwatch timer = Stopwatch.StartNew();
+			logger.LogInformation("Syncing into layer '{LayerId}'...", LayerId);
+
+			DirectoryNode node = await handle.ReadNodeAsync<DirectoryNode>();
+			await node.CopyToDirectoryAsync(workspace.RootDir.ToDirectoryInfo(), logger, CancellationToken.None);
+			await workspace.SaveAsync(cancellationToken);
+
+			logger.LogInformation("Elapsed: {Time}s", timer.Elapsed.TotalSeconds);
+
+			if (Stats)
+			{
+				BundleStorageClient? bundleStorageClient = store as BundleStorageClient;
+				if (bundleStorageClient != null)
+				{
+					logger.LogInformation("Num bytes read: {NumBytes:n0}", bundleStorageClient.BundleReader.NumBytesRead);
+					logger.LogInformation("Num header reads: {NumReads:n0}", bundleStorageClient.BundleReader.NumHeaderReads);
+					logger.LogInformation("Num packet reads: {NumReads:n0}", bundleStorageClient.BundleReader.NumPacketReads);
+				}
+			}
+
+			return 0;
+		}
+	}
+}
diff --git a/Engine/Source/Programs/Shared/EpicGames.Core/FileTransaction.cs b/Engine/Source/Programs/Shared/EpicGames.Core/FileTransaction.cs
--- a/Engine/Source/Programs/Shared/EpicGames.Core/FileTransaction.cs
+++ b/Engine/Source/Programs/Shared/EpicGames.Core/FileTransaction.cs
@@ -78,7 +78,7 @@ namespace EpicGames.Core
 		public static FileTransactionStream OpenWrite(FileReference file)
 		{
 			FileReference incomingFile = GetIncomingFile(file);
-			Stream stream = FileReference.Open(incomingFile, FileMode.CreateNew);
+			Stream stream = FileReference.Open(incomingFile, FileMode.Create, FileAccess.ReadWrite, FileShare.Read);
 			return new FileTransactionStream(stream, incomingFile, file);
 		}
 
@@ -145,10 +145,11 @@ namespace EpicGames.Core
 				throw new InvalidOperationException("Stream cannot be written to");
 			}
 
+			FileReference finalFile = _finalFile;
 			Close();
 
-			FileReference.Delete(_finalFile);
-			FileReference.Move(_file, _finalFile);
+			FileReference.Delete(finalFile);
+			FileReference.Move(_file, finalFile);
 
 			_finalFile = null;
 		}
diff --git a/Engine/Source/Programs/Shared/EpicGames.Horde/Storage/Workspace.cs b/Engine/Source/Programs/Shared/EpicGames.Horde/Storage/Workspace.cs
new file mode 100644
--- /dev/null
+++ b/Engine/Source/Programs/Shared/EpicGames.Horde/Storage/Workspace.cs
@@ -0,0 +1,407 @@
+// Copyright Epic Games, Inc. All Rights Reserved.
+
+using System;
+using System.Collections.Generic;
+using System.Diagnostics;
+using System.IO;
+using System.Linq;
+using System.Threading;
+using System.Threading.Tasks;
+using EpicGames.Core;
+using EpicGames.Horde.Storage.Nodes;
+using Microsoft.Extensions.Logging;
+
+namespace EpicGames.Horde.Storage
+{
+	/// <summary>
+	/// Provides functionality to extract and patch data in a local workspace
+	/// </summary>
+	public class Workspace
+	{
+		class DirectoryState
+		{
+			public Dictionary<Utf8String, DirectoryState> Directories { get; }
+			public Dictionary<Utf8String, FileState> Files { get; }
+
+			public ulong LayerFlags { get; set; }
+
+			public DirectoryState()
+			{
+				Directories = new Dictionary<Utf8String, DirectoryState>(Utf8StringComparer.Ordinal);
+				Files = new Dictionary<Utf8String, FileState>(Utf8StringComparer.Ordinal);
+			}
+
+			public DirectoryState(IMemoryReader reader)
+			{
+				int numDirectories = reader.ReadInt32();
+				Directories = new Dictionary<Utf8String, DirectoryState>(numDirectories, Utf8StringComparer.Ordinal);
+
+				for (int idx = 0; idx < numDirectories; idx++)
+				{
+					Utf8String name = reader.ReadUtf8String();
+					DirectoryState directory = new DirectoryState(reader);
+					Directories[name] = directory;
+				}
+
+				int numFiles = reader.ReadInt32();
+				Files = new Dictionary<Utf8String, FileState>(numFiles, Utf8StringComparer.Ordinal);
+
+				for (int idx = 0; idx < numFiles; idx++)
+				{
+					Utf8String name = reader.ReadUtf8String();
+					FileState file = new FileState(reader);
+					Files[name] = file;
+				}
+
+				LayerFlags = reader.ReadUnsignedVarInt();
+			}
+
+			public void Write(IMemoryWriter writer)
+			{
+				writer.WriteInt32(Directories.Count);
+				foreach (DirectoryState directory in Directories.Values)
+				{
+					directory.Write(writer);
+				}
+
+				writer.WriteInt32(Files.Count);
+				foreach (FileState file in Files.Values)
+				{
+					file.Write(writer);
+				}
+
+				writer.WriteUnsignedVarInt(LayerFlags);
+			}
+		}
+
+		class FileState
+		{
+			public long Length { get; private set; }
+			public long LastModifiedTimeUtc { get; private set; }
+			public IoHash Hash { get; }
+
+			public ulong LayerFlags { get; set; }
+			public List<ChunkInfo> Chunks { get; } = new List<ChunkInfo>();
+
+			public FileState(FileInfo fileInfo, IoHash hash)
+			{
+				Update(fileInfo);
+				Hash = hash;
+			}
+
+			public FileState(IMemoryReader reader)
+			{
+				Length = reader.ReadInt64();
+				LastModifiedTimeUtc = reader.ReadInt64();
+				Hash = reader.ReadIoHash();
+
+				LayerFlags = reader.ReadUnsignedVarInt();
+				Chunks = reader.ReadList(() => new ChunkInfo(reader));
+			}
+
+			public bool Modified(FileInfo fileInfo) => Length != fileInfo.Length || LastModifiedTimeUtc != fileInfo.LastWriteTimeUtc.Ticks;
+
+			public void Update(FileInfo fileInfo)
+			{
+				Length = fileInfo.Length;
+				LastModifiedTimeUtc = fileInfo.LastWriteTimeUtc.Ticks;
+			}
+
+			public void Write(IMemoryWriter writer)
+			{
+				Debug.Assert(Hash != IoHash.Zero);
+				writer.WriteInt64(Length);
+				writer.WriteInt64(LastModifiedTimeUtc);
+				writer.WriteIoHash(Hash);
+				writer.WriteUnsignedVarInt(LayerFlags);
+				writer.WriteList(Chunks, x => x.Write(writer));
+			}
+		}
+
+		class ChunkInfo
+		{
+			public IoHash Hash { get; }
+			public long Offset { get; }
+			public long Length { get; }
+
+			public ChunkInfo(IoHash hash, long offset, long length)
+			{
+				Hash = hash;
+				Offset = offset;
+				Length = length;
+			}
+
+			public ChunkInfo(IMemoryReader reader)
+			{
+				Hash = reader.ReadIoHash();
+				Offset = (long)reader.ReadUnsignedVarInt();
+				Length = (long)reader.ReadUnsignedVarInt();
+			}
+
+			public void Write(IMemoryWriter writer)
+			{
+				writer.WriteIoHash(Hash);
+				writer.WriteUnsignedVarInt((ulong)Offset);
+				writer.WriteUnsignedVarInt((ulong)Length);
+			}
+		}
+
+		class LayerState
+		{
+			public WorkspaceLayerId Id { get; }
+			public ulong Flag { get; }
+
+			public LayerState(WorkspaceLayerId id, ulong flag)
+			{
+				Id = id;
+				Flag = flag;
+			}
+
+			public LayerState(IMemoryReader reader)
+			{
+				Id = new WorkspaceLayerId(new StringId(reader.ReadUtf8String()));
+				Flag = reader.ReadUnsignedVarInt();
+			}
+
+			public void Write(IMemoryWriter writer)
+			{
+				writer.WriteUtf8String(Id.Id.Text);
+				writer.WriteUnsignedVarInt(Flag);
+			}
+		}
+
+		class WorkspaceState
+		{
+			public DirectoryState Root { get; set; }
+			public List<LayerState> Layers { get; private set; }
+
+			public WorkspaceState()
+			{
+				Root = new DirectoryState();
+				Layers = new List<LayerState>();
+			}
+
+			public WorkspaceState(IMemoryReader reader)
+			{
+				Root = new DirectoryState(reader);
+				Layers = reader.ReadList(() => new LayerState(reader));
+			}
+
+			public void Write(IMemoryWriter writer)
+			{
+				Root.Write(writer);
+				writer.WriteList(Layers, x => x.Write(writer));
+			}
+		}
+
+		readonly DirectoryReference _rootDir;
+		readonly FileReference _stateFile;
+		readonly WorkspaceState _state;
+		readonly ILogger _logger;
+
+		const string HordeDirName = ".horde";
+		const string StateFileName = "contents.dat";
+
+		/// <summary>
+		/// Root directory for the workspace
+		/// </summary>
+		public DirectoryReference RootDir => _rootDir;
+
+		/// <summary>
+		/// Layers current in this workspace
+		/// </summary>
+		public IReadOnlyList<WorkspaceLayerId> Layers => _state.Layers.Select(x => x.Id).ToList();
+
+		/// <summary>
+		/// Constructor
+		/// </summary>
+		/// <param name="rootDir">Directory for the workspace</param>
+		/// <param name="stateFile">Path to the state file for this directory</param>
+		/// <param name="state">Current state</param>
+		/// <param name="logger">Logger for diagnostic output</param>
+		private Workspace(DirectoryReference rootDir, FileReference stateFile, WorkspaceState state, ILogger logger)
+		{
+			_rootDir = rootDir;
+			_stateFile = stateFile;
+			_state = state;
+			_logger = logger;
+		}
+
+		/// <summary>
+		/// Create a new workspace instance in the given location. Opens the existing instance if it already contains workspace data.
+		/// </summary>
+		/// <param name="rootDir">Root directory for the workspace</param>
+		/// <param name="logger">Logger for output</param>
+		/// <param name="cancellationToken">Cancellation token for the operation</param>
+		/// <returns>Workspace instance</returns>
+		public static async Task<Workspace> CreateAsync(DirectoryReference rootDir, ILogger logger, CancellationToken cancellationToken = default)
+		{
+			FileReference stateFile = FileReference.Combine(rootDir, HordeDirName, StateFileName);
+
+			using FileStream? stream = FileTransaction.OpenRead(stateFile);
+			if (stream != null)
+			{
+				throw new InvalidOperationException($"Workspace already exists in {rootDir}; use Open instead.");
+			}
+
+			Workspace workspace = new Workspace(rootDir, stateFile, new WorkspaceState(), logger);
+			await workspace.SaveAsync(cancellationToken);
+
+			return workspace;
+		}
+
+		/// <summary>
+		/// Attempts to open an existing workspace for the current directory. 
+		/// </summary>
+		/// <param name="currentDir">Root directory for the workspace</param>
+		/// <param name="logger">Logger for output</param>
+		/// <param name="cancellationToken">Cancellation token for the operation</param>
+		/// <returns>Workspace instance</returns>
+		public static async Task<Workspace?> TryOpenAsync(DirectoryReference currentDir, ILogger logger, CancellationToken cancellationToken = default)
+		{
+			for (DirectoryReference? testDir = currentDir; testDir != null; testDir = testDir.ParentDirectory)
+			{
+				FileReference stateFile = FileReference.Combine(testDir, HordeDirName, StateFileName);
+				using (FileStream? stream = FileTransaction.OpenRead(stateFile))
+				{
+					if (stream != null)
+					{
+						byte[] data = await stream.ReadAllBytesAsync(cancellationToken);
+						WorkspaceState state = new WorkspaceState(new MemoryReader(data));
+						return new Workspace(testDir, stateFile, state, logger);
+					}
+				}
+			}
+			return null;
+		}
+
+		/// <summary>
+		/// Save the current state of the workspace
+		/// </summary>
+		/// <param name="cancellationToken">Cancellation token for the operation</param>
+		public async Task SaveAsync(CancellationToken cancellationToken)
+		{
+			DirectoryReference.CreateDirectory(_stateFile.Directory);
+
+			using (FileTransactionStream stream = FileTransaction.OpenWrite(_stateFile))
+			{
+				using (ChunkedMemoryWriter writer = new ChunkedMemoryWriter(64 * 1024))
+				{
+					_state.Write(writer);
+					await writer.CopyToAsync(stream, cancellationToken);
+				}
+				stream.CompleteTransaction();
+			}
+		}
+
+		#region Layers
+
+		/// <summary>
+		/// Add or update a layer with the given identifier
+		/// </summary>
+		/// <param name="id">Identifier for the layer</param>
+		public void AddLayer(WorkspaceLayerId id)
+		{
+			if (_state.Layers.Any(x => x.Id == id))
+			{
+				throw new InvalidOperationException($"Layer {id} already exists");
+			}
+
+			ulong flags = 0;
+			for (int idx = 0; idx < _state.Layers.Count; idx++)
+			{
+				flags |= _state.Layers[idx].Flag;
+			}
+			if (flags == ~0UL)
+			{
+				throw new InvalidOperationException("Maximum number of layers reached");
+			}
+
+			ulong nextFlag = (flags + 1) ^ flags;
+			_state.Layers.Add(new LayerState(id, nextFlag));
+		}
+
+		/// <summary>
+		/// Removes a layer with the given identifier. Does not remove any files in the workspace.
+		/// </summary>
+		/// <param name="layerId">Layer to update</param>
+		public void RemoveLayer(WorkspaceLayerId layerId)
+		{
+			int layerIdx = _state.Layers.FindIndex(x => x.Id == layerId);
+			if (layerIdx != -1)
+			{
+				LayerState layer = _state.Layers[layerIdx];
+				if ((_state.Root.LayerFlags & layer.Flag) != 0)
+				{
+					throw new InvalidOperationException($"Workspace still contains files for layer {layerId}");
+				}
+				_state.Layers.RemoveAt(layerIdx);
+			}
+		}
+
+		LayerState? GetLayerState(WorkspaceLayerId layerId) => _state.Layers.FirstOrDefault(x => x.Id == layerId);
+
+		#endregion
+
+		/// <summary>
+		/// Syncs a layer to the given contents
+		/// </summary>
+		/// <param name="layerId">Identifier for the layer</param>
+		/// <param name="contents">New contents for the layer</param>
+		/// <param name="cancellationToken">Cancellation token for the operation</param>
+		public async Task SyncAsync(WorkspaceLayerId layerId, DirectoryNode contents, CancellationToken cancellationToken = default)
+		{
+			LayerState? layerState = GetLayerState(layerId);
+			if (layerState == null)
+			{
+				throw new ArgumentException($"Layer {layerId} does not exist", nameof(layerId));
+			}
+
+			_state.Root = await SyncDirectoryAsync(contents, _rootDir, _state.Root, layerState.Flag, _logger, cancellationToken);
+		}
+
+		static async Task<DirectoryState> SyncDirectoryAsync(DirectoryNode directoryNode, DirectoryReference dirPath, DirectoryState? directoryState, ulong flag, ILogger logger, CancellationToken cancellationToken)
+		{
+			DirectoryReference.CreateDirectory(dirPath);
+
+			DirectoryState newState = new DirectoryState();
+
+			foreach ((Utf8String name, DirectoryEntry? subDirEntry, DirectoryState? subDirState) in directoryNode.NameToDirectory.Zip(directoryState?.Directories))
+			{
+				DirectoryReference subDirPath = DirectoryReference.Combine(dirPath, name.ToString());
+				if (subDirEntry != null)
+				{
+					DirectoryNode subDirNode = await subDirEntry.ExpandAsync(cancellationToken);
+					newState.Directories[name] = await SyncDirectoryAsync(subDirNode, subDirPath, subDirState, flag, logger, cancellationToken);
+				}
+			}
+
+			foreach ((Utf8String name, FileEntry? fileEntry, FileState? fileState) in directoryNode.NameToFile.Zip(directoryState?.Files))
+			{
+				FileReference filePath = FileReference.Combine(dirPath, name.ToString());
+				if (fileEntry != null)
+				{
+					if (fileState == null || fileState.Hash != fileEntry.Hash)
+					{
+						newState.Files[name] = await CheckoutFileAsync(fileEntry, filePath.ToFileInfo(), logger, cancellationToken);
+					}
+					else
+					{
+						newState.Files[name] = fileState;
+					}
+				}
+			}
+
+			return newState;
+		}
+
+		static async Task<FileState> CheckoutFileAsync(FileEntry fileRef, FileInfo fileInfo, ILogger logger, CancellationToken cancellationToken)
+		{
+			logger.LogInformation("Updating {File} to {Hash}", fileInfo, fileRef.Hash);
+			ChunkedDataNode fileNode = await fileRef.ExpandAsync(cancellationToken);
+			await fileNode.CopyToFileAsync(fileInfo, cancellationToken);
+			fileInfo.Refresh();
+			return new FileState(fileInfo, fileRef.Hash);
+		}
+	}
+}
diff --git a/Engine/Source/Programs/Shared/EpicGames.Horde/Storage/WorkspaceLayerId.cs b/Engine/Source/Programs/Shared/EpicGames.Horde/Storage/WorkspaceLayerId.cs
new file mode 100644
--- /dev/null
+++ b/Engine/Source/Programs/Shared/EpicGames.Horde/Storage/WorkspaceLayerId.cs
@@ -0,0 +1,35 @@
+// Copyright Epic Games, Inc. All Rights Reserved.
+
+using System.ComponentModel;
+
+namespace EpicGames.Horde.Storage
+{
+	/// <summary>
+	/// Identifier for a workspace layer
+	/// </summary>
+	/// <param name="Id">Identifier for the layer</param>
+	[TypeConverter(typeof(StringIdTypeConverter<WorkspaceLayerId, WorkspaceLayerIdConverter>))]
+	[BinaryIdConverter(typeof(WorkspaceLayerIdConverter))]
+	public record struct WorkspaceLayerId(StringId Id)
+	{
+		/// <summary>
+		/// Name of the default layer
+		/// </summary>
+		public static WorkspaceLayerId Default { get; } = new WorkspaceLayerId(new StringId("default"));
+
+		/// <inheritdoc/>
+		public override string ToString() => Id.ToString();
+	}
+
+	/// <summary>
+	/// Converter class to and from ObjectId values
+	/// </summary>
+	class WorkspaceLayerIdConverter : StringIdConverter<WorkspaceLayerId>
+	{
+		/// <inheritdoc/>
+		public override WorkspaceLayerId FromStringId(StringId id) => new WorkspaceLayerId(id);
+
+		/// <inheritdoc/>
+		public override StringId ToStringId(WorkspaceLayerId value) => value.Id;
+	}
+}
