Horde: Automatically retry Horde requests resulting in HTTP errors, and add methods for other HTTP verbs to HordeHttpClient.

diff --git a/Engine/Source/Programs/Horde/Horde/Settings.cs b/Engine/Source/Programs/Horde/Horde/Settings.cs
--- a/Engine/Source/Programs/Horde/Horde/Settings.cs
+++ b/Engine/Source/Programs/Horde/Horde/Settings.cs
@@ -158,7 +158,7 @@ namespace Horde
 				throw new Exception("Unable to log in to server.");
 			}
 
-			return new HordeHttpClient(server, token);
+			return new HordeHttpClient(server, token, logger);
 		}
 	}
 }
diff --git a/Engine/Source/Programs/Shared/EpicGames.Horde/HordeHttpClient.cs b/Engine/Source/Programs/Shared/EpicGames.Horde/HordeHttpClient.cs
--- a/Engine/Source/Programs/Shared/EpicGames.Horde/HordeHttpClient.cs
+++ b/Engine/Source/Programs/Shared/EpicGames.Horde/HordeHttpClient.cs
@@ -7,6 +7,14 @@ using System.Text.Json;
 using System.Text.Json.Serialization;
 using System.Threading;
 using System.Threading.Tasks;
+using EpicGames.Core;
+using EpicGames.Horde.Storage;
+using EpicGames.Horde.Storage.Backends;
+using Microsoft.Extensions.Http;
+using Microsoft.Extensions.Logging;
+using Polly;
+using Polly.Extensions.Http;
+using Polly.Retry;
 
 #pragma warning disable CA2234 // Pass system uri objects instead of strings
 
@@ -17,39 +25,68 @@ namespace EpicGames.Horde
 	/// </summary>
 	public sealed class HordeHttpClient : IDisposable
 	{
-		/// <summary>
-		/// The configured HTTP client
-		/// </summary>
-		public HttpClient HttpClient { get; }
-
+		readonly HttpMessageHandler _httpMessageHandler;
+		readonly Uri _baseUri;
+		readonly AuthenticationHeaderValue? _authHeader;
 		readonly JsonSerializerOptions _jsonSerializerOptions;
+		readonly ILogger _logger;
 
 		/// <summary>
 		/// Constructor
 		/// </summary>
-		/// <param name="httpClient">Client to use for communication</param>
-		public HordeHttpClient(HttpClient httpClient)
+		/// <param name="uri">Base URI of the server</param>
+		/// <param name="token">Access token for the connection</param>
+		/// <param name="logger">Logger for output</param>
+		public HordeHttpClient(Uri uri, string token, ILogger logger)
 		{
-			HttpClient = httpClient;
+			_baseUri = uri;
+			_authHeader = new AuthenticationHeaderValue("Bearer", token);
+
+			_httpMessageHandler = CreateHttpMessageHandler();
 
 			_jsonSerializerOptions = new JsonSerializerOptions();
 			ConfigureJsonSerializer(_jsonSerializerOptions);
+
+			_logger = logger;
+		}
+
+		/// <inheritdoc/>
+		public void Dispose()
+		{
+			_httpMessageHandler.Dispose();
 		}
 
 		/// <summary>
-		/// Constructor
+		/// Creates an HTTP client using the underlying message handler
 		/// </summary>
-		/// <param name="uri">Base URI of the server</param>
-		/// <param name="token">Access token for the connection</param>
-		public HordeHttpClient(Uri uri, string token)
-			: this(new HttpClient { BaseAddress = uri, DefaultRequestHeaders = { Authorization = new AuthenticationHeaderValue("Bearer", token) } })
+		public HttpClient CreateHttpClient()
 		{
+			HttpClient httpClient = new HttpClient(_httpMessageHandler, false);
+			httpClient.BaseAddress = _baseUri;
+			httpClient.DefaultRequestHeaders.Authorization = _authHeader;
+			return httpClient;
 		}
 
-		/// <inheritdoc/>
-		public void Dispose()
+		/// <summary>
+		/// Creates an HTTP client using the underlying message handler
+		/// </summary>
+		public HttpClient CreateRedirectHttpClient()
 		{
-			HttpClient.Dispose();
+			return new HttpClient(_httpMessageHandler, false);
+		}
+
+		/// <summary>
+		/// Creates a HTTP message handler
+		/// </summary>
+		/// <returns></returns>
+		public static HttpMessageHandler CreateHttpMessageHandler()
+		{
+			AsyncRetryPolicy<HttpResponseMessage> retryPolicy = HttpPolicyExtensions
+				.HandleTransientHttpError()
+				.WaitAndRetryAsync(new[] { TimeSpan.FromSeconds(2.0), TimeSpan.FromSeconds(5.0), TimeSpan.FromSeconds(10), TimeSpan.FromSeconds(30) });
+
+			SocketsHttpHandler socketsHandler = new SocketsHttpHandler { PooledConnectionLifetime = TimeSpan.FromMinutes(15) };
+			return new PolicyHttpMessageHandler(retryPolicy) { InnerHandler = socketsHandler };
 		}
 
 		/// <summary>
@@ -68,6 +105,40 @@ namespace EpicGames.Horde
 			options.Converters.Add(new BinaryIdJsonConverterFactory());
 		}
 
+		/// <summary>
+		/// Creates a storage client for this Horde instance
+		/// </summary>
+		/// <param name="namespaceId">Namespace to manipulate</param>
+		/// <param name="storageCache">Cache for storage operations</param>
+		/// <returns>New storage client instance</returns>
+		public HttpStorageClient CreateStorageClient(NamespaceId namespaceId, StorageCache storageCache) => CreateStorageClient($"api/v1/storage/{namespaceId}/", storageCache);
+
+		/// <summary>
+		/// Creates a storage client for this Horde instance
+		/// </summary>
+		/// <param name="path">Base path for storage requests</param>
+		/// <param name="storageCache">Cache for storage operations</param>
+		/// <returns>New storage client instance</returns>
+		public HttpStorageClient CreateStorageClient(string path, StorageCache storageCache)
+		{
+			if (!path.EndsWith("/", StringComparison.Ordinal))
+			{
+				path += "/";
+			}
+
+			Uri baseUri = new Uri(_baseUri, path);
+
+			HttpClient CreateStorageHttpClient()
+			{
+				HttpClient httpClient = new HttpClient(_httpMessageHandler, false);
+				httpClient.BaseAddress = baseUri;
+				httpClient.DefaultRequestHeaders.Authorization = _authHeader;
+				return httpClient;
+			}
+
+			return new HttpStorageClient(CreateStorageHttpClient, CreateRedirectHttpClient, storageCache, _logger);
+		}
+
 		/// <summary>
 		/// Gets a resource from an HTTP endpoint and parses it as a JSON object
 		/// </summary>
@@ -77,14 +148,66 @@ namespace EpicGames.Horde
 		/// <returns>New instance of the object</returns>
 		public async Task<TResponse> GetAsync<TResponse>(string relativePath, CancellationToken cancellationToken = default)
 		{
-			using (HttpResponseMessage response = await HttpClient.GetAsync(relativePath, cancellationToken))
-			{
-				response.EnsureSuccessStatusCode();
+			using HttpClient httpClient = CreateHttpClient();
+			return await httpClient.GetAsync<TResponse>(relativePath, cancellationToken);
+		}
 
-				byte[] bytes = await response.Content.ReadAsByteArrayAsync(cancellationToken);
-				TResponse result = JsonSerializer.Deserialize<TResponse>(bytes, _jsonSerializerOptions)!;
-				return result;
-			}
+		/// <summary>
+		/// Posts an object to an HTTP endpoint as a JSON object, and parses the response object
+		/// </summary>
+		/// <typeparam name="TRequest">The object type to post</typeparam>
+		/// <param name="relativePath">The url to retrieve</param>
+		/// <param name="request">The object to post</param>
+		/// <param name="cancellationToken">Cancels the request</param>
+		/// <returns>The response parsed into the requested type</returns>
+		public async Task<HttpResponseMessage> PostAsync<TRequest>(string relativePath, TRequest request, CancellationToken cancellationToken = default)
+		{
+			using HttpClient httpClient = CreateHttpClient();
+			return await httpClient.PostAsync<TRequest>(relativePath, request, cancellationToken);
+		}
+
+		/// <summary>
+		/// Posts an object to an HTTP endpoint as a JSON object, and parses the response object
+		/// </summary>
+		/// <typeparam name="TResponse">The object type to return</typeparam>
+		/// <typeparam name="TRequest">The object type to post</typeparam>
+		/// <param name="relativePath">The url to retrieve</param>
+		/// <param name="request">The object to post</param>
+		/// <param name="cancellationToken">Cancels the request</param>
+		/// <returns>The response parsed into the requested type</returns>
+		public async Task<TResponse> PostAsync<TResponse, TRequest>(string relativePath, TRequest request, CancellationToken cancellationToken = default)
+		{
+			using HttpClient httpClient = CreateHttpClient();
+			return await httpClient.PostAsync<TResponse, TRequest>(relativePath, request, cancellationToken);
+		}
+
+		/// <summary>
+		/// Puts an object to an HTTP endpoint as a JSON object
+		/// </summary>
+		/// <typeparam name="TRequest">The object type to post</typeparam>
+		/// <param name="relativePath">The url to write to</param>
+		/// <param name="request">The object to post</param>
+		/// <param name="cancellationToken">Cancels the request</param>
+		/// <returns>Response message</returns>
+		public async Task<HttpResponseMessage> PutAsync<TRequest>(string relativePath, TRequest request, CancellationToken cancellationToken)
+		{
+			using HttpClient httpClient = CreateHttpClient();
+			return await httpClient.PutAsync<TRequest>(relativePath, request, cancellationToken);
+		}
+
+		/// <summary>
+		/// Puts an object to an HTTP endpoint as a JSON object
+		/// </summary>
+		/// <typeparam name="TResponse">The object type to return</typeparam>
+		/// <typeparam name="TRequest">The object type to post</typeparam>
+		/// <param name="relativePath">The url to write to</param>
+		/// <param name="request">The object to post</param>
+		/// <param name="cancellationToken">Cancels the request</param>
+		/// <returns>Response message</returns>
+		public async Task<TResponse> PutAsync<TResponse, TRequest>(string relativePath, TRequest request, CancellationToken cancellationToken)
+		{
+			using HttpClient httpClient = CreateHttpClient();
+			return await httpClient.PutAsync<TResponse, TRequest>(relativePath, request, cancellationToken);
 		}
 	}
 }
