Horde: Do not overwrite workspace files that have not changed.

diff --git a/Engine/Source/Programs/Shared/EpicGames.Core/FileTransaction.cs b/Engine/Source/Programs/Shared/EpicGames.Core/FileTransaction.cs
--- a/Engine/Source/Programs/Shared/EpicGames.Core/FileTransaction.cs
+++ b/Engine/Source/Programs/Shared/EpicGames.Core/FileTransaction.cs
@@ -43,7 +43,7 @@ namespace EpicGames.Core
 
 			try
 			{
-				return FileReference.Open(file, FileMode.Open);
+				return FileReference.Open(file, FileMode.Open, FileAccess.Read, FileShare.Read | FileShare.Delete);
 			}
 			catch (FileNotFoundException)
 			{
diff --git a/Engine/Source/Programs/Shared/EpicGames.Horde/Storage/Workspace.cs b/Engine/Source/Programs/Shared/EpicGames.Horde/Storage/Workspace.cs
--- a/Engine/Source/Programs/Shared/EpicGames.Horde/Storage/Workspace.cs
+++ b/Engine/Source/Programs/Shared/EpicGames.Horde/Storage/Workspace.cs
@@ -20,27 +20,35 @@ namespace EpicGames.Horde.Storage
 	{
 		class DirectoryState
 		{
+			public DirectoryState? Parent { get; }
+			public Utf8String Name { get; }
 			public Dictionary<Utf8String, DirectoryState> Directories { get; }
 			public Dictionary<Utf8String, FileState> Files { get; }
 
 			public ulong LayerFlags { get; set; }
 
-			public DirectoryState()
+			public DirectoryState(DirectoryState? parent, Utf8String name)
 			{
+				Parent = parent;
+				Name = name;
+
 				Directories = new Dictionary<Utf8String, DirectoryState>(Utf8StringComparer.Ordinal);
 				Files = new Dictionary<Utf8String, FileState>(Utf8StringComparer.Ordinal);
 			}
 
-			public DirectoryState(IMemoryReader reader)
+			public DirectoryState(DirectoryState? parent, Utf8String name, IMemoryReader reader)
 			{
+				Parent = parent;
+				Name = name;
+
 				int numDirectories = reader.ReadInt32();
 				Directories = new Dictionary<Utf8String, DirectoryState>(numDirectories, Utf8StringComparer.Ordinal);
 
 				for (int idx = 0; idx < numDirectories; idx++)
 				{
-					Utf8String name = reader.ReadUtf8String();
-					DirectoryState directory = new DirectoryState(reader);
-					Directories[name] = directory;
+					Utf8String subDirName = reader.ReadUtf8String();
+					DirectoryState subDirState = new DirectoryState(this, subDirName, reader);
+					Directories.Add(subDirName, subDirState);
 				}
 
 				int numFiles = reader.ReadInt32();
@@ -48,9 +56,9 @@ namespace EpicGames.Horde.Storage
 
 				for (int idx = 0; idx < numFiles; idx++)
 				{
-					Utf8String name = reader.ReadUtf8String();
-					FileState file = new FileState(reader);
-					Files[name] = file;
+					Utf8String fileName = reader.ReadUtf8String();
+					FileState fileState = new FileState(this, fileName, reader);
+					Files.Add(fileName, fileState);
 				}
 
 				LayerFlags = reader.ReadUnsignedVarInt();
@@ -61,12 +69,14 @@ namespace EpicGames.Horde.Storage
 				writer.WriteInt32(Directories.Count);
 				foreach (DirectoryState directory in Directories.Values)
 				{
+					writer.WriteUtf8String(directory.Name);
 					directory.Write(writer);
 				}
 
 				writer.WriteInt32(Files.Count);
 				foreach (FileState file in Files.Values)
 				{
+					writer.WriteUtf8String(file.Name);
 					file.Write(writer);
 				}
 
@@ -76,21 +86,26 @@ namespace EpicGames.Horde.Storage
 
 		class FileState
 		{
+			public DirectoryState Parent { get; }
+			public Utf8String Name { get; }
 			public long Length { get; private set; }
 			public long LastModifiedTimeUtc { get; private set; }
-			public IoHash Hash { get; }
+			public IoHash Hash { get; set; }
 
 			public ulong LayerFlags { get; set; }
 			public List<ChunkInfo> Chunks { get; } = new List<ChunkInfo>();
 
-			public FileState(FileInfo fileInfo, IoHash hash)
+			public FileState(DirectoryState parent, Utf8String name)
 			{
-				Update(fileInfo);
-				Hash = hash;
+				Parent = parent;
+				Name = name;
 			}
 
-			public FileState(IMemoryReader reader)
+			public FileState(DirectoryState parent, Utf8String name, IMemoryReader reader)
 			{
+				Parent = parent;
+				Name = name;
+
 				Length = reader.ReadInt64();
 				LastModifiedTimeUtc = reader.ReadInt64();
 				Hash = reader.ReadIoHash();
@@ -177,13 +192,13 @@ namespace EpicGames.Horde.Storage
 
 			public WorkspaceState()
 			{
-				Root = new DirectoryState();
+				Root = new DirectoryState(null, Utf8String.Empty);
 				Layers = new List<LayerState>();
 			}
 
 			public WorkspaceState(IMemoryReader reader)
 			{
-				Root = new DirectoryState(reader);
+				Root = new DirectoryState(null, Utf8String.Empty, reader);
 				Layers = reader.ReadList(() => new LayerState(reader));
 			}
 
@@ -354,54 +369,80 @@ namespace EpicGames.Horde.Storage
 			LayerState? layerState = GetLayerState(layerId);
 			if (layerState == null)
 			{
-				throw new ArgumentException($"Layer {layerId} does not exist", nameof(layerId));
+				throw new InvalidOperationException($"Layer '{layerId}' does not exist");
 			}
 
-			_state.Root = await SyncDirectoryAsync(contents, _rootDir, _state.Root, layerState.Flag, _logger, cancellationToken);
+			await SyncDirectoryAsync(_rootDir, _state.Root, contents, layerState.Flag, _logger, cancellationToken);
 		}
 
-		static async Task<DirectoryState> SyncDirectoryAsync(DirectoryNode directoryNode, DirectoryReference dirPath, DirectoryState? directoryState, ulong flag, ILogger logger, CancellationToken cancellationToken)
+		static async Task SyncDirectoryAsync(DirectoryReference dirPath, DirectoryState dirState, DirectoryNode dirNode, ulong flag, ILogger logger, CancellationToken cancellationToken)
 		{
 			DirectoryReference.CreateDirectory(dirPath);
 
-			DirectoryState newState = new DirectoryState();
+			// Remove any directories that no longer exist
+			foreach ((Utf8String subDirName, _) in dirState.Directories)
+			{
+				if (!dirNode.TryGetDirectoryEntry(subDirName, out _))
+				{
+					DirectoryReference subDirPath = DirectoryReference.Combine(dirPath, subDirName.ToString());
+					DirectoryReference.Delete(subDirPath, true);
+				}
+			}
 
-			foreach ((Utf8String name, DirectoryEntry? subDirEntry, DirectoryState? subDirState) in directoryNode.NameToDirectory.Zip(directoryState?.Directories))
+			// Remove any files that no longer exist
+			foreach ((Utf8String fileName, _) in dirState.Files)
 			{
-				DirectoryReference subDirPath = DirectoryReference.Combine(dirPath, name.ToString());
-				if (subDirEntry != null)
+				if (!dirNode.TryGetFileEntry(fileName, out _))
 				{
-					DirectoryNode subDirNode = await subDirEntry.ExpandAsync(cancellationToken);
-					newState.Directories[name] = await SyncDirectoryAsync(subDirNode, subDirPath, subDirState, flag, logger, cancellationToken);
+					FileReference filePath = FileReference.Combine(dirPath, fileName.ToString());
+					FileReference.Delete(filePath);
 				}
 			}
 
-			foreach ((Utf8String name, FileEntry? fileEntry, FileState? fileState) in directoryNode.NameToFile.Zip(directoryState?.Files))
+			// Update directories
+			foreach (DirectoryEntry subDirEntry in dirNode.Directories)
 			{
-				FileReference filePath = FileReference.Combine(dirPath, name.ToString());
-				if (fileEntry != null)
+				DirectoryState? subDirState;
+				if (!dirState.Directories.TryGetValue(subDirEntry.Name, out subDirState))
 				{
-					if (fileState == null || fileState.Hash != fileEntry.Hash)
-					{
-						newState.Files[name] = await CheckoutFileAsync(fileEntry, filePath.ToFileInfo(), logger, cancellationToken);
-					}
-					else
-					{
-						newState.Files[name] = fileState;
-					}
+					subDirState = new DirectoryState(dirState, subDirEntry.Name);
+					dirState.Directories.Add(subDirEntry.Name, subDirState);
 				}
+
+				DirectoryReference subDirPath = DirectoryReference.Combine(dirPath, subDirEntry.Name.ToString());
+
+				DirectoryNode subDirNode = await subDirEntry.ExpandAsync(cancellationToken);
+				await SyncDirectoryAsync(subDirPath, subDirState, subDirNode, flag, logger, cancellationToken);
 			}
 
-			return newState;
+			// Update files
+			foreach (FileEntry fileEntry in dirNode.Files)
+			{
+				FileState? fileState;
+				if (!dirState.Files.TryGetValue(fileEntry.Name, out fileState))
+				{
+					fileState = new FileState(dirState, fileEntry.Name);
+					dirState.Files.Add(fileEntry.Name, fileState);
+				}
+
+				FileReference filePath = FileReference.Combine(dirPath, fileEntry.Name.ToString());
+				await SyncFileAsync(filePath, fileState, fileEntry, logger, cancellationToken);
+			}
 		}
 
-		static async Task<FileState> CheckoutFileAsync(FileEntry fileRef, FileInfo fileInfo, ILogger logger, CancellationToken cancellationToken)
+		static async Task SyncFileAsync(FileReference filePath, FileState fileState, FileEntry fileEntry, ILogger logger, CancellationToken cancellationToken)
 		{
-			logger.LogInformation("Updating {File} to {Hash}", fileInfo, fileRef.Hash);
-			ChunkedDataNode fileNode = await fileRef.ExpandAsync(cancellationToken);
-			await fileNode.CopyToFileAsync(fileInfo, cancellationToken);
-			fileInfo.Refresh();
-			return new FileState(fileInfo, fileRef.Hash);
+			FileInfo fileInfo = filePath.ToFileInfo();
+			if (fileState.Hash != fileEntry.Hash)
+			{
+				logger.LogInformation("Updating {File} to {Hash}", fileInfo, fileEntry.Hash);
+				ChunkedDataNode fileNode = await fileEntry.ExpandAsync(cancellationToken);
+				await fileNode.CopyToFileAsync(fileInfo, cancellationToken);
+				fileInfo.Refresh();
+
+				fileState.Hash = fileEntry.Hash;
+				fileState.Update(fileInfo);
+			}
 		}
 	}
 }
