Horde: Respect layer flags when determining whether to delete a file or directory.

diff --git a/Engine/Source/Programs/Shared/EpicGames.Horde/Storage/Workspace.cs b/Engine/Source/Programs/Shared/EpicGames.Horde/Storage/Workspace.cs
--- a/Engine/Source/Programs/Shared/EpicGames.Horde/Storage/Workspace.cs
+++ b/Engine/Source/Programs/Shared/EpicGames.Horde/Storage/Workspace.cs
@@ -3,6 +3,7 @@
 using System;
 using System.Collections.Generic;
 using System.Diagnostics;
+using System.Diagnostics.CodeAnalysis;
 using System.IO;
 using System.Linq;
 using System.Threading;
@@ -18,12 +19,13 @@ namespace EpicGames.Horde.Storage
 	/// </summary>
 	public class Workspace
 	{
+		[DebuggerDisplay("{Name}")]
 		class DirectoryState
 		{
 			public DirectoryState? Parent { get; }
 			public Utf8String Name { get; }
-			public Dictionary<Utf8String, DirectoryState> Directories { get; }
-			public Dictionary<Utf8String, FileState> Files { get; }
+			public List<DirectoryState> Directories { get; }
+			public List<FileState> Files { get; }
 
 			public ulong LayerFlags { get; set; }
 
@@ -32,8 +34,8 @@ namespace EpicGames.Horde.Storage
 				Parent = parent;
 				Name = name;
 
-				Directories = new Dictionary<Utf8String, DirectoryState>(Utf8StringComparer.Ordinal);
-				Files = new Dictionary<Utf8String, FileState>(Utf8StringComparer.Ordinal);
+				Directories = new List<DirectoryState>();
+				Files = new List<FileState>();
 			}
 
 			public DirectoryState(DirectoryState? parent, Utf8String name, IMemoryReader reader)
@@ -42,39 +44,99 @@ namespace EpicGames.Horde.Storage
 				Name = name;
 
 				int numDirectories = reader.ReadInt32();
-				Directories = new Dictionary<Utf8String, DirectoryState>(numDirectories, Utf8StringComparer.Ordinal);
+				Directories = new List<DirectoryState>(numDirectories);
 
 				for (int idx = 0; idx < numDirectories; idx++)
 				{
 					Utf8String subDirName = reader.ReadUtf8String();
 					DirectoryState subDirState = new DirectoryState(this, subDirName, reader);
-					Directories.Add(subDirName, subDirState);
+					Directories.Add(subDirState);
 				}
 
 				int numFiles = reader.ReadInt32();
-				Files = new Dictionary<Utf8String, FileState>(numFiles, Utf8StringComparer.Ordinal);
+				Files = new List<FileState>(numFiles);
 
 				for (int idx = 0; idx < numFiles; idx++)
 				{
 					Utf8String fileName = reader.ReadUtf8String();
 					FileState fileState = new FileState(this, fileName, reader);
-					Files.Add(fileName, fileState);
+					Files.Add(fileState);
 				}
 
 				LayerFlags = reader.ReadUnsignedVarInt();
 			}
 
+			public bool TryGetFile(Utf8String name, [NotNullWhen(true)] out FileState? fileState)
+			{
+				int index = Files.BinarySearch(x => x.Name, name);
+				if (index >= 0)
+				{
+					fileState = Files[index];
+					return true;
+				}
+				else
+				{
+					fileState = null;
+					return false;
+				}
+			}
+
+			public FileState FindOrAddFile(Utf8String name)
+			{
+				int index = Files.BinarySearch(x => x.Name, name);
+				if (index >= 0)
+				{
+					return Files[index];
+				}
+				else
+				{
+					FileState fileState = new FileState(this, name);
+					Files.Insert(~index, fileState);
+					return fileState;
+				}
+			}
+
+			public bool TryGetDirectory(Utf8String name, [NotNullWhen(true)] out DirectoryState? directoryState)
+			{
+				int index = Directories.BinarySearch(x => x.Name, name);
+				if (index >= 0)
+				{
+					directoryState = Directories[index];
+					return true;
+				}
+				else
+				{
+					directoryState = null;
+					return false;
+				}
+			}
+
+			public DirectoryState FindOrAddDirectory(Utf8String name)
+			{
+				int index = Directories.BinarySearch(x => x.Name, name);
+				if (index >= 0)
+				{
+					return Directories[index];
+				}
+				else
+				{
+					DirectoryState subDirState = new DirectoryState(this, name);
+					Directories.Insert(~index, subDirState);
+					return subDirState;
+				}
+			}
+
 			public void Write(IMemoryWriter writer)
 			{
 				writer.WriteInt32(Directories.Count);
-				foreach (DirectoryState directory in Directories.Values)
+				foreach (DirectoryState directory in Directories)
 				{
 					writer.WriteUtf8String(directory.Name);
 					directory.Write(writer);
 				}
 
 				writer.WriteInt32(Files.Count);
-				foreach (FileState file in Files.Values)
+				foreach (FileState file in Files)
 				{
 					writer.WriteUtf8String(file.Name);
 					file.Write(writer);
@@ -84,6 +146,7 @@ namespace EpicGames.Horde.Storage
 			}
 		}
 
+		[DebuggerDisplay("{Name}")]
 		class FileState
 		{
 			public DirectoryState Parent { get; }
@@ -193,7 +256,7 @@ namespace EpicGames.Horde.Storage
 			public WorkspaceState()
 			{
 				Root = new DirectoryState(null, Utf8String.Empty);
-				Layers = new List<LayerState>();
+				Layers = new List<LayerState> { new LayerState(WorkspaceLayerId.Default, 1) };
 			}
 
 			public WorkspaceState(IMemoryReader reader)
@@ -343,7 +406,7 @@ namespace EpicGames.Horde.Storage
 		public void RemoveLayer(WorkspaceLayerId layerId)
 		{
 			int layerIdx = _state.Layers.FindIndex(x => x.Id == layerId);
-			if (layerIdx != -1)
+			if (layerIdx > 0) // Note: Excluding default layer at index 0
 			{
 				LayerState layer = _state.Layers[layerIdx];
 				if ((_state.Root.LayerFlags & layer.Flag) != 0)
@@ -364,7 +427,7 @@ namespace EpicGames.Horde.Storage
 		/// <param name="layerId">Identifier for the layer</param>
 		/// <param name="contents">New contents for the layer</param>
 		/// <param name="cancellationToken">Cancellation token for the operation</param>
-		public async Task SyncAsync(WorkspaceLayerId layerId, DirectoryNode contents, CancellationToken cancellationToken = default)
+		public async Task SyncAsync(WorkspaceLayerId layerId, DirectoryNode? contents, CancellationToken cancellationToken = default)
 		{
 			LayerState? layerState = GetLayerState(layerId);
 			if (layerState == null)
@@ -375,66 +438,93 @@ namespace EpicGames.Horde.Storage
 			await SyncDirectoryAsync(_rootDir, _state.Root, contents, layerState.Flag, _logger, cancellationToken);
 		}
 
-		static async Task SyncDirectoryAsync(DirectoryReference dirPath, DirectoryState dirState, DirectoryNode dirNode, ulong flag, ILogger logger, CancellationToken cancellationToken)
+		static async Task SyncDirectoryAsync(DirectoryReference dirPath, DirectoryState dirState, DirectoryNode? dirNode, ulong flag, ILogger logger, CancellationToken cancellationToken)
 		{
-			DirectoryReference.CreateDirectory(dirPath);
-
 			// Remove any directories that no longer exist
-			foreach ((Utf8String subDirName, _) in dirState.Directories)
+			for (int subDirIdx = 0; subDirIdx < dirState.Directories.Count; subDirIdx++)
 			{
-				if (!dirNode.TryGetDirectoryEntry(subDirName, out _))
+				DirectoryState subDirState = dirState.Directories[subDirIdx];
+				if ((subDirState.LayerFlags & flag) != 0)
 				{
-					DirectoryReference subDirPath = DirectoryReference.Combine(dirPath, subDirName.ToString());
-					DirectoryReference.Delete(subDirPath, true);
+					if (dirNode == null || !dirNode.TryGetDirectoryEntry(subDirState.Name, out _))
+					{
+						DirectoryReference subDirPath = DirectoryReference.Combine(dirPath, subDirState.Name.ToString());
+						await SyncDirectoryAsync(subDirPath, subDirState, null, flag, logger, cancellationToken);
+					}
 				}
 			}
 
 			// Remove any files that no longer exist
-			foreach ((Utf8String fileName, _) in dirState.Files)
+			for (int fileIdx = 0; fileIdx < dirState.Files.Count; fileIdx++)
 			{
-				if (!dirNode.TryGetFileEntry(fileName, out _))
+				FileState fileState = dirState.Files[fileIdx];
+				if ((fileState.LayerFlags & flag) != 0)
 				{
-					FileReference filePath = FileReference.Combine(dirPath, fileName.ToString());
-					FileReference.Delete(filePath);
+					if (dirNode == null || !dirNode.TryGetFileEntry(fileState.Name, out _))
+					{
+						FileReference filePath = FileReference.Combine(dirPath, fileState.Name.ToString());
+						await SyncFileAsync(filePath, fileState, null, flag, logger, cancellationToken);
+					}
 				}
 			}
 
-			// Update directories
-			foreach (DirectoryEntry subDirEntry in dirNode.Directories)
+			// Actually delete all the unreferenced directories
+			dirState.Directories.RemoveAll(x => x.LayerFlags == 0);
+			dirState.Files.RemoveAll(x => x.LayerFlags == 0);
+
+			// Clear out the layer flag for this directory. It'll be added back if we add/reuse files below.
+			dirState.LayerFlags &= ~flag;
+
+			// Add files for this directory
+			if (dirNode != null)
 			{
-				DirectoryState? subDirState;
-				if (!dirState.Directories.TryGetValue(subDirEntry.Name, out subDirState))
+				DirectoryReference.CreateDirectory(dirPath);
+
+				// Update directories
+				foreach (DirectoryEntry subDirEntry in dirNode.Directories)
 				{
-					subDirState = new DirectoryState(dirState, subDirEntry.Name);
-					dirState.Directories.Add(subDirEntry.Name, subDirState);
-				}
+					DirectoryReference subDirPath = DirectoryReference.Combine(dirPath, subDirEntry.Name.ToString());
+					DirectoryState subDirState = dirState.FindOrAddDirectory(subDirEntry.Name);
 
-				DirectoryReference subDirPath = DirectoryReference.Combine(dirPath, subDirEntry.Name.ToString());
+					DirectoryNode subDirNode = await subDirEntry.ExpandAsync(cancellationToken);
+					await SyncDirectoryAsync(subDirPath, subDirState, subDirNode, flag, logger, cancellationToken);
 
-				DirectoryNode subDirNode = await subDirEntry.ExpandAsync(cancellationToken);
-				await SyncDirectoryAsync(subDirPath, subDirState, subDirNode, flag, logger, cancellationToken);
-			}
+					dirState.LayerFlags |= flag;
+				}
 
-			// Update files
-			foreach (FileEntry fileEntry in dirNode.Files)
-			{
-				FileState? fileState;
-				if (!dirState.Files.TryGetValue(fileEntry.Name, out fileState))
+				// Update files
+				foreach (FileEntry fileEntry in dirNode.Files)
 				{
-					fileState = new FileState(dirState, fileEntry.Name);
-					dirState.Files.Add(fileEntry.Name, fileState);
+					FileReference filePath = FileReference.Combine(dirPath, fileEntry.Name.ToString());
+					FileState fileState = dirState.FindOrAddFile(fileEntry.Name);
+
+					await SyncFileAsync(filePath, fileState, fileEntry, flag, logger, cancellationToken);
+
+					dirState.LayerFlags |= flag;
 				}
+			}
 
-				FileReference filePath = FileReference.Combine(dirPath, fileEntry.Name.ToString());
-				await SyncFileAsync(filePath, fileState, fileEntry, logger, cancellationToken);
+			// Delete the directory if it's no longer needed
+			if (dirState.LayerFlags == 0 && dirState.Parent != null)
+			{
+				FileUtils.ForceDeleteDirectory(dirPath);
 			}
 		}
 
-		static async Task SyncFileAsync(FileReference filePath, FileState fileState, FileEntry fileEntry, ILogger logger, CancellationToken cancellationToken)
+		static async Task SyncFileAsync(FileReference filePath, FileState fileState, FileEntry? fileEntry, ulong flag, ILogger logger, CancellationToken cancellationToken)
 		{
-			FileInfo fileInfo = filePath.ToFileInfo();
-			if (fileState.Hash != fileEntry.Hash)
+			if (fileEntry == null)
+			{
+				fileState.LayerFlags &= ~flag;
+				if (fileState.LayerFlags == 0)
+				{
+					FileUtils.ForceDeleteFile(filePath);
+				}
+			}
+			else if (fileState.Hash != fileEntry.Hash)
 			{
+				FileInfo fileInfo = filePath.ToFileInfo();
+
 				logger.LogInformation("Updating {File} to {Hash}", fileInfo, fileEntry.Hash);
 				ChunkedDataNode fileNode = await fileEntry.ExpandAsync(cancellationToken);
 				await fileNode.CopyToFileAsync(fileInfo, cancellationToken);
