Horde: Add a list of tracked chunk hashes to workspace state.

diff --git a/Engine/Source/Programs/Shared/EpicGames.Horde/Storage/Workspace.cs b/Engine/Source/Programs/Shared/EpicGames.Horde/Storage/Workspace.cs
--- a/Engine/Source/Programs/Shared/EpicGames.Horde/Storage/Workspace.cs
+++ b/Engine/Source/Programs/Shared/EpicGames.Horde/Storage/Workspace.cs
@@ -19,13 +19,13 @@ namespace EpicGames.Horde.Storage
 	/// </summary>
 	public class Workspace
 	{
-		[DebuggerDisplay("{Name}")]
+		// Tracked state of a directory
 		class DirectoryState
 		{
 			public DirectoryState? Parent { get; }
 			public Utf8String Name { get; }
-			public List<DirectoryState> Directories { get; }
-			public List<FileState> Files { get; }
+			public List<DirectoryState> Directories { get; } = new List<DirectoryState>();
+			public List<FileState> Files { get; } = new List<FileState>();
 
 			public ulong LayerFlags { get; set; }
 
@@ -33,37 +33,6 @@ namespace EpicGames.Horde.Storage
 			{
 				Parent = parent;
 				Name = name;
-
-				Directories = new List<DirectoryState>();
-				Files = new List<FileState>();
-			}
-
-			public DirectoryState(DirectoryState? parent, Utf8String name, IMemoryReader reader)
-			{
-				Parent = parent;
-				Name = name;
-
-				int numDirectories = reader.ReadInt32();
-				Directories = new List<DirectoryState>(numDirectories);
-
-				for (int idx = 0; idx < numDirectories; idx++)
-				{
-					Utf8String subDirName = reader.ReadUtf8String();
-					DirectoryState subDirState = new DirectoryState(this, subDirName, reader);
-					Directories.Add(subDirState);
-				}
-
-				int numFiles = reader.ReadInt32();
-				Files = new List<FileState>(numFiles);
-
-				for (int idx = 0; idx < numFiles; idx++)
-				{
-					Utf8String fileName = reader.ReadUtf8String();
-					FileState fileState = new FileState(this, fileName, reader);
-					Files.Add(fileState);
-				}
-
-				LayerFlags = reader.ReadUnsignedVarInt();
 			}
 
 			public bool TryGetFile(Utf8String name, [NotNullWhen(true)] out FileState? fileState)
@@ -126,6 +95,56 @@ namespace EpicGames.Horde.Storage
 				}
 			}
 
+			public Utf8String GetPath()
+			{
+				Utf8StringBuilder builder = new Utf8StringBuilder();
+				GetPath(builder);
+				return builder.ToUtf8String();
+			}
+
+			public void GetPath(Utf8StringBuilder builder)
+			{
+				if (Parent != null)
+				{
+					Parent.GetPath(builder);
+				}
+				builder.Append(Name);
+				builder.Append((byte)'/');
+			}
+
+			public void Read(IMemoryReader reader)
+			{
+				int numDirectories = reader.ReadInt32();
+				Directories.Capacity = numDirectories;
+				Directories.Clear();
+
+				for (int idx = 0; idx < numDirectories; idx++)
+				{
+					Utf8String subDirName = reader.ReadUtf8String();
+
+					DirectoryState subDirState = new DirectoryState(this, subDirName);
+					subDirState.Read(reader);
+
+					Directories.Add(subDirState);
+				}
+
+				int numFiles = reader.ReadInt32();
+				Files.Capacity = numFiles;
+				Files.Clear();
+
+				for (int idx = 0; idx < numFiles; idx++)
+				{
+					Utf8String fileName = reader.ReadUtf8String();
+				
+					FileState fileState = new FileState(this, fileName);
+					fileState.Read(reader);
+
+					Files.Add(fileState);
+				}
+
+				LayerFlags = reader.ReadUnsignedVarInt();
+			}
+
 			public void Write(IMemoryWriter writer)
 			{
 				writer.WriteInt32(Directories.Count);
@@ -144,8 +163,11 @@ namespace EpicGames.Horde.Storage
 
 				writer.WriteUnsignedVarInt(LayerFlags);
 			}
+
+			public override string ToString() => GetPath().ToString();
 		}
 
+		// Tracked state of a file
 		[DebuggerDisplay("{Name}")]
 		class FileState
 		{
@@ -154,9 +176,7 @@ namespace EpicGames.Horde.Storage
 			public long Length { get; private set; }
 			public long LastModifiedTimeUtc { get; private set; }
 			public IoHash Hash { get; set; }
-
 			public ulong LayerFlags { get; set; }
-			public List<ChunkInfo> Chunks { get; } = new List<ChunkInfo>();
 
 			public FileState(DirectoryState parent, Utf8String name)
 			{
@@ -164,17 +184,13 @@ namespace EpicGames.Horde.Storage
 				Name = name;
 			}
 
-			public FileState(DirectoryState parent, Utf8String name, IMemoryReader reader)
+			public void Read(IMemoryReader reader)
 			{
-				Parent = parent;
-				Name = name;
-
 				Length = reader.ReadInt64();
 				LastModifiedTimeUtc = reader.ReadInt64();
 				Hash = reader.ReadIoHash();
 
 				LayerFlags = reader.ReadUnsignedVarInt();
-				Chunks = reader.ReadList(() => new ChunkInfo(reader));
 			}
 
 			public bool Modified(FileInfo fileInfo) => Length != fileInfo.Length || LastModifiedTimeUtc != fileInfo.LastWriteTimeUtc.Ticks;
@@ -192,38 +208,68 @@ namespace EpicGames.Horde.Storage
 				writer.WriteInt64(LastModifiedTimeUtc);
 				writer.WriteIoHash(Hash);
 				writer.WriteUnsignedVarInt(LayerFlags);
-				writer.WriteList(Chunks, x => x.Write(writer));
 			}
+
+			Utf8String GetPath()
+			{
+				Utf8StringBuilder builder = new Utf8StringBuilder();
+				GetPath(builder);
+				return builder.ToUtf8String();
+			}
+
+			void GetPath(Utf8StringBuilder builder)
+			{
+				Parent.GetPath(builder);
+				builder.Append(Name);
+			}
+
+			public override string ToString() => GetPath().ToString();
 		}
 
-		class ChunkInfo
+		// Collates lists of files and chunks with a particular hash
+		[DebuggerDisplay("{Hash}")]
+		class HashInfo
 		{
+			public int Index { get; set; }
+
 			public IoHash Hash { get; }
+			public List<FileState> Files { get; } = new List<FileState>();
+			public List<ChunkInfo> Chunks { get; } = new List<ChunkInfo>();
+
+			public HashInfo(IoHash hash) => Hash = hash;
+		}
+
+		// Hashed chunk within another hashed object
+		[DebuggerDisplay("{Offset}+{Length}")]
+		class ChunkInfo
+		{
+			public HashInfo WithinHashInfo { get; }
 			public long Offset { get; }
 			public long Length { get; }
 
-			public ChunkInfo(IoHash hash, long offset, long length)
+			public ChunkInfo(HashInfo withinHashInfo, long offset, long length)
 			{
-				Hash = hash;
+				WithinHashInfo = withinHashInfo;
 				Offset = offset;
 				Length = length;
 			}
 
-			public ChunkInfo(IMemoryReader reader)
+			public ChunkInfo(IMemoryReader reader, HashInfo[] hashes)
 			{
-				Hash = reader.ReadIoHash();
+				WithinHashInfo = hashes[(int)reader.ReadUnsignedVarInt()];
 				Offset = (long)reader.ReadUnsignedVarInt();
 				Length = (long)reader.ReadUnsignedVarInt();
 			}
 
 			public void Write(IMemoryWriter writer)
 			{
-				writer.WriteIoHash(Hash);
+				writer.WriteUnsignedVarInt(WithinHashInfo.Index);
 				writer.WriteUnsignedVarInt((ulong)Offset);
 				writer.WriteUnsignedVarInt((ulong)Length);
 			}
 		}
 
+		// Maps a layer id to a flag
 		class LayerState
 		{
 			public WorkspaceLayerId Id { get; }
@@ -248,33 +294,11 @@ namespace EpicGames.Horde.Storage
 			}
 		}
 
-		class WorkspaceState
-		{
-			public DirectoryState Root { get; set; }
-			public List<LayerState> Layers { get; private set; }
-
-			public WorkspaceState()
-			{
-				Root = new DirectoryState(null, Utf8String.Empty);
-				Layers = new List<LayerState> { new LayerState(WorkspaceLayerId.Default, 1) };
-			}
-
-			public WorkspaceState(IMemoryReader reader)
-			{
-				Root = new DirectoryState(null, Utf8String.Empty, reader);
-				Layers = reader.ReadList(() => new LayerState(reader));
-			}
-
-			public void Write(IMemoryWriter writer)
-			{
-				Root.Write(writer);
-				writer.WriteList(Layers, x => x.Write(writer));
-			}
-		}
-
 		readonly DirectoryReference _rootDir;
 		readonly FileReference _stateFile;
-		readonly WorkspaceState _state;
+		readonly DirectoryState _root;
+		readonly List<LayerState> _layers;
+		readonly Dictionary<IoHash, HashInfo> _hashes = new Dictionary<IoHash, HashInfo>();
 		readonly ILogger _logger;
 
 		const string HordeDirName = ".horde";
@@ -288,20 +312,20 @@ namespace EpicGames.Horde.Storage
 		/// <summary>
 		/// Layers current in this workspace
 		/// </summary>
-		public IReadOnlyList<WorkspaceLayerId> Layers => _state.Layers.Select(x => x.Id).ToList();
+		public IReadOnlyList<WorkspaceLayerId> Layers => _layers.Select(x => x.Id).ToList();
 
 		/// <summary>
 		/// Constructor
 		/// </summary>
 		/// <param name="rootDir">Directory for the workspace</param>
 		/// <param name="stateFile">Path to the state file for this directory</param>
-		/// <param name="state">Current state</param>
 		/// <param name="logger">Logger for diagnostic output</param>
-		private Workspace(DirectoryReference rootDir, FileReference stateFile, WorkspaceState state, ILogger logger)
+		private Workspace(DirectoryReference rootDir, FileReference stateFile, ILogger logger)
 		{
 			_rootDir = rootDir;
 			_stateFile = stateFile;
-			_state = state;
+			_root = new DirectoryState(null, Utf8String.Empty);
+			_layers = new List<LayerState> { new LayerState(WorkspaceLayerId.Default, 1) };
 			_logger = logger;
 		}
 
@@ -322,7 +346,7 @@ namespace EpicGames.Horde.Storage
 				throw new InvalidOperationException($"Workspace already exists in {rootDir}; use Open instead.");
 			}
 
-			Workspace workspace = new Workspace(rootDir, stateFile, new WorkspaceState(), logger);
+			Workspace workspace = new Workspace(rootDir, stateFile, logger);
 			await workspace.SaveAsync(cancellationToken);
 
 			return workspace;
@@ -345,8 +369,11 @@ namespace EpicGames.Horde.Storage
 					if (stream != null)
 					{
 						byte[] data = await stream.ReadAllBytesAsync(cancellationToken);
-						WorkspaceState state = new WorkspaceState(new MemoryReader(data));
-						return new Workspace(testDir, stateFile, state, logger);
+
+						Workspace workspace = new Workspace(testDir, stateFile, logger);
+						workspace.Read(new MemoryReader(data));
+
+						return workspace;
 					}
 				}
 			}
@@ -365,13 +392,58 @@ namespace EpicGames.Horde.Storage
 			{
 				using (ChunkedMemoryWriter writer = new ChunkedMemoryWriter(64 * 1024))
 				{
-					_state.Write(writer);
+					Write(writer);
 					await writer.CopyToAsync(stream, cancellationToken);
 				}
 				stream.CompleteTransaction();
 			}
 		}
 
+		void Read(IMemoryReader reader)
+		{
+			_root.Read(reader);
+			reader.ReadList(_layers, () => new LayerState(reader));
+
+			// Read the hash lookup
+			int numHashes = reader.ReadInt32();
+			HashInfo[] hashInfoArray = new HashInfo[numHashes];
+
+			_hashes.EnsureCapacity(numHashes);
+			_hashes.Clear();
+
+			for (int idx = 0; idx < numHashes; idx++)
+			{
+				IoHash hash = reader.ReadIoHash();
+				hashInfoArray[idx] = new HashInfo(hash);
+				_hashes.Add(hash, hashInfoArray[idx]);
+			}
+			for (int idx = 0; idx < numHashes; idx++)
+			{
+				HashInfo hashInfo = hashInfoArray[idx];
+				reader.ReadList(hashInfoArray[idx].Chunks, () => new ChunkInfo(reader, hashInfoArray));
+			}
+		}
+
+		void Write(IMemoryWriter writer)
+		{
+			_root.Write(writer);
+			writer.WriteList(_layers, x => x.Write(writer));
+
+			// Write the hash lookup
+			writer.WriteInt32(_hashes.Count);
+
+			int nextIndex = 0;
+			foreach (HashInfo hashInfo in _hashes.Values)
+			{
+				writer.WriteIoHash(hashInfo.Hash);
+				hashInfo.Index = nextIndex++;
+			}
+			foreach (HashInfo hashInfo in _hashes.Values)
+			{
+				writer.WriteList(hashInfo.Chunks, x => x.Write(writer));
+			}
+		}
+
 		#region Layers
 
 		/// <summary>
@@ -380,15 +452,15 @@ namespace EpicGames.Horde.Storage
 		/// <param name="id">Identifier for the layer</param>
 		public void AddLayer(WorkspaceLayerId id)
 		{
-			if (_state.Layers.Any(x => x.Id == id))
+			if (_layers.Any(x => x.Id == id))
 			{
 				throw new InvalidOperationException($"Layer {id} already exists");
 			}
 
 			ulong flags = 0;
-			for (int idx = 0; idx < _state.Layers.Count; idx++)
+			for (int idx = 0; idx < _layers.Count; idx++)
 			{
-				flags |= _state.Layers[idx].Flag;
+				flags |= _layers[idx].Flag;
 			}
 			if (flags == ~0UL)
 			{
@@ -396,7 +468,7 @@ namespace EpicGames.Horde.Storage
 			}
 
 			ulong nextFlag = (flags + 1) ^ flags;
-			_state.Layers.Add(new LayerState(id, nextFlag));
+			_layers.Add(new LayerState(id, nextFlag));
 		}
 
 		/// <summary>
@@ -405,19 +477,19 @@ namespace EpicGames.Horde.Storage
 		/// <param name="layerId">Layer to update</param>
 		public void RemoveLayer(WorkspaceLayerId layerId)
 		{
-			int layerIdx = _state.Layers.FindIndex(x => x.Id == layerId);
+			int layerIdx = _layers.FindIndex(x => x.Id == layerId);
 			if (layerIdx > 0) // Note: Excluding default layer at index 0
 			{
-				LayerState layer = _state.Layers[layerIdx];
-				if ((_state.Root.LayerFlags & layer.Flag) != 0)
+				LayerState layer = _layers[layerIdx];
+				if ((_root.LayerFlags & layer.Flag) != 0)
 				{
 					throw new InvalidOperationException($"Workspace still contains files for layer {layerId}");
 				}
-				_state.Layers.RemoveAt(layerIdx);
+				_layers.RemoveAt(layerIdx);
 			}
 		}
 
-		LayerState? GetLayerState(WorkspaceLayerId layerId) => _state.Layers.FirstOrDefault(x => x.Id == layerId);
+		LayerState? GetLayerState(WorkspaceLayerId layerId) => _layers.FirstOrDefault(x => x.Id == layerId);
 
 		#endregion
 
@@ -435,10 +507,10 @@ namespace EpicGames.Horde.Storage
 				throw new InvalidOperationException($"Layer '{layerId}' does not exist");
 			}
 
-			await SyncDirectoryAsync(_rootDir, _state.Root, contents, layerState.Flag, _logger, cancellationToken);
+			await SyncDirectoryAsync(_rootDir, _root, contents, layerState.Flag, cancellationToken);
 		}
 
-		static async Task SyncDirectoryAsync(DirectoryReference dirPath, DirectoryState dirState, DirectoryNode? dirNode, ulong flag, ILogger logger, CancellationToken cancellationToken)
+		async Task SyncDirectoryAsync(DirectoryReference dirPath, DirectoryState dirState, DirectoryNode? dirNode, ulong flag, CancellationToken cancellationToken)
 		{
 			// Remove any directories that no longer exist
 			for (int subDirIdx = 0; subDirIdx < dirState.Directories.Count; subDirIdx++)
@@ -449,7 +521,7 @@ namespace EpicGames.Horde.Storage
 					if (dirNode == null || !dirNode.TryGetDirectoryEntry(subDirState.Name, out _))
 					{
 						DirectoryReference subDirPath = DirectoryReference.Combine(dirPath, subDirState.Name.ToString());
-						await SyncDirectoryAsync(subDirPath, subDirState, null, flag, logger, cancellationToken);
+						await SyncDirectoryAsync(subDirPath, subDirState, null, flag, cancellationToken);
 					}
 				}
 			}
@@ -463,7 +535,7 @@ namespace EpicGames.Horde.Storage
 					if (dirNode == null || !dirNode.TryGetFileEntry(fileState.Name, out _))
 					{
 						FileReference filePath = FileReference.Combine(dirPath, fileState.Name.ToString());
-						await SyncFileAsync(filePath, fileState, null, flag, logger, cancellationToken);
+						await SyncFileAsync(filePath, fileState, null, flag, cancellationToken);
 					}
 				}
 			}
@@ -487,7 +559,7 @@ namespace EpicGames.Horde.Storage
 					DirectoryState subDirState = dirState.FindOrAddDirectory(subDirEntry.Name);
 
 					DirectoryNode subDirNode = await subDirEntry.ExpandAsync(cancellationToken);
-					await SyncDirectoryAsync(subDirPath, subDirState, subDirNode, flag, logger, cancellationToken);
+					await SyncDirectoryAsync(subDirPath, subDirState, subDirNode, flag, cancellationToken);
 
 					dirState.LayerFlags |= flag;
 				}
@@ -496,9 +568,9 @@ namespace EpicGames.Horde.Storage
 				foreach (FileEntry fileEntry in dirNode.Files)
 				{
 					FileReference filePath = FileReference.Combine(dirPath, fileEntry.Name.ToString());
-					FileState fileState = dirState.FindOrAddFile(fileEntry.Name);
 
-					await SyncFileAsync(filePath, fileState, fileEntry, flag, logger, cancellationToken);
+					FileState fileState = dirState.FindOrAddFile(fileEntry.Name);
+					await SyncFileAsync(filePath, fileState, fileEntry, flag, cancellationToken);
 
 					dirState.LayerFlags |= flag;
 				}
@@ -511,7 +583,7 @@ namespace EpicGames.Horde.Storage
 			}
 		}
 
-		static async Task SyncFileAsync(FileReference filePath, FileState fileState, FileEntry? fileEntry, ulong flag, ILogger logger, CancellationToken cancellationToken)
+		async Task SyncFileAsync(FileReference filePath, FileState fileState, FileEntry? fileEntry, ulong flag, CancellationToken cancellationToken)
 		{
 			if (fileEntry == null)
 			{
@@ -519,19 +591,43 @@ namespace EpicGames.Horde.Storage
 				if (fileState.LayerFlags == 0)
 				{
 					FileUtils.ForceDeleteFile(filePath);
+					RemoveFileFromHashLookup(fileState);
 				}
 			}
 			else if (fileState.Hash != fileEntry.Hash)
 			{
 				FileInfo fileInfo = filePath.ToFileInfo();
 
-				logger.LogInformation("Updating {File} to {Hash}", fileInfo, fileEntry.Hash);
+				_logger.LogInformation("Updating {File} to {Hash}", fileInfo, fileEntry.Hash);
 				ChunkedDataNode fileNode = await fileEntry.ExpandAsync(cancellationToken);
 				await fileNode.CopyToFileAsync(fileInfo, cancellationToken);
 				fileInfo.Refresh();
 
+				fileState.LayerFlags |= flag;
 				fileState.Hash = fileEntry.Hash;
 				fileState.Update(fileInfo);
+
+				AddFileToHashLookup(fileState);
+			}
+		}
+
+		void AddFileToHashLookup(FileState file)
+		{
+			HashInfo? hashInfo;
+			if (!_hashes.TryGetValue(file.Hash, out hashInfo))
+			{
+				hashInfo = new HashInfo(file.Hash);
+				_hashes.Add(file.Hash, hashInfo);
+			}
+			hashInfo.Files.Add(file);
+		}
+
+		void RemoveFileFromHashLookup(FileState file)
+		{
+			HashInfo? hashInfo;
+			if (_hashes.TryGetValue(file.Hash, out hashInfo))
+			{
+				hashInfo.Files.Remove(file);
 			}
 		}
 	}
