Missing file.

diff --git a/Engine/Source/Programs/Shared/EpicGames.Core/MemoryReader.cs b/Engine/Source/Programs/Shared/EpicGames.Core/MemoryReader.cs
--- a/Engine/Source/Programs/Shared/EpicGames.Core/MemoryReader.cs
+++ b/Engine/Source/Programs/Shared/EpicGames.Core/MemoryReader.cs
@@ -250,16 +250,26 @@ namespace EpicGames.Core
 		/// </summary>
 		/// <param name="reader">Reader to deserialize from</param>
 		/// <param name="readItem">Delegate to write an individual item</param>
-		public static List<T> ReadList<T>(this IMemoryReader reader, Func<T> readItem)
+		public static void ReadList<T>(this IMemoryReader reader, List<T> list, Func<T> readItem)
 		{
 			int length = (int)reader.ReadUnsignedVarInt();
+			list.EnsureCapacity(list.Count + length);
 
-			List<T> list = new List<T>(length);
 			for (int idx = 0; idx < length; idx++)
 			{
 				list.Add(readItem());
 			}
+		}
 
+		/// <summary>
+		/// Reads a variable length list
+		/// </summary>
+		/// <param name="reader">Reader to deserialize from</param>
+		/// <param name="readItem">Delegate to write an individual item</param>
+		public static List<T> ReadList<T>(this IMemoryReader reader, Func<T> readItem)
+		{
+			List<T> list = new List<T>();
+			ReadList(reader, list, readItem);
 			return list;
 		}
 
