Horde: Make parameters to ManagedWorkspace configurable

diff --git a/Engine/Source/Programs/Horde/Horde.Agent.Tests/WorkspaceInfoTest.cs b/Engine/Source/Programs/Horde/Horde.Agent.Tests/WorkspaceInfoTest.cs
--- a/Engine/Source/Programs/Horde/Horde.Agent.Tests/WorkspaceInfoTest.cs
+++ b/Engine/Source/Programs/Horde/Horde.Agent.Tests/WorkspaceInfoTest.cs
@@ -1,10 +1,13 @@
 // Copyright Epic Games, Inc. All Rights Reserved.
 
+using EpicGames.Perforce.Managed;
 using Horde.Agent.Utility;
 using Microsoft.VisualStudio.TestTools.UnitTesting;
 
 namespace Horde.Agent.Tests;
 
+using MWO = ManagedWorkspaceOptions;
+
 [TestClass]
 public class WorkspaceInfoTest
 {
@@ -21,4 +24,20 @@ public class WorkspaceInfoTest
 		Assert.IsFalse(WorkspaceInfo.ShouldUseHaveTable("name=managedWorkspace&useHaveTable=false"));
 		Assert.IsFalse(WorkspaceInfo.ShouldUseHaveTable("name=ManagedWorkspace&useHaveTable=FalsE"));
 	}
+
+	[TestMethod]
+	public void GetManagedWorkspaceOptions()
+	{
+		ManagedWorkspaceOptions defaultOptions = new ();
+		Assert.AreEqual(defaultOptions, WorkspaceInfo.GetMwOptions(null));
+		Assert.AreEqual(defaultOptions, WorkspaceInfo.GetMwOptions(""));
+		Assert.AreEqual(defaultOptions, WorkspaceInfo.GetMwOptions("name=somethingElse&numParallelSyncThreads=111"));
+		Assert.AreEqual(new MWO { NumParallelSyncThreads = 111 }, WorkspaceInfo.GetMwOptions("name=managedWorkspace&numParallelSyncThreads=111"));
+		Assert.AreEqual(new MWO { MaxFileConcurrency = 222 }, WorkspaceInfo.GetMwOptions("name=managedWorkspace&maxFileConcurrency=222"));
+		Assert.AreEqual(new MWO { MinScratchSpace = 333 }, WorkspaceInfo.GetMwOptions("name=managedWorkspace&minScratchSpace=333"));
+		Assert.AreEqual(new MWO { UseHaveTable = false }, WorkspaceInfo.GetMwOptions("name=managedWorkspace&useHaveTable=false"));
+		Assert.AreEqual(new MWO { UseHaveTable = false }, WorkspaceInfo.GetMwOptions("name=managedWorkspace&useHaveTable=FalSe"));
+		Assert.AreEqual(new MWO { UseHaveTable = true }, WorkspaceInfo.GetMwOptions("name=managedWorkspace&useHaveTable=true"));
+		Assert.AreEqual(new MWO { UseHaveTable = true }, WorkspaceInfo.GetMwOptions("name=managedWorkspace&useHaveTable=TrUE"));
+	}
 }
diff --git a/Engine/Source/Programs/Horde/Horde.Agent/Execution/ManagedWorkspaceMaterializer.cs b/Engine/Source/Programs/Horde/Horde.Agent/Execution/ManagedWorkspaceMaterializer.cs
--- a/Engine/Source/Programs/Horde/Horde.Agent/Execution/ManagedWorkspaceMaterializer.cs
+++ b/Engine/Source/Programs/Horde/Horde.Agent/Execution/ManagedWorkspaceMaterializer.cs
@@ -4,6 +4,7 @@ using System.Collections.Generic;
 using System.Threading;
 using System.Threading.Tasks;
 using EpicGames.Core;
+using EpicGames.Perforce.Managed;
 using Horde.Agent.Utility;
 using HordeCommon.Rpc.Messages;
 using Microsoft.Extensions.Logging;
@@ -43,8 +44,8 @@ public class ManagedWorkspaceMaterializer : IWorkspaceMaterializer
 	{
 		using IScope scope = CreateTraceSpan("ManagedWorkspaceMaterializer.InitializeAsync");
 		
-		bool useHaveTable = WorkspaceInfo.ShouldUseHaveTable(_agentWorkspace.Method);
-		_workspace = await WorkspaceInfo.SetupWorkspaceAsync(_agentWorkspace, _workingDir, useHaveTable, logger, cancellationToken);
+		ManagedWorkspaceOptions options = WorkspaceInfo.GetMwOptions(_agentWorkspace.Method);
+		_workspace = await WorkspaceInfo.SetupWorkspaceAsync(_agentWorkspace, _workingDir, options, logger, cancellationToken);
 		return await GetSettingsAsync(cancellationToken);
 	}
 
diff --git a/Engine/Source/Programs/Horde/Horde.Agent/Execution/PerforceExecutor.cs b/Engine/Source/Programs/Horde/Horde.Agent/Execution/PerforceExecutor.cs
--- a/Engine/Source/Programs/Horde/Horde.Agent/Execution/PerforceExecutor.cs
+++ b/Engine/Source/Programs/Horde/Horde.Agent/Execution/PerforceExecutor.cs
@@ -59,8 +59,8 @@ namespace Horde.Agent.Execution
 			{
 				using IScope _ = GlobalTracer.Instance.BuildSpan("Workspace").WithResourceName("AutoSDK").StartActive();
 				
-				bool useHaveTable = WorkspaceInfo.ShouldUseHaveTable(_autoSdkWorkspaceInfo.Method);
-				_autoSdkWorkspace = await WorkspaceInfo.SetupWorkspaceAsync(_autoSdkWorkspaceInfo, _rootDir, useHaveTable, logger, cancellationToken);
+				ManagedWorkspaceOptions options = WorkspaceInfo.GetMwOptions(_autoSdkWorkspaceInfo.Method);
+				_autoSdkWorkspace = await WorkspaceInfo.SetupWorkspaceAsync(_autoSdkWorkspaceInfo, _rootDir, options, logger, cancellationToken);
 
 				DirectoryReference legacyDir = DirectoryReference.Combine(_autoSdkWorkspace.MetadataDir, "HostWin64");
 				if (DirectoryReference.Exists(legacyDir))
@@ -104,8 +104,8 @@ namespace Horde.Agent.Execution
 			using (IScope scope = GlobalTracer.Instance.BuildSpan("Workspace").WithResourceName(_workspaceInfo.Identifier).StartActive())
 			{
 				// Sync the regular workspace
-				bool useHaveTable = WorkspaceInfo.ShouldUseHaveTable(_workspaceInfo.Method);
-				_workspace = await WorkspaceInfo.SetupWorkspaceAsync(_workspaceInfo, _rootDir, useHaveTable, logger, cancellationToken);
+				ManagedWorkspaceOptions options = WorkspaceInfo.GetMwOptions(_workspaceInfo.Method);
+				_workspace = await WorkspaceInfo.SetupWorkspaceAsync(_workspaceInfo, _rootDir, options, logger, cancellationToken);
 
 				// Figure out the change to build
 				if (_batch.Change == 0)
@@ -261,8 +261,8 @@ namespace Horde.Agent.Execution
 			List<WorkspaceInfo> workspaces = new List<WorkspaceInfo>();
 			foreach (AgentWorkspace pendingWorkspace in pendingWorkspaces)
 			{
-				bool useHaveTable = WorkspaceInfo.ShouldUseHaveTable(pendingWorkspace.Method);
-				WorkspaceInfo workspace = await WorkspaceInfo.SetupWorkspaceAsync(pendingWorkspace, rootDir, useHaveTable, logger, cancellationToken);
+				ManagedWorkspaceOptions options = WorkspaceInfo.GetMwOptions(pendingWorkspace.Method);
+				WorkspaceInfo workspace = await WorkspaceInfo.SetupWorkspaceAsync(pendingWorkspace, rootDir, options, logger, cancellationToken);
 				workspaces.Add(workspace);
 			}
 
diff --git a/Engine/Source/Programs/Horde/Horde.Agent/Utility/WorkspaceInfo.cs b/Engine/Source/Programs/Horde/Horde.Agent/Utility/WorkspaceInfo.cs
--- a/Engine/Source/Programs/Horde/Horde.Agent/Utility/WorkspaceInfo.cs
+++ b/Engine/Source/Programs/Horde/Horde.Agent/Utility/WorkspaceInfo.cs
@@ -121,11 +121,11 @@ namespace Horde.Agent.Utility
 		/// </summary>
 		/// <param name="workspace">The workspace definition</param>
 		/// <param name="rootDir">Root directory for storing the workspace</param>
-		/// <param name="useHaveTable">Use the client's have table when syncing</param>
+		/// <param name="options">Extra options for ManagedWorkspace</param>
 		/// <param name="logger">Logger output</param>
 		/// <param name="cancellationToken">Cancellation token</param>
 		/// <returns>New workspace info</returns>
-		public static async Task<WorkspaceInfo> SetupWorkspaceAsync(AgentWorkspace workspace, DirectoryReference rootDir, bool useHaveTable, ILogger logger, CancellationToken cancellationToken)
+		public static async Task<WorkspaceInfo> SetupWorkspaceAsync(AgentWorkspace workspace, DirectoryReference rootDir, ManagedWorkspaceOptions options, ILogger logger, CancellationToken cancellationToken)
 		{
 			// Fill in the default credentials iff they are not set
 			string? serverAndPort = String.IsNullOrEmpty(workspace.ServerAndPort)? null : workspace.ServerAndPort;
@@ -157,7 +157,7 @@ namespace Horde.Agent.Utility
 					logger.LogInformation("Using locally logged in session for {UserName}", userName);
 				}
 			}
-			return await SetupWorkspaceAsync(perforce, workspace.Stream, workspace.Identifier, workspace.View, !workspace.Incremental, rootDir, useHaveTable, logger, cancellationToken);
+			return await SetupWorkspaceAsync(perforce, workspace.Stream, workspace.Identifier, workspace.View, !workspace.Incremental, rootDir, options, logger, cancellationToken);
 		}
 
 		/// <summary>
@@ -169,11 +169,11 @@ namespace Horde.Agent.Utility
 		/// <param name="view">View for this workspace</param>
 		/// <param name="removeUntrackedFiles">Whether untracked files should be removed when cleaning this workspace</param>
 		/// <param name="rootDir">Root directory for storing the workspace</param>
-		/// <param name="useHaveTable">Use the client's have table when syncing</param>
+		/// <param name="options">Extra options for ManagedWorkspace</param>
 		/// <param name="logger">Logger output</param>
 		/// <param name="cancellationToken">Cancellation token</param>
 		/// <returns>New workspace info</returns>
-		public static async Task<WorkspaceInfo> SetupWorkspaceAsync(IPerforceConnection perforce, string streamName, string identifier, IList<string> view, bool removeUntrackedFiles, DirectoryReference rootDir, bool useHaveTable, ILogger logger, CancellationToken cancellationToken)
+		public static async Task<WorkspaceInfo> SetupWorkspaceAsync(IPerforceConnection perforce, string streamName, string identifier, IList<string> view, bool removeUntrackedFiles, DirectoryReference rootDir, ManagedWorkspaceOptions options, ILogger logger, CancellationToken cancellationToken)
 		{
 			// Get the host name, and fill in any missing metadata about the connection
 			InfoRecord info = await perforce.GetInfoAsync(InfoOptions.ShortOutput, cancellationToken);
@@ -184,7 +184,7 @@ namespace Horde.Agent.Utility
 				throw new Exception("Unable to determine Perforce host name");
 			}
 			
-			if (!useHaveTable)
+			if (!options.UseHaveTable)
 			{
 				logger.LogInformation("Skipping use of have table");
 			}
@@ -221,7 +221,7 @@ namespace Horde.Agent.Utility
 			DirectoryReference workspaceDir = DirectoryReference.Combine(metadataDir, "Sync");
 
 			// Create the repository
-			ManagedWorkspace newRepository = await ManagedWorkspace.LoadOrCreateAsync(hostName, metadataDir, true, useHaveTable, logger, cancellationToken);
+			ManagedWorkspace newRepository = await ManagedWorkspace.LoadOrCreateAsync(hostName, metadataDir, true, options, logger, cancellationToken);
 			if (removeUntrackedFiles)
 			{
 				await newRepository.DeleteClientAsync(perforceClient, cancellationToken);
@@ -266,6 +266,56 @@ namespace Horde.Agent.Utility
 
 			return true;
 		}
+
+		/// <summary>
+		/// Create ManagedWorkspace options from a URL-encoded query string.
+		/// </summary>
+		/// <param name="method">URL-encoded query string</param>
+		/// <returns></returns>
+		public static ManagedWorkspaceOptions GetMwOptions(string? method)
+		{
+			const string NameKey = "name";
+			const string ManagedWorkspaceValue = "managedWorkspace";
+			const string NumParallelSyncThreadsKey = "numParallelSyncThreads";
+			const string MaxFileConcurrencyKey = "maxFileConcurrency";
+			const string MinScratchSpaceKey = "minScratchSpace";
+			const string UseHaveTableKey = "useHaveTable";
+
+			ManagedWorkspaceOptions defaultOptions = new ();
+
+			int numParallelSyncThreads = defaultOptions.NumParallelSyncThreads;
+			int maxFileConcurrency = defaultOptions.MaxFileConcurrency;
+			long minScratchSpace = defaultOptions.MinScratchSpace;
+			bool useHaveTable = defaultOptions.UseHaveTable;
+			
+			if (String.IsNullOrEmpty(method))
+			{
+				return defaultOptions;
+			}
+			
+			NameValueCollection nameValues = HttpUtility.ParseQueryString(method);
+			if (!String.Equals(nameValues[NameKey], ManagedWorkspaceValue, StringComparison.OrdinalIgnoreCase))
+			{
+				return defaultOptions;
+			}
+
+			if (Int32.TryParse(nameValues[NumParallelSyncThreadsKey], out int v)) { numParallelSyncThreads = v; }
+			if (Int32.TryParse(nameValues[MaxFileConcurrencyKey], out v)) { maxFileConcurrency = v; }
+			if (Int32.TryParse(nameValues[MinScratchSpaceKey], out v)) { minScratchSpace = v; }
+			
+			if (String.Equals(nameValues[UseHaveTableKey], "false", StringComparison.OrdinalIgnoreCase))
+			{
+				useHaveTable = false;
+			}
+
+			return new ManagedWorkspaceOptions
+			{
+				NumParallelSyncThreads = numParallelSyncThreads,
+				MaxFileConcurrency = maxFileConcurrency,
+				MinScratchSpace = minScratchSpace,
+				UseHaveTable = useHaveTable
+			};
+		}
 		
 		/// <summary>
 		/// Gets the latest change in the stream
diff --git a/Engine/Source/Programs/Horde/Horde/Commands/Perforce/PerforceBase.cs b/Engine/Source/Programs/Horde/Horde/Commands/Perforce/PerforceBase.cs
--- a/Engine/Source/Programs/Horde/Horde/Commands/Perforce/PerforceBase.cs
+++ b/Engine/Source/Programs/Horde/Horde/Commands/Perforce/PerforceBase.cs
@@ -63,7 +63,8 @@ namespace Horde.Commands.Perforce
 			InfoRecord info = await perforce.GetInfoAsync(InfoOptions.ShortOutput, CancellationToken.None);
 			bool useHaveTable = UseHaveTable.Equals("true", StringComparison.Ordinal);
 
-			ManagedWorkspace repo = await ManagedWorkspace.LoadOrCreateAsync(info.ClientHost!, BaseDir, Overwrite, useHaveTable, logger, CancellationToken.None);
+			ManagedWorkspaceOptions options = new() { UseHaveTable = useHaveTable };
+			ManagedWorkspace repo = await ManagedWorkspace.LoadOrCreateAsync(info.ClientHost!, BaseDir, Overwrite, options, logger, CancellationToken.None);
 			await ExecuteAsync(perforce, repo, logger);
 			return 0;
 		}
diff --git a/Engine/Source/Programs/Shared/EpicGames.Perforce.Managed.Tests/ManagedWorkspaceTest.cs b/Engine/Source/Programs/Shared/EpicGames.Perforce.Managed.Tests/ManagedWorkspaceTest.cs
--- a/Engine/Source/Programs/Shared/EpicGames.Perforce.Managed.Tests/ManagedWorkspaceTest.cs
+++ b/Engine/Source/Programs/Shared/EpicGames.Perforce.Managed.Tests/ManagedWorkspaceTest.cs
@@ -34,7 +34,7 @@ public class ManagedWorkspaceTest : BasePerforceFixtureTest
 	[DataRow(false, DisplayName = "Without have-table")]
 	public async Task SyncSingleChangelist(bool useHaveTable)
 	{
-		ManagedWorkspace ws = await GetManagedWorkspace(useHaveTable);
+		ManagedWorkspace ws = await CreateManagedWorkspace(useHaveTable);
 		await AssertHaveTableFileCount(0);
 
 		await SyncAsync(ws, 6);
@@ -45,7 +45,7 @@ public class ManagedWorkspaceTest : BasePerforceFixtureTest
 	[TestMethod]
 	public async Task SyncBackwardsToOlderChangelistRemoveUntracked()
 	{
-		ManagedWorkspace ws = await GetManagedWorkspace(true);
+		ManagedWorkspace ws = await CreateManagedWorkspace(true);
 
 		await SyncAsync(ws, 6, removeUntracked: false);
 		Stream.GetChangelist(6).AssertDepotFiles(SyncDir);
@@ -65,7 +65,7 @@ public class ManagedWorkspaceTest : BasePerforceFixtureTest
 	[DataRow(false, DisplayName = "Without have-table")]
 	public async Task SyncBackwardsToOlderChangelist(bool useHaveTable)
 	{
-		ManagedWorkspace ws = await GetManagedWorkspace(useHaveTable);
+		ManagedWorkspace ws = await CreateManagedWorkspace(useHaveTable);
 
 		await SyncAsync(ws, 6);
 		Stream.GetChangelist(6).AssertDepotFiles(SyncDir);
@@ -86,7 +86,7 @@ public class ManagedWorkspaceTest : BasePerforceFixtureTest
 	[DataRow(false, DisplayName = "Without have-table")]
 	public async Task SyncUsingCacheFiles(bool useHaveTable)
 	{
-		ManagedWorkspace ws = await GetManagedWorkspace(useHaveTable);
+		ManagedWorkspace ws = await CreateManagedWorkspace(useHaveTable);
 
 		FileReference GetCacheFilePath(int changeNumber)
 		{
@@ -115,7 +115,7 @@ public class ManagedWorkspaceTest : BasePerforceFixtureTest
 	[DataRow(false, DisplayName = "Without have-table")]
 	public async Task SyncWithViewExclusivePathFirst(bool useHaveTable)
 	{
-		ManagedWorkspace ws = await GetManagedWorkspace(useHaveTable);
+		ManagedWorkspace ws = await CreateManagedWorkspace(useHaveTable);
 		ChangelistFixture cl = Stream.GetChangelist(6);
 
 		List<string> view = new() { "-/Data/...", "-/shared.h", };
@@ -135,7 +135,7 @@ public class ManagedWorkspaceTest : BasePerforceFixtureTest
 	[DataRow(false, DisplayName = "Without have-table")]
 	public async Task SyncWithViewInclusivePathFirst(bool useHaveTable)
 	{
-		ManagedWorkspace ws = await GetManagedWorkspace(useHaveTable);
+		ManagedWorkspace ws = await CreateManagedWorkspace(useHaveTable);
 		ChangelistFixture cl = Stream.GetChangelist(6);
 
 		List<string> view = new() { "/...", "-/Data/..." };
@@ -154,7 +154,7 @@ public class ManagedWorkspaceTest : BasePerforceFixtureTest
 	[DataRow(false, DisplayName = "Without have-table")]
 	public async Task Populate(bool useHaveTable)
 	{
-		ManagedWorkspace ws = await GetManagedWorkspace(useHaveTable);
+		ManagedWorkspace ws = await CreateManagedWorkspace(useHaveTable);
 		List<PopulateRequest> populateRequests = new () { new PopulateRequest(PerforceConnection, StreamName, new List<string>()) };
 		await ws.PopulateAsync(populateRequests, false, CancellationToken.None);
 		Stream.LatestChangelist.AssertDepotFiles(SyncDir);
@@ -166,7 +166,7 @@ public class ManagedWorkspaceTest : BasePerforceFixtureTest
 	[DataRow(false, DisplayName = "Without have-table")]
 	public async Task PopulateWithView(bool useHaveTable)
 	{
-		ManagedWorkspace ws = await GetManagedWorkspace(useHaveTable);
+		ManagedWorkspace ws = await CreateManagedWorkspace(useHaveTable);
 		List<string> view = new () { "-/Data/...", "-/shared.h" };
 		ChangelistFixture cl = Stream.LatestChangelist;
 		List<DepotFileFixture> filtered = cl.StreamFiles
@@ -186,7 +186,7 @@ public class ManagedWorkspaceTest : BasePerforceFixtureTest
 	[DataRow(false, DisplayName = "Without have-table")]
 	public async Task Unshelve(bool useHaveTable)
 	{
-		ManagedWorkspace ws = await GetManagedWorkspace(useHaveTable);
+		ManagedWorkspace ws = await CreateManagedWorkspace(useHaveTable);
 		await SyncAsync(ws, 7);
 		await ws.UnshelveAsync(PerforceConnection, 8, CancellationToken.None);
 		Stream.GetChangelist(8).AssertDepotFiles(SyncDir);
@@ -201,7 +201,7 @@ public class ManagedWorkspaceTest : BasePerforceFixtureTest
 	public async Task Caching_SyncingNewChangelist_UnusedFilesMovedToCache(bool useHaveTable)
 	{
 		// Arrange
-		ManagedWorkspace ws = await GetManagedWorkspace(useHaveTable);
+		ManagedWorkspace ws = await CreateManagedWorkspace(useHaveTable);
 		await SyncAsync(ws, 7);
 		
 		// Act
@@ -217,7 +217,7 @@ public class ManagedWorkspaceTest : BasePerforceFixtureTest
 	public async Task Caching_SyncingWithCachedData_FilesPopulatedFromCache(bool useHaveTable)
 	{
 		// Arrange
-		ManagedWorkspace ws = await GetManagedWorkspace(useHaveTable);
+		ManagedWorkspace ws = await CreateManagedWorkspace(useHaveTable);
 		await SyncAsync(ws, 7);
 		await SyncAsync(ws, 1);
 		
@@ -275,7 +275,7 @@ public class ManagedWorkspaceTest : BasePerforceFixtureTest
 	public async Task ReusePerforceClientWithoutHaveTable()
 	{
 		// Sync with have-table as normal
-		ManagedWorkspace wsWithHave = await GetManagedWorkspace(true);
+		ManagedWorkspace wsWithHave = await CreateManagedWorkspace(true);
 		await wsWithHave.SetupAsync(PerforceConnection, StreamName, CancellationToken.None);
 		ChangelistFixture cl = Stream.GetChangelist(6);
 		await SyncAsync(wsWithHave, cl.Number);
@@ -284,7 +284,7 @@ public class ManagedWorkspaceTest : BasePerforceFixtureTest
 
 		// Create a new ManagedWorkspace without have-table but re-use same Perforce connection
 		// The have-table remnants from above should not interfere with this workspace
-		ManagedWorkspace wsWithoutHave = await GetManagedWorkspace(false);
+		ManagedWorkspace wsWithoutHave = await CreateManagedWorkspace(false);
 		await wsWithoutHave.SetupAsync(PerforceConnection, StreamName, CancellationToken.None);
 		await SyncAsync(wsWithoutHave, cl.Number);
 		await AssertHaveTableFileCount(0);
@@ -292,9 +292,18 @@ public class ManagedWorkspaceTest : BasePerforceFixtureTest
 		cl.AssertDepotFiles(SyncDir);
 	}
 
-	private async Task<ManagedWorkspace> GetManagedWorkspace(bool useHaveTable)
+	private async Task<ManagedWorkspace> CreateManagedWorkspace(bool useHaveTable)
 	{
-		ManagedWorkspace ws = await ManagedWorkspace.CreateAsync(Environment.MachineName, TempDir, useHaveTable, _mwLogger, CancellationToken.None);
+		ManagedWorkspaceOptions options = new() { UseHaveTable = useHaveTable };
+		ManagedWorkspace ws = await ManagedWorkspace.CreateAsync(Environment.MachineName, TempDir, options, _mwLogger, CancellationToken.None);
+		await ws.SetupAsync(PerforceConnection, StreamName, CancellationToken.None);
+		return ws;
+	}
+	
+	private async Task<ManagedWorkspace> LoadManagedWorkspace(bool useHaveTable)
+	{
+		ManagedWorkspaceOptions options = new() { UseHaveTable = useHaveTable };
+		ManagedWorkspace ws = await ManagedWorkspace.LoadAsync(Environment.MachineName, TempDir, options, _mwLogger, CancellationToken.None);
 		await ws.SetupAsync(PerforceConnection, StreamName, CancellationToken.None);
 		return ws;
 	}
diff --git a/Engine/Source/Programs/Shared/EpicGames.Perforce.Managed/ManagedWorkspace.cs b/Engine/Source/Programs/Shared/EpicGames.Perforce.Managed/ManagedWorkspace.cs
--- a/Engine/Source/Programs/Shared/EpicGames.Perforce.Managed/ManagedWorkspace.cs
+++ b/Engine/Source/Programs/Shared/EpicGames.Perforce.Managed/ManagedWorkspace.cs
@@ -64,6 +64,59 @@ namespace EpicGames.Perforce.Managed
 			View = view;
 		}
 	}
+	
+	/// <summary>
+	/// Extra options for configuring ManagedWorkspace
+	/// </summary>
+	public class ManagedWorkspaceOptions
+	{
+		/// <summary>
+		/// Maximum number of threads to sync in parallel
+		/// </summary>
+		public int NumParallelSyncThreads { get; init; } = 4;
+
+		/// <summary>
+		/// Maximum number of concurrent file system operations (copying, moving, deleting etc)
+		/// </summary>
+		public int MaxFileConcurrency { get; init; } = 4;
+		
+		/// <summary>
+		/// Minimum amount of space that must be on a drive after a branch is synced
+		/// </summary>
+		public long MinScratchSpace { get; init; } = 50L * 1024 * 1024 * 1024;
+		
+		/// <summary>
+		/// Use the client's have table when syncing.
+		/// 
+		/// When set to false, updates to the have table will be prevented through use of "sync -p".
+		/// Actual files to sync will be gathered through "fstat".
+		/// This puts less strain on the Perforce server and can improve sync performance.
+		/// </summary>
+		public bool UseHaveTable { get; init; } = true;
+
+		protected bool Equals(ManagedWorkspaceOptions other)
+		{
+			return NumParallelSyncThreads == other.NumParallelSyncThreads && MaxFileConcurrency == other.MaxFileConcurrency && MinScratchSpace == other.MinScratchSpace && UseHaveTable == other.UseHaveTable;
+		}
+
+		public override bool Equals(object? obj)
+		{
+			if (ReferenceEquals(null, obj)) { return false; }
+			if (ReferenceEquals(this, obj)) { return true; }
+			if (obj.GetType() != GetType()) { return false; }
+			return Equals((ManagedWorkspaceOptions)obj);
+		}
+
+		public override int GetHashCode()
+		{
+			return HashCode.Combine(NumParallelSyncThreads, MaxFileConcurrency, MinScratchSpace, UseHaveTable);
+		}
+
+		public override string ToString()
+		{
+			return $"{nameof(NumParallelSyncThreads)}={NumParallelSyncThreads} {nameof(MaxFileConcurrency)}={MaxFileConcurrency} {nameof(MinScratchSpace)}={MinScratchSpace} {nameof(UseHaveTable)}={UseHaveTable}";
+		}
+	}
 
 	/// <summary>
 	/// Version number for managed workspace cache files
@@ -111,19 +164,9 @@ namespace EpicGames.Perforce.Managed
 		static int CurrentVersion { get; } = Enum.GetValues(typeof(ManagedWorkspaceVersion)).Cast<int>().Max();
 
 		/// <summary>
-		/// Maximum number of threads to sync in parallel
+		/// Externally configurable options
 		/// </summary>
-		const int NumParallelSyncThreads = 4;
-		
-		/// <summary>
-		/// Maximum number of concurrent file system operations (copying, moving, deleting etc)
-		/// </summary>
-		const int MaxFileConcurrency = 4;
-
-		/// <summary>
-		/// Minimum amount of space that must be on a drive after a branch is synced
-		/// </summary>
-		const long MinScratchSpace = 50L * 1024 * 1024 * 1024;
+		private readonly ManagedWorkspaceOptions _options;
 
 		/// <summary>
 		/// Constant for syncing the latest change number
@@ -160,15 +203,6 @@ namespace EpicGames.Perforce.Managed
 		/// </summary>
 		readonly ILogger _logger;
 
-		/// <summary>
-		/// Use the client's have table when syncing.
-		/// 
-		/// When set to false, updates to the have table will be prevented through use of "sync -p".
-		/// Actual files to sync will gathered through "fstat".
-		/// This puts less strain on the Perforce server and can improve sync performance.
-		/// </summary>
-		readonly bool _useHaveTable;
-
 		/// <summary>
 		/// The root directory for the stash
 		/// </summary>
@@ -210,15 +244,15 @@ namespace EpicGames.Perforce.Managed
 		/// <param name="hostName">Name of the current host</param>
 		/// <param name="nextSequenceNumber">The next sequence number for operations</param>
 		/// <param name="baseDir">The root directory for the stash</param>
-		/// <param name="useHaveTable">Use the client's have table when syncing</param>
+		/// <param name="options">Extra options</param>
 		/// <param name="logger">The log output device</param>
-		private ManagedWorkspace(string hostName, uint nextSequenceNumber, DirectoryReference baseDir, bool useHaveTable, ILogger logger)
+		private ManagedWorkspace(string hostName, uint nextSequenceNumber, DirectoryReference baseDir, ManagedWorkspaceOptions options, ILogger logger)
 		{
 			// Save the Perforce settings
 			_hostName = hostName;
 			_nextSequenceNumber = nextSequenceNumber;
-			_useHaveTable = useHaveTable;
 			_logger = logger;
+			_options = options;
 
 			// Get all the directories
 			_baseDir = baseDir;
@@ -240,17 +274,17 @@ namespace EpicGames.Perforce.Managed
 		/// <param name="hostName">Name of the current machine. Will be automatically detected from the host settings if not present.</param>
 		/// <param name="baseDir">The base directory for the repository</param>
 		/// <param name="overwrite">Whether to allow overwriting a repository that's not up to date</param>
-		/// <param name="useHaveTable">Use the client's have table when syncing</param>
+		/// <param name="options">Extra options</param>
 		/// <param name="logger">The logging interface</param>
 		/// <param name="cancellationToken">Cancellation token for this operation</param>
 		/// <returns></returns>
-		public static async Task<ManagedWorkspace> LoadOrCreateAsync(string hostName, DirectoryReference baseDir, bool overwrite, bool useHaveTable, ILogger logger, CancellationToken cancellationToken)
+		public static async Task<ManagedWorkspace> LoadOrCreateAsync(string hostName, DirectoryReference baseDir, bool overwrite, ManagedWorkspaceOptions options, ILogger logger, CancellationToken cancellationToken)
 		{
 			if (Exists(baseDir))
 			{
 				try
 				{
-					return await LoadAsync(hostName, baseDir, useHaveTable, logger, cancellationToken);
+					return await LoadAsync(hostName, baseDir, options, logger, cancellationToken);
 				}
 				catch (Exception ex)
 				{
@@ -265,7 +299,7 @@ namespace EpicGames.Perforce.Managed
 				}
 			}
 
-			return await CreateAsync(hostName, baseDir, useHaveTable, logger, cancellationToken);
+			return await CreateAsync(hostName, baseDir, options, logger, cancellationToken);
 		}
 		/*
 				public static PerforceConnection GetPerforceConnection(PerforceConnection Perforce)
@@ -298,11 +332,11 @@ namespace EpicGames.Perforce.Managed
 		/// </summary>
 		/// <param name="hostName">Name of the current machine.</param>
 		/// <param name="baseDir">The base directory for the repository</param>
-		/// <param name="useHaveTable">Use the client's have table when syncing</param>
+		/// <param name="options">Extra options</param>
 		/// <param name="logger">The log output device</param>
 		/// <param name="cancellationToken">Cancellation token for this operation</param>
 		/// <returns>New repository instance</returns>
-		public static async Task<ManagedWorkspace> CreateAsync(string hostName, DirectoryReference baseDir, bool useHaveTable, ILogger logger, CancellationToken cancellationToken)
+		public static async Task<ManagedWorkspace> CreateAsync(string hostName, DirectoryReference baseDir, ManagedWorkspaceOptions options, ILogger logger, CancellationToken cancellationToken)
 		{
 			logger.LogInformation("Creating repository at {Location}...", baseDir);
 
@@ -310,7 +344,7 @@ namespace EpicGames.Perforce.Managed
 			DirectoryReference.CreateDirectory(baseDir);
 			FileUtils.ForceDeleteDirectoryContents(baseDir);
 
-			ManagedWorkspace repo = new ManagedWorkspace(hostName, 1, baseDir, useHaveTable, logger);
+			ManagedWorkspace repo = new ManagedWorkspace(hostName, 1, baseDir, options, logger);
 			await repo.SaveAsync(TransactionState.Clean, cancellationToken);
 			repo.CreateCacheHierarchy();
 
@@ -350,10 +384,10 @@ namespace EpicGames.Perforce.Managed
 		/// </summary>
 		/// <param name="hostName">Name of the current host. Will be obtained from a 'p4 info' call if not specified</param>
 		/// <param name="baseDir">The base directory for the repository</param>
-		/// <param name="useHaveTable">Use the client's have table when syncing</param> 
+		/// <param name="options">Extra options</param> 
 		/// <param name="logger">The log output device</param>
 		/// <param name="cancellationToken">Cancellation token for this command</param>
-		public static async Task<ManagedWorkspace> LoadAsync(string hostName, DirectoryReference baseDir, bool useHaveTable, ILogger logger, CancellationToken cancellationToken)
+		public static async Task<ManagedWorkspace> LoadAsync(string hostName, DirectoryReference baseDir, ManagedWorkspaceOptions options, ILogger logger, CancellationToken cancellationToken)
 		{
 			if (!Exists(baseDir))
 			{
@@ -375,7 +409,7 @@ namespace EpicGames.Perforce.Managed
 			bool requiresRepair = reader.ReadBoolean();
 			uint nextSequenceNumber = reader.ReadUInt32();
 
-			ManagedWorkspace repo = new (hostName, nextSequenceNumber, baseDir, useHaveTable, logger);
+			ManagedWorkspace repo = new (hostName, nextSequenceNumber, baseDir, options, logger);
 			repo._requiresRepair = requiresRepair;
 
 			int numTrackedFiles = reader.ReadInt32();
@@ -455,7 +489,7 @@ namespace EpicGames.Perforce.Managed
 			{
 				Stopwatch timer = Stopwatch.StartNew();
 
-				(FileInfo[] filesToDelete, DirectoryInfo[] directoriesToDelete) refreshResult = await _workspace.Refresh(removeUntracked, MaxFileConcurrency);
+				(FileInfo[] filesToDelete, DirectoryInfo[] directoriesToDelete) refreshResult = await _workspace.Refresh(removeUntracked, _options.MaxFileConcurrency);
 				filesToDelete = refreshResult.filesToDelete;
 				directoriesToDelete = refreshResult.directoriesToDelete;
 
@@ -483,7 +517,7 @@ namespace EpicGames.Perforce.Managed
 				{
 					Stopwatch timer = Stopwatch.StartNew();
 					
-					ParallelOptions options = new() { MaxDegreeOfParallelism = MaxFileConcurrency, CancellationToken = cancellationToken };
+					ParallelOptions options = new() { MaxDegreeOfParallelism = _options.MaxFileConcurrency, CancellationToken = cancellationToken };
 					await Parallel.ForEachAsync(filesToDelete, options, (fileToDelete, ct) =>
 					{
 						FileUtils.ForceDeleteFile(fileToDelete);
@@ -765,7 +799,7 @@ namespace EpicGames.Perforce.Managed
 			_logger.LogInformation("Stage contains {NumFiles:n0} files, {TotalSize:n1}mb", _workspace.GetFiles().Count, _workspace.GetFiles().Sum(x => x._length) / (1024.0 * 1024.0));
 
 			// Print the contents of the workspace
-			string[] differences = await _workspace.FindDifferencesAsync(MaxFileConcurrency);
+			string[] differences = await _workspace.FindDifferencesAsync(_options.MaxFileConcurrency);
 			if (differences.Length > 0)
 			{
 				_logger.LogInformation("Local changes:");
@@ -816,7 +850,7 @@ namespace EpicGames.Perforce.Managed
 				_logger.LogInformation("Syncing to {StreamName} at CL {CL}", streamName, changeNumber);
 			}
 
-			using (_logger.WithProperty("useHaveTable", _useHaveTable).BeginScope())
+			using (_logger.WithProperty("useHaveTable", _options.UseHaveTable).BeginScope())
 			using (_logger.BeginIndentScope("  "))
 			{
 				// Update the client to the current stream
@@ -828,7 +862,7 @@ namespace EpicGames.Perforce.Managed
 					changeNumber = await GetLatestClientChangeAsync(perforce, cancellationToken);
 				}
 
-				if (_useHaveTable)
+				if (_options.UseHaveTable)
 				{
 					// Revert any open files
 					await RevertInternalAsync(perforce, cancellationToken);
@@ -847,7 +881,7 @@ namespace EpicGames.Perforce.Managed
 				StreamSnapshot? contents;
 				if (cacheFile == null)
 				{
-					if (_useHaveTable)
+					if (_options.UseHaveTable)
 					{
 						contents = await FindClientContentsAsync(perforce, changeNumber, cancellationToken);						
 					}
@@ -1019,7 +1053,7 @@ namespace EpicGames.Perforce.Managed
 					int changeNumber = await GetLatestClientChangeAsync(request.PerforceClient, cancellationToken);
 					_logger.LogInformation("Latest change is CL {CL}", changeNumber);
 
-					if (_useHaveTable)
+					if (_options.UseHaveTable)
 					{
 						await UpdateClientHaveTableAsync(request.PerforceClient, changeNumber, request.View, cancellationToken);
 
@@ -1100,7 +1134,7 @@ namespace EpicGames.Perforce.Managed
 
 				using (_logger.BeginIndentScope("  "))
 				{
-					if (_useHaveTable)
+					if (_options.UseHaveTable)
 					{
 						await DeleteClientAsync(request.PerforceClient, cancellationToken);
 						await UpdateClientAsync(request.PerforceClient, streamName, cancellationToken);
@@ -1168,7 +1202,7 @@ namespace EpicGames.Perforce.Managed
 				
 				// Partitioned and read-only types store their have table separately on the server, compared to normal (writeable) clients
 				// Clients that sync without updating the have table cannot submit so they're marked as read-only. 
-				client.Type = _useHaveTable ? "partitioned" : "readonly";
+				client.Type = _options.UseHaveTable ? "partitioned" : "readonly";
 
 				using IPerforceConnection perforce = await perforceClient.WithoutClientAsync();
 
@@ -1179,7 +1213,7 @@ namespace EpicGames.Perforce.Managed
 					await perforceClient.CreateClientAsync(client, cancellationToken);
 				}
 
-				if (!_useHaveTable)
+				if (!_options.UseHaveTable)
 				{
 					// If have table is not used, make client is fully reset as it may have been re-used
 					await UpdateHaveTablePathAsync(perforceClient, $"//{clientName}/...#0", cancellationToken);
@@ -1675,7 +1709,7 @@ namespace EpicGames.Perforce.Managed
 		/// <returns>Contents of the workspace</returns>
 		private async Task<StreamSnapshotFromMemory> FindAndSaveClientContentsAsync(IPerforceConnection perforceClient, Utf8String basePath, IReadOnlyList<string> view, int changeNumber, FileReference cacheFile, CancellationToken cancellationToken)
 		{
-			StreamSnapshotFromMemory contents = _useHaveTable
+			StreamSnapshotFromMemory contents = _options.UseHaveTable
 				? await FindClientContentsAsync(perforceClient, changeNumber, cancellationToken)
 				: await FindClientContentsWithoutHaveTableAsync(perforceClient, basePath.ToString(), view, changeNumber, cancellationToken);
 
@@ -1724,7 +1758,7 @@ namespace EpicGames.Perforce.Managed
 			{
 				Stopwatch timer = Stopwatch.StartNew();
 
-				transaction = await RemoveTransaction.CreateAsync(_workspace, contents, _contentIdToTrackedFile, MaxFileConcurrency);
+				transaction = await RemoveTransaction.CreateAsync(_workspace, contents, _contentIdToTrackedFile, _options.MaxFileConcurrency);
 
 				scope.Progress = $"({timer.Elapsed.TotalSeconds:0.0}s)";
 			}
@@ -1753,7 +1787,7 @@ namespace EpicGames.Perforce.Managed
 					await SaveAsync(TransactionState.Dirty, cancellationToken);
 
 					// Execute all the moves and deletes
-					ParallelOptions options = new() { MaxDegreeOfParallelism = MaxFileConcurrency, CancellationToken = cancellationToken };
+					ParallelOptions options = new() { MaxDegreeOfParallelism = _options.MaxFileConcurrency, CancellationToken = cancellationToken };
 					await Parallel.ForEachAsync(sourceAndTargetFiles, options, (sourceAndTargetFile, ct) =>
 					{
 						FileUtils.ForceMoveFile(sourceAndTargetFile.Key, sourceAndTargetFile.Value);
@@ -1773,7 +1807,7 @@ namespace EpicGames.Perforce.Managed
 				{
 					Stopwatch timer = Stopwatch.StartNew();
 
-					ParallelOptions options = new() { MaxDegreeOfParallelism = MaxFileConcurrency, CancellationToken = cancellationToken };
+					ParallelOptions options = new() { MaxDegreeOfParallelism = _options.MaxFileConcurrency, CancellationToken = cancellationToken };
 					await Parallel.ForEachAsync(filesToDelete, options, (fileToDelete, ct) =>
 					{
 						RemoveFile(fileToDelete);
@@ -1860,7 +1894,7 @@ namespace EpicGames.Perforce.Managed
 			{
 				Stopwatch timer = Stopwatch.StartNew();
 
-				transaction = await AddTransaction.CreateAsync(_workspace, stream, _contentIdToTrackedFile, MaxFileConcurrency);
+				transaction = await AddTransaction.CreateAsync(_workspace, stream, _contentIdToTrackedFile, _options.MaxFileConcurrency);
 				_workspace = transaction._newWorkspaceRootDir;
 				await SaveAsync(TransactionState.Dirty, cancellationToken);
 
@@ -1875,7 +1909,7 @@ namespace EpicGames.Perforce.Managed
 				using (ILoggerProgress status = _logger.BeginProgressScope(String.Format("Moving {0} {1} from cache...", filesToMove.Length, (filesToMove.Length == 1) ? "file" : "files")))
 				{
 					Stopwatch timer = Stopwatch.StartNew();
-					ParallelOptions options = new() { MaxDegreeOfParallelism = MaxFileConcurrency, CancellationToken = cancellationToken };
+					ParallelOptions options = new() { MaxDegreeOfParallelism = _options.MaxFileConcurrency, CancellationToken = cancellationToken };
 					await Parallel.ForEachAsync(filesToMove, options, (fileToMove, ct) =>
 					{
 						MoveFileFromCache(fileToMove, transaction._filesToSync);
@@ -1895,7 +1929,7 @@ namespace EpicGames.Perforce.Managed
 				using (ILoggerProgress status = _logger.BeginProgressScope(String.Format("Copying {0} {1} within workspace...", filesToCopy.Length, (filesToCopy.Length == 1) ? "file" : "files")))
 				{
 					Stopwatch timer = Stopwatch.StartNew();
-					ParallelOptions options = new() { MaxDegreeOfParallelism = MaxFileConcurrency, CancellationToken = cancellationToken };
+					ParallelOptions options = new() { MaxDegreeOfParallelism = _options.MaxFileConcurrency, CancellationToken = cancellationToken };
 					await Parallel.ForEachAsync(filesToCopy, options, (fileToCopy, ct) =>
 					{
 						CopyFileWithinWorkspace(fileToCopy, transaction._filesToSync);
@@ -1915,15 +1949,15 @@ namespace EpicGames.Perforce.Managed
 				if (RuntimeInformation.IsOSPlatform(OSPlatform.Windows))
 				{
 					long freeSpace = new DriveInfo(Path.GetPathRoot(_baseDir.FullName)!).AvailableFreeSpace;
-					if (freeSpace - syncSize < MinScratchSpace)
+					if (freeSpace - syncSize < _options.MinScratchSpace)
 					{
-						throw new InsufficientSpaceException($"Not enough space to sync new files (free space: {freeSpace / (1024.0 * 1024.0):n1}mb, sync size: {syncSize / (1024.0 * 1024.0):n1}mb, min scratch space: {MinScratchSpace / (1024.0 * 1024.0):n1}mb)");
+						throw new InsufficientSpaceException($"Not enough space to sync new files (free space: {freeSpace / (1024.0 * 1024.0):n1}mb, sync size: {syncSize / (1024.0 * 1024.0):n1}mb, min scratch space: {_options.MinScratchSpace / (1024.0 * 1024.0):n1}mb)");
 					}
 				}
 
 				// Sync all the files
 				using (Trace("SyncFiles"))
-				using (ILoggerProgress status = _logger.BeginProgressScope(String.Format("Syncing {0} {1} using {2} threads...", filesToSync.Length, (filesToSync.Length == 1) ? "file" : "files", NumParallelSyncThreads)))
+				using (ILoggerProgress status = _logger.BeginProgressScope(String.Format("Syncing {0} {1} using {2} threads...", filesToSync.Length, (filesToSync.Length == 1) ? "file" : "files", _options.NumParallelSyncThreads)))
 				{
 					Stopwatch timer = Stopwatch.StartNew();
 
@@ -1963,7 +1997,7 @@ namespace EpicGames.Perforce.Managed
 						while (tasks.Count > 0 || nextBatchIdx < batches.Count)
 						{
 							// Create new tasks
-							while (tasks.Count < NumParallelSyncThreads && nextBatchIdx < batches.Count)
+							while (tasks.Count < _options.NumParallelSyncThreads && nextBatchIdx < batches.Count)
 							{
 								(int batchBeginIdx, int batchEndIdx) = batches[nextBatchIdx];
 
@@ -1982,7 +2016,7 @@ namespace EpicGames.Perforce.Managed
 							(int beginIdx, int endIdx) = batches[batchIdx];
 
 							int[] indexesToUpdate = Enumerable.Range(beginIdx, endIdx - beginIdx).ToArray();
-							ParallelOptions options = new() { MaxDegreeOfParallelism = MaxFileConcurrency, CancellationToken = cancellationToken };
+							ParallelOptions options = new() { MaxDegreeOfParallelism = _options.MaxFileConcurrency, CancellationToken = cancellationToken };
 							await Parallel.ForEachAsync(indexesToUpdate, options, (idx, ct) =>
 							{
 								filesToSync[idx]._workspaceFile.UpdateMetadata();
@@ -2050,7 +2084,7 @@ namespace EpicGames.Perforce.Managed
 					}
 				}
 
-				if (_useHaveTable)
+				if (_options.UseHaveTable)
 				{
 					using PerforceConnection clientWithFileList = new (client.Settings, client.Logger);
 					clientWithFileList.GlobalOptions.Add($"-x\"{syncFileName}\"");
