Horde: Move unused files to the cache during a workspace sync, allowing them to be reused for replacement files.

diff --git a/Engine/Source/Programs/Horde/Horde/Commands/StorageCommandBase.cs b/Engine/Source/Programs/Horde/Horde/Commands/StorageCommandBase.cs
--- a/Engine/Source/Programs/Horde/Horde/Commands/StorageCommandBase.cs
+++ b/Engine/Source/Programs/Horde/Horde/Commands/StorageCommandBase.cs
@@ -1,10 +1,8 @@
 // Copyright Epic Games, Inc. All Rights Reserved.
 
-using System.Net.Http.Headers;
 using EpicGames.Core;
+using EpicGames.Horde;
 using EpicGames.Horde.Storage;
-using EpicGames.Horde.Storage.Backends;
-using EpicGames.Horde.Storage.Bundles;
 using Microsoft.Extensions.Logging;
 
 namespace Horde.Commands
@@ -12,13 +10,15 @@ namespace Horde.Commands
 	/// <summary>
 	/// Base class for commands that require a configured storage client
 	/// </summary>
-	abstract class StorageCommandBase : Command
+	abstract class StorageCommandBase : Command, IDisposable
 	{
+		HordeHttpClient? _hordeHttpClient;
+
 		/// <summary>
 		/// Namespace to use
 		/// </summary>
 		[CommandLine("-Namespace=", Description = "Namespace for data to manipulate")]
-		public string Namespace { get; set; } = "default";
+		public NamespaceId Namespace { get; set; } = new NamespaceId("default");
 
 		/// <summary>
 		/// Base URI to upload to
@@ -39,6 +39,12 @@ namespace Horde.Commands
 			StorageCache = storageCache;
 		}
 
+		/// <inheritdoc/>
+		public void Dispose()
+		{
+			_hordeHttpClient?.Dispose();
+		}
+
 		/// <summary>
 		/// Creates a new client instance
 		/// </summary>
@@ -46,29 +52,15 @@ namespace Horde.Commands
 		/// <param name="cancellationToken"></param>
 		public async Task<IStorageClient> CreateStorageClientAsync(ILogger logger, CancellationToken cancellationToken = default)
 		{
-			Uri? server = await Settings.GetServerAsync(cancellationToken);
-			if (server == null)
-			{
-				throw new Exception("No server is configured. Run 'horde login -server=...' to set up.");
-			}
-
-			string? token = await Settings.GetAccessTokenAsync(logger, cancellationToken);
-			if (token == null)
-			{
-				throw new Exception("Unable to log in to server.");
-			}
-
+			_hordeHttpClient ??= await Settings.GetHttpClientAsync(logger, cancellationToken);
 			if (String.IsNullOrEmpty(Path))
 			{
-				Path = $"api/v1/storage/{Namespace}/";
+				return _hordeHttpClient.CreateStorageClient(Namespace, StorageCache);
 			}
-			else if(!Path.EndsWith("/", StringComparison.Ordinal))
+			else
 			{
-				Path += "/";
+				return _hordeHttpClient.CreateStorageClient(Path, StorageCache);
 			}
-
-			server = new Uri(server, Path);
-			return new HttpStorageClient(server, token, StorageCache, logger);
 		}
 	}
 }
diff --git a/Engine/Source/Programs/Horde/Horde/Commands/Workspace/WorkspaceSync.cs b/Engine/Source/Programs/Horde/Horde/Commands/Workspace/WorkspaceSync.cs
--- a/Engine/Source/Programs/Horde/Horde/Commands/Workspace/WorkspaceSync.cs
+++ b/Engine/Source/Programs/Horde/Horde/Commands/Workspace/WorkspaceSync.cs
@@ -76,8 +76,8 @@ namespace Horde.Commands.Workspace
 			Stopwatch timer = Stopwatch.StartNew();
 			logger.LogInformation("Syncing into layer '{LayerId}'...", LayerId);
 
-			DirectoryNode node = await handle.ReadNodeAsync<DirectoryNode>();
-			await node.CopyToDirectoryAsync(workspace.RootDir.ToDirectoryInfo(), logger, CancellationToken.None);
+			DirectoryNode contents = await handle.ReadNodeAsync<DirectoryNode>();
+			await workspace.SyncAsync(LayerId, contents, cancellationToken);
 			await workspace.SaveAsync(cancellationToken);
 
 			logger.LogInformation("Elapsed: {Time}s", timer.Elapsed.TotalSeconds);
diff --git a/Engine/Source/Programs/Shared/EpicGames.Horde/Storage/Nodes/ChunkedDataNode.cs b/Engine/Source/Programs/Shared/EpicGames.Horde/Storage/Nodes/ChunkedDataNode.cs
--- a/Engine/Source/Programs/Shared/EpicGames.Horde/Storage/Nodes/ChunkedDataNode.cs
+++ b/Engine/Source/Programs/Shared/EpicGames.Horde/Storage/Nodes/ChunkedDataNode.cs
@@ -92,6 +92,11 @@ namespace EpicGames.Horde.Storage.Nodes
 	[NodeType("{B27AFB68-9E20-4A4B-A4D8-788A4098D439}", 1)]
 	public sealed class LeafChunkedDataNode : ChunkedDataNode
 	{
+		/// <summary>
+		/// Static accessor for the blob type
+		/// </summary>
+		public static BlobType BlobType { get; } = Node.GetNodeType<LeafChunkedDataNode>();
+
 		/// <summary>
 		/// Data for this node
 		/// </summary>
@@ -283,6 +288,11 @@ namespace EpicGames.Horde.Storage.Nodes
 	[NodeType("{F4DEDDBC-70CB-4C7A-8347-F011AFCCCDB9}", 1)]
 	public class InteriorChunkedDataNode : ChunkedDataNode
 	{
+		/// <summary>
+		/// Static accessor for the blob type
+		/// </summary>
+		public static BlobType BlobType { get; } = Node.GetNodeType<InteriorChunkedDataNode>();
+
 		/// <summary>
 		/// Child nodes
 		/// </summary>
diff --git a/Engine/Source/Programs/Shared/EpicGames.Horde/Storage/Workspace.cs b/Engine/Source/Programs/Shared/EpicGames.Horde/Storage/Workspace.cs
--- a/Engine/Source/Programs/Shared/EpicGames.Horde/Storage/Workspace.cs
+++ b/Engine/Source/Programs/Shared/EpicGames.Horde/Storage/Workspace.cs
@@ -6,6 +6,7 @@ using System.Diagnostics;
 using System.Diagnostics.CodeAnalysis;
 using System.IO;
 using System.Linq;
+using System.Text;
 using System.Threading;
 using System.Threading.Tasks;
 using EpicGames.Core;
@@ -35,6 +36,29 @@ namespace EpicGames.Horde.Storage
 				Name = name;
 			}
 
+			public void AddFile(FileState fileState)
+			{
+				int index = Files.BinarySearch(x => x.Name, fileState.Name);
+				if (index >= 0)
+				{
+					throw new InvalidOperationException($"File {fileState.Name} already exists");
+				}
+
+				fileState.Parent = this;
+				Files.Insert(~index, fileState);
+			}
+
+			public void RemoveFile(FileState fileState)
+			{
+				int index = Files.BinarySearch(x => x.Name, fileState.Name);
+				if (index >= 0 && Files[index] == fileState)
+				{
+					Files.RemoveAt(index);
+				}
+			}
+
+			public bool ContainsFile(Utf8String name) => TryGetFile(name, out _);
+
 			public bool TryGetFile(Utf8String name, [NotNullWhen(true)] out FileState? fileState)
 			{
 				int index = Files.BinarySearch(x => x.Name, name);
@@ -95,21 +119,34 @@ namespace EpicGames.Horde.Storage
 				}
 			}
 
-			public Utf8String GetPath()
+			public string GetPath()
 			{
-				Utf8StringBuilder builder = new Utf8StringBuilder();
-				GetPath(builder);
-				return builder.ToUtf8String();
+				StringBuilder builder = new StringBuilder();
+				AppendPath(builder);
+				return builder.ToString();
 			}
 
-			public void GetPath(Utf8StringBuilder builder)
+			public void AppendPath(StringBuilder builder)
 			{
 				if (Parent != null)
 				{
-					Parent.GetPath(builder);
+					Parent.AppendPath(builder);
 				}
-				builder.Append(Name);
-				builder.Append((byte)'/');
+				if (!Name.IsEmpty)
+				{
+					builder.Append(Name);
+				}
+				if (builder.Length == 0 || builder[^1] != Path.DirectorySeparatorChar)
+				{
+					builder.Append(Path.DirectorySeparatorChar);
+				}
+			}
+
+			public DirectoryReference GetDirectoryReference(DirectoryReference baseDir)
+			{
+				StringBuilder builder = new StringBuilder(baseDir.FullName);
+				AppendPath(builder);
+				return new DirectoryReference(builder.ToString(), DirectoryReference.Sanitize.None);
 			}
 
 			public void Read(IMemoryReader reader)
@@ -164,15 +201,15 @@ namespace EpicGames.Horde.Storage
 				writer.WriteUnsignedVarInt(LayerFlags);
 			}
 
-			public override string ToString() => GetPath().ToString();
+			public override string ToString() => GetPath();
 		}
 
 		// Tracked state of a file
 		[DebuggerDisplay("{Name}")]
 		class FileState
 		{
-			public DirectoryState Parent { get; }
-			public Utf8String Name { get; }
+			public DirectoryState Parent { get; set; }
+			public Utf8String Name { get; private set; }
 			public long Length { get; private set; }
 			public long LastModifiedTimeUtc { get; private set; }
 			public IoHash Hash { get; set; }
@@ -193,7 +230,19 @@ namespace EpicGames.Horde.Storage
 				LayerFlags = reader.ReadUnsignedVarInt();
 			}
 
-			public bool Modified(FileInfo fileInfo) => Length != fileInfo.Length || LastModifiedTimeUtc != fileInfo.LastWriteTimeUtc.Ticks;
+			public void MoveTo(DirectoryState newParent, Utf8String newName)
+			{
+				Parent.RemoveFile(this);
+				Name = newName;
+				newParent.AddFile(this);
+			}
+
+			public void Delete()
+			{
+				Parent.RemoveFile(this);
+			}
+
+			public bool IsModified(FileInfo fileInfo) => Length != fileInfo.Length || LastModifiedTimeUtc != fileInfo.LastWriteTimeUtc.Ticks;
 
 			public void Update(FileInfo fileInfo)
 			{
@@ -210,20 +259,27 @@ namespace EpicGames.Horde.Storage
 				writer.WriteUnsignedVarInt(LayerFlags);
 			}
 
-			Utf8String GetPath()
+			public string GetPath()
 			{
-				Utf8StringBuilder builder = new Utf8StringBuilder();
-				GetPath(builder);
-				return builder.ToUtf8String();
+				StringBuilder builder = new StringBuilder();
+				AppendPath(builder);
+				return builder.ToString();
 			}
 
-			void GetPath(Utf8StringBuilder builder)
+			public void AppendPath(StringBuilder builder)
 			{
-				Parent.GetPath(builder);
+				Parent.AppendPath(builder);
 				builder.Append(Name);
 			}
 
-			public override string ToString() => GetPath().ToString();
+			public FileReference GetFileReference(DirectoryReference baseDir)
+			{
+				StringBuilder builder = new StringBuilder(baseDir.FullName);
+				AppendPath(builder);
+				return new FileReference(builder.ToString(), FileReference.Sanitize.None);
+			}
+
+			public override string ToString() => GetPath();
 		}
 
 		// Collates lists of files and chunks with a particular hash
@@ -233,10 +289,16 @@ namespace EpicGames.Horde.Storage
 			public int Index { get; set; }
 
 			public IoHash Hash { get; }
+			public long Length { get; }
+
 			public List<FileState> Files { get; } = new List<FileState>();
 			public List<ChunkInfo> Chunks { get; } = new List<ChunkInfo>();
 
-			public HashInfo(IoHash hash) => Hash = hash;
+			public HashInfo(IoHash hash, long length)
+			{
+				Hash = hash;
+				Length = length;
+			}
 		}
 
 		// Hashed chunk within another hashed object
@@ -296,12 +358,13 @@ namespace EpicGames.Horde.Storage
 
 		readonly DirectoryReference _rootDir;
 		readonly FileReference _stateFile;
-		readonly DirectoryState _root;
+		readonly DirectoryState _rootDirState;
 		readonly List<LayerState> _layers;
 		readonly Dictionary<IoHash, HashInfo> _hashes = new Dictionary<IoHash, HashInfo>();
 		readonly ILogger _logger;
 
 		const string HordeDirName = ".horde";
+		const string CacheDirName = "cache";
 		const string StateFileName = "contents.dat";
 
 		/// <summary>
@@ -324,7 +387,7 @@ namespace EpicGames.Horde.Storage
 		{
 			_rootDir = rootDir;
 			_stateFile = stateFile;
-			_root = new DirectoryState(null, Utf8String.Empty);
+			_rootDirState = new DirectoryState(null, Utf8String.Empty);
 			_layers = new List<LayerState> { new LayerState(WorkspaceLayerId.Default, 1) };
 			_logger = logger;
 		}
@@ -401,7 +464,7 @@ namespace EpicGames.Horde.Storage
 
 		void Read(IMemoryReader reader)
 		{
-			_root.Read(reader);
+			_rootDirState.Read(reader);
 			reader.ReadList(_layers, () => new LayerState(reader));
 
 			// Read the hash lookup
@@ -414,7 +477,8 @@ namespace EpicGames.Horde.Storage
 			for (int idx = 0; idx < numHashes; idx++)
 			{
 				IoHash hash = reader.ReadIoHash();
-				hashInfoArray[idx] = new HashInfo(hash);
+				long length = (long)reader.ReadUnsignedVarInt();
+				hashInfoArray[idx] = new HashInfo(hash, length);
 				_hashes.Add(hash, hashInfoArray[idx]);
 			}
 			for (int idx = 0; idx < numHashes; idx++)
@@ -426,7 +490,7 @@ namespace EpicGames.Horde.Storage
 
 		void Write(IMemoryWriter writer)
 		{
-			_root.Write(writer);
+			_rootDirState.Write(writer);
 			writer.WriteList(_layers, x => x.Write(writer));
 
 			// Write the hash lookup
@@ -436,6 +500,7 @@ namespace EpicGames.Horde.Storage
 			foreach (HashInfo hashInfo in _hashes.Values)
 			{
 				writer.WriteIoHash(hashInfo.Hash);
+				writer.WriteUnsignedVarInt((ulong)hashInfo.Length);
 				hashInfo.Index = nextIndex++;
 			}
 			foreach (HashInfo hashInfo in _hashes.Values)
@@ -481,7 +546,7 @@ namespace EpicGames.Horde.Storage
 			if (layerIdx > 0) // Note: Excluding default layer at index 0
 			{
 				LayerState layer = _layers[layerIdx];
-				if ((_root.LayerFlags & layer.Flag) != 0)
+				if ((_rootDirState.LayerFlags & layer.Flag) != 0)
 				{
 					throw new InvalidOperationException($"Workspace still contains files for layer {layerId}");
 				}
@@ -493,6 +558,8 @@ namespace EpicGames.Horde.Storage
 
 		#endregion
 
+		#region Syncing
+
 		/// <summary>
 		/// Syncs a layer to the given contents
 		/// </summary>
@@ -507,10 +574,13 @@ namespace EpicGames.Horde.Storage
 				throw new InvalidOperationException($"Layer '{layerId}' does not exist");
 			}
 
-			await SyncDirectoryAsync(_rootDir, _root, contents, layerState.Flag, cancellationToken);
+			DirectoryState hordeDirState = _rootDirState.FindOrAddDirectory(HordeDirName);
+			DirectoryState cacheDirState = hordeDirState.FindOrAddDirectory(CacheDirName);
+
+			await SyncDirectoryAsync(_rootDir, _rootDirState, contents, layerState.Flag, cacheDirState, cancellationToken);
 		}
 
-		async Task SyncDirectoryAsync(DirectoryReference dirPath, DirectoryState dirState, DirectoryNode? dirNode, ulong flag, CancellationToken cancellationToken)
+		async Task SyncDirectoryAsync(DirectoryReference dirRef, DirectoryState dirState, DirectoryNode? dirNode, ulong flag, DirectoryState cacheDirState, CancellationToken cancellationToken)
 		{
 			// Remove any directories that no longer exist
 			for (int subDirIdx = 0; subDirIdx < dirState.Directories.Count; subDirIdx++)
@@ -520,8 +590,8 @@ namespace EpicGames.Horde.Storage
 				{
 					if (dirNode == null || !dirNode.TryGetDirectoryEntry(subDirState.Name, out _))
 					{
-						DirectoryReference subDirPath = DirectoryReference.Combine(dirPath, subDirState.Name.ToString());
-						await SyncDirectoryAsync(subDirPath, subDirState, null, flag, cancellationToken);
+						DirectoryReference subDirPath = DirectoryReference.Combine(dirRef, subDirState.Name.ToString());
+						await SyncDirectoryAsync(subDirPath, subDirState, null, flag, cacheDirState, cancellationToken);
 					}
 				}
 			}
@@ -532,10 +602,15 @@ namespace EpicGames.Horde.Storage
 				FileState fileState = dirState.Files[fileIdx];
 				if ((fileState.LayerFlags & flag) != 0)
 				{
-					if (dirNode == null || !dirNode.TryGetFileEntry(fileState.Name, out _))
+					if (dirNode == null || !dirNode.TryGetFileEntry(fileState.Name, out FileEntry? entry) || entry.Hash != fileState.Hash)
 					{
-						FileReference filePath = FileReference.Combine(dirPath, fileState.Name.ToString());
-						await SyncFileAsync(filePath, fileState, null, flag, cancellationToken);
+						fileState.LayerFlags &= ~flag;
+
+						if (fileState.LayerFlags == 0)
+						{
+							MoveFileToCache(fileState, cacheDirState);
+							fileIdx--;
+						}
 					}
 				}
 			}
@@ -550,16 +625,16 @@ namespace EpicGames.Horde.Storage
 			// Add files for this directory
 			if (dirNode != null)
 			{
-				DirectoryReference.CreateDirectory(dirPath);
+				DirectoryReference.CreateDirectory(dirRef);
 
 				// Update directories
 				foreach (DirectoryEntry subDirEntry in dirNode.Directories)
 				{
-					DirectoryReference subDirPath = DirectoryReference.Combine(dirPath, subDirEntry.Name.ToString());
+					DirectoryReference subDirPath = DirectoryReference.Combine(dirRef, subDirEntry.Name.ToString());
 					DirectoryState subDirState = dirState.FindOrAddDirectory(subDirEntry.Name);
 
 					DirectoryNode subDirNode = await subDirEntry.ExpandAsync(cancellationToken);
-					await SyncDirectoryAsync(subDirPath, subDirState, subDirNode, flag, cancellationToken);
+					await SyncDirectoryAsync(subDirPath, subDirState, subDirNode, flag, cacheDirState, cancellationToken);
 
 					dirState.LayerFlags |= flag;
 				}
@@ -567,7 +642,7 @@ namespace EpicGames.Horde.Storage
 				// Update files
 				foreach (FileEntry fileEntry in dirNode.Files)
 				{
-					FileReference filePath = FileReference.Combine(dirPath, fileEntry.Name.ToString());
+					FileReference filePath = FileReference.Combine(dirRef, fileEntry.Name.ToString());
 
 					FileState fileState = dirState.FindOrAddFile(fileEntry.Name);
 					await SyncFileAsync(filePath, fileState, fileEntry, flag, cancellationToken);
@@ -579,7 +654,7 @@ namespace EpicGames.Horde.Storage
 			// Delete the directory if it's no longer needed
 			if (dirState.LayerFlags == 0 && dirState.Parent != null)
 			{
-				FileUtils.ForceDeleteDirectory(dirPath);
+				FileUtils.ForceDeleteDirectory(dirRef);
 			}
 		}
 
@@ -597,10 +672,18 @@ namespace EpicGames.Horde.Storage
 			else if (fileState.Hash != fileEntry.Hash)
 			{
 				FileInfo fileInfo = filePath.ToFileInfo();
+				if (fileInfo.Exists && (fileInfo.Attributes & FileAttributes.ReadOnly) != 0)
+				{
+					fileInfo.Attributes &= ~FileAttributes.ReadOnly;
+				}
 
 				_logger.LogInformation("Updating {File} to {Hash}", fileInfo, fileEntry.Hash);
-				ChunkedDataNode fileNode = await fileEntry.ExpandAsync(cancellationToken);
-				await fileNode.CopyToFileAsync(fileInfo, cancellationToken);
+
+				using (FileStream stream = fileInfo.Open(FileMode.Create, FileAccess.Write, FileShare.Read))
+				{
+					await ExtractDataAsync(fileEntry, stream, cancellationToken);
+				}
+
 				fileInfo.Refresh();
 
 				fileState.LayerFlags |= flag;
@@ -611,12 +694,129 @@ namespace EpicGames.Horde.Storage
 			}
 		}
 
+		async Task ExtractDataAsync(NodeRef<ChunkedDataNode> nodeRef, Stream outputStream, CancellationToken cancellationToken)
+		{
+			BlobType blobType = await nodeRef.Handle.GetTypeAsync(cancellationToken);
+			if (blobType.Guid == LeafChunkedDataNode.BlobType.Guid)
+			{
+				await ExtractLeafDataAsync(nodeRef, outputStream, cancellationToken);
+			}
+			else if (blobType.Guid == InteriorChunkedDataNode.BlobType.Guid)
+			{
+				InteriorChunkedDataNode interiorNode = (InteriorChunkedDataNode)await nodeRef.ExpandAsync(cancellationToken);
+				foreach (NodeRef<ChunkedDataNode> childNodeRef in interiorNode.Children)
+				{
+					await ExtractDataAsync(childNodeRef, outputStream, cancellationToken);
+				}
+			}
+			else
+			{
+				throw new InvalidDataException($"Invalid node type {blobType}");
+			}
+		}
+
+		async Task ExtractLeafDataAsync(NodeRef<ChunkedDataNode> nodeRef, Stream outputStream, CancellationToken cancellationToken)
+		{
+			// Try to copy cached data for this node
+			HashInfo? hashInfo;
+			if (_hashes.TryGetValue(nodeRef.Handle.Hash, out hashInfo))
+			{
+				if (await TryCopyCachedDataAsync(hashInfo, 0, hashInfo.Length, outputStream, cancellationToken))
+				{
+					return;
+				}
+			}
+
+			// Otherwise 
+			BlobData nodeData = await nodeRef.Handle.ReadAsync(cancellationToken);
+			await LeafChunkedDataNode.CopyToStreamAsync(nodeData, outputStream, cancellationToken);
+		}
+
+		async Task<bool> TryCopyCachedDataAsync(HashInfo hashInfo, long offset, long length, Stream outputStream, CancellationToken cancellationToken)
+		{
+			foreach (FileState file in hashInfo.Files)
+			{
+				if (await TryCopyChunkAsync(file, offset, length, outputStream, cancellationToken))
+				{
+					return true;
+				}
+			}
+			foreach (ChunkInfo chunk in hashInfo.Chunks)
+			{
+				if (await TryCopyCachedDataAsync(chunk.WithinHashInfo, chunk.Offset, chunk.Length, outputStream, cancellationToken))
+				{
+					return true;
+				}
+			}
+			return false;
+		}
+
+		async Task<bool> TryCopyChunkAsync(FileState file, long offset, long length, Stream outputStream, CancellationToken cancellationToken)
+		{
+			FileInfo fileInfo = GetFileInfo(file);
+			if (fileInfo.Exists)
+			{
+				using FileStream inputStream = fileInfo.OpenRead();
+				inputStream.Seek(offset, SeekOrigin.Begin);
+
+				byte[] tempBuffer = new byte[65536];
+				while (length > 0)
+				{
+					int readSize = await inputStream.ReadAsync(tempBuffer.AsMemory(0, (int)Math.Min(length, tempBuffer.Length)), cancellationToken);
+					await outputStream.WriteAsync(tempBuffer.AsMemory(0, readSize), cancellationToken);
+				}
+			}
+			return true;
+		}
+
+		#endregion
+
+		#region Verify
+
+		/// <summary>
+		/// Checks that all files within the workspace have the correct hash
+		/// </summary>
+		/// <param name="cancellationToken">Cancellation token for the operation</param>
+		public async Task VerifyAsync(CancellationToken cancellationToken = default)
+		{
+			await VerifyAsync(_rootDirState, cancellationToken);
+		}
+
+		async Task VerifyAsync(DirectoryState dirState, CancellationToken cancellationToken)
+		{
+			foreach (DirectoryState subDirState in dirState.Directories)
+			{
+				await VerifyAsync(subDirState, cancellationToken);
+			}
+			foreach (FileState fileState in dirState.Files)
+			{
+				FileInfo fileInfo = GetFileInfo(fileState);
+				if (fileState.IsModified(fileInfo))
+				{
+					throw new InvalidDataException($"File {fileInfo.FullName} has been modified");
+				}
+
+				IoHash hash;
+				using (FileStream stream = GetFileInfo(fileState).OpenRead())
+				{
+					hash = await IoHash.ComputeAsync(stream, cancellationToken);
+				}
+
+				if (hash != fileState.Hash)
+				{
+					throw new InvalidDataException($"Hash for {fileInfo.FullName} was {hash}, expected {fileState.Hash}");
+				}
+			}
+		}
+
+		#endregion
+
 		void AddFileToHashLookup(FileState file)
 		{
 			HashInfo? hashInfo;
 			if (!_hashes.TryGetValue(file.Hash, out hashInfo))
 			{
-				hashInfo = new HashInfo(file.Hash);
+				hashInfo = new HashInfo(file.Hash, file.Length);
 				_hashes.Add(file.Hash, hashInfo);
 			}
 			hashInfo.Files.Add(file);
@@ -630,5 +830,34 @@ namespace EpicGames.Horde.Storage
 				hashInfo.Files.Remove(file);
 			}
 		}
+
+		void MoveFileToCache(FileState fileState, DirectoryState cacheDirState)
+		{
+			Utf8String name = fileState.Hash.ToUtf8String();
+			for (int idx = 0; idx < 2; idx++)
+			{
+				cacheDirState = cacheDirState.FindOrAddDirectory(name.Slice(idx * 2, 2));
+			}
+
+			FileReference fileRef = fileState.GetFileReference(_rootDir);
+			if (cacheDirState.ContainsFile(name))
+			{
+				fileState.Delete();
+				FileUtils.ForceDeleteFile(fileRef);
+			}
+			else
+			{
+				fileState.MoveTo(cacheDirState, name);
+
+				FileReference cachedFileRef = fileState.GetFileReference(_rootDir);
+				DirectoryReference.CreateDirectory(cachedFileRef.Directory);
+				FileReference.Move(fileRef, cachedFileRef);
+			}
+		}
+
+		FileInfo GetFileInfo(FileState file)
+		{
+			return file.GetFileReference(_rootDir).ToFileInfo();
+		}
 	}
 }
