Horde: Add checks to ensure values written to bundles are not out of range.

diff --git a/Engine/Source/Programs/Shared/EpicGames.Horde/Storage/Bundles/Bundle.cs b/Engine/Source/Programs/Shared/EpicGames.Horde/Storage/Bundles/Bundle.cs
--- a/Engine/Source/Programs/Shared/EpicGames.Horde/Storage/Bundles/Bundle.cs
+++ b/Engine/Source/Programs/Shared/EpicGames.Horde/Storage/Bundles/Bundle.cs
@@ -307,6 +307,10 @@ namespace EpicGames.Horde.Storage.Bundles
 
 		static void WriteSectionHeader(IMemoryWriter writer, BundleSectionType type, int length)
 		{
+			if (length >= 0x1000000)
+			{
+				throw new InvalidDataException($"Section header exceeds maximum allowed size (type: {type}, length: {length:n0});
+			}
 			writer.WriteInt32((int)type | (length << 8));
 		}
 
@@ -803,7 +807,7 @@ namespace EpicGames.Horde.Storage.Bundles
 		readonly BundlePacket[] _packets;
 
 		/// <inheritdoc/>
-		public int Count => _packets.Length;
+		public int Count => _packets?.Length ?? 0;
 
 		/// <summary>
 		/// Constructor
@@ -1182,6 +1186,15 @@ namespace EpicGames.Horde.Storage.Bundles
 		/// </summary>
 		public void CopyTo(Span<byte> data)
 		{
+			if (ImportIdx > 0xffff)
+			{
+				throw new InvalidDataException($"Import index is out of range ({ImportIdx})");
+			}
+			if (NodeIdx > 0xffff)
+			{
+				throw new InvalidDataException($"Node index is out of range ({NodeIdx})");
+			}
+
 			BinaryPrimitives.WriteInt16LittleEndian(data, (short)ImportIdx);
 			BinaryPrimitives.WriteUInt16LittleEndian(data.Slice(2), (ushort)NodeIdx);
 			Hash.CopyTo(data.Slice(4));
