Missing quote.

diff --git a/Engine/Source/Programs/Shared/EpicGames.Horde/Storage/Bundles/Bundle.cs b/Engine/Source/Programs/Shared/EpicGames.Horde/Storage/Bundles/Bundle.cs
--- a/Engine/Source/Programs/Shared/EpicGames.Horde/Storage/Bundles/Bundle.cs
+++ b/Engine/Source/Programs/Shared/EpicGames.Horde/Storage/Bundles/Bundle.cs
@@ -309,7 +309,7 @@ namespace EpicGames.Horde.Storage.Bundles
 		{
 			if (length >= 0x1000000)
 			{
-				throw new InvalidDataException($"Section header exceeds maximum allowed size (type: {type}, length: {length:n0});
+				throw new InvalidDataException($"Section header exceeds maximum allowed size (type: {type}, length: {length:n0}");
 			}
 			writer.WriteInt32((int)type | (length << 8));
 		}
