Horde: Missing comment.

diff --git a/Engine/Source/Programs/Shared/EpicGames.Core/MemoryReader.cs b/Engine/Source/Programs/Shared/EpicGames.Core/MemoryReader.cs
--- a/Engine/Source/Programs/Shared/EpicGames.Core/MemoryReader.cs
+++ b/Engine/Source/Programs/Shared/EpicGames.Core/MemoryReader.cs
@@ -249,6 +249,7 @@ namespace EpicGames.Core
 		/// Reads a variable length list
 		/// </summary>
 		/// <param name="reader">Reader to deserialize from</param>
+		/// <param name="list">List to receive the items that were read</param>
 		/// <param name="readItem">Delegate to write an individual item</param>
 		public static void ReadList<T>(this IMemoryReader reader, List<T> list, Func<T> readItem)
 		{
