Horde: Prevent too many exports being serialized into a single bundle section.

diff --git a/Engine/Source/Programs/Shared/EpicGames.Horde/Storage/Bundles/Bundle.cs b/Engine/Source/Programs/Shared/EpicGames.Horde/Storage/Bundles/Bundle.cs
--- a/Engine/Source/Programs/Shared/EpicGames.Horde/Storage/Bundles/Bundle.cs
+++ b/Engine/Source/Programs/Shared/EpicGames.Horde/Storage/Bundles/Bundle.cs
@@ -67,10 +67,20 @@ namespace EpicGames.Horde.Storage.Bundles
 	/// </summary>
 	public class Bundle
 	{
+		/// <summary>
+		/// Maximum size for a section
+		/// </summary>
+		public const int MaxSectionSize = 0x1000000;
+
 		/// <summary>
 		/// Maximum number of exports from a single bundle
 		/// </summary>
-		public const int MaxExports = 60000;
+		public const int MaxExports = 65535;
+
+		/// <summary>
+		/// Maximum number of export refs from a single bundle
+		/// </summary>
+		public const int MaxExportRefs = MaxSectionSize / BundleExportRef.NumBytes;
 
 		/// <summary>
 		/// Header for the bundle
diff --git a/Engine/Source/Programs/Shared/EpicGames.Horde/Storage/Bundles/BundleWriter.cs b/Engine/Source/Programs/Shared/EpicGames.Horde/Storage/Bundles/BundleWriter.cs
--- a/Engine/Source/Programs/Shared/EpicGames.Horde/Storage/Bundles/BundleWriter.cs
+++ b/Engine/Source/Programs/Shared/EpicGames.Horde/Storage/Bundles/BundleWriter.cs
@@ -298,6 +298,9 @@ namespace EpicGames.Horde.Storage.Bundles
 			// Queue of nodes for the current bundle
 			readonly List<PendingNode> _queue = new List<PendingNode>();
 
+			// Number of references in the queue
+			int _queuedRefs = 0;
+
 			// List of packets in the current bundle
 			readonly List<BundlePacket> _packets = new List<BundlePacket>();
 
@@ -350,7 +353,7 @@ namespace EpicGames.Horde.Storage.Bundles
 			}
 
 			// Whether this bundle is full
-			public bool IsFull() => _queue.Count >= Bundle.MaxExports;
+			public bool IsFull() => (_queue.Count + 1000) >= Bundle.MaxExports || (_queuedRefs + 1000) >= Bundle.MaxExportRefs;
 
 			// Whether this bundle has finished writing
 			public bool IsComplete() => CompleteTask.IsCompleted;
@@ -433,6 +436,7 @@ namespace EpicGames.Horde.Storage.Bundles
 
 				PendingNode pendingNode = new PendingNode(_treeReader, nodeKey, _currentPacketIdx, offset, (int)size, refs.ConvertAll(x => (BundleNodeHandle)x), this);
 				_queue.Add(pendingNode);
+				_queuedRefs += refs.Count;
 				_nodeKeyToInfo.Add(nodeKey, pendingNode);
 
 				return pendingNode;
