Horde: Fix items moved from workspace cache not being correctly added back into the hash lookup.

diff --git a/Engine/Source/Programs/Shared/EpicGames.Horde/Storage/Workspace.cs b/Engine/Source/Programs/Shared/EpicGames.Horde/Storage/Workspace.cs
--- a/Engine/Source/Programs/Shared/EpicGames.Horde/Storage/Workspace.cs
+++ b/Engine/Source/Programs/Shared/EpicGames.Horde/Storage/Workspace.cs
@@ -367,6 +367,21 @@ namespace EpicGames.Horde.Storage
 		const string CacheDirName = "cache";
 		const string StateFileName = "contents.dat";
 
+		/// <summary>
+		/// Flag for the default layer
+		/// </summary>
+		const ulong DefaultLayerFlag = 1;
+
+		/// <summary>
+		/// Flag for the cache layer
+		/// </summary>
+		const ulong CacheLayerFlag = 2;
+
+		/// <summary>
+		/// Flags for user layers
+		/// </summary>
+		const ulong ReservedLayerFlags = DefaultLayerFlag | CacheLayerFlag;
+
 		/// <summary>
 		/// Root directory for the workspace
 		/// </summary>
@@ -486,6 +501,9 @@ namespace EpicGames.Horde.Storage
 				HashInfo hashInfo = hashInfoArray[idx];
 				reader.ReadList(hashInfoArray[idx].Chunks, () => new ChunkInfo(reader, hashInfoArray));
 			}
+
+			// Add all the files to the hash lookup
+			AddDirToHashLookup(_rootDirState);
 		}
 
 		void Write(IMemoryWriter writer)
@@ -522,7 +540,7 @@ namespace EpicGames.Horde.Storage
 				throw new InvalidOperationException($"Layer {id} already exists");
 			}
 
-			ulong flags = 0;
+			ulong flags = ReservedLayerFlags;
 			for (int idx = 0; idx < _layers.Count; idx++)
 			{
 				flags |= _layers[idx].Flag;
@@ -575,7 +593,10 @@ namespace EpicGames.Horde.Storage
 			}
 
 			DirectoryState hordeDirState = _rootDirState.FindOrAddDirectory(HordeDirName);
+			hordeDirState.LayerFlags |= CacheLayerFlag;
+
 			DirectoryState cacheDirState = hordeDirState.FindOrAddDirectory(CacheDirName);
+			cacheDirState.LayerFlags |= CacheLayerFlag;
 
 			await SyncDirectoryAsync(_rootDir, _rootDirState, contents, layerState.Flag, cacheDirState, cancellationToken);
 		}
@@ -592,6 +613,12 @@ namespace EpicGames.Horde.Storage
 					{
 						DirectoryReference subDirPath = DirectoryReference.Combine(dirRef, subDirState.Name.ToString());
 						await SyncDirectoryAsync(subDirPath, subDirState, null, flag, cacheDirState, cancellationToken);
+
+						if (subDirState.LayerFlags == 0)
+						{
+							dirState.Directories.Remove(subDirState);
+							subDirIdx--;
+						}
 					}
 				}
 			}
@@ -615,10 +642,6 @@ namespace EpicGames.Horde.Storage
 				}
 			}
 
-			// Actually delete all the unreferenced directories
-			dirState.Directories.RemoveAll(x => x.LayerFlags == 0);
-			dirState.Files.RemoveAll(x => x.LayerFlags == 0);
-
 			// Clear out the layer flag for this directory. It'll be added back if we add/reuse files below.
 			dirState.LayerFlags &= ~flag;
 
@@ -642,11 +665,7 @@ namespace EpicGames.Horde.Storage
 				// Update files
 				foreach (FileEntry fileEntry in dirNode.Files)
 				{
-					FileReference filePath = FileReference.Combine(dirRef, fileEntry.Name.ToString());
-
-					FileState fileState = dirState.FindOrAddFile(fileEntry.Name);
-					await SyncFileAsync(filePath, fileState, fileEntry, flag, cancellationToken);
-
+					await SyncFileAsync(dirRef, dirState, fileEntry, flag, cancellationToken);
 					dirState.LayerFlags |= flag;
 				}
 			}
@@ -658,39 +677,42 @@ namespace EpicGames.Horde.Storage
 			}
 		}
 
-		async Task SyncFileAsync(FileReference filePath, FileState fileState, FileEntry? fileEntry, ulong flag, CancellationToken cancellationToken)
+		async Task SyncFileAsync(DirectoryReference dirRef, DirectoryState dirState, FileEntry fileEntry, ulong flag, CancellationToken cancellationToken)
 		{
-			if (fileEntry == null)
+			FileState? fileState;
+			if (dirState.TryGetFile(fileEntry.Name, out fileState))
 			{
-				fileState.LayerFlags &= ~flag;
-				if (fileState.LayerFlags == 0)
+				if (fileState.Hash == fileEntry.Hash)
+				{
+					fileState.LayerFlags |= flag;
+				}
+				else
 				{
-					FileUtils.ForceDeleteFile(filePath);
-					RemoveFileFromHashLookup(fileState);
+					throw new InvalidOperationException($"Conflicting hash for file {fileState.GetFileReference(_rootDir)}");
 				}
 			}
-			else if (fileState.Hash != fileEntry.Hash)
+			else
 			{
-				FileInfo fileInfo = filePath.ToFileInfo();
-				if (fileInfo.Exists && (fileInfo.Attributes & FileAttributes.ReadOnly) != 0)
+				FileInfo fileInfo = FileReference.Combine(dirRef, fileEntry.Name.ToString()).ToFileInfo();
+
+				fileState = TryMoveCachedDataAsync(fileEntry.Hash, dirState, fileEntry.Name);
+				if(fileState == null)
 				{
-					fileInfo.Attributes &= ~FileAttributes.ReadOnly;
-				}
+					fileState = dirState.FindOrAddFile(fileEntry.Name);
 
-				_logger.LogInformation("Updating {File} to {Hash}", fileInfo, fileEntry.Hash);
+					_logger.LogInformation("Updating {File} to {Hash}", fileInfo, fileEntry.Hash);
+					using (FileStream stream = fileInfo.Open(FileMode.Create, FileAccess.Write, FileShare.Read))
+					{
+						await ExtractDataAsync(fileEntry, stream, cancellationToken);
+					}
 
-				using (FileStream stream = fileInfo.Open(FileMode.Create, FileAccess.Write, FileShare.Read))
-				{
-					await ExtractDataAsync(fileEntry, stream, cancellationToken);
+					fileState.Hash = fileEntry.Hash;
+					AddFileToHashLookup(fileState);
 				}
 
-				fileInfo.Refresh();
-
 				fileState.LayerFlags |= flag;
-				fileState.Hash = fileEntry.Hash;
+				fileInfo.Refresh();
 				fileState.Update(fileInfo);
-
-				AddFileToHashLookup(fileState);
 			}
 		}
 
@@ -732,6 +754,24 @@ namespace EpicGames.Horde.Storage
 			await LeafChunkedDataNode.CopyToStreamAsync(nodeData, outputStream, cancellationToken);
 		}
 
+		FileState? TryMoveCachedDataAsync(IoHash hash, DirectoryState targetDirState, Utf8String targetName)
+		{
+			HashInfo? hashInfo;
+			if (_hashes.TryGetValue(hash, out hashInfo))
+			{
+				foreach (FileState file in hashInfo.Files)
+				{
+					if (file.LayerFlags == CacheLayerFlag)
+					{
+						MoveFile(file, targetDirState, targetName);
+						file.LayerFlags = 0;
+						return file;
+					}
+				}
+			}
+			return null;
+		}
+
 		async Task<bool> TryCopyCachedDataAsync(HashInfo hashInfo, long offset, long length, Stream outputStream, CancellationToken cancellationToken)
 		{
 			foreach (FileState file in hashInfo.Files)
@@ -811,6 +851,18 @@ namespace EpicGames.Horde.Storage
 
 		#endregion
 
+		void AddDirToHashLookup(DirectoryState dirState)
+		{
+			foreach (DirectoryState subDirState in dirState.Directories)
+			{
+				AddDirToHashLookup(subDirState);
+			}
+			foreach (FileState fileState in dirState.Files)
+			{
+				AddFileToHashLookup(fileState);
+			}
+		}
+
 		void AddFileToHashLookup(FileState file)
 		{
 			HashInfo? hashInfo;
@@ -839,20 +891,40 @@ namespace EpicGames.Horde.Storage
 				cacheDirState = cacheDirState.FindOrAddDirectory(name.Slice(idx * 2, 2));
 			}
 
-			FileReference fileRef = fileState.GetFileReference(_rootDir);
-			if (cacheDirState.ContainsFile(name))
+			FileInfo fileInfo = fileState.GetFileReference(_rootDir).ToFileInfo();
+			if (fileInfo.Exists)
 			{
-				fileState.Delete();
-				FileUtils.ForceDeleteFile(fileRef);
+				if (cacheDirState.ContainsFile(name) || fileState.IsModified(fileInfo))
+				{
+					DeleteFile(fileState);
+				}
+				else
+				{
+					MoveFile(fileState, cacheDirState, name);
+					fileState.LayerFlags = CacheLayerFlag;
+				}
 			}
-			else
-			{
-				fileState.MoveTo(cacheDirState, name);
+		}
 
-				FileReference cachedFileRef = fileState.GetFileReference(_rootDir);
-				DirectoryReference.CreateDirectory(cachedFileRef.Directory);
-				FileReference.Move(fileRef, cachedFileRef);
-			}
+		void MoveFile(FileState fileState, DirectoryState targetDirState, Utf8String targetName)
+		{
+			FileReference sourceFile = fileState.GetFileReference(_rootDir);
+			fileState.MoveTo(targetDirState, targetName);
+
+			FileReference targetFile = fileState.GetFileReference(_rootDir);
+			DirectoryReference.CreateDirectory(targetFile.Directory);
+			FileReference.Move(sourceFile, targetFile, true);
+
+			_logger.LogDebug("Moving file from {Source} to {Target}", sourceFile, targetFile);
+		}
+
+		void DeleteFile(FileState fileState)
+		{
+			FileReference fileRef = fileState.GetFileReference(_rootDir);
+			FileUtils.ForceDeleteFile(fileRef);
+			fileState.Delete();
+
+			_logger.LogDebug("Deleting file {File}", fileRef);
 		}
 
 		FileInfo GetFileInfo(FileState file)
