Horde: Fix warnings highlighted by VS 17.8 preview static analyzer.

diff --git a/Engine/Source/Programs/Shared/EpicGames.Perforce.Managed/ManagedWorkspace.cs b/Engine/Source/Programs/Shared/EpicGames.Perforce.Managed/ManagedWorkspace.cs
--- a/Engine/Source/Programs/Shared/EpicGames.Perforce.Managed/ManagedWorkspace.cs
+++ b/Engine/Source/Programs/Shared/EpicGames.Perforce.Managed/ManagedWorkspace.cs
@@ -994,7 +994,7 @@ namespace EpicGames.Perforce.Managed
 			}
 
 			// Use common paths with wild cards speed up the print operation with one call instead of many calls to print.
-			_logger.LogInformation("Writing files from shelved changelist", unshelveChangelist);
+			_logger.LogInformation("Writing files from shelved changelist {Change}", unshelveChangelist);
 			PerforceResponseList<PrintRecord> printResponse = await perforce.TryPrintAsync($"{perforceClient.Root}{Path.DirectorySeparatorChar}...", $"//{perforceClient.Name}/...@={unshelveChangelist}", cancellationToken);
 			if (!printResponse.Succeeded)
 			{
@@ -2070,7 +2070,7 @@ namespace EpicGames.Perforce.Managed
 				{
 					FileReference localFile = filesToSync[idx]._workspaceFile.GetLocation();
 					DirectoryReference.CreateDirectory(localFile.Directory);
-					FileReference.WriteAllBytes(localFile, Array.Empty<byte>());
+					await FileReference.WriteAllBytesAsync(localFile, Array.Empty<byte>(), cancellationToken);
 				}
 			}
 			else
