Horde: Output stats when copying data from temp storage.

diff --git a/Engine/Source/Programs/Horde/Horde.Agent/Utility/TempStorage.cs b/Engine/Source/Programs/Horde/Horde.Agent/Utility/TempStorage.cs
--- a/Engine/Source/Programs/Horde/Horde.Agent/Utility/TempStorage.cs
+++ b/Engine/Source/Programs/Horde/Horde.Agent/Utility/TempStorage.cs
@@ -733,7 +733,7 @@ namespace Horde.Storage.Utility
 
 				// Add all the files and flush the ref
 				DirectoryNode rootDirNode = await rootDirEntry.ExpandAsync(cancellationToken);
-				await rootDirNode.CopyToDirectoryAsync(rootDir.ToDirectoryInfo(), logger, cancellationToken);
+				await rootDirNode.CopyToDirectoryAsync(rootDir.ToDirectoryInfo(), new CopyStatsLogger(logger), logger, cancellationToken);
 
 				timer.Stop();
 				logger.LogInformation("Elapsed: {Elapsed}s, Num bytes: {NumBytes:n0}, Num headers: {NumHeaders:n0}, Num packets: {NumPackets:n0}, Header cache: {Size:n0}mb, Packet cache {Size:n0}mb", (int)timer.Elapsed.TotalSeconds, bundleStorageClient.BundleReader.NumBytesRead - initialNumBytesRead, bundleStorageClient.BundleReader.NumHeaderReads - initialNumHeaderReads, bundleStorageClient.BundleReader.NumPacketReads - initialNumPacketReads, bundleStorageClient.BundleReader.Cache.HeaderCacheSize / (1024.0 * 1024.0), bundleStorageClient.BundleReader.Cache.PacketCacheSize / (1024.0 * 1024.0));
diff --git a/Engine/Source/Programs/Shared/EpicGames.Horde/Storage/Nodes/DirectoryNode.cs b/Engine/Source/Programs/Shared/EpicGames.Horde/Storage/Nodes/DirectoryNode.cs
--- a/Engine/Source/Programs/Shared/EpicGames.Horde/Storage/Nodes/DirectoryNode.cs
+++ b/Engine/Source/Programs/Shared/EpicGames.Horde/Storage/Nodes/DirectoryNode.cs
@@ -58,7 +58,14 @@ namespace EpicGames.Horde.Storage.Nodes
 		/// <inheritdoc/>
 		public void Report(ICopyStats stats)
 		{
-			_logger.LogInformation("Copied {NumFiles:n0}/{TotalFiles:n0} ({Size:n1}/{TotalSize:n1}mb, {Pct}%)", stats.CopiedCount, stats.TotalCount, stats.CopiedSize / (1024.0 * 1024.0), stats.TotalSize / (1024.0 * 1024.0), (int)((Math.Max(stats.CopiedCount, 1) * 100) / Math.Max(stats.TotalCount, 1)));
+			if (stats.TotalCount == 0)
+			{
+				_logger.LogInformation("Copied {NumFiles:n0} files ({Size:n1}/{TotalSize:n1}mb, {Pct}%)", stats.CopiedCount, stats.CopiedSize / (1024.0 * 1024.0), stats.TotalSize / (1024.0 * 1024.0), (int)((Math.Max(stats.CopiedCount, 1) * 100) / Math.Max(stats.TotalCount, 1)));
+			}
+			else
+			{
+				_logger.LogInformation("Copied {NumFiles:n0}/{TotalFiles:n0} files ({Size:n1}/{TotalSize:n1}mb, {Pct}%)", stats.CopiedCount, stats.TotalCount, stats.CopiedSize / (1024.0 * 1024.0), stats.TotalSize / (1024.0 * 1024.0), (int)((Math.Max(stats.CopiedCount, 1) * 100) / Math.Max(stats.TotalCount, 1)));
+			}
 		}
 	}
 
@@ -753,12 +760,26 @@ namespace EpicGames.Horde.Storage.Nodes
 		/// <param name="directoryInfo"></param>
 		/// <param name="logger"></param>
 		/// <param name="cancellationToken"></param>
-		/// <returns></returns>
-		public async Task CopyToDirectoryAsync(DirectoryInfo directoryInfo, ILogger logger, CancellationToken cancellationToken)
+		public Task CopyToDirectoryAsync(DirectoryInfo directoryInfo, ILogger logger, CancellationToken cancellationToken) => CopyToDirectoryAsync(directoryInfo, null, logger, cancellationToken);
+
+		/// <summary>
+		/// Utility function to allow extracting a packed directory to disk
+		/// </summary>
+		/// <param name="directoryInfo">Direcotry to write to</param>
+		/// <param name="progress">Sink for progress updates</param>
+		/// <param name="logger">Logger for output</param>
+		/// <param name="cancellationToken">Cancellation token for the operation</param>
+		public async Task CopyToDirectoryAsync(DirectoryInfo directoryInfo, IProgress<ICopyStats>? progress, ILogger logger, CancellationToken cancellationToken)
 		{
 			int NumThreads = Math.Min(1 + (int)(Length / (10 * 1024 * 1024)), 4);
 			logger.LogInformation("Splitting read into {NumThreads} threads", NumThreads);
 
+			CopyStats? copyStats = null;
+			if (progress != null)
+			{
+				copyStats = new CopyStats(0, Length, progress);
+			}
+
 			List<Task> tasks = new List<Task>();
 			try
 			{
@@ -767,7 +788,7 @@ namespace EpicGames.Horde.Storage.Nodes
 				{
 					long minOffset = offset;
 					long maxOffset = (Length * (threadIdx + 1)) / NumThreads;
-					tasks.Add(Task.Run(() => CopyToDirectoryInternalAsync(directoryInfo, minOffset, maxOffset - minOffset, logger, cancellationToken), cancellationToken));
+					tasks.Add(Task.Run(() => CopyToDirectoryInternalAsync(directoryInfo, minOffset, maxOffset - minOffset, copyStats, logger, cancellationToken), cancellationToken));
 					offset = maxOffset;
 				}
 			}
@@ -777,7 +798,7 @@ namespace EpicGames.Horde.Storage.Nodes
 			}
 		}
 
-		async Task CopyToDirectoryInternalAsync(DirectoryInfo directoryInfo, long windowOffset, long windowLength, ILogger logger, CancellationToken cancellationToken)
+		async Task CopyToDirectoryInternalAsync(DirectoryInfo directoryInfo, long windowOffset, long windowLength, CopyStats copyStats, ILogger logger, CancellationToken cancellationToken)
 		{
 			directoryInfo.Create();
 
@@ -788,6 +809,7 @@ namespace EpicGames.Horde.Storage.Nodes
 				{
 					FileInfo fileInfo = new FileInfo(Path.Combine(directoryInfo.FullName, fileEntry.Name.ToString()));
 					await fileEntry.CopyToFileAsync(fileInfo, cancellationToken);
+					copyStats.Update(1, fileEntry.Length);
 				}
 				windowOffset -= fileEntry.Length;
 			}
@@ -799,7 +821,7 @@ namespace EpicGames.Horde.Storage.Nodes
 				{
 					DirectoryInfo subDirectoryInfo = directoryInfo.CreateSubdirectory(directoryEntry.Name.ToString());
 					DirectoryNode subDirectoryNode = await directoryEntry.ExpandAsync(cancellationToken);
-					await subDirectoryNode.CopyToDirectoryInternalAsync(subDirectoryInfo, windowOffset, windowLength, logger, cancellationToken);
+					await subDirectoryNode.CopyToDirectoryInternalAsync(subDirectoryInfo, windowOffset, windowLength, copyStats, logger, cancellationToken);
 				}
 				windowOffset -= directoryEntry.Length;
 			}
