Horde: Fix null reference warning.

diff --git a/Engine/Source/Programs/Shared/EpicGames.Horde/Storage/Nodes/DirectoryNode.cs b/Engine/Source/Programs/Shared/EpicGames.Horde/Storage/Nodes/DirectoryNode.cs
--- a/Engine/Source/Programs/Shared/EpicGames.Horde/Storage/Nodes/DirectoryNode.cs
+++ b/Engine/Source/Programs/Shared/EpicGames.Horde/Storage/Nodes/DirectoryNode.cs
@@ -798,7 +798,7 @@ namespace EpicGames.Horde.Storage.Nodes
 			}
 		}
 
-		async Task CopyToDirectoryInternalAsync(DirectoryInfo directoryInfo, long windowOffset, long windowLength, CopyStats copyStats, ILogger logger, CancellationToken cancellationToken)
+		async Task CopyToDirectoryInternalAsync(DirectoryInfo directoryInfo, long windowOffset, long windowLength, CopyStats? copyStats, ILogger logger, CancellationToken cancellationToken)
 		{
 			directoryInfo.Create();
 
@@ -809,7 +809,7 @@ namespace EpicGames.Horde.Storage.Nodes
 				{
 					FileInfo fileInfo = new FileInfo(Path.Combine(directoryInfo.FullName, fileEntry.Name.ToString()));
 					await fileEntry.CopyToFileAsync(fileInfo, cancellationToken);
-					copyStats.Update(1, fileEntry.Length);
+					copyStats?.Update(1, fileEntry.Length);
 				}
 				windowOffset -= fileEntry.Length;
 			}
