Horde: Add linear hash for files to FileEntry instances, rather than using the hash of the Merkle tree underneath it. Allows verifying integrity of files extracted to the workspace.

diff --git a/Engine/Source/Programs/Horde/Horde.Agent/Utility/TempStorage.cs b/Engine/Source/Programs/Horde/Horde.Agent/Utility/TempStorage.cs
--- a/Engine/Source/Programs/Horde/Horde.Agent/Utility/TempStorage.cs
+++ b/Engine/Source/Programs/Horde/Horde.Agent/Utility/TempStorage.cs
@@ -626,10 +626,10 @@ namespace Horde.Storage.Utility
 			FileReference localFileListLocation = GetTagManifestLocation(manifestDir, nodeName, tagName);
 			fileList.Save(localFileListLocation);
 
-			ChunkedDataWriter fileNodeWriter = new ChunkedDataWriter(writer, new ChunkingOptions());
-			NodeRef<ChunkedDataNode> rootRef = await fileNodeWriter.CreateAsync(localFileListLocation.ToFileInfo(), cancellationToken);
+			using ChunkedDataWriter fileNodeWriter = new ChunkedDataWriter(writer, new ChunkingOptions());
+			ChunkedData fileNodeData = await fileNodeWriter.CreateAsync(localFileListLocation.ToFileInfo(), cancellationToken);
 
-			return new FileEntry(localFileListLocation.GetFileName(), FileEntryFlags.None, fileNodeWriter.Length, rootRef);
+			return new FileEntry(localFileListLocation.GetFileName(), FileEntryFlags.None, fileNodeWriter.Length, fileNodeData);
 		}
 
 		/// <summary>
diff --git a/Engine/Source/Programs/Horde/Horde.Server/Commands/Bundles/BundleWriteTest.cs b/Engine/Source/Programs/Horde/Horde.Server/Commands/Bundles/BundleWriteTest.cs
--- a/Engine/Source/Programs/Horde/Horde.Server/Commands/Bundles/BundleWriteTest.cs
+++ b/Engine/Source/Programs/Horde/Horde.Server/Commands/Bundles/BundleWriteTest.cs
@@ -60,7 +60,7 @@ namespace Horde.Server.Commands.Bundles
 			long length = 0;
 			double nextTime = 2.0;
 
-			ChunkedDataWriter fileNodeWriter = new ChunkedDataWriter(writer, chunkingOptions);
+			using ChunkedDataWriter fileNodeWriter = new ChunkedDataWriter(writer, chunkingOptions);
 			for (; ; )
 			{
 				await fileNodeWriter.AppendAsync(buffer, default);
diff --git a/Engine/Source/Programs/Horde/Horde.Server/Perforce/PerforceReplicator.cs b/Engine/Source/Programs/Horde/Horde.Server/Perforce/PerforceReplicator.cs
--- a/Engine/Source/Programs/Horde/Horde.Server/Perforce/PerforceReplicator.cs
+++ b/Engine/Source/Programs/Horde/Horde.Server/Perforce/PerforceReplicator.cs
@@ -102,23 +102,30 @@ namespace Horde.Server.Perforce
 			}
 		}
 
-		record class FileInfo(Utf8String Path, FileEntryFlags Flags, long Length, byte[] Md5, NodeRef<ChunkedDataNode> NodeRef);
+		record class FileInfo(Utf8String Path, FileEntryFlags Flags, long Length, byte[] Md5, ChunkedData ChunkedData);
 
 		class FileWriter : IDisposable
 		{
-			class Handle
+			class Handle : IDisposable
 			{
 				public Utf8String _path;
 				public FileEntryFlags _flags;
 				public readonly ChunkedDataWriter _fileWriter;
 				public long _size;
 				public long _sizeWritten;
-				public readonly IncrementalHash _hash = IncrementalHash.CreateHash(HashAlgorithmName.MD5);
+				public readonly IncrementalHash _hash;
 
 				public Handle(IStorageWriter writer, ChunkingOptions options)
 				{
+					_hash = IncrementalHash.CreateHash(HashAlgorithmName.MD5);
 					_fileWriter = new ChunkedDataWriter(writer, options);
 				}
+
+				public void Dispose()
+				{
+					_hash.Dispose();
+					_fileWriter.Dispose();
+				}
 			}
 
 			readonly IStorageWriter _writer;
@@ -136,11 +143,11 @@ namespace Horde.Server.Perforce
 			{
 				foreach (Handle handle in _freeHandles)
 				{
-					handle._hash.Dispose();
+					handle.Dispose();
 				}
 				foreach (Handle handle in _openHandles.Values)
 				{
-					handle._hash.Dispose();
+					handle.Dispose();
 				}
 			}
 
@@ -178,9 +185,9 @@ namespace Horde.Server.Perforce
 					throw new ReplicationException($"Invalid size for replicated file '{handle._path}'. Expected {handle._size}, got {handle._sizeWritten}.");
 				}
 
-				NodeRef<ChunkedDataNode> nodeRef = await handle._fileWriter.FlushAsync(cancellationToken);
+				ChunkedData chunkedData = await handle._fileWriter.FlushAsync(cancellationToken);
 				byte[] hash = handle._hash.GetHashAndReset();
-				FileInfo info = new FileInfo(handle._path, handle._flags, handle._size, hash, nodeRef);
+				FileInfo info = new FileInfo(handle._path, handle._flags, handle._size, hash, chunkedData);
 
 				_openHandles.Remove(fd);
 				_freeHandles.Push(handle);
@@ -486,7 +493,7 @@ namespace Horde.Server.Perforce
 						else if (io.Command == PerforceIoCommand.Close)
 						{
 							FileInfo info = await fileWriter.CloseAsync(io.File, cancellationToken);
-							FileEntry entry = rootUpdate.AddFile(info.Path.ToString(), info.Flags, info.Length, info.NodeRef);
+							FileEntry entry = rootUpdate.AddFile(info.Path.ToString(), info.Flags, info.Length, info.ChunkedData);
 							entry.CustomData = info.Md5;
 						}
 						else if (io.Command == PerforceIoCommand.Unlink)
diff --git a/Engine/Source/Programs/Horde/Horde/Commands/Workspace/WorkspaceVerify.cs b/Engine/Source/Programs/Horde/Horde/Commands/Workspace/WorkspaceVerify.cs
new file mode 100644
--- /dev/null
+++ b/Engine/Source/Programs/Horde/Horde/Commands/Workspace/WorkspaceVerify.cs
@@ -0,0 +1,38 @@
+// Copyright Epic Games, Inc. All Rights Reserved.
+
+using System.Diagnostics;
+using EpicGames.Core;
+using Microsoft.Extensions.Logging;
+
+namespace Horde.Commands.Workspace
+{
+	using Workspace = EpicGames.Horde.Storage.Workspace;
+
+	[Command("workspace", "verify", "Verify the integrity of the workspace")]
+	class WorkspaceVerify : Command
+	{
+		[CommandLine("-Root=")]
+		public DirectoryReference? RootDir { get; set; }
+
+		public override async Task<int> ExecuteAsync(ILogger logger)
+		{
+			CancellationToken cancellationToken = CancellationToken.None;
+			RootDir ??= DirectoryReference.GetCurrentDirectory();
+
+			Workspace? workspace = await Workspace.TryOpenAsync(RootDir, logger, cancellationToken);
+			if (workspace == null)
+			{
+				logger.LogError("No workspace has been initialized in {RootDir}. Use 'workspace init' to create a new workspace.", RootDir);
+				return 1;
+			}
+
+			Stopwatch timer = Stopwatch.StartNew();
+			logger.LogInformation("Checking workspace integrity...");
+
+			await workspace.VerifyAsync(cancellationToken);
+			logger.LogInformation("Elapsed: {Time}s", timer.Elapsed.TotalSeconds);
+
+			return 0;
+		}
+	}
+}
diff --git a/Engine/Source/Programs/Shared/EpicGames.Horde.Tests/BundleStoreTests.cs b/Engine/Source/Programs/Shared/EpicGames.Horde.Tests/BundleStoreTests.cs
--- a/Engine/Source/Programs/Shared/EpicGames.Horde.Tests/BundleStoreTests.cs
+++ b/Engine/Source/Programs/Shared/EpicGames.Horde.Tests/BundleStoreTests.cs
@@ -265,8 +265,8 @@ namespace EpicGames.Horde.Tests
 			{
 				await using IStorageWriter writer = store.CreateWriter();
 
-				ChunkedDataWriter fileWriter = new ChunkedDataWriter(writer, new ChunkingOptions());
-				NodeRef<ChunkedDataNode> fileHandle = await fileWriter.CreateAsync(Encoding.UTF8.GetBytes("world"), CancellationToken.None);
+				using ChunkedDataWriter fileWriter = new ChunkedDataWriter(writer, new ChunkingOptions());
+				ChunkedData fileHandle = await fileWriter.CreateAsync(Encoding.UTF8.GetBytes("world"), CancellationToken.None);
 
 				List<FileUpdate> fileUpdates = new List<FileUpdate>();
 				fileUpdates.Add(new FileUpdate("hello/world", FileEntryFlags.None, fileWriter.Length, fileHandle));
@@ -321,13 +321,13 @@ namespace EpicGames.Horde.Tests
 				ChunkingOptions options = new ChunkingOptions();
 				options.LeafOptions = new LeafChunkedDataNodeOptions(128, 256, 64 * 1024);
 
-				ChunkedDataWriter fileWriter = new ChunkedDataWriter(writer, options);
+				using ChunkedDataWriter fileWriter = new ChunkedDataWriter(writer, options);
 				for (int idx = 0; idx < chunk.Length / 16; idx++)
 				{
 					await fileWriter.AppendAsync(chunk.AsMemory(idx * 16, 16), CancellationToken.None);
 				}
 
-				nodeRef = await fileWriter.FlushAsync(CancellationToken.None);
+				nodeRef = (await fileWriter.FlushAsync(CancellationToken.None)).Root;
 			}
 
 			// Check we can read it back in
@@ -368,11 +368,11 @@ namespace EpicGames.Horde.Tests
 				ChunkingOptions options = new ChunkingOptions();
 				options.LeafOptions = new LeafChunkedDataNodeOptions(128, 256, 64 * 1024);
 
-				ChunkedDataWriter fileWriter = new ChunkedDataWriter(writer, options);
-				NodeRef<ChunkedDataNode> nodeRef = await fileWriter.CreateAsync(data, CancellationToken.None);
+				using ChunkedDataWriter fileWriter = new ChunkedDataWriter(writer, options);
+				ChunkedData chunkedData = await fileWriter.CreateAsync(data, CancellationToken.None);
 
 				root = new DirectoryNode(DirectoryFlags.None);
-				root.AddFile("test", FileEntryFlags.None, fileWriter.Length, nodeRef);
+				root.AddFile("test", FileEntryFlags.None, fileWriter.Length, chunkedData);
 
 				NodeRef<DirectoryNode> rootRef = await writer.WriteNodeAsync(root);
 				await store.WriteRefTargetAsync(new RefName("test"), rootRef);
diff --git a/Engine/Source/Programs/Shared/EpicGames.Horde.Tests/BundleTests.cs b/Engine/Source/Programs/Shared/EpicGames.Horde.Tests/BundleTests.cs
--- a/Engine/Source/Programs/Shared/EpicGames.Horde.Tests/BundleTests.cs
+++ b/Engine/Source/Programs/Shared/EpicGames.Horde.Tests/BundleTests.cs
@@ -75,25 +75,25 @@ namespace EpicGames.Horde.Tests
 			ChunkingOptions options = new ChunkingOptions();
 			options.LeafOptions = new LeafChunkedDataNodeOptions(8, 8, 8);
 
-			ChunkedDataWriter fileNodeWriter = new ChunkedDataWriter(writer, options);
+			using ChunkedDataWriter fileNodeWriter = new ChunkedDataWriter(writer, options);
 
 			ChunkedDataNode node;
 			NodeRef<ChunkedDataNode> nodeRef;
 			byte[] data = CreateBuffer(1024);
 
-			nodeRef = await fileNodeWriter.CreateAsync(data.AsMemory(0, 7), CancellationToken.None);
+			nodeRef = (await fileNodeWriter.CreateAsync(data.AsMemory(0, 7), CancellationToken.None)).Root;
 			node = await nodeRef.ExpandAsync();
 			Assert.IsTrue(node is LeafChunkedDataNode);
 			Assert.AreEqual(7, ((LeafChunkedDataNode)node).Data.Length);
 			await TestBufferlessReadsAsync(nodeRef, data.AsMemory(0, 7));
 
-			nodeRef = await fileNodeWriter.CreateAsync(data.AsMemory(0, 8), CancellationToken.None);
+			nodeRef = (await fileNodeWriter.CreateAsync(data.AsMemory(0, 8), CancellationToken.None)).Root;
 			node = await nodeRef.ExpandAsync();
 			Assert.IsTrue(node is LeafChunkedDataNode);
 			Assert.AreEqual(8, ((LeafChunkedDataNode)node).Data.Length);
 			await TestBufferlessReadsAsync(nodeRef, data.AsMemory(0, 8));
 
-			nodeRef = await fileNodeWriter.CreateAsync(data.AsMemory(0, 9), CancellationToken.None);
+			nodeRef = (await fileNodeWriter.CreateAsync(data.AsMemory(0, 9), CancellationToken.None)).Root;
 			node = await nodeRef.ExpandAsync();
 			Assert.IsTrue(node is InteriorChunkedDataNode);
 			Assert.AreEqual(2, ((InteriorChunkedDataNode)node).Children.Count);
@@ -109,7 +109,7 @@ namespace EpicGames.Horde.Tests
 			Assert.IsTrue(childNode2 is LeafChunkedDataNode);
 			Assert.AreEqual(1, ((LeafChunkedDataNode)childNode2!).Data.Length);
 
-			nodeRef = await fileNodeWriter.CreateAsync(data, CancellationToken.None);
+			nodeRef = (await fileNodeWriter.CreateAsync(data, CancellationToken.None)).Root;
 			node = await nodeRef.ExpandAsync();
 			Assert.IsTrue(node is InteriorChunkedDataNode);
 			await TestBufferlessReadsAsync(nodeRef, data);
diff --git a/Engine/Source/Programs/Shared/EpicGames.Horde.Tests/ChunkingTests.cs b/Engine/Source/Programs/Shared/EpicGames.Horde.Tests/ChunkingTests.cs
--- a/Engine/Source/Programs/Shared/EpicGames.Horde.Tests/ChunkingTests.cs
+++ b/Engine/Source/Programs/Shared/EpicGames.Horde.Tests/ChunkingTests.cs
@@ -60,11 +60,11 @@ namespace EpicGames.Horde.Tests
 				options.InteriorOptions = new InteriorChunkedDataNodeOptions(4, 4, 4);
 
 				using MemoryStream emptyStream = new MemoryStream();
-				List<NodeRef<ChunkedDataNode>> leafNodes = await LeafChunkedDataNode.CreateFromStreamAsync(writer, emptyStream, new LeafChunkedDataNodeOptions(64, 64, 64), CancellationToken.None);
-				NodeRef<ChunkedDataNode> dataNode = await InteriorChunkedDataNode.CreateTreeAsync(leafNodes, new InteriorChunkedDataNodeOptions(4, 4, 4), writer, CancellationToken.None); 
+				LeafChunkedData leafChunkedData = await LeafChunkedDataNode.CreateFromStreamAsync(writer, emptyStream, new LeafChunkedDataNodeOptions(64, 64, 64), CancellationToken.None);
+				ChunkedData chunkedData = await InteriorChunkedDataNode.CreateTreeAsync(leafChunkedData, new InteriorChunkedDataNodeOptions(4, 4, 4), writer, CancellationToken.None); 
 
 				DirectoryNode directory = new DirectoryNode();
-				directory.AddFile("test.foo", FileEntryFlags.None, 0, dataNode);
+				directory.AddFile("test.foo", FileEntryFlags.None, 0, chunkedData);
 
 				NodeRef<DirectoryNode> directoryRef = await writer.WriteNodeAsync(directory);
 				await writer.WriteRefAsync(directoryRef.Handle);
@@ -111,14 +111,14 @@ namespace EpicGames.Horde.Tests
 
 			const int NumIterations = 100;
 			{
-				ChunkedDataWriter fileWriter = new ChunkedDataWriter(writer, options);
+				using ChunkedDataWriter fileWriter = new ChunkedDataWriter(writer, options);
 
 				for (int idx = 0; idx < NumIterations; idx++)
 				{
 					await fileWriter.AppendAsync(data, CancellationToken.None);
 				}
 
-				handle = await fileWriter.FlushAsync(CancellationToken.None);
+				handle = (await fileWriter.FlushAsync(CancellationToken.None)).Root;
 			}
 
 			ChunkedDataNode root = await handle.ExpandAsync();
diff --git a/Engine/Source/Programs/Shared/EpicGames.Horde.Tests/ComputeSocketTests.cs b/Engine/Source/Programs/Shared/EpicGames.Horde.Tests/ComputeSocketTests.cs
--- a/Engine/Source/Programs/Shared/EpicGames.Horde.Tests/ComputeSocketTests.cs
+++ b/Engine/Source/Programs/Shared/EpicGames.Horde.Tests/ComputeSocketTests.cs
@@ -156,11 +156,11 @@ namespace EpicGames.Horde.Tests
 
 						byte[] data = Encoding.UTF8.GetBytes("Hello world");
 
-						ChunkedDataWriter writer = new ChunkedDataWriter(treeWriter, new ChunkingOptions());
-						NodeRef<ChunkedDataNode> nodeRef = await writer.CreateAsync(data, CancellationToken.None);
+						using ChunkedDataWriter writer = new ChunkedDataWriter(treeWriter, new ChunkingOptions());
+						ChunkedData chunkedData = await writer.CreateAsync(data, CancellationToken.None);
 
 						DirectoryNode directory = new DirectoryNode();
-						directory.AddFile("hello.txt", FileEntryFlags.None, data.Length, nodeRef);
+						directory.AddFile("hello.txt", FileEntryFlags.None, data.Length, chunkedData);
 
 						NodeRef<DirectoryNode> directoryRef = await treeWriter.WriteNodeAsync(directory);
 
diff --git a/Engine/Source/Programs/Shared/EpicGames.Horde/Storage/Nodes/ChunkedDataNode.cs b/Engine/Source/Programs/Shared/EpicGames.Horde/Storage/Nodes/ChunkedDataNode.cs
--- a/Engine/Source/Programs/Shared/EpicGames.Horde/Storage/Nodes/ChunkedDataNode.cs
+++ b/Engine/Source/Programs/Shared/EpicGames.Horde/Storage/Nodes/ChunkedDataNode.cs
@@ -86,6 +86,13 @@ namespace EpicGames.Horde.Storage.Nodes
 		}
 	}
 
+	/// <summary>
+	/// Stores the flat list of chunks produced from chunking a single data stream
+	/// </summary>
+	/// <param name="Hash">Hash of the data</param>
+	/// <param name="LeafHandles">Handles to the leaf chunks</param>
+	public record class LeafChunkedData(IoHash Hash, List<NodeRef<ChunkedDataNode>> LeafHandles);
+
 	/// <summary>
 	/// File node that contains a chunk of data
 	/// </summary>
@@ -145,12 +152,12 @@ namespace EpicGames.Horde.Storage.Nodes
 		/// <summary>
 		/// Creates nodes from the given file
 		/// </summary>
-		/// <param name="writer"></param>
-		/// <param name="file"></param>
-		/// <param name="options"></param>
-		/// <param name="cancellationToken"></param>
-		/// <returns></returns>
-		public static async Task<List<NodeRef<ChunkedDataNode>>> CreateFromFileAsync(IStorageWriter writer, FileReference file, LeafChunkedDataNodeOptions options, CancellationToken cancellationToken)
+		/// <param name="writer">Writer for output nodes</param>
+		/// <param name="file">File info</param>
+		/// <param name="options">Options for finding chunk boundaries</param>
+		/// <param name="cancellationToken">Cancellation token for the operation</param>
+		/// <returns>Hash of the full file data</returns>
+		public static async Task<LeafChunkedData> CreateFromFileAsync(IStorageWriter writer, FileReference file, LeafChunkedDataNodeOptions options, CancellationToken cancellationToken)
 		{
 			using (FileStream stream = FileReference.Open(file, FileMode.Open, FileAccess.Read, FileShare.Read))
 			{
@@ -161,12 +168,12 @@ namespace EpicGames.Horde.Storage.Nodes
 		/// <summary>
 		/// Creates nodes from the given file
 		/// </summary>
-		/// <param name="writer"></param>
-		/// <param name="file"></param>
-		/// <param name="options"></param>
-		/// <param name="cancellationToken"></param>
-		/// <returns></returns>
-		public static async Task<List<NodeRef<ChunkedDataNode>>> CreateFromFileAsync(IStorageWriter writer, FileInfo file, LeafChunkedDataNodeOptions options, CancellationToken cancellationToken)
+		/// <param name="writer">Writer for output nodes</param>
+		/// <param name="file">File info</param>
+		/// <param name="options">Options for finding chunk boundaries</param>
+		/// <param name="cancellationToken">Cancellation token for the operation</param>
+		/// <returns>Hash of the full file data</returns>
+		public static async Task<LeafChunkedData> CreateFromFileAsync(IStorageWriter writer, FileInfo file, LeafChunkedDataNodeOptions options, CancellationToken cancellationToken)
 		{
 			using (FileStream stream = file.Open(FileMode.Open, FileAccess.Read, FileShare.Read))
 			{
@@ -177,17 +184,18 @@ namespace EpicGames.Horde.Storage.Nodes
 		/// <summary>
 		/// Creates nodes from the given file
 		/// </summary>
-		/// <param name="writer"></param>
-		/// <param name="stream"></param>
-		/// <param name="options"></param>
-		/// <param name="cancellationToken"></param>
-		/// <returns></returns>
-		public static async Task<List<NodeRef<ChunkedDataNode>>> CreateFromStreamAsync(IStorageWriter writer, Stream stream, LeafChunkedDataNodeOptions options, CancellationToken cancellationToken)
+		/// <param name="writer">Writer for output nodes</param>
+		/// <param name="stream">Stream to read from</param>
+		/// <param name="options">Options for finding chunk boundaries</param>
+		/// <param name="cancellationToken">Cancellation token for the operation</param>
+		/// <returns>Hash of the full file data</returns>
+		public static async Task<LeafChunkedData> CreateFromStreamAsync(IStorageWriter writer, Stream stream, LeafChunkedDataNodeOptions options, CancellationToken cancellationToken)
 		{
-			List<NodeRef<ChunkedDataNode>> handles = new List<NodeRef<ChunkedDataNode>>();
-
+			using Blake3.Hasher hasher = Blake3.Hasher.New();
 			using IMemoryOwner<byte> readBuffer = MemoryPool<byte>.Shared.Rent(options.MaxSize);
 
+			List<NodeRef<ChunkedDataNode>> handles = new List<NodeRef<ChunkedDataNode>>();
+
 			int size = 0;
 			for (; ; )
 			{
@@ -201,6 +209,7 @@ namespace EpicGames.Horde.Storage.Nodes
 
 				Memory<byte> outputBuffer = writer.GetOutputBuffer(0, nextLength);
 				readBuffer.Memory.Slice(0, nextLength).CopyTo(outputBuffer);
+				hasher.Update(outputBuffer.Span);
 
 				BlobHandle handle = await writer.WriteNodeAsync(nextLength, Array.Empty<BlobHandle>(), GetNodeType<LeafChunkedDataNode>(), cancellationToken);
 				handles.Add(new NodeRef<ChunkedDataNode>(handle));
@@ -209,7 +218,8 @@ namespace EpicGames.Horde.Storage.Nodes
 				size -= nextLength;
 			}
 
-			return handles;
+			IoHash hash = IoHash.FromBlake3(hasher);
+			return new LeafChunkedData(hash, handles);
 		}
 
 		/// <summary>
@@ -330,6 +340,20 @@ namespace EpicGames.Horde.Storage.Nodes
 			}
 		}
 
+		/// <summary>
+		/// Create a tree of nodes from the given list of handles, splitting nodes in each layer based on the hash of the last node.
+		/// </summary>
+		/// <param name="leafChunkedData">List of leaf handles</param>
+		/// <param name="options">Options for splitting the tree</param>
+		/// <param name="writer">Output writer for new interior nodes</param>
+		/// <param name="cancellationToken">Cancellation token for the operation</param>
+		/// <returns>Handle to the root node of the tree</returns>
+		public static async Task<ChunkedData> CreateTreeAsync(LeafChunkedData leafChunkedData, InteriorChunkedDataNodeOptions options, IStorageWriter writer, CancellationToken cancellationToken)
+		{
+			NodeRef<ChunkedDataNode> rootRef = await CreateTreeAsync(leafChunkedData.LeafHandles, options, writer, cancellationToken);
+			return new ChunkedData(leafChunkedData.Hash, rootRef);
+		}
+
 		/// <summary>
 		/// Create a tree of nodes from the given list of handles, splitting nodes in each layer based on the hash of the last node.
 		/// </summary>
diff --git a/Engine/Source/Programs/Shared/EpicGames.Horde/Storage/Nodes/ChunkedDataWriter.cs b/Engine/Source/Programs/Shared/EpicGames.Horde/Storage/Nodes/ChunkedDataWriter.cs
--- a/Engine/Source/Programs/Shared/EpicGames.Horde/Storage/Nodes/ChunkedDataWriter.cs
+++ b/Engine/Source/Programs/Shared/EpicGames.Horde/Storage/Nodes/ChunkedDataWriter.cs
@@ -54,10 +54,17 @@ namespace EpicGames.Horde.Storage.Nodes
 		}
 	}
 
+	/// <summary>
+	/// Describes a chunked data stream
+	/// </summary>
+	/// <param name="Hash">Hash of the stream</param>
+	/// <param name="Root">Handle to the root chunk containing the data</param>
+	public record class ChunkedData(IoHash Hash, NodeRef<ChunkedDataNode> Root);
+
 	/// <summary>
 	/// Utility class for generating FileNode data directly into <see cref="IStorageWriter"/> instances, without constructing node representations first.
 	/// </summary>
-	public class ChunkedDataWriter
+	public sealed class ChunkedDataWriter : IDisposable
 	{
 		/// <summary>
 		/// Default buffer length when calling CreateAsync/AppendAsync
@@ -68,6 +75,7 @@ namespace EpicGames.Horde.Storage.Nodes
 
 		readonly IStorageWriter _writer;
 		readonly ChunkingOptions _options;
+		readonly Blake3.Hasher _hasher;
 
 		// Tree state
 		long _totalLength;
@@ -91,6 +99,13 @@ namespace EpicGames.Horde.Storage.Nodes
 		{
 			_writer = writer;
 			_options = options;
+			_hasher = Blake3.Hasher.New();
+		}
+
+		/// <inheritdoc/>
+		public void Dispose()
+		{
+			_hasher.Dispose();
 		}
 
 		/// <summary>
@@ -101,6 +116,7 @@ namespace EpicGames.Horde.Storage.Nodes
 			_leafHandles.Clear();
 			ResetLeafState();
 			_totalLength = 0;
+			_hasher.Reset();
 		}
 
 		/// <summary>
@@ -117,7 +133,7 @@ namespace EpicGames.Horde.Storage.Nodes
 		/// </summary>
 		/// <param name="fileInfo">File to append</param>
 		/// <param name="cancellationToken">Cancellation token for the operation</param>
-		public async Task<NodeRef<ChunkedDataNode>> CreateAsync(FileInfo fileInfo, CancellationToken cancellationToken)
+		public async Task<ChunkedData> CreateAsync(FileInfo fileInfo, CancellationToken cancellationToken)
 		{
 			return await CreateAsync(fileInfo, DefaultBufferLength, cancellationToken);
 		}
@@ -128,7 +144,7 @@ namespace EpicGames.Horde.Storage.Nodes
 		/// <param name="fileInfo">File to append</param>
 		/// <param name="bufferLength">Size of the read buffer</param>
 		/// <param name="cancellationToken">Cancellation token for the operation</param>
-		public async Task<NodeRef<ChunkedDataNode>> CreateAsync(FileInfo fileInfo, int bufferLength, CancellationToken cancellationToken)
+		public async Task<ChunkedData> CreateAsync(FileInfo fileInfo, int bufferLength, CancellationToken cancellationToken)
 		{
 			using (FileStream stream = fileInfo.OpenRead())
 			{
@@ -141,7 +157,7 @@ namespace EpicGames.Horde.Storage.Nodes
 		/// </summary>
 		/// <param name="stream">Stream to append</param>
 		/// <param name="cancellationToken">Cancellation token for the operation</param>
-		public async Task<NodeRef<ChunkedDataNode>> CreateAsync(Stream stream, CancellationToken cancellationToken)
+		public async Task<ChunkedData> CreateAsync(Stream stream, CancellationToken cancellationToken)
 		{
 			return await CreateAsync(stream, DefaultBufferLength, cancellationToken);
 		}
@@ -152,7 +168,7 @@ namespace EpicGames.Horde.Storage.Nodes
 		/// <param name="stream">Stream to append</param>
 		/// <param name="bufferLength">Size of the read buffer</param>
 		/// <param name="cancellationToken">Cancellation token for the operation</param>
-		public async Task<NodeRef<ChunkedDataNode>> CreateAsync(Stream stream, int bufferLength, CancellationToken cancellationToken)
+		public async Task<ChunkedData> CreateAsync(Stream stream, int bufferLength, CancellationToken cancellationToken)
 		{
 			Reset();
 			await AppendAsync(stream, bufferLength, cancellationToken);
@@ -164,7 +180,7 @@ namespace EpicGames.Horde.Storage.Nodes
 		/// </summary>
 		/// <param name="data">Stream to append</param>
 		/// <param name="cancellationToken">Cancellation token for the operation</param>
-		public async Task<NodeRef<ChunkedDataNode>> CreateAsync(ReadOnlyMemory<byte> data, CancellationToken cancellationToken)
+		public async Task<ChunkedData> CreateAsync(ReadOnlyMemory<byte> data, CancellationToken cancellationToken)
 		{
 			Reset();
 			await AppendAsync(data, cancellationToken);
@@ -219,6 +235,7 @@ namespace EpicGames.Horde.Storage.Nodes
 
 				await FlushLeafNodeAsync(cancellationToken);
 			}
+			_hasher.Update(data.Span);
 		}
 
 		/// <summary>
@@ -303,11 +320,11 @@ namespace EpicGames.Horde.Storage.Nodes
 		/// </summary>
 		/// <param name="cancellationToken">Cancellation token for the operation</param>
 		/// <returns>Handle to the root node</returns>
-		public async Task<NodeRef<ChunkedDataNode>> CompleteAsync(CancellationToken cancellationToken)
+		public async Task<ChunkedData> CompleteAsync(CancellationToken cancellationToken)
 		{
 			await FlushLeafNodeAsync(cancellationToken);
 			NodeRef<ChunkedDataNode> rootHandle = await InteriorChunkedDataNode.CreateTreeAsync(_leafHandles, _options.InteriorOptions, _writer, cancellationToken);
-			return rootHandle;
+			return new ChunkedData(IoHash.FromBlake3(_hasher), rootHandle);
 		}
 
 		/// <summary>
@@ -315,11 +332,11 @@ namespace EpicGames.Horde.Storage.Nodes
 		/// </summary>
 		/// <param name="cancellationToken">Cancellation token for the operation</param>
 		/// <returns>Handle to the root FileNode</returns>
-		public async Task<NodeRef<ChunkedDataNode>> FlushAsync(CancellationToken cancellationToken)
+		public async Task<ChunkedData> FlushAsync(CancellationToken cancellationToken)
 		{
-			NodeRef<ChunkedDataNode> handle = await CompleteAsync(cancellationToken);
+			ChunkedData chunkedData = await CompleteAsync(cancellationToken);
 			await _writer.FlushAsync(cancellationToken);
-			return handle;
+			return chunkedData;
 		}
 
 		/// <summary>
diff --git a/Engine/Source/Programs/Shared/EpicGames.Horde/Storage/Nodes/DirectoryNode.cs b/Engine/Source/Programs/Shared/EpicGames.Horde/Storage/Nodes/DirectoryNode.cs
--- a/Engine/Source/Programs/Shared/EpicGames.Horde/Storage/Nodes/DirectoryNode.cs
+++ b/Engine/Source/Programs/Shared/EpicGames.Horde/Storage/Nodes/DirectoryNode.cs
@@ -205,11 +205,11 @@ namespace EpicGames.Horde.Storage.Nodes
 		/// <param name="name">Name of the new directory</param>
 		/// <param name="flags">Flags for the new file</param>
 		/// <param name="length">Length of the file</param>
-		/// <param name="dataRef">Handle to the file data</param>
+		/// <param name="data">Chunked data for the file</param>
 		/// <returns>The new directory object</returns>
-		public FileEntry AddFile(Utf8String name, FileEntryFlags flags, long length, NodeRef<ChunkedDataNode> dataRef)
+		public FileEntry AddFile(Utf8String name, FileEntryFlags flags, long length, ChunkedData data)
 		{
-			FileEntry entry = new FileEntry(name, flags, length, dataRef);
+			FileEntry entry = new FileEntry(name, flags, length, data);
 			AddFile(entry);
 			return entry;
 		}
@@ -549,11 +549,11 @@ namespace EpicGames.Horde.Storage.Nodes
 
 			// Partition them into blocks for parallel writers to process asynchronously
 			List<(int Start, int Count)> partitions = ComputePartitions(files, totalSize);
-			List<NodeRef<ChunkedDataNode>>[] leafChunks = new List<NodeRef<ChunkedDataNode>>[files.Count];
-			await Parallel.ForEachAsync(partitions, cancellationToken, (filePartition, ctx) => CreateLeafChunkNodesAsync(writer, files, leafChunks, filePartition.Start, filePartition.Count, copyStats, options, cancellationToken));
+			LeafChunkedData[] leafChunkedFiles = new LeafChunkedData[files.Count];
+			await Parallel.ForEachAsync(partitions, cancellationToken, (filePartition, ctx) => CreateLeafChunkNodesAsync(writer, files, leafChunkedFiles, filePartition.Start, filePartition.Count, copyStats, options, cancellationToken));
 
 			// Create interior nodes for all the leaf chunks
-			NodeRef<ChunkedDataNode>[] chunks = await CreateInteriorChunkNodesAsync(leafChunks, options.InteriorOptions, writer, cancellationToken);
+			ChunkedData[] chunkedFiles = await CreateInteriorChunkNodesAsync(leafChunkedFiles, options.InteriorOptions, writer, cancellationToken);
 
 			// Write all the interior nodes and generate the directory update
 			DirectoryUpdate update = new DirectoryUpdate();
@@ -579,7 +579,7 @@ namespace EpicGames.Horde.Storage.Nodes
 					}
 				}
 
-				FileEntry entry = new FileEntry(file.Name, flags, file.Length, chunks[idx]);
+				FileEntry entry = new FileEntry(file.Name, flags, file.Length, chunkedFiles[idx]);
 				update.AddFile(new FileReference(file).MakeRelativeTo(baseDir), entry);
 			}
 
@@ -623,7 +623,7 @@ namespace EpicGames.Horde.Storage.Nodes
 			return partitions;
 		}
 
-		static async ValueTask CreateLeafChunkNodesAsync(IStorageWriter writer, IReadOnlyList<FileInfo> files, List<NodeRef<ChunkedDataNode>>[] leafChunks, int start, int count, CopyStats? copyStats, ChunkingOptions options, CancellationToken cancellationToken)
+		static async ValueTask CreateLeafChunkNodesAsync(IStorageWriter writer, IReadOnlyList<FileInfo> files, LeafChunkedData[] leafChunks, int start, int count, CopyStats? copyStats, ChunkingOptions options, CancellationToken cancellationToken)
 		{
 			await using IStorageWriter writerFork = writer.Fork();
 			for (int idx = start; idx < start + count; idx++)
@@ -638,14 +638,14 @@ namespace EpicGames.Horde.Storage.Nodes
 			await writerFork.FlushAsync(cancellationToken);
 		}
 
-		static async Task<NodeRef<ChunkedDataNode>[]> CreateInteriorChunkNodesAsync(List<NodeRef<ChunkedDataNode>>[] chunks, InteriorChunkedDataNodeOptions options, IStorageWriter writer, CancellationToken cancellationToken)
+		static async Task<ChunkedData[]> CreateInteriorChunkNodesAsync(LeafChunkedData[] leafChunkedFiles, InteriorChunkedDataNodeOptions options, IStorageWriter writer, CancellationToken cancellationToken)
 		{
-			NodeRef<ChunkedDataNode>[] nodes = new NodeRef<ChunkedDataNode>[chunks.Length];
-			for (int idx = 0; idx < chunks.Length; idx++)
+			ChunkedData[] chunkedFiles = new ChunkedData[leafChunkedFiles.Length];
+			for (int idx = 0; idx < leafChunkedFiles.Length; idx++)
 			{
-				nodes[idx] = await InteriorChunkedDataNode.CreateTreeAsync(chunks[idx], options, writer, cancellationToken);
+				chunkedFiles[idx] = await InteriorChunkedDataNode.CreateTreeAsync(leafChunkedFiles[idx], options, writer, cancellationToken);
 			}
-			return nodes;
+			return chunkedFiles;
 		}
 
 		/// <summary>
@@ -707,7 +707,7 @@ namespace EpicGames.Horde.Storage.Nodes
 		public async Task CopyFromZipStreamAsync(Stream stream, IStorageWriter writer, ChunkingOptions options, CancellationToken cancellationToken = default)
 		{
 			// Create all the leaf nodes
-			List<(ZipArchiveEntry, List<NodeRef<ChunkedDataNode>>)> entries = new List<(ZipArchiveEntry, List<NodeRef<ChunkedDataNode>>)>();
+			List<(ZipArchiveEntry, LeafChunkedData)> entries = new List<(ZipArchiveEntry, LeafChunkedData)>();
 			using (ZipArchive archive = new ZipArchive(stream, ZipArchiveMode.Read, true))
 			{
 				foreach (ZipArchiveEntry entry in archive.Entries)
@@ -716,8 +716,8 @@ namespace EpicGames.Horde.Storage.Nodes
 					{
 						using (Stream entryStream = entry.Open())
 						{
-							List<NodeRef<ChunkedDataNode>> chunks = await LeafChunkedDataNode.CreateFromStreamAsync(writer, entryStream, options.LeafOptions, cancellationToken);
-							entries.Add((entry, chunks));
+							LeafChunkedData leafChunkedData = await LeafChunkedDataNode.CreateFromStreamAsync(writer, entryStream, options.LeafOptions, cancellationToken);
+							entries.Add((entry, leafChunkedData));
 						}
 					}
 				}
@@ -725,7 +725,7 @@ namespace EpicGames.Horde.Storage.Nodes
 
 			// Create all the interior nodes
 			List<FileUpdate> updates = new List<FileUpdate>();
-			foreach ((ZipArchiveEntry entry, List<NodeRef<ChunkedDataNode>> chunks) in entries)
+			foreach ((ZipArchiveEntry entry, LeafChunkedData leafChunkedFile) in entries)
 			{
 				FileEntryFlags flags = FileEntryFlags.None;
 				if ((entry.ExternalAttributes & (0b_001_001_001 << 16)) != 0)
@@ -733,8 +733,8 @@ namespace EpicGames.Horde.Storage.Nodes
 					flags |= FileEntryFlags.Executable;
 				}
 
-				NodeRef<ChunkedDataNode> chunk = await InteriorChunkedDataNode.CreateTreeAsync(chunks, options.InteriorOptions, writer, cancellationToken);
-				updates.Add(new FileUpdate(entry.FullName, flags, entry.Length, chunk));
+				ChunkedData chunkedFile = await InteriorChunkedDataNode.CreateTreeAsync(leafChunkedFile, options.InteriorOptions, writer, cancellationToken);
+				updates.Add(new FileUpdate(entry.FullName, flags, entry.Length, chunkedFile));
 			}
 
 			// Update the tree
@@ -841,8 +841,8 @@ namespace EpicGames.Horde.Storage.Nodes
 	/// <param name="Path">Path to the file</param>
 	/// <param name="Length">Length of the file data</param>
 	/// <param name="Flags">Flags for the new file entry</param>
-	/// <param name="DataRef">Reference to the root data node</param>
-	public record class FileUpdate(string Path, FileEntryFlags Flags, long Length, NodeRef<ChunkedDataNode> DataRef);
+	/// <param name="Data">Chunked data for the file</param>
+	public record class FileUpdate(string Path, FileEntryFlags Flags, long Length, ChunkedData Data);
 
 	/// <summary>
 	/// Describes an update to a directory node
@@ -903,8 +903,8 @@ namespace EpicGames.Horde.Storage.Nodes
 		/// <param name="path">Path to the file</param>
 		/// <param name="flags">Flags for the new file entry</param>
 		/// <param name="length">Length of the file</param>
-		/// <param name="dataRef">Reference to the file data</param>
-		public FileEntry AddFile(string path, FileEntryFlags flags, long length, NodeRef<ChunkedDataNode> dataRef)
+		/// <param name="chunkedData">Chunked data instance</param>
+		public FileEntry AddFile(string path, FileEntryFlags flags, long length, ChunkedData chunkedData)
 		{
 			string name = path;
 			for (int idx = path.Length - 1; idx >= 0; idx--)
@@ -916,7 +916,7 @@ namespace EpicGames.Horde.Storage.Nodes
 				}
 			}
 
-			FileEntry entry = new FileEntry(name, flags, length, dataRef);
+			FileEntry entry = new FileEntry(name, flags, length, chunkedData);
 			AddFile(path, entry);
 			return entry;
 		}
@@ -948,7 +948,7 @@ namespace EpicGames.Horde.Storage.Nodes
 				if (nextDirLength == -1)
 				{
 					string fileName = fileUpdate.Path.Substring(prefix.Length);
-					Files[fileName] = new FileEntry(fileName, fileUpdate.Flags, fileUpdate.Length, fileUpdate.DataRef);
+					Files[fileName] = new FileEntry(fileName, fileUpdate.Flags, fileUpdate.Length, fileUpdate.Data);
 
 					if (!files.MoveNext())
 					{
diff --git a/Engine/Source/Programs/Shared/EpicGames.Horde/Storage/Nodes/FileEntry.cs b/Engine/Source/Programs/Shared/EpicGames.Horde/Storage/Nodes/FileEntry.cs
--- a/Engine/Source/Programs/Shared/EpicGames.Horde/Storage/Nodes/FileEntry.cs
+++ b/Engine/Source/Programs/Shared/EpicGames.Horde/Storage/Nodes/FileEntry.cs
@@ -80,13 +80,13 @@ namespace EpicGames.Horde.Storage.Nodes
 		/// <summary>
 		/// Constructor
 		/// </summary>
-		public FileEntry(Utf8String name, FileEntryFlags flags, long length, NodeRef<ChunkedDataNode> node)
-			: base(node)
+		public FileEntry(Utf8String name, FileEntryFlags flags, long length, ChunkedData data)
+			: base(data.Root)
 		{
 			Name = name;
 			Flags = flags;
 			Length = length;
-			Hash = node.Handle.Hash;
+			Hash = data.Hash;
 		}
 
 		/// <summary>
diff --git a/Engine/Source/Programs/Shared/EpicGames.Horde/Storage/Workspace.cs b/Engine/Source/Programs/Shared/EpicGames.Horde/Storage/Workspace.cs
--- a/Engine/Source/Programs/Shared/EpicGames.Horde/Storage/Workspace.cs
+++ b/Engine/Source/Programs/Shared/EpicGames.Horde/Storage/Workspace.cs
@@ -2,6 +2,7 @@
 
 using System;
 using System.Collections.Generic;
+using System.ComponentModel;
 using System.Diagnostics;
 using System.Diagnostics.CodeAnalysis;
 using System.IO;
@@ -696,7 +697,7 @@ namespace EpicGames.Horde.Storage
 				FileInfo fileInfo = FileReference.Combine(dirRef, fileEntry.Name.ToString()).ToFileInfo();
 
 				fileState = TryMoveCachedDataAsync(fileEntry.Hash, dirState, fileEntry.Name);
-				if(fileState == null)
+				if (fileState == null)
 				{
 					fileState = dirState.FindOrAddFile(fileEntry.Name);
 
@@ -704,15 +705,26 @@ namespace EpicGames.Horde.Storage
 					using (FileStream stream = fileInfo.Open(FileMode.Create, FileAccess.Write, FileShare.Read))
 					{
 						await ExtractDataAsync(fileEntry, stream, cancellationToken);
+						if (stream.Length != fileEntry.Length)
+						{
+							throw new EndOfStreamException($"Incorrect length for extracted file {fileInfo.FullName}");
+						}
 					}
 
 					fileState.Hash = fileEntry.Hash;
+
+					fileInfo.Refresh();
+					fileState.Update(fileInfo);
+
 					AddFileToHashLookup(fileState);
 				}
+				else
+				{
+					fileInfo.Refresh();
+					fileState.Update(fileInfo);
+				}
 
 				fileState.LayerFlags |= flag;
-				fileInfo.Refresh();
-				fileState.Update(fileInfo);
 			}
 		}
 
@@ -796,6 +808,8 @@ namespace EpicGames.Horde.Storage
 			FileInfo fileInfo = GetFileInfo(file);
 			if (fileInfo.Exists)
 			{
+				long initialPosition = outputStream.Position;
+
 				using FileStream inputStream = fileInfo.OpenRead();
 				inputStream.Seek(offset, SeekOrigin.Begin);
 
@@ -803,7 +817,15 @@ namespace EpicGames.Horde.Storage
 				while (length > 0)
 				{
 					int readSize = await inputStream.ReadAsync(tempBuffer.AsMemory(0, (int)Math.Min(length, tempBuffer.Length)), cancellationToken);
+					if (readSize == 0)
+					{
+						outputStream.Seek(initialPosition, SeekOrigin.Begin);
+						outputStream.SetLength(initialPosition);
+						return false;
+					}
+
 					await outputStream.WriteAsync(tempBuffer.AsMemory(0, readSize), cancellationToken);
+					length -= readSize;
 				}
 			}
 			return true;
diff --git a/Engine/Source/Programs/UnrealBuildTool/Artifacts/HordeStorageArtifactCache.cs b/Engine/Source/Programs/UnrealBuildTool/Artifacts/HordeStorageArtifactCache.cs
--- a/Engine/Source/Programs/UnrealBuildTool/Artifacts/HordeStorageArtifactCache.cs
+++ b/Engine/Source/Programs/UnrealBuildTool/Artifacts/HordeStorageArtifactCache.cs
@@ -81,13 +81,14 @@ namespace UnrealBuildTool.Artifacts
 			InteriorChunkedDataNodeOptions interiorOptions = new(1, 10, 20);
 			ChunkingOptions options = new() { LeafOptions = leafOptions, InteriorOptions = interiorOptions };
 
-			ChunkedDataWriter fileWriter = new(writer, options);
+			using ChunkedDataWriter fileWriter = new(writer, options);
 			int index = 0;
 			foreach (ArtifactFile artifact in ArtifactAction.Outputs)
 			{
 				string outputName = artifact.GetFullPath(ArtifactAction.DirectoryMapping);
 				using FileStream stream = new(outputName, FileMode.Open, FileAccess.Read, FileShare.Read);
-				OutputRefs[index++] = new NodeRef<ChunkedDataNode>(await fileWriter.CreateAsync(stream, leafOptions.TargetSize, cancellationToken));
+				ChunkedData chunkedData = await fileWriter.CreateAsync(stream, leafOptions.TargetSize, cancellationToken);
+				OutputRefs[index++] = new NodeRef<ChunkedDataNode>(chunkedData.Root);
 			}
 		}
 	}
