Horde: Include download speed for log events output when retrieving bundles.

diff --git a/Engine/Source/Programs/Shared/EpicGames.Horde/Storage/Nodes/ChunkedDataNode.cs b/Engine/Source/Programs/Shared/EpicGames.Horde/Storage/Nodes/ChunkedDataNode.cs
--- a/Engine/Source/Programs/Shared/EpicGames.Horde/Storage/Nodes/ChunkedDataNode.cs
+++ b/Engine/Source/Programs/Shared/EpicGames.Horde/Storage/Nodes/ChunkedDataNode.cs
@@ -161,7 +161,7 @@ namespace EpicGames.Horde.Storage.Nodes
 		{
 			using (FileStream stream = FileReference.Open(file, FileMode.Open, FileAccess.Read, FileShare.Read))
 			{
-				return await CreateFromStreamAsync(writer, stream, options, cancellationToken);
+				return await CreateFromStreamAsync(writer, stream, options, null, cancellationToken);
 			}
 		}
 
@@ -177,7 +177,7 @@ namespace EpicGames.Horde.Storage.Nodes
 		{
 			using (FileStream stream = file.Open(FileMode.Open, FileAccess.Read, FileShare.Read))
 			{
-				return await CreateFromStreamAsync(writer, stream, options, cancellationToken);
+				return await CreateFromStreamAsync(writer, stream, options, null, cancellationToken);
 			}
 		}
 
@@ -189,7 +189,21 @@ namespace EpicGames.Horde.Storage.Nodes
 		/// <param name="options">Options for finding chunk boundaries</param>
 		/// <param name="cancellationToken">Cancellation token for the operation</param>
 		/// <returns>Hash of the full file data</returns>
-		public static async Task<LeafChunkedData> CreateFromStreamAsync(IStorageWriter writer, Stream stream, LeafChunkedDataNodeOptions options, CancellationToken cancellationToken)
+		public static Task<LeafChunkedData> CreateFromStreamAsync(IStorageWriter writer, Stream stream, LeafChunkedDataNodeOptions options, CancellationToken cancellationToken)
+		{
+			return CreateFromStreamAsync(writer, stream, options, null, cancellationToken);
+		}
+
+		/// <summary>
+		/// Creates nodes from the given file
+		/// </summary>
+		/// <param name="writer">Writer for output nodes</param>
+		/// <param name="stream">Stream to read from</param>
+		/// <param name="options">Options for finding chunk boundaries</param>
+		/// <param name="copyStats">Stats for the copy operation</param>
+		/// <param name="cancellationToken">Cancellation token for the operation</param>
+		/// <returns>Hash of the full file data</returns>
+		internal static async Task<LeafChunkedData> CreateFromStreamAsync(IStorageWriter writer, Stream stream, LeafChunkedDataNodeOptions options, CopyStats? copyStats, CancellationToken cancellationToken)
 		{
 			using Blake3.Hasher hasher = Blake3.Hasher.New();
 			using IMemoryOwner<byte> readBuffer = MemoryPool<byte>.Shared.Rent(options.MaxSize);
@@ -197,6 +211,7 @@ namespace EpicGames.Horde.Storage.Nodes
 			List<NodeRef<ChunkedDataNode>> handles = new List<NodeRef<ChunkedDataNode>>();
 
 			int size = 0;
+			int sizeSinceProgressUpdate = 0;
 			for (; ; )
 			{
 				size += await stream.ReadGreedyAsync(readBuffer.Memory.Slice(size), cancellationToken);
@@ -211,6 +226,13 @@ namespace EpicGames.Horde.Storage.Nodes
 				readBuffer.Memory.Slice(0, nextLength).CopyTo(outputBuffer);
 				hasher.Update(outputBuffer.Span);
 
+				sizeSinceProgressUpdate += nextLength;
+				if (sizeSinceProgressUpdate > 512 * 1024)
+				{
+					copyStats?.Update(0, sizeSinceProgressUpdate);
+					sizeSinceProgressUpdate = 0;
+				}
+
 				BlobHandle handle = await writer.WriteNodeAsync(nextLength, Array.Empty<BlobHandle>(), GetNodeType<LeafChunkedDataNode>(), cancellationToken);
 				handles.Add(new NodeRef<ChunkedDataNode>(handle));
 
@@ -218,6 +240,8 @@ namespace EpicGames.Horde.Storage.Nodes
 				size -= nextLength;
 			}
 
+			copyStats?.Update(1, sizeSinceProgressUpdate);
+
 			IoHash hash = IoHash.FromBlake3(hasher);
 			return new LeafChunkedData(hash, handles);
 		}
diff --git a/Engine/Source/Programs/Shared/EpicGames.Horde/Storage/Nodes/DirectoryNode.cs b/Engine/Source/Programs/Shared/EpicGames.Horde/Storage/Nodes/DirectoryNode.cs
--- a/Engine/Source/Programs/Shared/EpicGames.Horde/Storage/Nodes/DirectoryNode.cs
+++ b/Engine/Source/Programs/Shared/EpicGames.Horde/Storage/Nodes/DirectoryNode.cs
@@ -41,6 +41,61 @@ namespace EpicGames.Horde.Storage.Nodes
 		/// Total size of data to copy
 		/// </summary>
 		long TotalSize { get; }
+
+		/// <summary>
+		/// Download speed, in bytes per second
+		/// </summary>
+		double DownloadSpeed { get; }
+	}
+
+	/// <summary>
+	/// Reports progress info back to callers
+	/// </summary>
+	class CopyStats : ICopyStats
+	{
+		readonly object _lockObject = new object();
+		readonly Stopwatch _timer = Stopwatch.StartNew();
+		readonly IProgress<ICopyStats> _progress;
+		long _lastTotalSize;
+
+		public int CopiedCount { get; set; }
+		public long CopiedSize { get; set; }
+		public int TotalCount { get; }
+		public long TotalSize { get; }
+		public double DownloadSpeed { get; set; }
+
+		public CopyStats(int totalCount, long totalSize, IProgress<ICopyStats> progress)
+		{
+			TotalCount = totalCount;
+			TotalSize = totalSize;
+			_progress = progress;
+		}
+
+		public void Update(int count, long size)
+		{
+			lock (_lockObject)
+			{
+				CopiedCount += count;
+				CopiedSize += size;
+				if (_timer.Elapsed > TimeSpan.FromSeconds(10.0) || CopiedCount == count || CopiedSize == TotalSize)
+				{
+					DownloadSpeed = (CopiedSize - _lastTotalSize) / _timer.Elapsed.TotalSeconds;
+					_lastTotalSize = CopiedSize;
+
+					_progress.Report(this);
+					_timer.Restart();
+				}
+			}
+		}
+
+		public void Flush()
+		{
+			lock (_lockObject)
+			{
+				_progress.Report(this);
+				_timer.Restart();
+			}
+		}
 	}
 
 	/// <summary>
@@ -60,11 +115,11 @@ namespace EpicGames.Horde.Storage.Nodes
 		{
 			if (stats.TotalCount == 0)
 			{
-				_logger.LogInformation("Copied {NumFiles:n0} files ({Size:n1}/{TotalSize:n1}mb, {Pct}%)", stats.CopiedCount, stats.CopiedSize / (1024.0 * 1024.0), stats.TotalSize / (1024.0 * 1024.0), (int)((Math.Max(stats.CopiedCount, 1) * 100) / Math.Max(stats.TotalCount, 1)));
+				_logger.LogInformation("Copied {NumFiles:n0} files ({Size:n1}/{TotalSize:n1}mb, {Rate:n1}mb/s, {Pct}%)", stats.CopiedCount, stats.CopiedSize / (1024.0 * 1024.0), stats.TotalSize / (1024.0 * 1024.0), stats.DownloadSpeed / (1024.0 * 1024.0), (int)((Math.Max(stats.CopiedSize, 1) * 100) / Math.Max(stats.TotalSize, 1)));
 			}
 			else
 			{
-				_logger.LogInformation("Copied {NumFiles:n0}/{TotalFiles:n0} files ({Size:n1}/{TotalSize:n1}mb, {Pct}%)", stats.CopiedCount, stats.TotalCount, stats.CopiedSize / (1024.0 * 1024.0), stats.TotalSize / (1024.0 * 1024.0), (int)((Math.Max(stats.CopiedCount, 1) * 100) / Math.Max(stats.TotalCount, 1)));
+				_logger.LogInformation("Copied {NumFiles:n0}/{TotalFiles:n0} files ({Size:n1}/{TotalSize:n1}mb, {Rate:n1}mb/s, {Pct}%)", stats.CopiedCount, stats.TotalCount, stats.CopiedSize / (1024.0 * 1024.0), stats.TotalSize / (1024.0 * 1024.0), stats.DownloadSpeed / (1024.0 * 1024.0), (int)((Math.Max(stats.CopiedSize, 1) * 100) / Math.Max(stats.TotalSize, 1)));
 			}
 		}
 	}
@@ -450,51 +505,6 @@ namespace EpicGames.Horde.Storage.Nodes
 
 		#endregion
 
-		/// <summary>
-		/// Reports progress info back to callers
-		/// </summary>
-		class CopyStats : ICopyStats
-		{
-			readonly object _lockObject = new object();
-			readonly Stopwatch _timer = Stopwatch.StartNew();
-			readonly IProgress<ICopyStats> _progress;
-
-			public int CopiedCount { get; set; }
-			public long CopiedSize { get; set; }
-			public int TotalCount { get; }
-			public long TotalSize { get; }
-
-			public CopyStats(int totalCount, long totalSize, IProgress<ICopyStats> progress)
-			{
-				TotalCount = totalCount;
-				TotalSize = totalSize;
-				_progress = progress;
-			}
-
-			public void Update(int count, long size)
-			{
-				lock (_lockObject)
-				{
-					CopiedCount += count;
-					CopiedSize += size;
-					if (_timer.Elapsed > TimeSpan.FromSeconds(10.0) || CopiedCount == count)
-					{
-						_progress.Report(this);
-						_timer.Restart();
-					}
-				}
-			}
-
-			public void Flush()
-			{
-				lock (_lockObject)
-				{
-					_progress.Report(this);
-					_timer.Restart();
-				}
-			}
-		}
-
 		/// <summary>
 		/// Adds files from a flat list of paths
 		/// </summary>
@@ -631,9 +641,8 @@ namespace EpicGames.Horde.Storage.Nodes
 				FileInfo file = files[idx];
 				using (Stream stream = file.OpenRead())
 				{
-					leafChunks[idx] = await LeafChunkedDataNode.CreateFromStreamAsync(writerFork, stream, options.LeafOptions, cancellationToken);
+					leafChunks[idx] = await LeafChunkedDataNode.CreateFromStreamAsync(writerFork, stream, options.LeafOptions, copyStats, cancellationToken);
 				}
-				copyStats?.Update(1, file.Length);
 			}
 			await writerFork.FlushAsync(cancellationToken);
 		}
