Horde: Increase max number of concurrent tasks to 16.

diff --git a/Engine/Source/Programs/Shared/EpicGames.Horde/Storage/Nodes/DirectoryNode.cs b/Engine/Source/Programs/Shared/EpicGames.Horde/Storage/Nodes/DirectoryNode.cs
--- a/Engine/Source/Programs/Shared/EpicGames.Horde/Storage/Nodes/DirectoryNode.cs
+++ b/Engine/Source/Programs/Shared/EpicGames.Horde/Storage/Nodes/DirectoryNode.cs
@@ -780,8 +780,8 @@ namespace EpicGames.Horde.Storage.Nodes
 		/// <param name="cancellationToken">Cancellation token for the operation</param>
 		public async Task CopyToDirectoryAsync(DirectoryInfo directoryInfo, IProgress<ICopyStats>? progress, ILogger logger, CancellationToken cancellationToken)
 		{
-			int NumThreads = Math.Min(1 + (int)(Length / (10 * 1024 * 1024)), 4);
-			logger.LogInformation("Splitting read into {NumThreads} threads", NumThreads);
+			int NumTasks = Math.Min(1 + (int)(Length / (10 * 1024 * 1024)), 16);
+			logger.LogInformation("Splitting read into {NumThreads} threads", NumTasks);
 
 			CopyStats? copyStats = null;
 			if (progress != null)
@@ -793,10 +793,10 @@ namespace EpicGames.Horde.Storage.Nodes
 			try
 			{
 				long offset = 0;
-				for (int threadIdx = 0; threadIdx < NumThreads; threadIdx++)
+				for (int taskIdx = 0; taskIdx < NumTasks; taskIdx++)
 				{
 					long minOffset = offset;
-					long maxOffset = (Length * (threadIdx + 1)) / NumThreads;
+					long maxOffset = (Length * (taskIdx + 1)) / NumTasks;
 					tasks.Add(Task.Run(() => CopyToDirectoryInternalAsync(directoryInfo, minOffset, maxOffset - minOffset, copyStats, logger, cancellationToken), cancellationToken));
 					offset = maxOffset;
 				}
