Horde: Move IStorageClient implementations into the EpicGames.Horde.Storage.Clients namespace.

diff --git a/Engine/Source/Programs/AutomationTool/BuildGraph/Tasks/DeployToolTask.cs b/Engine/Source/Programs/AutomationTool/BuildGraph/Tasks/DeployToolTask.cs
--- a/Engine/Source/Programs/AutomationTool/BuildGraph/Tasks/DeployToolTask.cs
+++ b/Engine/Source/Programs/AutomationTool/BuildGraph/Tasks/DeployToolTask.cs
@@ -13,8 +13,8 @@ using System.Threading.Tasks;
 using System.Xml;
 using UnrealBuildBase;
 using Microsoft.Extensions.Logging;
-using EpicGames.Horde.Storage.Backends;
 using EpicGames.Horde.Storage.Bundles;
+using EpicGames.Horde.Storage.Clients;
 using EpicGames.Horde.Storage;
 using EpicGames.Horde.Storage.Nodes;
 using System.Threading;
diff --git a/Engine/Source/Programs/Horde/Horde.Agent.Tests/JsonRpcLoggerTests.cs b/Engine/Source/Programs/Horde/Horde.Agent.Tests/JsonRpcLoggerTests.cs
--- a/Engine/Source/Programs/Horde/Horde.Agent.Tests/JsonRpcLoggerTests.cs
+++ b/Engine/Source/Programs/Horde/Horde.Agent.Tests/JsonRpcLoggerTests.cs
@@ -8,8 +8,8 @@ using System.Threading;
 using System.Threading.Tasks;
 using EpicGames.Core;
 using EpicGames.Horde.Storage;
-using EpicGames.Horde.Storage.Backends;
 using EpicGames.Horde.Storage.Bundles;
+using EpicGames.Horde.Storage.Clients;
 using EpicGames.Horde.Logs;
 using Horde.Agent.Parser;
 using Horde.Agent.Utility;
diff --git a/Engine/Source/Programs/Horde/Horde.Agent/Commands/Compute/CppComputeCommand.cs b/Engine/Source/Programs/Horde/Horde.Agent/Commands/Compute/CppComputeCommand.cs
--- a/Engine/Source/Programs/Horde/Horde.Agent/Commands/Compute/CppComputeCommand.cs
+++ b/Engine/Source/Programs/Horde/Horde.Agent/Commands/Compute/CppComputeCommand.cs
@@ -8,8 +8,8 @@ using System.Threading.Tasks;
 using EpicGames.Core;
 using EpicGames.Horde.Compute;
 using EpicGames.Horde.Storage;
-using EpicGames.Horde.Storage.Backends;
 using EpicGames.Horde.Storage.Bundles;
+using EpicGames.Horde.Storage.Clients;
 using EpicGames.Horde.Storage.Nodes;
 using Horde.Agent.Parser;
 using Microsoft.Extensions.Logging;
diff --git a/Engine/Source/Programs/Horde/Horde.Agent/Program.cs b/Engine/Source/Programs/Horde/Horde.Agent/Program.cs
--- a/Engine/Source/Programs/Horde/Horde.Agent/Program.cs
+++ b/Engine/Source/Programs/Horde/Horde.Agent/Program.cs
@@ -20,7 +20,7 @@ using EpicGames.Core;
 using EpicGames.Horde.Compute;
 using EpicGames.Horde.Compute.Buffers;
 using EpicGames.Horde.Storage;
-using EpicGames.Horde.Storage.Backends;
+using EpicGames.Horde.Storage.Clients;
 using Horde.Agent.Execution;
 using Horde.Agent.Leases;
 using Horde.Agent.Leases.Handlers;
diff --git a/Engine/Source/Programs/Horde/Horde.Agent/Services/ServerStorageFactory.cs b/Engine/Source/Programs/Horde/Horde.Agent/Services/ServerStorageFactory.cs
--- a/Engine/Source/Programs/Horde/Horde.Agent/Services/ServerStorageFactory.cs
+++ b/Engine/Source/Programs/Horde/Horde.Agent/Services/ServerStorageFactory.cs
@@ -4,8 +4,8 @@ using System;
 using System.Net.Http;
 using System.Net.Http.Headers;
 using EpicGames.Horde.Storage;
-using EpicGames.Horde.Storage.Backends;
 using EpicGames.Horde.Storage.Bundles;
+using EpicGames.Horde.Storage.Clients;
 using Microsoft.Extensions.Caching.Memory;
 using Microsoft.Extensions.Logging;
 
diff --git a/Engine/Source/Programs/Horde/Horde.Agent/Utility/StorageClientFactory.cs b/Engine/Source/Programs/Horde/Horde.Agent/Utility/StorageClientFactory.cs
--- a/Engine/Source/Programs/Horde/Horde.Agent/Utility/StorageClientFactory.cs
+++ b/Engine/Source/Programs/Horde/Horde.Agent/Utility/StorageClientFactory.cs
@@ -5,7 +5,7 @@ using System.Net.Http.Headers;
 using System.Net.Http;
 using EpicGames.Horde.Storage;
 using Microsoft.Extensions.Logging;
-using EpicGames.Horde.Storage.Backends;
+using EpicGames.Horde.Storage.Clients;
 using Microsoft.Extensions.Options;
 using EpicGames.Core;
 using System.Threading;
diff --git a/Engine/Source/Programs/Horde/Horde.Agent/Utility/TempStorage.cs b/Engine/Source/Programs/Horde/Horde.Agent/Utility/TempStorage.cs
--- a/Engine/Source/Programs/Horde/Horde.Agent/Utility/TempStorage.cs
+++ b/Engine/Source/Programs/Horde/Horde.Agent/Utility/TempStorage.cs
@@ -2,7 +2,7 @@
 
 using EpicGames.Core;
 using EpicGames.Horde.Storage;
-using EpicGames.Horde.Storage.Bundles;
+using EpicGames.Horde.Storage.Clients;
 using EpicGames.Horde.Storage.Nodes;
 using Microsoft.Extensions.Logging;
 using System;
diff --git a/Engine/Source/Programs/Horde/Horde.Server.Tests/BlobStoreTests.cs b/Engine/Source/Programs/Horde/Horde.Server.Tests/BlobStoreTests.cs
--- a/Engine/Source/Programs/Horde/Horde.Server.Tests/BlobStoreTests.cs
+++ b/Engine/Source/Programs/Horde/Horde.Server.Tests/BlobStoreTests.cs
@@ -9,6 +9,7 @@ using Microsoft.VisualStudio.TestTools.UnitTesting;
 using MongoDB.Driver;
 using EpicGames.Horde.Storage;
 using EpicGames.Horde.Storage.Bundles;
+using EpicGames.Horde.Storage.Clients;
 using System.Buffers;
 using System.Threading;
 using Horde.Server.Server;
diff --git a/Engine/Source/Programs/Horde/Horde.Server.Tests/GcServiceTests.cs b/Engine/Source/Programs/Horde/Horde.Server.Tests/GcServiceTests.cs
--- a/Engine/Source/Programs/Horde/Horde.Server.Tests/GcServiceTests.cs
+++ b/Engine/Source/Programs/Horde/Horde.Server.Tests/GcServiceTests.cs
@@ -8,6 +8,7 @@ using EpicGames.Core;
 using Microsoft.VisualStudio.TestTools.UnitTesting;
 using EpicGames.Horde.Storage;
 using EpicGames.Horde.Storage.Bundles;
+using EpicGames.Horde.Storage.Clients;
 using Horde.Server.Storage;
 using System.Threading;
 using Horde.Server.Server;
diff --git a/Engine/Source/Programs/Horde/Horde.Server/Commands/Bundles/CreateCommand.cs b/Engine/Source/Programs/Horde/Horde.Server/Commands/Bundles/CreateCommand.cs
--- a/Engine/Source/Programs/Horde/Horde.Server/Commands/Bundles/CreateCommand.cs
+++ b/Engine/Source/Programs/Horde/Horde.Server/Commands/Bundles/CreateCommand.cs
@@ -4,7 +4,7 @@ using System.Threading;
 using System.Threading.Tasks;
 using EpicGames.Core;
 using EpicGames.Horde.Storage;
-using EpicGames.Horde.Storage.Backends;
+using EpicGames.Horde.Storage.Clients;
 using EpicGames.Horde.Storage.Nodes;
 using Horde.Server.Perforce;
 using Horde.Server.Storage;
diff --git a/Engine/Source/Programs/Horde/Horde.Server/Commands/Install/SetupCommand.cs b/Engine/Source/Programs/Horde/Horde.Server/Commands/Install/SetupCommand.cs
--- a/Engine/Source/Programs/Horde/Horde.Server/Commands/Install/SetupCommand.cs
+++ b/Engine/Source/Programs/Horde/Horde.Server/Commands/Install/SetupCommand.cs
@@ -14,7 +14,7 @@ using System.Text.RegularExpressions;
 using System.Threading.Tasks;
 using EpicGames.Core;
 using EpicGames.Horde.Storage;
-using EpicGames.Horde.Storage.Backends;
+using EpicGames.Horde.Storage.Clients;
 using EpicGames.Horde.Storage.Nodes;
 using Horde.Server.Agents;
 using Horde.Server.Utilities;
diff --git a/Engine/Source/Programs/Horde/Horde.Server/Program.cs b/Engine/Source/Programs/Horde/Horde.Server/Program.cs
--- a/Engine/Source/Programs/Horde/Horde.Server/Program.cs
+++ b/Engine/Source/Programs/Horde/Horde.Server/Program.cs
@@ -4,12 +4,18 @@ using System;
 using System.Collections.Generic;
 using System.Diagnostics;
 using System.IO;
+using System.Net.Http;
+using System.Net.Http.Headers;
 using System.Reflection;
 using System.Runtime.InteropServices;
+using System.Threading;
 using System.Threading.Tasks;
 using EpicGames.Core;
+using Grpc.Core;
+using Grpc.Net.Client;
 using Horde.Server.Commands;
 using Horde.Server.Utilities;
+using HordeCommon.Rpc;
 using Microsoft.Extensions.Configuration;
 using Microsoft.Extensions.DependencyInjection;
 using Microsoft.Extensions.Hosting;
diff --git a/Engine/Source/Programs/Horde/Horde.Server/Storage/StorageController.cs b/Engine/Source/Programs/Horde/Horde.Server/Storage/StorageController.cs
--- a/Engine/Source/Programs/Horde/Horde.Server/Storage/StorageController.cs
+++ b/Engine/Source/Programs/Horde/Horde.Server/Storage/StorageController.cs
@@ -14,6 +14,7 @@ using Amazon.EC2.Model;
 using EpicGames.Core;
 using EpicGames.Horde.Storage;
 using EpicGames.Horde.Storage.Bundles;
+using EpicGames.Horde.Storage.Clients;
 using EpicGames.Horde.Storage.Nodes;
 using EpicGames.Redis;
 using Horde.Server.Acls;
diff --git a/Engine/Source/Programs/Horde/Horde.Server/Storage/StorageService.cs b/Engine/Source/Programs/Horde/Horde.Server/Storage/StorageService.cs
--- a/Engine/Source/Programs/Horde/Horde.Server/Storage/StorageService.cs
+++ b/Engine/Source/Programs/Horde/Horde.Server/Storage/StorageService.cs
@@ -2,7 +2,6 @@
 
 using System;
 using System.Buffers;
-using System.Collections;
 using System.Collections.Generic;
 using System.IO;
 using System.Linq;
@@ -12,6 +11,7 @@ using System.Threading.Tasks;
 using EpicGames.Core;
 using EpicGames.Horde.Storage;
 using EpicGames.Horde.Storage.Bundles;
+using EpicGames.Horde.Storage.Clients;
 using EpicGames.Redis;
 using EpicGames.Redis.Utility;
 using Horde.Server.Server;
diff --git a/Engine/Source/Programs/Horde/Horde.Server/Tools/IToolCollection.cs b/Engine/Source/Programs/Horde/Horde.Server/Tools/IToolCollection.cs
--- a/Engine/Source/Programs/Horde/Horde.Server/Tools/IToolCollection.cs
+++ b/Engine/Source/Programs/Horde/Horde.Server/Tools/IToolCollection.cs
@@ -5,6 +5,7 @@ using System.Threading;
 using System.Threading.Tasks;
 using EpicGames.Horde.Storage;
 using EpicGames.Horde.Storage.Bundles;
+using EpicGames.Horde.Storage.Clients;
 using Horde.Server.Server;
 using Horde.Server.Storage;
 
diff --git a/Engine/Source/Programs/Horde/Horde.Server/Tools/ToolCollection.cs b/Engine/Source/Programs/Horde/Horde.Server/Tools/ToolCollection.cs
--- a/Engine/Source/Programs/Horde/Horde.Server/Tools/ToolCollection.cs
+++ b/Engine/Source/Programs/Horde/Horde.Server/Tools/ToolCollection.cs
@@ -2,8 +2,8 @@
 
 using EpicGames.Core;
 using EpicGames.Horde.Storage;
-using EpicGames.Horde.Storage.Backends;
 using EpicGames.Horde.Storage.Bundles;
+using EpicGames.Horde.Storage.Clients;
 using EpicGames.Horde.Storage.Nodes;
 using EpicGames.Serialization;
 using Horde.Server.Server;
diff --git a/Engine/Source/Programs/Horde/Horde.Server/Tools/ToolsController.cs b/Engine/Source/Programs/Horde/Horde.Server/Tools/ToolsController.cs
--- a/Engine/Source/Programs/Horde/Horde.Server/Tools/ToolsController.cs
+++ b/Engine/Source/Programs/Horde/Horde.Server/Tools/ToolsController.cs
@@ -3,6 +3,7 @@
 using EpicGames.Core;
 using EpicGames.Horde.Storage;
 using EpicGames.Horde.Storage.Bundles;
+using EpicGames.Horde.Storage.Clients;
 using EpicGames.Horde.Storage.Nodes;
 using Horde.Server.Acls;
 using Horde.Server.Server;
diff --git a/Engine/Source/Programs/Horde/Horde/Commands/Bundles/BundleCreate.cs b/Engine/Source/Programs/Horde/Horde/Commands/Bundles/BundleCreate.cs
--- a/Engine/Source/Programs/Horde/Horde/Commands/Bundles/BundleCreate.cs
+++ b/Engine/Source/Programs/Horde/Horde/Commands/Bundles/BundleCreate.cs
@@ -3,8 +3,8 @@
 using System.Diagnostics;
 using EpicGames.Core;
 using EpicGames.Horde.Storage;
-using EpicGames.Horde.Storage.Backends;
 using EpicGames.Horde.Storage.Bundles;
+using EpicGames.Horde.Storage.Clients;
 using EpicGames.Horde.Storage.Nodes;
 using Microsoft.Extensions.Logging;
 
diff --git a/Engine/Source/Programs/Horde/Horde/Commands/Bundles/BundleExtract.cs b/Engine/Source/Programs/Horde/Horde/Commands/Bundles/BundleExtract.cs
--- a/Engine/Source/Programs/Horde/Horde/Commands/Bundles/BundleExtract.cs
+++ b/Engine/Source/Programs/Horde/Horde/Commands/Bundles/BundleExtract.cs
@@ -3,8 +3,8 @@
 using System.Diagnostics;
 using EpicGames.Core;
 using EpicGames.Horde.Storage;
-using EpicGames.Horde.Storage.Backends;
 using EpicGames.Horde.Storage.Bundles;
+using EpicGames.Horde.Storage.Clients;
 using EpicGames.Horde.Storage.Nodes;
 using Microsoft.Extensions.Logging;
 
diff --git a/Engine/Source/Programs/Horde/Horde/Commands/Compute/ComputeRun.cs b/Engine/Source/Programs/Horde/Horde/Commands/Compute/ComputeRun.cs
--- a/Engine/Source/Programs/Horde/Horde/Commands/Compute/ComputeRun.cs
+++ b/Engine/Source/Programs/Horde/Horde/Commands/Compute/ComputeRun.cs
@@ -7,8 +7,8 @@ using EpicGames.Core;
 using EpicGames.Horde.Common;
 using EpicGames.Horde.Compute;
 using EpicGames.Horde.Compute.Clients;
-using EpicGames.Horde.Storage.Backends;
 using EpicGames.Horde.Storage.Bundles;
+using EpicGames.Horde.Storage.Clients;
 using EpicGames.Horde.Storage.Nodes;
 using Microsoft.Extensions.Logging;
 
diff --git a/Engine/Source/Programs/Horde/Horde/Commands/Workspace/WorkspaceDrop.cs b/Engine/Source/Programs/Horde/Horde/Commands/Workspace/WorkspaceDrop.cs
--- a/Engine/Source/Programs/Horde/Horde/Commands/Workspace/WorkspaceDrop.cs
+++ b/Engine/Source/Programs/Horde/Horde/Commands/Workspace/WorkspaceDrop.cs
@@ -3,8 +3,8 @@
 using System.Diagnostics;
 using EpicGames.Core;
 using EpicGames.Horde.Storage;
-using EpicGames.Horde.Storage.Backends;
 using EpicGames.Horde.Storage.Bundles;
+using EpicGames.Horde.Storage.Clients;
 using EpicGames.Horde.Storage.Nodes;
 using Microsoft.Extensions.Logging;
 
diff --git a/Engine/Source/Programs/Horde/Horde/Commands/Workspace/WorkspaceSync.cs b/Engine/Source/Programs/Horde/Horde/Commands/Workspace/WorkspaceSync.cs
--- a/Engine/Source/Programs/Horde/Horde/Commands/Workspace/WorkspaceSync.cs
+++ b/Engine/Source/Programs/Horde/Horde/Commands/Workspace/WorkspaceSync.cs
@@ -3,8 +3,8 @@
 using System.Diagnostics;
 using EpicGames.Core;
 using EpicGames.Horde.Storage;
-using EpicGames.Horde.Storage.Backends;
 using EpicGames.Horde.Storage.Bundles;
+using EpicGames.Horde.Storage.Clients;
 using EpicGames.Horde.Storage.Nodes;
 using Microsoft.Extensions.Logging;
 
diff --git a/Engine/Source/Programs/Horde/Samples/RemoteClient/RemoteClientApp.cs b/Engine/Source/Programs/Horde/Samples/RemoteClient/RemoteClientApp.cs
--- a/Engine/Source/Programs/Horde/Samples/RemoteClient/RemoteClientApp.cs
+++ b/Engine/Source/Programs/Horde/Samples/RemoteClient/RemoteClientApp.cs
@@ -8,8 +8,8 @@ using EpicGames.Horde.Common;
 using EpicGames.Horde.Compute;
 using EpicGames.Horde.Compute.Buffers;
 using EpicGames.Horde.Compute.Clients;
-using EpicGames.Horde.Storage.Backends;
 using EpicGames.Horde.Storage.Bundles;
+using EpicGames.Horde.Storage.Clients;
 using EpicGames.Horde.Storage.Nodes;
 using EpicGames.OIDC;
 using Microsoft.Extensions.Configuration;
diff --git a/Engine/Source/Programs/Shared/EpicGames.Horde.Tests/BundleStoreTests.cs b/Engine/Source/Programs/Shared/EpicGames.Horde.Tests/BundleStoreTests.cs
--- a/Engine/Source/Programs/Shared/EpicGames.Horde.Tests/BundleStoreTests.cs
+++ b/Engine/Source/Programs/Shared/EpicGames.Horde.Tests/BundleStoreTests.cs
@@ -1,7 +1,7 @@
 // Copyright Epic Games, Inc. All Rights Reserved.
 
 using EpicGames.Horde.Storage;
-using EpicGames.Horde.Storage.Backends;
+using EpicGames.Horde.Storage.Clients;
 using EpicGames.Horde.Storage.Nodes;
 using Microsoft.VisualStudio.TestTools.UnitTesting;
 using System;
diff --git a/Engine/Source/Programs/Shared/EpicGames.Horde.Tests/BundleTests.cs b/Engine/Source/Programs/Shared/EpicGames.Horde.Tests/BundleTests.cs
--- a/Engine/Source/Programs/Shared/EpicGames.Horde.Tests/BundleTests.cs
+++ b/Engine/Source/Programs/Shared/EpicGames.Horde.Tests/BundleTests.cs
@@ -2,8 +2,8 @@
 
 using EpicGames.Core;
 using EpicGames.Horde.Storage;
-using EpicGames.Horde.Storage.Backends;
 using EpicGames.Horde.Storage.Bundles;
+using EpicGames.Horde.Storage.Clients;
 using EpicGames.Horde.Storage.Nodes;
 using Microsoft.Extensions.Caching.Memory;
 using Microsoft.Extensions.Logging.Abstractions;
diff --git a/Engine/Source/Programs/Shared/EpicGames.Horde.Tests/ChunkingTests.cs b/Engine/Source/Programs/Shared/EpicGames.Horde.Tests/ChunkingTests.cs
--- a/Engine/Source/Programs/Shared/EpicGames.Horde.Tests/ChunkingTests.cs
+++ b/Engine/Source/Programs/Shared/EpicGames.Horde.Tests/ChunkingTests.cs
@@ -9,7 +9,7 @@ using System.Threading.Tasks;
 using System.IO;
 using System.Threading;
 using EpicGames.Horde.Storage;
-using EpicGames.Horde.Storage.Backends;
+using EpicGames.Horde.Storage.Clients;
 using Microsoft.Extensions.Caching.Memory;
 using Microsoft.Extensions.Logging.Abstractions;
 using System.Collections.Generic;
diff --git a/Engine/Source/Programs/Shared/EpicGames.Horde.Tests/ComputeSocketTests.cs b/Engine/Source/Programs/Shared/EpicGames.Horde.Tests/ComputeSocketTests.cs
--- a/Engine/Source/Programs/Shared/EpicGames.Horde.Tests/ComputeSocketTests.cs
+++ b/Engine/Source/Programs/Shared/EpicGames.Horde.Tests/ComputeSocketTests.cs
@@ -1,10 +1,8 @@
 // Copyright Epic Games, Inc. All Rights Reserved.
 
 using System;
-using System.Buffers;
 using System.Collections.Generic;
 using System.IO.Pipelines;
-using System.IO.Pipes;
 using System.Linq;
 using System.Net;
 using System.Net.Sockets;
@@ -13,13 +11,11 @@ using System.Threading;
 using System.Threading.Tasks;
 using EpicGames.Core;
 using EpicGames.Horde.Compute;
-using EpicGames.Horde.Compute.Buffers;
 using EpicGames.Horde.Compute.Transports;
 using EpicGames.Horde.Storage;
-using EpicGames.Horde.Storage.Backends;
 using EpicGames.Horde.Storage.Bundles;
+using EpicGames.Horde.Storage.Clients;
 using EpicGames.Horde.Storage.Nodes;
-using Microsoft.Extensions.Caching.Memory;
 using Microsoft.Extensions.Logging;
 using Microsoft.Extensions.Logging.Abstractions;
 using Microsoft.VisualStudio.TestTools.UnitTesting;
diff --git a/Engine/Source/Programs/Shared/EpicGames.Horde.Tests/LogTests.cs b/Engine/Source/Programs/Shared/EpicGames.Horde.Tests/LogTests.cs
--- a/Engine/Source/Programs/Shared/EpicGames.Horde.Tests/LogTests.cs
+++ b/Engine/Source/Programs/Shared/EpicGames.Horde.Tests/LogTests.cs
@@ -9,7 +9,7 @@ using System.Threading;
 using System.Threading.Tasks;
 using EpicGames.Horde.Logs;
 using EpicGames.Horde.Storage;
-using EpicGames.Horde.Storage.Backends;
+using EpicGames.Horde.Storage.Clients;
 using EpicGames.Horde.Tests.Properties;
 using Microsoft.Extensions.Caching.Memory;
 using Microsoft.Extensions.Logging.Abstractions;
diff --git a/Engine/Source/Programs/Shared/EpicGames.Horde/Compute/AgentMessage.cs b/Engine/Source/Programs/Shared/EpicGames.Horde/Compute/AgentMessage.cs
--- a/Engine/Source/Programs/Shared/EpicGames.Horde/Compute/AgentMessage.cs
+++ b/Engine/Source/Programs/Shared/EpicGames.Horde/Compute/AgentMessage.cs
@@ -8,8 +8,8 @@ using System.IO;
 using System.Threading;
 using System.Threading.Tasks;
 using EpicGames.Core;
-using EpicGames.Horde.Storage;
 using EpicGames.Horde.Storage.Bundles;
+using EpicGames.Horde.Storage.Clients;
 using EpicGames.Horde.Storage.Nodes;
 
 namespace EpicGames.Horde.Compute
diff --git a/Engine/Source/Programs/Shared/EpicGames.Horde/Compute/AgentStorageClient.cs b/Engine/Source/Programs/Shared/EpicGames.Horde/Compute/AgentStorageClient.cs
--- a/Engine/Source/Programs/Shared/EpicGames.Horde/Compute/AgentStorageClient.cs
+++ b/Engine/Source/Programs/Shared/EpicGames.Horde/Compute/AgentStorageClient.cs
@@ -9,6 +9,7 @@ using EpicGames.Core;
 using EpicGames.Horde.Storage;
 using Microsoft.Extensions.Logging.Abstractions;
 using EpicGames.Horde.Storage.Bundles;
+using EpicGames.Horde.Storage.Clients;
 
 namespace EpicGames.Horde.Compute
 {
diff --git a/Engine/Source/Programs/Shared/EpicGames.Horde/HordeHttpClient.cs b/Engine/Source/Programs/Shared/EpicGames.Horde/HordeHttpClient.cs
--- a/Engine/Source/Programs/Shared/EpicGames.Horde/HordeHttpClient.cs
+++ b/Engine/Source/Programs/Shared/EpicGames.Horde/HordeHttpClient.cs
@@ -9,7 +9,7 @@ using System.Threading;
 using System.Threading.Tasks;
 using EpicGames.Core;
 using EpicGames.Horde.Storage;
-using EpicGames.Horde.Storage.Backends;
+using EpicGames.Horde.Storage.Clients;
 using Microsoft.Extensions.Http;
 using Microsoft.Extensions.Logging;
 using Polly;
diff --git a/Engine/Source/Programs/Shared/EpicGames.Horde/Storage/Backends/FileStorageClient.cs b/Engine/Source/Programs/Shared/EpicGames.Horde/Storage/Backends/FileStorageClient.cs
deleted file mode 100644
diff --git a/Engine/Source/Programs/Shared/EpicGames.Horde/Storage/Backends/HttpStorageClient.cs b/Engine/Source/Programs/Shared/EpicGames.Horde/Storage/Backends/HttpStorageClient.cs
deleted file mode 100644
diff --git a/Engine/Source/Programs/Shared/EpicGames.Horde/Storage/Backends/MemoryStorageClient.cs b/Engine/Source/Programs/Shared/EpicGames.Horde/Storage/Backends/MemoryStorageClient.cs
deleted file mode 100644
diff --git a/Engine/Source/Programs/Shared/EpicGames.Horde/Storage/Bundles/BundleReader.cs b/Engine/Source/Programs/Shared/EpicGames.Horde/Storage/Bundles/BundleReader.cs
--- a/Engine/Source/Programs/Shared/EpicGames.Horde/Storage/Bundles/BundleReader.cs
+++ b/Engine/Source/Programs/Shared/EpicGames.Horde/Storage/Bundles/BundleReader.cs
@@ -6,11 +6,11 @@ using System.Linq;
 using System.Threading;
 using System.Threading.Tasks;
 using EpicGames.Core;
+using EpicGames.Horde.Storage.Clients;
 using Microsoft.Extensions.Logging;
 using Microsoft.CodeAnalysis;
 using System.Diagnostics;
 using System.IO;
-using System.Diagnostics.CodeAnalysis;
 
 namespace EpicGames.Horde.Storage.Bundles
 {
diff --git a/Engine/Source/Programs/Shared/EpicGames.Horde/Storage/Bundles/BundleStorageClient.cs b/Engine/Source/Programs/Shared/EpicGames.Horde/Storage/Bundles/BundleStorageClient.cs
deleted file mode 100644
diff --git a/Engine/Source/Programs/Shared/EpicGames.Horde/Storage/Bundles/BundleWriter.cs b/Engine/Source/Programs/Shared/EpicGames.Horde/Storage/Bundles/BundleWriter.cs
--- a/Engine/Source/Programs/Shared/EpicGames.Horde/Storage/Bundles/BundleWriter.cs
+++ b/Engine/Source/Programs/Shared/EpicGames.Horde/Storage/Bundles/BundleWriter.cs
@@ -9,6 +9,7 @@ using System.Linq;
 using System.Threading;
 using System.Threading.Tasks;
 using EpicGames.Core;
+using EpicGames.Horde.Storage.Clients;
 using Microsoft.Extensions.Logging;
 
 namespace EpicGames.Horde.Storage.Bundles
diff --git a/Engine/Source/Programs/Shared/EpicGames.Horde/Storage/Clients/BundleStorageClient.cs b/Engine/Source/Programs/Shared/EpicGames.Horde/Storage/Clients/BundleStorageClient.cs
new file mode 100644
--- /dev/null
+++ b/Engine/Source/Programs/Shared/EpicGames.Horde/Storage/Clients/BundleStorageClient.cs
@@ -0,0 +1,141 @@
+// Copyright Epic Games, Inc. All Rights Reserved.
+
+using System.Collections.Generic;
+using System.IO;
+using System.Runtime.CompilerServices;
+using System.Threading;
+using System.Threading.Tasks;
+using EpicGames.Core;
+using EpicGames.Horde.Storage.Bundles;
+using Microsoft.Extensions.Logging;
+
+namespace EpicGames.Horde.Storage.Clients
+{
+	/// <summary>
+	/// Base class for an implementation of <see cref="IStorageClient"/>, providing implementations for some common functionality using bundles.
+	/// </summary>
+	public abstract class BundleStorageClient : IStorageClient
+	{
+		/// <summary>
+		/// Reader for node data
+		/// </summary>
+		public BundleReader BundleReader { get; }
+
+		/// <summary>
+		/// Constructor
+		/// </summary>
+		protected BundleStorageClient(StorageCache cache, ILogger logger)
+		{
+			BundleReader = new BundleReader(this, cache, logger);
+		}
+
+		#region Blobs
+
+		/// <summary>
+		/// Opens a bundle stream for reading
+		/// </summary>
+		/// <param name="locator">Locator for the bundle</param>
+		/// <param name="cancellationToken">Cancellation token for the operation</param>
+		/// <returns>Stream for reading from the bundle</returns>
+		public Task<Stream> OpenAsync(BundleLocator locator, CancellationToken cancellationToken = default) => OpenAsync(locator, 0, null, cancellationToken);
+
+		/// <inheritdoc/>
+		public abstract Task<Stream> OpenAsync(BundleLocator locator, int offset, int? length = null, CancellationToken cancellationToken = default);
+
+		/// <summary>
+		/// Reads an entire bundle into memory
+		/// </summary>
+		/// <param name="locator">Locator for the bundle</param>
+		/// <param name="cancellationToken">Cancellation token for the operation</param>
+		/// <returns>Bundle that was read</returns>
+		public async Task<Bundle> ReadBundleAsync(BundleLocator locator, CancellationToken cancellationToken = default)
+		{
+			using Stream stream = await OpenAsync(locator, cancellationToken);
+			return await Bundle.FromStreamAsync(stream, cancellationToken);
+		}
+
+		/// <inheritdoc/>
+		public abstract Task<BundleLocator> WriteBundleAsync(Bundle bundle, Utf8String prefix = default, CancellationToken cancellationToken = default);
+
+		#endregion
+
+		#region Nodes
+
+		/// <summary>
+		/// Creates a handle to a node from its locator
+		/// </summary>
+		public BlobHandle CreateNodeHandle(BundleNodeLocator locator) => new FlushedNodeHandle(BundleReader, locator);
+
+		/// <inheritdoc/>
+		public BundleWriter CreateWriter(RefName refName = default, BundleOptions? options = null)
+		{
+			return new BundleWriter(this, BundleReader, refName, options);
+		}
+
+		/// <inheritdoc/>
+		IStorageWriter IStorageClient.CreateWriter(RefName refName) => CreateWriter(refName);
+
+		#endregion
+
+		#region Aliases
+
+		/// <inheritdoc/>
+		Task IStorageClient.AddAliasAsync(Utf8String name, BlobHandle handle, int rank, CancellationToken cancellationToken) => AddAliasAsync(name, (BundleNodeHandle)handle, rank, cancellationToken);
+
+		/// <inheritdoc/>
+		public abstract Task AddAliasAsync(Utf8String name, BundleNodeHandle handle, int rank = 0, CancellationToken cancellationToken = default);
+
+		/// <inheritdoc/>
+		Task IStorageClient.RemoveAliasAsync(Utf8String name, BlobHandle handle, CancellationToken cancellationToken) => RemoveAliasAsync(name, (BundleNodeHandle)handle, cancellationToken);
+
+		/// <inheritdoc/>
+		public abstract Task RemoveAliasAsync(Utf8String name, BundleNodeHandle handle, CancellationToken cancellationToken = default);
+
+		/// <inheritdoc/>
+		async IAsyncEnumerable<BlobHandle> IStorageClient.FindAliasAsync(Utf8String name, [EnumeratorCancellation] CancellationToken cancellationToken)
+		{
+			await foreach (BundleNodeHandle handle in FindAliasAsync(name, cancellationToken))
+			{
+				yield return handle;
+			}
+		}
+
+		/// <inheritdoc/>
+		public abstract IAsyncEnumerable<BundleNodeHandle> FindAliasAsync(Utf8String name, CancellationToken cancellationToken = default);
+
+		#endregion
+
+		#region Refs
+
+		/// <inheritdoc/>
+		public abstract Task<bool> DeleteRefAsync(RefName name, CancellationToken cancellationToken = default);
+
+		/// <inheritdoc cref="StorageClientExtensions.ReadRefTargetAsync(IStorageClient, RefName, RefCacheTime, CancellationToken)"/>
+		public async Task<BundleNodeHandle> ReadRefTargetAsync(RefName name, RefCacheTime cacheTime = default, CancellationToken cancellationToken = default)
+		{
+			BundleNodeHandle? refTarget = await TryReadRefTargetAsync(name, cacheTime, cancellationToken);
+			if (refTarget == null)
+			{
+				throw new RefNameNotFoundException(name);
+			}
+			return refTarget;
+		}
+
+		/// <inheritdoc/>
+		public abstract Task<BundleNodeHandle?> TryReadRefTargetAsync(RefName name, RefCacheTime cacheTime = default, CancellationToken cancellationToken = default);
+
+		/// <inheritdoc/>
+		async Task<BlobHandle?> IStorageClient.TryReadRefTargetAsync(RefName name, RefCacheTime cacheTime, CancellationToken cancellationToken) => await TryReadRefTargetAsync(name, cacheTime, cancellationToken);
+
+		/// <inheritdoc/>
+		public Task WriteRefTargetAsync(RefName name, BlobHandle target, RefOptions? options = null, CancellationToken cancellationToken = default)
+		{
+			return WriteRefTargetAsync(name, (BundleNodeHandle)target, options, cancellationToken);
+		}
+
+		/// <inheritdoc/>
+		public abstract Task WriteRefTargetAsync(RefName name, BundleNodeHandle target, RefOptions? options = null, CancellationToken cancellationToken = default);
+
+		#endregion
+	}
+}
diff --git a/Engine/Source/Programs/Shared/EpicGames.Horde/Storage/Clients/FileStorageClient.cs b/Engine/Source/Programs/Shared/EpicGames.Horde/Storage/Clients/FileStorageClient.cs
new file mode 100644
--- /dev/null
+++ b/Engine/Source/Programs/Shared/EpicGames.Horde/Storage/Clients/FileStorageClient.cs
@@ -0,0 +1,167 @@
+// Copyright Epic Games, Inc. All Rights Reserved.
+
+using System;
+using System.Buffers;
+using System.Collections.Generic;
+using System.Diagnostics;
+using System.Globalization;
+using System.IO;
+using System.Linq;
+using System.Runtime.Serialization;
+using System.Threading;
+using System.Threading.Tasks;
+using EpicGames.Core;
+using Microsoft.Extensions.Caching.Memory;
+using Microsoft.Extensions.Logging;
+using EpicGames.Horde.Storage.Bundles;
+
+namespace EpicGames.Horde.Storage.Clients
+{
+	/// <summary>
+	/// Implementation of <see cref="IStorageClient"/> which writes data to files on disk.
+	/// </summary>
+	public class FileStorageClient : BundleStorageClient
+	{
+		readonly DirectoryReference _rootDir;
+		readonly ILogger _logger;
+
+		/// <summary>
+		/// Constructor
+		/// </summary>
+		/// <param name="rootDir">Root directory for storing blobs</param>
+		/// <param name="cache">Memory cache for read data</param>
+		/// <param name="logger">Logger interface</param>
+		public FileStorageClient(DirectoryReference rootDir, StorageCache cache, ILogger logger)
+			: base(cache, logger)
+		{
+			_rootDir = rootDir;
+			_logger = logger;
+
+			DirectoryReference.CreateDirectory(_rootDir);
+		}
+
+		/// <summary>
+		/// Reads a ref from a file on disk
+		/// </summary>
+		public async ValueTask<BlobHandle> ReadRefAsync(FileReference file)
+		{
+			string text = await FileReference.ReadAllTextAsync(file);
+			return new FlushedNodeHandle(BundleReader, BundleNodeLocator.Parse(text));
+		}
+
+		FileReference GetRefFile(RefName name) => FileReference.Combine(_rootDir, name.ToString() + ".ref");
+		FileReference GetBlobFile(BundleLocator id) => FileReference.Combine(_rootDir, id.Path.ToString() + ".blob");
+
+		#region Blobs
+
+		/// <inheritdoc/>
+		public override Task<Stream> OpenAsync(BundleLocator id, int offset, int? length, CancellationToken cancellationToken = default)
+		{
+			FileReference file = GetBlobFile(id);
+			_logger.LogInformation("Reading {File} ({Offset}-{Length})", file, offset, length);
+
+			Stream stream = FileReference.Open(file, FileMode.Open, FileAccess.Read, FileShare.Read | FileShare.Delete);
+			if (offset > 0)
+			{
+				stream.Seek(offset, SeekOrigin.Begin);
+			}
+			return Task.FromResult(stream);
+		}
+
+		/// <inheritdoc/>
+		public override async Task<BundleLocator> WriteBundleAsync(Bundle bundle, Utf8String prefix = default, CancellationToken cancellationToken = default)
+		{
+			BundleLocator id = BundleLocator.CreateUnique(prefix);
+			FileReference file = GetBlobFile(id);
+			DirectoryReference.CreateDirectory(file.Directory);
+			_logger.LogInformation("Writing {File}", file);
+
+			using (FileStream fileStream = FileReference.Open(file, FileMode.Create, FileAccess.Write, FileShare.ReadWrite))
+			{
+				foreach (ReadOnlyMemory<byte> segment in bundle.AsSequence())
+				{
+					await fileStream.WriteAsync(segment, cancellationToken);
+				}
+			}
+
+			return id;
+		}
+
+		#endregion
+
+		#region Aliases
+
+		/// <inheritdoc/>
+		public override Task AddAliasAsync(Utf8String name, BundleNodeHandle locator, int rank = 0, CancellationToken cancellationToken = default)
+		{
+			throw new NotSupportedException("File storage client does not currently support aliases.");
+		}
+
+		/// <inheritdoc/>
+		public override Task RemoveAliasAsync(Utf8String name, BundleNodeHandle locator, CancellationToken cancellationToken = default)
+		{
+			throw new NotSupportedException("File storage client does not currently support aliases.");
+		}
+
+		/// <inheritdoc/>
+		public override IAsyncEnumerable<BundleNodeHandle> FindAliasAsync(Utf8String alias, CancellationToken cancellationToken = default)
+		{
+			throw new NotSupportedException("File storage client does not currently support aliases.");
+		}
+
+		#endregion
+
+		#region Refs
+
+		/// <inheritdoc/>
+		public override Task<bool> DeleteRefAsync(RefName name, CancellationToken cancellationToken = default)
+		{
+			FileInfo file = GetRefFile(name).ToFileInfo();
+			if (file.Exists)
+			{
+				file.Delete();
+				return Task.FromResult(true);
+			}
+			return Task.FromResult(false);
+		}
+
+		/// <inheritdoc/>
+		public override async Task<BundleNodeHandle?> TryReadRefTargetAsync(RefName name, RefCacheTime cacheTime = default, CancellationToken cancellationToken = default)
+		{
+			FileReference file = GetRefFile(name);
+			if (!FileReference.Exists(file))
+			{
+				return null;
+			}
+
+			_logger.LogInformation("Reading {File}", file);
+			string text = await FileReference.ReadAllTextAsync(file, cancellationToken);
+			return new FlushedNodeHandle(BundleReader, BundleNodeLocator.Parse(text));
+		}
+
+		/// <inheritdoc/>
+		public override async Task WriteRefTargetAsync(RefName name, BundleNodeHandle target, RefOptions? options = null, CancellationToken cancellationToken = default)
+		{
+			FileReference file = GetRefFile(name);
+			DirectoryReference.CreateDirectory(file.Directory);
+			_logger.LogInformation("Writing {File}", file);
+
+			BundleNodeLocator locator = await target.FlushAsync(cancellationToken);
+			for (int attempt = 0; ; attempt++)
+			{
+				try
+				{
+					await FileReference.WriteAllTextAsync(file, locator.ToString());
+					break;
+				}
+				catch (IOException ex) when (attempt < 3)
+				{
+					_logger.LogDebug(ex, "Unable to write to {File}; retrying...", file);
+					await Task.Delay(100 * attempt, cancellationToken);
+				}
+			}
+		}
+
+		#endregion
+	}
+}
diff --git a/Engine/Source/Programs/Shared/EpicGames.Horde/Storage/Clients/HttpStorageClient.cs b/Engine/Source/Programs/Shared/EpicGames.Horde/Storage/Clients/HttpStorageClient.cs
new file mode 100644
--- /dev/null
+++ b/Engine/Source/Programs/Shared/EpicGames.Horde/Storage/Clients/HttpStorageClient.cs
@@ -0,0 +1,351 @@
+// Copyright Epic Games, Inc. All Rights Reserved.
+
+using System;
+using System.Collections.Generic;
+using System.IO;
+using System.Net;
+using System.Net.Http;
+using System.Net.Http.Headers;
+using System.Net.Http.Json;
+using System.Runtime.CompilerServices;
+using System.Threading;
+using System.Threading.Tasks;
+using System.Web;
+using EpicGames.Core;
+using Microsoft.Extensions.Caching.Memory;
+using Microsoft.Extensions.Logging;
+using EpicGames.Horde.Storage.Bundles;
+using Microsoft.Extensions.Http;
+using Polly.Extensions.Http;
+using Polly;
+using Polly.Retry;
+
+namespace EpicGames.Horde.Storage.Clients
+{
+	/// <summary>
+	/// Implementation of <see cref="IStorageClient"/> which communicates with an upstream Horde instance via HTTP.
+	/// </summary>
+	public class HttpStorageClient : BundleStorageClient
+	{
+		/// <summary>
+		/// Name of clients created from the http client factory
+		/// </summary>
+		public const string HttpClientName = "Horde.HttpStorageClient";
+
+		class WriteBlobResponse
+		{
+			public BundleLocator Blob { get; set; }
+			public Uri? UploadUrl { get; set; }
+			public bool? SupportsRedirects { get; set; }
+		}
+
+		class FindNodeResponse
+		{
+			public IoHash Hash { get; set; }
+			public BundleLocator Blob { get; set; }
+			public int ExportIdx { get; set; }
+		}
+
+		class FindNodesResponse
+		{
+			public List<FindNodeResponse> Nodes { get; set; } = new List<FindNodeResponse>();
+		}
+
+		class ReadRefResponse
+		{
+			public IoHash Hash { get; set; }
+			public BundleLocator Blob { get; set; }
+			public int ExportIdx { get; set; }
+		}
+
+		readonly Func<HttpClient> _createClient;
+		readonly Func<HttpClient> _createRedirectClient;
+		readonly ILogger _logger;
+		bool _supportsUploadRedirects = true;
+
+		static readonly HttpMessageHandler s_defaultHttpMessageHandler = CreateDefaultHttpMessageHandler();
+
+		/// <summary>
+		/// Constructor
+		/// </summary>
+		public HttpStorageClient(Uri baseAddress, string? bearerToken, StorageCache cache, ILogger logger)
+			: this(() => CreateAuthenticatedClient(null, baseAddress, bearerToken), () => CreateClient(null), cache, logger)
+		{
+		}
+
+		/// <summary>
+		/// Constructor
+		/// </summary>
+		public HttpStorageClient(Func<HttpClient> createClient, Func<HttpClient> createRedirectClient, StorageCache cache, ILogger logger) 
+			: base(cache, logger)
+		{
+			_createClient = createClient;
+			_createRedirectClient = createRedirectClient;
+			_logger = logger;
+		}
+
+		/// <summary>
+		/// Constructor
+		/// </summary>
+		public HttpStorageClient(IHttpClientFactory httpClientFactory, Uri baseAddress, string? bearerToken, StorageCache cache, ILogger logger)
+			: this(() => CreateAuthenticatedClient(httpClientFactory, baseAddress, bearerToken), () => CreateClient(httpClientFactory), cache, logger)
+		{
+		}
+
+		static HttpMessageHandler CreateDefaultHttpMessageHandler()
+		{
+			AsyncRetryPolicy<HttpResponseMessage> retryPolicy = HttpPolicyExtensions
+				.HandleTransientHttpError()
+				.WaitAndRetryAsync(new[] { TimeSpan.FromSeconds(2.0), TimeSpan.FromSeconds(5.0), TimeSpan.FromSeconds(10), TimeSpan.FromSeconds(30) });
+
+			SocketsHttpHandler socketsHandler = new SocketsHttpHandler { PooledConnectionLifetime = TimeSpan.FromMinutes(15) };
+			return new PolicyHttpMessageHandler(retryPolicy) { InnerHandler = socketsHandler };
+		}
+
+		/// <summary>
+		/// Helper method to create an HTTP client from the given factory
+		/// </summary>
+		static HttpClient CreateClient(IHttpClientFactory? httpClientFactory)
+		{
+			if (httpClientFactory == null)
+			{
+				return new HttpClient(s_defaultHttpMessageHandler, disposeHandler: false);
+			}
+			else
+			{
+				return httpClientFactory.CreateClient(HttpClientName);
+			}
+		}
+
+		/// <summary>
+		/// Helper method to add the base address and auth header to an HTTP client
+		/// </summary>
+		static HttpClient CreateAuthenticatedClient(IHttpClientFactory? httpClientFactory, Uri baseAddress, string? bearerToken)
+		{
+			HttpClient httpClient = CreateClient(httpClientFactory);
+			httpClient.BaseAddress = baseAddress;
+			if (bearerToken != null)
+			{
+				httpClient.DefaultRequestHeaders.Authorization = new AuthenticationHeaderValue("Bearer", bearerToken);
+			}
+			return httpClient;
+		}
+
+		#region Blobs
+
+		/// <inheritdoc/>
+		public override async Task<Stream> OpenAsync(BundleLocator locator, int offset, int? length, CancellationToken cancellationToken = default)
+		{
+			if (offset == 0 && length == null)
+			{
+				_logger.LogDebug("Reading {Locator}", locator);
+			}
+			else if (length == null)
+			{
+				_logger.LogDebug("Reading {Locator} ({Offset}..)", locator, offset);
+			}
+			else
+			{
+				_logger.LogDebug("Reading {Locator} ({Offset}+{Length})", locator, offset, length);
+			}
+
+			if (length.HasValue && length.Value == 0)
+			{
+				return new MemoryStream(Array.Empty<byte>());
+			}
+
+			using (HttpClient httpClient = _createClient())
+			{
+				using (HttpRequestMessage request = new HttpRequestMessage(HttpMethod.Get, $"blobs/{locator}"))
+				{
+					if (offset != 0 || length != null)
+					{
+						request.Headers.Range = new RangeHeaderValue(offset, (length == null) ? null : (offset + (length - 1)));
+					}
+
+					HttpResponseMessage response = await httpClient.SendAsync(request, cancellationToken);
+					response.EnsureSuccessStatusCode();
+					return await response.Content.ReadAsStreamAsync(cancellationToken);
+				}
+			}
+		}
+
+		/// <inheritdoc/>
+		public override async Task<BundleLocator> WriteBundleAsync(Bundle bundle, Utf8String prefix = default, CancellationToken cancellationToken = default)
+		{
+			using StreamContent streamContent = new StreamContent(new ReadOnlySequenceStream(bundle.AsSequence()));
+
+			if (_supportsUploadRedirects)
+			{
+				using (HttpClient redirectHttpClient = _createRedirectClient())
+				{
+					WriteBlobResponse redirectResponse = await SendWriteRequestAsync(null, prefix, cancellationToken);
+					if (redirectResponse.UploadUrl != null)
+					{
+						using (HttpResponseMessage uploadResponse = await redirectHttpClient.PutAsync(redirectResponse.UploadUrl, streamContent, cancellationToken))
+						{
+							if (!uploadResponse.IsSuccessStatusCode)
+							{
+								string body = await uploadResponse.Content.ReadAsStringAsync(cancellationToken);
+								throw new StorageException($"Unable to upload data to redirected URL: {body}");
+							}
+						}
+						_logger.LogDebug("Written {Locator} (using redirect)", redirectResponse.Blob);
+						return redirectResponse.Blob;
+					}
+				}
+			}
+
+			WriteBlobResponse response = await SendWriteRequestAsync(streamContent, prefix, cancellationToken);
+			_supportsUploadRedirects = response.SupportsRedirects ?? false;
+			_logger.LogDebug("Written {Locator} (direct)", response.Blob);
+			return response.Blob;
+		}
+
+		async Task<WriteBlobResponse> SendWriteRequestAsync(StreamContent? streamContent, Utf8String prefix = default, CancellationToken cancellationToken = default)
+		{
+			using (HttpClient httpClient = _createClient())
+			{
+				using (HttpRequestMessage request = new HttpRequestMessage(HttpMethod.Post, $"blobs"))
+				{
+					using StringContent stringContent = new StringContent(prefix.ToString());
+
+					MultipartFormDataContent form = new MultipartFormDataContent();
+					if (streamContent != null)
+					{
+						form.Add(streamContent, "file", "filename");
+					}
+					form.Add(stringContent, "prefix");
+
+					request.Content = form;
+					using (HttpResponseMessage response = await httpClient.SendAsync(request, cancellationToken))
+					{
+						response.EnsureSuccessStatusCode();
+						WriteBlobResponse? data = await response.Content.ReadFromJsonAsync<WriteBlobResponse>(cancellationToken: cancellationToken);
+						return data!;
+					}
+				}
+			}
+		}
+
+		#endregion
+
+		#region Nodes
+
+		/// <inheritdoc/>
+		public override Task AddAliasAsync(Utf8String name, BundleNodeHandle locator, int rank = 0, CancellationToken cancellationToken = default)
+		{
+			throw new NotSupportedException("Http storage client does not currently support aliases.");
+		}
+
+		/// <inheritdoc/>
+		public override Task RemoveAliasAsync(Utf8String name, BundleNodeHandle locator, CancellationToken cancellationToken = default)
+		{
+			throw new NotSupportedException("Http storage client does not currently support aliases.");
+		}
+
+		/// <inheritdoc/>
+		public override async IAsyncEnumerable<BundleNodeHandle> FindAliasAsync(Utf8String alias, [EnumeratorCancellation] CancellationToken cancellationToken = default)
+		{
+			_logger.LogDebug("Finding nodes with alias {Alias}", alias);
+			using (HttpClient httpClient = _createClient())
+			{
+				using (HttpRequestMessage request = new HttpRequestMessage(HttpMethod.Get, $"nodes?alias={HttpUtility.UrlEncode(alias.ToString())}"))
+				{
+					using (HttpResponseMessage response = await httpClient.SendAsync(request, cancellationToken))
+					{
+						response.EnsureSuccessStatusCode();
+
+						FindNodesResponse? message = await response.Content.ReadFromJsonAsync<FindNodesResponse>(cancellationToken: cancellationToken);
+						foreach (FindNodeResponse node in message!.Nodes)
+						{
+							yield return new FlushedNodeHandle(BundleReader, new BundleNodeLocator(node.Hash, node.Blob, node.ExportIdx));
+						}
+					}
+				}
+			}
+		}
+
+		#endregion
+
+		#region Refs
+
+		/// <inheritdoc/>
+		public override async Task<bool> DeleteRefAsync(RefName name, CancellationToken cancellationToken)
+		{
+			_logger.LogDebug("Deleting ref {RefName}", name);
+			using (HttpClient httpClient = _createClient())
+			{
+				using (HttpRequestMessage request = new HttpRequestMessage(HttpMethod.Delete, $"refs/{name}"))
+				{
+					using (HttpResponseMessage response = await httpClient.SendAsync(request, cancellationToken))
+					{
+						if (response.IsSuccessStatusCode)
+						{
+							return true;
+						}
+						if (response.StatusCode == HttpStatusCode.NotFound)
+						{
+							return false;
+						}
+
+						response.EnsureSuccessStatusCode();
+						return false;
+					}
+				}
+			}
+		}
+
+		/// <inheritdoc/>
+		public override async Task<BundleNodeHandle?> TryReadRefTargetAsync(RefName name, RefCacheTime cacheTime = default, CancellationToken cancellationToken = default)
+		{
+			using (HttpClient httpClient = _createClient())
+			{
+				using (HttpRequestMessage request = new HttpRequestMessage(HttpMethod.Get, $"refs/{name}"))
+				{
+					if (cacheTime.IsSet())
+					{
+						request.Headers.CacheControl = new CacheControlHeaderValue { MaxAge = cacheTime.MaxAge };
+					}
+
+					using (HttpResponseMessage response = await httpClient.SendAsync(request, cancellationToken))
+					{
+						if (response.StatusCode == HttpStatusCode.NotFound)
+						{
+							_logger.LogDebug("Read ref {RefName} -> None", name);
+							return null;
+						}
+						else if (!response.IsSuccessStatusCode)
+						{
+							_logger.LogError("Unable to read ref {RefName} (status: {StatusCode}, body: {Body})", name, response.StatusCode, await response.Content.ReadAsStringAsync(cancellationToken));
+							throw new StorageException($"Unable to read ref '{name}'");
+						}
+						else
+						{
+							response.EnsureSuccessStatusCode();
+							ReadRefResponse? data = await response.Content.ReadFromJsonAsync<ReadRefResponse>(cancellationToken: cancellationToken);
+							_logger.LogDebug("Read ref {RefName} -> {Blob}#{ExportIdx}", name, data!.Blob, data!.ExportIdx);
+							return new FlushedNodeHandle(BundleReader, new BundleNodeLocator(data.Hash, data!.Blob, data!.ExportIdx));
+						}
+					}
+				}
+			}
+		}
+
+		/// <inheritdoc/>
+		public override async Task WriteRefTargetAsync(RefName name, BundleNodeHandle target, RefOptions? options = null, CancellationToken cancellationToken = default)
+		{
+			_logger.LogDebug("Writing ref {RefName} -> {RefTarget}", name, target);
+			using (HttpClient httpClient = _createClient())
+			{
+				BundleNodeLocator locator = await target.FlushAsync(cancellationToken);
+				using (HttpResponseMessage response = await httpClient.PutAsync($"refs/{name}", new { blob = locator.Blob, exportIdx = locator.ExportIdx, options }, cancellationToken))
+				{
+					response.EnsureSuccessStatusCode();
+				}
+			}
+		}
+
+		#endregion
+	}
+}
diff --git a/Engine/Source/Programs/Shared/EpicGames.Horde/Storage/Clients/MemoryStorageClient.cs b/Engine/Source/Programs/Shared/EpicGames.Horde/Storage/Clients/MemoryStorageClient.cs
new file mode 100644
--- /dev/null
+++ b/Engine/Source/Programs/Shared/EpicGames.Horde/Storage/Clients/MemoryStorageClient.cs
@@ -0,0 +1,137 @@
+// Copyright Epic Games, Inc. All Rights Reserved.
+
+using System;
+using System.Buffers;
+using System.Collections.Concurrent;
+using System.Collections.Generic;
+using System.IO;
+using System.Runtime.CompilerServices;
+using System.Threading;
+using System.Threading.Tasks;
+using EpicGames.Core;
+using Microsoft.Extensions.Caching.Memory;
+using Microsoft.Extensions.Logging;
+using Microsoft.Extensions.Logging.Abstractions;
+using EpicGames.Horde.Storage.Bundles;
+
+namespace EpicGames.Horde.Storage.Clients
+{
+	/// <summary>
+	/// Implementation of <see cref="IStorageClient"/> which stores data in memory. Not intended for production use.
+	/// </summary>
+	public class MemoryStorageClient : BundleStorageClient
+	{
+		record class ExportEntry(BundleNodeHandle Handle, int Rank, ExportEntry? Next);
+
+		/// <summary>
+		/// Map of blob id to blob data
+		/// </summary>
+		readonly ConcurrentDictionary<BundleLocator, Bundle> _bundles = new ConcurrentDictionary<BundleLocator, Bundle>();
+
+		/// <summary>
+		/// Map of ref name to ref data
+		/// </summary>
+		readonly ConcurrentDictionary<RefName, BundleNodeLocator> _refs = new ConcurrentDictionary<RefName, BundleNodeLocator>();
+
+		/// <summary>
+		/// Content addressed data lookup
+		/// </summary>
+		readonly ConcurrentDictionary<Utf8String, ExportEntry> _exports = new ConcurrentDictionary<Utf8String, ExportEntry>();
+
+		/// <inheritdoc cref="_bundles"/>
+		public IReadOnlyDictionary<BundleLocator, Bundle> Bundles => _bundles;
+
+		/// <inheritdoc cref="_refs"/>
+		public IReadOnlyDictionary<RefName, BundleNodeLocator> Refs => _refs;
+
+		/// <summary>
+		/// Constructor
+		/// </summary>
+		public MemoryStorageClient() 
+			: base(StorageCache.None, NullLogger.Instance)
+		{
+		}
+
+		#region Blobs
+
+		/// <inheritdoc/>
+		public override Task<Stream> OpenAsync(BundleLocator locator, int offset, int? length, CancellationToken cancellationToken = default)
+		{
+			ReadOnlySequence<byte> sequence = _bundles[locator].AsSequence().Slice(offset);
+			if (length != null && sequence.Length > length)
+			{
+				sequence = sequence.Slice(0, length.Value);
+			}
+			return Task.FromResult<Stream>(new ReadOnlySequenceStream(sequence));
+		}
+
+		/// <inheritdoc/>
+		public override async Task<BundleLocator> WriteBundleAsync(Bundle bundle, Utf8String prefix = default, CancellationToken cancellationToken = default)
+		{
+			BundleLocator locator = BundleLocator.CreateUnique(prefix);
+			using ReadOnlySequenceStream stream = new ReadOnlySequenceStream(bundle.AsSequence());
+			_bundles[locator] = await Bundle.FromStreamAsync(stream, cancellationToken);
+			return locator;
+		}
+
+		#endregion
+
+		#region Aliases
+
+		/// <inheritdoc/>
+		public override Task AddAliasAsync(Utf8String name, BundleNodeHandle handle, int rank = 0, CancellationToken cancellationToken = default)
+		{
+			_exports.AddOrUpdate(name, _ => new ExportEntry(handle, rank, null), (_, entry) => new ExportEntry(handle, rank, entry));
+			return Task.CompletedTask;
+		}
+
+		/// <inheritdoc/>
+		public override Task RemoveAliasAsync(Utf8String name, BundleNodeHandle handle, CancellationToken cancellationToken = default)
+		{
+			throw new NotSupportedException();
+		}
+
+		/// <inheritdoc/>
+		public override async IAsyncEnumerable<BundleNodeHandle> FindAliasAsync(Utf8String alias, [EnumeratorCancellation] CancellationToken cancellationToken = default)
+		{
+			if(_exports.TryGetValue(alias, out ExportEntry? entry))
+			{
+				for (; entry != null; entry = entry.Next)
+				{
+					cancellationToken.ThrowIfCancellationRequested();
+					await Task.Yield();
+					yield return entry.Handle;
+				}
+			}
+		}
+
+		#endregion
+
+		#region Refs
+
+		/// <inheritdoc/>
+		public override Task<bool> DeleteRefAsync(RefName name, CancellationToken cancellationToken) => Task.FromResult(_refs.TryRemove(name, out _));
+
+		/// <inheritdoc/>
+		public override Task<BundleNodeHandle?> TryReadRefTargetAsync(RefName name, RefCacheTime cacheTime = default, CancellationToken cancellationToken = default)
+		{
+			BundleNodeLocator hashedLocator;
+			if (_refs.TryGetValue(name, out hashedLocator))
+			{
+				return Task.FromResult<BundleNodeHandle?>(new FlushedNodeHandle(BundleReader, hashedLocator)); 
+			}
+			else
+			{
+				return Task.FromResult<BundleNodeHandle?>(null);
+			}
+		}
+
+		/// <inheritdoc/>
+		public override async Task WriteRefTargetAsync(RefName name, BundleNodeHandle target, RefOptions? options = null, CancellationToken cancellationToken = default)
+		{
+			_refs[name] = await target.FlushAsync(cancellationToken);
+		}
+
+		#endregion
+	}
+}
diff --git a/Engine/Source/Programs/UnrealBuildTool/Artifacts/HordeStorageArtifactCache.cs b/Engine/Source/Programs/UnrealBuildTool/Artifacts/HordeStorageArtifactCache.cs
--- a/Engine/Source/Programs/UnrealBuildTool/Artifacts/HordeStorageArtifactCache.cs
+++ b/Engine/Source/Programs/UnrealBuildTool/Artifacts/HordeStorageArtifactCache.cs
@@ -8,7 +8,7 @@ using System.Threading;
 using System.Threading.Tasks;
 using EpicGames.Core;
 using EpicGames.Horde.Storage;
-using EpicGames.Horde.Storage.Backends;
+using EpicGames.Horde.Storage.Clients;
 using EpicGames.Horde.Storage.Nodes;
 using Microsoft.Extensions.Logging;
 using Microsoft.Extensions.Logging.Abstractions;
