Horde: Move storage backend implementations into EpicGames.Horde.

diff --git a/Engine/Source/Programs/Horde/Horde.Server.Tests/LogFileServiceTests.cs b/Engine/Source/Programs/Horde/Horde.Server.Tests/LogFileServiceTests.cs
--- a/Engine/Source/Programs/Horde/Horde.Server.Tests/LogFileServiceTests.cs
+++ b/Engine/Source/Programs/Horde/Horde.Server.Tests/LogFileServiceTests.cs
@@ -9,6 +9,8 @@ using System.Text;
 using System.Threading;
 using System.Threading.Tasks;
 using EpicGames.Core;
+using EpicGames.Horde.Storage;
+using EpicGames.Horde.Storage.Backends;
 using Horde.Server.Agents.Sessions;
 using Horde.Server.Jobs;
 using Horde.Server.Logs;
@@ -16,13 +18,11 @@ using Horde.Server.Logs.Builder;
 using Horde.Server.Logs.Data;
 using Horde.Server.Logs.Storage;
 using Horde.Server.Storage;
-using Horde.Server.Storage.Backends;
 using Horde.Server.Utilities;
 using HordeCommon;
 using Microsoft.Extensions.Logging;
 using Microsoft.Extensions.Logging.Abstractions;
 using Microsoft.VisualStudio.TestTools.UnitTesting;
-using MongoDB.Bson;
 using OpenTelemetry.Trace;
 
 namespace Horde.Server.Tests
diff --git a/Engine/Source/Programs/Horde/Horde.Server.Tests/TestSetup.cs b/Engine/Source/Programs/Horde/Horde.Server.Tests/TestSetup.cs
--- a/Engine/Source/Programs/Horde/Horde.Server.Tests/TestSetup.cs
+++ b/Engine/Source/Programs/Horde/Horde.Server.Tests/TestSetup.cs
@@ -14,6 +14,7 @@ using Datadog.Trace;
 using EpicGames.Core;
 using EpicGames.Horde.Api;
 using EpicGames.Horde.Storage;
+using EpicGames.Horde.Storage.Backends;
 using Horde.Server.Acls;
 using Horde.Server.Agents;
 using Horde.Server.Agents.Pools;
diff --git a/Engine/Source/Programs/Horde/Horde.Server/Jobs/Artifacts/ArtifactCollectionV1.cs b/Engine/Source/Programs/Horde/Horde.Server/Jobs/Artifacts/ArtifactCollectionV1.cs
--- a/Engine/Source/Programs/Horde/Horde.Server/Jobs/Artifacts/ArtifactCollectionV1.cs
+++ b/Engine/Source/Programs/Horde/Horde.Server/Jobs/Artifacts/ArtifactCollectionV1.cs
@@ -10,6 +10,7 @@ using Horde.Server.Utilities;
 using MongoDB.Bson;
 using MongoDB.Bson.Serialization.Attributes;
 using MongoDB.Driver;
+using EpicGames.Horde.Storage;
 
 namespace Horde.Server.Jobs.Artifacts
 {
diff --git a/Engine/Source/Programs/Horde/Horde.Server/Logs/Storage/PersistentLogStorage.cs b/Engine/Source/Programs/Horde/Horde.Server/Logs/Storage/PersistentLogStorage.cs
--- a/Engine/Source/Programs/Horde/Horde.Server/Logs/Storage/PersistentLogStorage.cs
+++ b/Engine/Source/Programs/Horde/Horde.Server/Logs/Storage/PersistentLogStorage.cs
@@ -3,6 +3,7 @@
 using System;
 using System.Threading.Tasks;
 using EpicGames.Core;
+using EpicGames.Horde.Storage;
 using Horde.Server.Logs.Data;
 using Horde.Server.Storage;
 using Horde.Server.Utilities;
diff --git a/Engine/Source/Programs/Horde/Horde.Server/ServerSettings.cs b/Engine/Source/Programs/Horde/Horde.Server/ServerSettings.cs
--- a/Engine/Source/Programs/Horde/Horde.Server/ServerSettings.cs
+++ b/Engine/Source/Programs/Horde/Horde.Server/ServerSettings.cs
@@ -55,8 +55,13 @@ namespace Horde.Server
 	/// <summary>
 	/// Common settings for different storage backends
 	/// </summary>
-	public interface IStorageBackendOptions : IFileSystemStorageOptions, IAwsStorageOptions
+	public interface IStorageBackendOptions : IAwsStorageOptions
 	{
+		/// <summary>
+		/// Base directory for filesystem storage
+		/// </summary>
+		string? BaseDir { get; }
+
 		/// <summary>
 		/// The type of storage backend to use
 		/// </summary>
diff --git a/Engine/Source/Programs/Horde/Horde.Server/Startup.cs b/Engine/Source/Programs/Horde/Horde.Server/Startup.cs
--- a/Engine/Source/Programs/Horde/Horde.Server/Startup.cs
+++ b/Engine/Source/Programs/Horde/Horde.Server/Startup.cs
@@ -22,6 +22,7 @@ using Amazon.SQS;
 using EpicGames.AspNet;
 using EpicGames.Core;
 using EpicGames.Horde.Storage;
+using EpicGames.Horde.Storage.Backends;
 using Grpc.Core;
 using Grpc.Core.Interceptors;
 using Horde.Server.Acls;
@@ -31,7 +32,6 @@ using Horde.Server.Agents.Fleet;
 using Horde.Server.Agents.Leases;
 using Horde.Server.Agents.Pools;
 using Horde.Server.Agents.Sessions;
-using Horde.Server.Agents.Software;
 using Horde.Server.Agents.Telemetry;
 using Horde.Server.Configuration;
 using Horde.Server.Dashboard;
@@ -112,7 +112,6 @@ using StackExchange.Redis;
 
 namespace Horde.Server
 {
-	using IStorageBackend = Horde.Server.Storage.IStorageBackend;
 	using ContentHash = EpicGames.Core.ContentHash;
 	using ReferenceResolver = Horde.Server.Ddc.ReferenceResolver;
 
@@ -287,7 +286,7 @@ namespace Horde.Server
 			switch (options.Type ?? StorageBackendType.FileSystem)
 			{
 				case StorageBackendType.FileSystem:
-					return new FileSystemStorageBackend(options);
+					return new FileStorageBackend(DirectoryReference.Combine(Program.DataDir, options.BaseDir ?? "Storage"));
 				case StorageBackendType.Aws:
 					return new AwsStorageBackend(sp.GetRequiredService<IConfiguration>(), options, sp.GetRequiredService<ILogger<AwsStorageBackend>>());
 				case StorageBackendType.Memory:
diff --git a/Engine/Source/Programs/Horde/Horde.Server/Storage/Backends/AwsStorageBackend.cs b/Engine/Source/Programs/Horde/Horde.Server/Storage/Backends/AwsStorageBackend.cs
--- a/Engine/Source/Programs/Horde/Horde.Server/Storage/Backends/AwsStorageBackend.cs
+++ b/Engine/Source/Programs/Horde/Horde.Server/Storage/Backends/AwsStorageBackend.cs
@@ -18,6 +18,7 @@ using Google.Protobuf.WellKnownTypes;
 using Horde.Server.Utilities;
 using Microsoft.Extensions.Configuration;
 using Microsoft.Extensions.Logging;
+using EpicGames.Horde.Storage;
 using OpenTelemetry.Trace;
 
 namespace Horde.Server.Storage.Backends
diff --git a/Engine/Source/Programs/Horde/Horde.Server/Storage/Backends/FileStorageBackend.cs b/Engine/Source/Programs/Horde/Horde.Server/Storage/Backends/FileStorageBackend.cs
deleted file mode 100644
diff --git a/Engine/Source/Programs/Horde/Horde.Server/Storage/Backends/MemoryStorageBackend.cs b/Engine/Source/Programs/Horde/Horde.Server/Storage/Backends/MemoryStorageBackend.cs
deleted file mode 100644
diff --git a/Engine/Source/Programs/Horde/Horde.Server/Storage/IStorageBackend.cs b/Engine/Source/Programs/Horde/Horde.Server/Storage/IStorageBackend.cs
deleted file mode 100644
diff --git a/Engine/Source/Programs/Horde/Horde.Server/Storage/IStorageBackendProvider.cs b/Engine/Source/Programs/Horde/Horde.Server/Storage/IStorageBackendProvider.cs
--- a/Engine/Source/Programs/Horde/Horde.Server/Storage/IStorageBackendProvider.cs
+++ b/Engine/Source/Programs/Horde/Horde.Server/Storage/IStorageBackendProvider.cs
@@ -1,5 +1,7 @@
 // Copyright Epic Games, Inc. All Rights Reserved.
 
+using EpicGames.Horde.Storage;
+
 namespace Horde.Server.Storage
 {
 	/// <summary>
diff --git a/Engine/Source/Programs/Horde/Horde.Server/Storage/StorageBackendProvider.cs b/Engine/Source/Programs/Horde/Horde.Server/Storage/StorageBackendProvider.cs
--- a/Engine/Source/Programs/Horde/Horde.Server/Storage/StorageBackendProvider.cs
+++ b/Engine/Source/Programs/Horde/Horde.Server/Storage/StorageBackendProvider.cs
@@ -7,6 +7,8 @@ using System.Text.Json;
 using System.Threading;
 using System.Threading.Tasks;
 using EpicGames.Core;
+using EpicGames.Horde.Storage;
+using EpicGames.Horde.Storage.Backends;
 using Horde.Server.Storage.Backends;
 using Microsoft.Extensions.Configuration;
 using Microsoft.Extensions.DependencyInjection;
@@ -166,7 +168,7 @@ namespace Horde.Server.Storage
 			switch (config.Type ?? StorageBackendType.FileSystem)
 			{
 				case StorageBackendType.FileSystem:
-					return new FileSystemStorageBackend(config);
+					return new FileStorageBackend(DirectoryReference.Combine(Program.DataDir, config.BaseDir ?? "Storage"));
 				case StorageBackendType.Aws:
 					return new AwsStorageBackend(_serviceProvider.GetRequiredService<IConfiguration>(), config, _serviceProvider.GetRequiredService<ILogger<AwsStorageBackend>>());
 				case StorageBackendType.Memory:
diff --git a/Engine/Source/Programs/Horde/Horde.Server/Storage/TypedStorageBackend.cs b/Engine/Source/Programs/Horde/Horde.Server/Storage/TypedStorageBackend.cs
new file mode 100644
--- /dev/null
+++ b/Engine/Source/Programs/Horde/Horde.Server/Storage/TypedStorageBackend.cs
@@ -0,0 +1,126 @@
+// Copyright Epic Games, Inc. All Rights Reserved.
+
+using System;
+using System.Collections.Generic;
+using System.IO;
+using System.Threading;
+using System.Threading.Tasks;
+using EpicGames.Core;
+using EpicGames.Horde.Storage;
+
+namespace Horde.Server.Storage
+{
+	/// <summary>
+	/// Generic version of IStorageBackend, to allow for dependency injection of different singletons
+	/// </summary>
+	/// <typeparam name="T">Type distinguishing different singletons</typeparam>
+	public interface IStorageBackend<T> : IStorageBackend
+	{
+	}
+
+	/// <summary>
+	/// Extension methods for <see cref="IStorageBackend"/>
+	/// </summary>
+	public static class StorageBackend
+	{
+		/// <summary>
+		/// Wrapper for <see cref="IStorageBackend"/>
+		/// </summary>
+		/// <typeparam name="T"></typeparam>
+		sealed class TypedStorageBackend<T> : IStorageBackend<T>
+		{
+			readonly IStorageBackend _inner;
+
+			/// <inheritdoc/>
+			public bool SupportsRedirects => _inner.SupportsRedirects;
+
+			/// <summary>
+			/// Constructor
+			/// </summary>
+			/// <param name="inner"></param>
+			public TypedStorageBackend(IStorageBackend inner) => _inner = inner;
+
+			/// <inheritdoc/>
+			public void Dispose() => _inner.Dispose();
+
+			/// <inheritdoc/>
+			public Task<Stream> ReadAsync(string path, int offset, int? length, CancellationToken cancellationToken) => _inner.ReadAsync(path, offset, length, cancellationToken);
+
+			/// <inheritdoc/>
+			public Task WriteAsync(string path, Stream stream, CancellationToken cancellationToken) => _inner.WriteAsync(path, stream, cancellationToken);
+
+			/// <inheritdoc/>
+			public Task DeleteAsync(string path, CancellationToken cancellationToken) => _inner.DeleteAsync(path, cancellationToken);
+
+			/// <inheritdoc/>
+			public Task<bool> ExistsAsync(string path, CancellationToken cancellationToken) => _inner.ExistsAsync(path, cancellationToken);
+
+			/// <inheritdoc/>
+			public IAsyncEnumerable<string> EnumerateAsync(CancellationToken cancellationToken = default) => _inner.EnumerateAsync(cancellationToken);
+
+			/// <inheritdoc/>
+			public ValueTask<Uri?> TryGetReadRedirectAsync(string path, CancellationToken cancellationToken = default) => _inner.TryGetReadRedirectAsync(path, cancellationToken);
+
+			/// <inheritdoc/>
+			public ValueTask<Uri?> TryGetWriteRedirectAsync(string path, CancellationToken cancellationToken = default) => _inner.TryGetWriteRedirectAsync(path, cancellationToken);
+		}
+
+		/// <summary>
+		/// Extension for blob files
+		/// </summary>
+		public const string BlobExtension = ".blob";
+
+		/// <summary>
+		/// Gets the blob id from a path
+		/// </summary>
+		/// <param name="path">Path to the blob</param>
+		/// <returns>Path to the blob</returns>
+		public static Utf8String GetBlobPathFromFileName(string path)
+		{
+			Utf8String blobPath;
+			if (!TryGetBlobPathFromFileName(path, out blobPath))
+			{
+				throw new ArgumentException("Path is not a valid blob identifier", nameof(path));
+			}
+			return blobPath;
+		}
+
+		/// <summary>
+		/// Gets the path to a blob
+		/// </summary>
+		/// <param name="blobPath">Blob identifier</param>
+		/// <returns>Path to the blob</returns>
+		public static string GetBlobFileName(Utf8String blobPath) => $"{blobPath}{BlobExtension}";
+
+		/// <summary>
+		/// Gets a blob id from a path within the storage backend
+		/// </summary>
+		/// <param name="path">Path to the file</param>
+		/// <param name="blobPath">Receives the blob id on success</param>
+		/// <returns>True on success</returns>
+		public static bool TryGetBlobPathFromFileName(string path, out Utf8String blobPath)
+		{
+			if (path.EndsWith(BlobExtension, StringComparison.Ordinal))
+			{
+				blobPath = new Utf8String(path.Substring(0, path.Length - BlobExtension.Length));
+				return true;
+			}
+			else
+			{
+				blobPath = default;
+				return false;
+			}
+		}
+
+		/// <summary>
+		/// Creates a typed wrapper around the given storage backend
+		/// </summary>
+		/// <typeparam name="T"></typeparam>
+		/// <param name="backend"></param>
+		/// <returns></returns>
+		public static IStorageBackend<T> ForType<T>(this IStorageBackend backend)
+		{
+			return new TypedStorageBackend<T>(backend);
+		}
+	}
+}
diff --git a/Engine/Source/Programs/Shared/EpicGames.Horde/Storage/Backends/FileStorageBackend.cs b/Engine/Source/Programs/Shared/EpicGames.Horde/Storage/Backends/FileStorageBackend.cs
new file mode 100644
--- /dev/null
+++ b/Engine/Source/Programs/Shared/EpicGames.Horde/Storage/Backends/FileStorageBackend.cs
@@ -0,0 +1,163 @@
+// Copyright Epic Games, Inc. All Rights Reserved.
+
+using System;
+using System.Collections.Generic;
+using System.IO;
+using System.Runtime.CompilerServices;
+using System.Threading;
+using System.Threading.Tasks;
+using EpicGames.Core;
+
+namespace EpicGames.Horde.Storage.Backends
+{
+	/// <summary>
+	/// Storage backend that utilizes the local filesystem
+	/// </summary>
+	public sealed class FileStorageBackend : IStorageBackend
+	{
+		/// <summary>
+		/// Base directory for log files
+		/// </summary>
+		private readonly DirectoryReference _baseDir;
+
+		/// <summary>
+		/// Unique identifier for this instance
+		/// </summary>
+		private readonly string _instanceId;
+
+		/// <summary>
+		/// Unique id for each write
+		/// </summary>
+		private int _uniqueId;
+
+		/// <inheritdoc/>
+		public bool SupportsRedirects => false;
+
+		/// <summary>
+		/// Constructor
+		/// </summary>
+		/// <param name="baseDir">Base directory for the store</param>
+		public FileStorageBackend(DirectoryReference baseDir)
+		{
+			_baseDir = baseDir;
+			_instanceId = Guid.NewGuid().ToString("N");
+			DirectoryReference.CreateDirectory(_baseDir);
+		}
+
+		/// <inheritdoc/>
+		public void Dispose()
+		{
+		}
+
+		/// <inheritdoc/>
+		public Task<Stream> ReadAsync(string path, CancellationToken cancellationToken)
+		{
+			FileReference location = FileReference.Combine(_baseDir, path);
+			return Task.FromResult<Stream>(FileReference.Open(location, FileMode.Open, FileAccess.Read, FileShare.Read));
+		}
+
+		/// <inheritdoc/>
+		public async Task<Stream> ReadAsync(string path, int offset, int? length, CancellationToken cancellationToken)
+		{
+			Stream stream = await ReadAsync(path, cancellationToken);
+			stream.Seek(offset, SeekOrigin.Begin);
+			return stream;
+		}
+
+		/// <inheritdoc/>
+		public async Task WriteAsync(string path, Stream stream, CancellationToken cancellationToken)
+		{
+			FileReference finalLocation = FileReference.Combine(_baseDir, path);
+			if (!FileReference.Exists(finalLocation))
+			{
+				// Write to a temp file first
+				int newUniqueId = Interlocked.Increment(ref _uniqueId);
+
+				DirectoryReference.CreateDirectory(finalLocation.Directory);
+				FileReference tempLocation = new FileReference($"{finalLocation}.{_instanceId}.{newUniqueId:x8}");
+
+				using (Stream outputStream = FileReference.Open(tempLocation, FileMode.Create, FileAccess.Write, FileShare.Read))
+				{
+					await stream.CopyToAsync(outputStream, cancellationToken);
+				}
+
+				// Move the temp file into place
+				try
+				{
+					FileReference.Move(tempLocation, finalLocation, true);
+				}
+				catch (IOException) // Already exists
+				{
+					if (FileReference.Exists(finalLocation))
+					{
+						FileReference.Delete(tempLocation);
+					}
+					else
+					{
+						throw;
+					}
+				}
+			}
+		}
+
+		/// <inheritdoc/>
+		public Task<bool> ExistsAsync(string path, CancellationToken cancellationToken)
+		{
+			FileReference location = FileReference.Combine(_baseDir, path);
+			return Task.FromResult(FileReference.Exists(location));
+		}
+
+		/// <inheritdoc/>
+		public Task DeleteAsync(string path, CancellationToken cancellationToken)
+		{
+			FileReference location = FileReference.Combine(_baseDir, path);
+			FileReference.Delete(location);
+			return Task.CompletedTask;
+		}
+
+		/// <inheritdoc/>
+		public async IAsyncEnumerable<string> EnumerateAsync([EnumeratorCancellation] CancellationToken cancellationToken = default)
+		{
+			Stack<IEnumerator<DirectoryInfo>> queue = new Stack<IEnumerator<DirectoryInfo>>();
+			try
+			{
+				queue.Push(new List<DirectoryInfo> { _baseDir.ToDirectoryInfo() }.GetEnumerator());
+				while (queue.Count > 0)
+				{
+					IEnumerator<DirectoryInfo> top = queue.Peek();
+					if (!top.MoveNext())
+					{
+						top.Dispose();
+						queue.Pop();
+						continue;
+					}
+
+					DirectoryInfo current = top.Current;
+					foreach (FileInfo fileInfo in current.EnumerateFiles("*"))
+					{
+						string path = fileInfo.FullName.Substring(_baseDir.FullName.Length + 1).Replace(Path.DirectorySeparatorChar, '/');
+						yield return path;
+					}
+
+					queue.Push(current.EnumerateDirectories().GetEnumerator());
+
+					cancellationToken.ThrowIfCancellationRequested();
+					await Task.Yield();
+				}
+			}
+			finally
+			{
+				while (queue.TryPop(out IEnumerator<DirectoryInfo>? enumerator))
+				{
+					enumerator.Dispose();
+				}
+			}
+		}
+
+		/// <inheritdoc/>
+		public ValueTask<Uri?> TryGetReadRedirectAsync(string path, CancellationToken cancellationToken = default) => default;
+
+		/// <inheritdoc/>
+		public ValueTask<Uri?> TryGetWriteRedirectAsync(string path, CancellationToken cancellationToken = default) => default;
+	}
+}
diff --git a/Engine/Source/Programs/Shared/EpicGames.Horde/Storage/Backends/MemoryStorageBackend.cs b/Engine/Source/Programs/Shared/EpicGames.Horde/Storage/Backends/MemoryStorageBackend.cs
new file mode 100644
--- /dev/null
+++ b/Engine/Source/Programs/Shared/EpicGames.Horde/Storage/Backends/MemoryStorageBackend.cs
@@ -0,0 +1,86 @@
+// Copyright Epic Games, Inc. All Rights Reserved.
+
+using System;
+using System.Collections.Concurrent;
+using System.Collections.Generic;
+using System.IO;
+using System.Runtime.CompilerServices;
+using System.Threading;
+using System.Threading.Tasks;
+using EpicGames.Core;
+
+namespace EpicGames.Horde.Storage.Backends
+{
+	/// <summary>
+	/// In-memory implementation of a storage backend
+	/// </summary>
+	public sealed class MemoryStorageBackend : IStorageBackend
+	{
+		/// <summary>
+		/// Data storage
+		/// </summary>
+		readonly ConcurrentDictionary<string, byte[]> _pathToData = new ConcurrentDictionary<string, byte[]>();
+
+		/// <inheritdoc/>
+		public bool SupportsRedirects => false;
+
+		/// <inheritdoc/>
+		public void Dispose()
+		{
+		}
+
+		/// <inheritdoc/>
+		public Task<Stream> ReadAsync(string path, CancellationToken cancellationToken) => ReadAsync(path, 0, null, cancellationToken);
+
+		/// <inheritdoc/>
+		public Task<Stream> ReadAsync(string path, int offset, int? length, CancellationToken cancellationToken)
+		{
+			ReadOnlyMemory<byte> data = _pathToData[path].AsMemory(offset);
+			if (length != null && length.Value < data.Length)
+			{
+				data = data.Slice(0, length.Value);
+			}
+			return Task.FromResult<Stream>(new ReadOnlyMemoryStream(data));
+		}
+
+		/// <inheritdoc/>
+		public async Task WriteAsync(string path, Stream stream, CancellationToken cancellationToken)
+		{
+			using (MemoryStream buffer = new MemoryStream())
+			{
+				await stream.CopyToAsync(buffer, cancellationToken);
+				_pathToData[path] = buffer.ToArray();
+			}
+		}
+
+		/// <inheritdoc/>
+		public Task<bool> ExistsAsync(string path, CancellationToken cancellationToken)
+		{
+			return Task.FromResult(_pathToData.ContainsKey(path));
+		}
+
+		/// <inheritdoc/>
+		public Task DeleteAsync(string path, CancellationToken cancellationToken)
+		{
+			_pathToData.TryRemove(path, out _);
+			return Task.CompletedTask;
+		}
+
+		/// <inheritdoc/>
+		public async IAsyncEnumerable<string> EnumerateAsync([EnumeratorCancellation] CancellationToken cancellationToken = default)
+		{
+			foreach (string path in _pathToData.Keys)
+			{
+				yield return path;
+				cancellationToken.ThrowIfCancellationRequested();
+				await Task.Yield();
+			}
+		}
+
+		/// <inheritdoc/>
+		public ValueTask<Uri?> TryGetReadRedirectAsync(string path, CancellationToken cancellationToken = default) => default;
+
+		/// <inheritdoc/>
+		public ValueTask<Uri?> TryGetWriteRedirectAsync(string path, CancellationToken cancellationToken = default) => default;
+	}
+}
