Missing file.

diff --git a/Engine/Source/Programs/Shared/EpicGames.Horde/Storage/IStorageBackend.cs b/Engine/Source/Programs/Shared/EpicGames.Horde/Storage/IStorageBackend.cs
new file mode 100644
--- /dev/null
+++ b/Engine/Source/Programs/Shared/EpicGames.Horde/Storage/IStorageBackend.cs
@@ -0,0 +1,128 @@
+// Copyright Epic Games, Inc. All Rights Reserved.
+
+using System;
+using System.Collections.Generic;
+using System.IO;
+using System.Threading;
+using System.Threading.Tasks;
+using EpicGames.Core;
+
+namespace EpicGames.Horde.Storage
+{
+	/// <summary>
+	/// Interface for a low-level storage backend.
+	/// </summary>
+	public interface IStorageBackend : IDisposable
+	{
+		/// <summary>
+		/// Whether this storage backend supports HTTP redirects for reads and writes
+		/// </summary>
+		bool SupportsRedirects { get; }
+
+		/// <summary>
+		/// Attempts to open a read stream for the given path.
+		/// </summary>
+		/// <param name="path">Relative path within the bucket</param>
+		/// <param name="offset">Offset to start reading from</param>
+		/// <param name="length">Length of data to read</param>
+		/// <param name="cancellationToken">Cancellation token for the operation</param>
+		/// <returns></returns>
+		Task<Stream> ReadAsync(string path, int offset, int? length, CancellationToken cancellationToken = default);
+
+		/// <summary>
+		/// Writes a stream to the given path. If the stream throws an exception during read, the write will be aborted.
+		/// </summary>
+		/// <param name="path">Relative path within the bucket</param>
+		/// <param name="stream">Stream to write</param>
+		/// <param name="cancellationToken">Cancellation token for the operation</param>
+		Task WriteAsync(string path, Stream stream, CancellationToken cancellationToken = default);
+
+		/// <summary>
+		/// Tests whether the given path exists
+		/// </summary>
+		/// <param name="path">Relative path within the bucket</param>
+		/// <param name="cancellationToken">Cancellation token for the operation</param>
+		/// <returns></returns>
+		Task<bool> ExistsAsync(string path, CancellationToken cancellationToken = default);
+
+		/// <summary>
+		/// Deletes a file with the given path
+		/// </summary>
+		/// <param name="path">Relative path within the bucket</param>
+		/// <param name="cancellationToken">Cancellation token for the operation</param>
+		/// <returns>Async task</returns>
+		Task DeleteAsync(string path, CancellationToken cancellationToken = default);
+
+		/// <summary>
+		/// Enumerates all the objects in the store
+		/// </summary>
+		/// <param name="cancellationToken">Cancellation token for the operation</param>
+		/// <returns>Sequence of object paths</returns>
+		IAsyncEnumerable<string> EnumerateAsync(CancellationToken cancellationToken = default);
+
+		/// <summary>
+		/// Gets a HTTP redirect for a read request
+		/// </summary>
+		/// <param name="path">Path to read from</param>
+		/// <param name="cancellationToken">Cancellation token for the operation</param>
+		/// <returns>Path to upload the data to</returns>
+		ValueTask<Uri?> TryGetReadRedirectAsync(string path, CancellationToken cancellationToken = default);
+
+		/// <summary>
+		/// Gets a HTTP redirect for a write request
+		/// </summary>
+		/// <param name="path">Path to write to</param>
+		/// <param name="cancellationToken">Cancellation token for the operation</param>
+		/// <returns>Path to upload the data to</returns>
+		ValueTask<Uri?> TryGetWriteRedirectAsync(string path, CancellationToken cancellationToken = default);
+	}
+
+	/// <summary>
+	/// Extension methods for <see cref="IStorageBackend"/>
+	/// </summary>
+	public static class StorageBackendExtensions
+	{
+		/// <summary>
+		/// Attempts to open a read stream for the given path.
+		/// </summary>
+		/// <param name="storageBackend">Backend to read from</param>
+		/// <param name="path">Object name within the store</param>
+		/// <param name="cancellationToken">Cancellation token for the operation</param>
+		/// <returns>Stream for the object</returns>
+		public static Task<Stream> ReadAsync(this IStorageBackend storageBackend, string path, CancellationToken cancellationToken = default) => storageBackend.ReadAsync(path, 0, null, cancellationToken);
+
+		/// <summary>
+		/// Reads an object as an array of bytes
+		/// </summary>
+		/// <param name="storageBackend">Backend to read from</param>
+		/// <param name="path">Object name within the store</param>
+		/// <param name="cancellationToken">Cancellation token for the operation</param>
+		/// <returns>Contents of the object</returns>
+		public static async Task<byte[]> ReadBytesAsync(this IStorageBackend storageBackend, string path, CancellationToken cancellationToken = default)
+		{
+			using (Stream inputStream = await storageBackend.ReadAsync(path, cancellationToken))
+			{
+				using (MemoryStream outputStream = new MemoryStream())
+				{
+					await inputStream.CopyToAsync(outputStream, cancellationToken);
+					return outputStream.ToArray();
+				}
+			}
+		}
+
+		/// <summary>
+		/// Writes a block of memory to storage
+		/// </summary>
+		/// <param name="storageBackend">Backend to read from</param>
+		/// <param name="path">Object name within the store</param>
+		/// <param name="data">Data to be written</param>
+		/// <param name="cancellationToken">Cancellation token for the operation</param>
+		public static async Task WriteBytesAsync(this IStorageBackend storageBackend, string path, ReadOnlyMemory<byte> data, CancellationToken cancellationToken = default)
+		{
+			using (ReadOnlyMemoryStream stream = new ReadOnlyMemoryStream(data))
+			{
+				await storageBackend.WriteAsync(path, stream, cancellationToken);
+			}
+		}
+	}
+}
