Horde: Deprecate writing an explicit path through IStorageBackend; add an override which allows the backend to generate a unique id. This unifies server-side codepaths (where the generation of a unique id is trusted) and client-side codepaths (where it is not). By forcing unique ids to be generated on the server, we can further abstract logic around creating bundles and make it more of an optimization over storing small blobs.

diff --git a/Engine/Source/Programs/Horde/Horde.Server.Tests/GcServiceTests.cs b/Engine/Source/Programs/Horde/Horde.Server.Tests/GcServiceTests.cs
--- a/Engine/Source/Programs/Horde/Horde.Server.Tests/GcServiceTests.cs
+++ b/Engine/Source/Programs/Horde/Horde.Server.Tests/GcServiceTests.cs
@@ -52,7 +52,7 @@ namespace Horde.Server.Tests
 
 			await Clock.AdvanceAsync(TimeSpan.FromDays(1.0));
 
-			BundleLocator[] remaining = await store.Backend.EnumerateAsync().Select(x => new BundleLocator(StorageBackend.GetBlobPathFromFileName(x))).ToArrayAsync();
+			BundleLocator[] remaining = await store.Backend.EnumerateAsync().Select(x => new BundleLocator(x)).ToArrayAsync();
 			Assert.AreEqual(nodes.Count, remaining.Length);
 			Assert.IsTrue(remaining.All(x => nodes.Contains(x)));
 		}
diff --git a/Engine/Source/Programs/Horde/Horde.Server/Jobs/Artifacts/ArtifactCollectionV1.cs b/Engine/Source/Programs/Horde/Horde.Server/Jobs/Artifacts/ArtifactCollectionV1.cs
--- a/Engine/Source/Programs/Horde/Horde.Server/Jobs/Artifacts/ArtifactCollectionV1.cs
+++ b/Engine/Source/Programs/Horde/Horde.Server/Jobs/Artifacts/ArtifactCollectionV1.cs
@@ -87,7 +87,9 @@ namespace Horde.Server.Jobs.Artifacts
 		{
 			// upload first
 			string artifactName = ValidateName(name);
-			await _storageBackend.WriteAsync(GetPath(jobId, stepId, artifactName), data);
+#pragma warning disable CS0618
+			await _storageBackend.WriteExplicitPathAsync(GetPath(jobId, stepId, artifactName), data);
+#pragma warning restore CS0618
 
 			// then create entry
 			Artifact newArtifact = new Artifact(jobId, stepId, artifactName, data.Length, mimeType);
@@ -179,7 +181,9 @@ namespace Horde.Server.Jobs.Artifacts
 
 				// re-upload the data to external
 				string artifactName = ValidateName(artifact.Name);
-				await _storageBackend.WriteAsync(GetPath(artifact.JobId, artifact.StepId, artifactName), newData);
+#pragma warning disable CS0618
+				await _storageBackend.WriteExplicitPathAsync(GetPath(artifact.JobId, artifact.StepId, artifactName), newData);
+#pragma warning restore CS0618
 
 				if (await TryUpdateArtifactAsync((Artifact)artifact, updateBuilder.Combine(updates)))
 				{
diff --git a/Engine/Source/Programs/Horde/Horde.Server/Logs/Storage/PersistentLogStorage.cs b/Engine/Source/Programs/Horde/Horde.Server/Logs/Storage/PersistentLogStorage.cs
--- a/Engine/Source/Programs/Horde/Horde.Server/Logs/Storage/PersistentLogStorage.cs
+++ b/Engine/Source/Programs/Horde/Horde.Server/Logs/Storage/PersistentLogStorage.cs
@@ -6,7 +6,6 @@ using EpicGames.Core;
 using EpicGames.Horde.Storage;
 using Horde.Server.Logs.Data;
 using Horde.Server.Storage;
-using Horde.Server.Utilities;
 using Microsoft.Extensions.Logging;
 
 namespace Horde.Server.Logs.Storage
@@ -16,14 +15,7 @@ namespace Horde.Server.Logs.Storage
 	/// </summary>
 	class PersistentLogStorage : ILogStorage
 	{
-		/// <summary>
-		/// The bulk storage provider to use
-		/// </summary>
 		readonly IStorageBackend _storageProvider;
-
-		/// <summary>
-		/// Log provider
-		/// </summary>
 		readonly ILogger _logger;
 
 		/// <summary>
@@ -53,13 +45,16 @@ namespace Horde.Server.Logs.Storage
 		}
 
 		/// <inheritdoc/>
-		public Task WriteIndexAsync(LogId logId, long length, LogIndexData indexData)
+		public async Task WriteIndexAsync(LogId logId, long length, LogIndexData indexData)
 		{
 			_logger.LogDebug("Writing log {LogId} index length {Length} to persistent storage", logId, length);
 
 			string path = $"{logId}/index_{length}";
-			ReadOnlyMemory<byte> data = indexData.ToByteArray();
-			return _storageProvider.WriteBytesAsync(path, data);
+			using ReadOnlyMemoryStream stream = new ReadOnlyMemoryStream(indexData.ToByteArray());
+
+#pragma warning disable CS0618
+			await _storageProvider.WriteExplicitPathAsync(path, stream);
+#pragma warning restore CS0618
 		}
 
 		/// <inheritdoc/>
@@ -82,7 +77,7 @@ namespace Horde.Server.Logs.Storage
 		}
 
 		/// <inheritdoc/>
-		public Task WriteChunkAsync(LogId logId, long offset, LogChunkData chunkData)
+		public async Task WriteChunkAsync(LogId logId, long offset, LogChunkData chunkData)
 		{
 			_logger.LogDebug("Writing log {LogId} chunk offset {Offset} to persistent storage", logId, offset);
 
@@ -92,7 +87,10 @@ namespace Horde.Server.Logs.Storage
 			writer.WriteLogChunkData(chunkData, _logger);
 			writer.CheckEmpty();
 
-			return _storageProvider.WriteBytesAsync(path, data);
+			using ReadOnlyMemoryStream stream = new ReadOnlyMemoryStream(data);
+#pragma warning disable CS0618
+			await _storageProvider.WriteExplicitPathAsync(path, stream);
+#pragma warning restore CS0618
 		}
 	}
 }
diff --git a/Engine/Source/Programs/Horde/Horde.Server/Storage/Backends/AwsStorageBackend.cs b/Engine/Source/Programs/Horde/Horde.Server/Storage/Backends/AwsStorageBackend.cs
--- a/Engine/Source/Programs/Horde/Horde.Server/Storage/Backends/AwsStorageBackend.cs
+++ b/Engine/Source/Programs/Horde/Horde.Server/Storage/Backends/AwsStorageBackend.cs
@@ -331,7 +331,20 @@ namespace Horde.Server.Storage.Backends
 		public ValueTask<Uri?> TryGetReadRedirectAsync(string path, CancellationToken cancellationToken = default) => new ValueTask<Uri?>(GetPresignedUrl(path, HttpVerb.GET));
 
 		/// <inheritdoc/>
-		public ValueTask<Uri?> TryGetWriteRedirectAsync(string path, CancellationToken cancellationToken = default) => new ValueTask<Uri?>(GetPresignedUrl(path, HttpVerb.PUT));
+		public ValueTask<(string, Uri)?> TryGetWriteRedirectAsync(string? prefix, CancellationToken cancellationToken = default)
+		{
+			string path = StorageHelpers.CreateUniqueName(prefix);
+
+			Uri? url = GetPresignedUrl(path, HttpVerb.PUT);
+			if (url == null)
+			{
+				return default;
+			}
+			else
+			{
+				return new ValueTask<(string, Uri)?>((path, url));
+			}
+		}
 
 		/// <summary>
 		/// Helper method to generate a presigned URL for a request
@@ -363,7 +376,15 @@ namespace Horde.Server.Storage.Backends
 		}
 
 		/// <inheritdoc/>
-		public async Task WriteAsync(string path, Stream inputStream, CancellationToken cancellationToken)
+		public async Task<string> WriteAsync(Stream inputStream, string? prefix = null, CancellationToken cancellationToken = default)
+		{
+			string path = StorageHelpers.CreateUniqueName(prefix);
+			await WriteExplicitPathAsync(path, inputStream, cancellationToken);
+			return path;
+		}
+
+		/// <inheritdoc/>
+		public async Task WriteExplicitPathAsync(string path, Stream inputStream, CancellationToken cancellationToken = default)
 		{
 			TimeSpan[] retryTimes =
 			{
diff --git a/Engine/Source/Programs/Horde/Horde.Server/Storage/StorageBackendProvider.cs b/Engine/Source/Programs/Horde/Horde.Server/Storage/StorageBackendProvider.cs
--- a/Engine/Source/Programs/Horde/Horde.Server/Storage/StorageBackendProvider.cs
+++ b/Engine/Source/Programs/Horde/Horde.Server/Storage/StorageBackendProvider.cs
@@ -89,13 +89,18 @@ namespace Horde.Server.Storage
 			}
 
 			/// <inheritdoc/>
-			public Task WriteAsync(string path, Stream stream, CancellationToken cancellationToken = default) => _backend.WriteAsync(path, stream, cancellationToken);
+			public Task<string> WriteAsync(Stream stream, string? prefix = null, CancellationToken cancellationToken = default) => _backend.WriteAsync(stream, prefix, cancellationToken);
+
+#pragma warning disable CS0618
+			/// <inheritdoc/>
+			public Task WriteExplicitPathAsync(string path, Stream stream, CancellationToken cancellationToken = default) => _backend.WriteExplicitPathAsync(path, stream, cancellationToken);
+#pragma warning restore CS0618
 
 			/// <inheritdoc/>
 			public ValueTask<Uri?> TryGetReadRedirectAsync(string path, CancellationToken cancellationToken = default) => _backend.TryGetReadRedirectAsync(path, cancellationToken);
 
 			/// <inheritdoc/>
-			public ValueTask<Uri?> TryGetWriteRedirectAsync(string path, CancellationToken cancellationToken = default) => _backend.TryGetWriteRedirectAsync(path, cancellationToken);
+			public ValueTask<(string, Uri)?> TryGetWriteRedirectAsync(string? prefix = null, CancellationToken cancellationToken = default) => _backend.TryGetWriteRedirectAsync(prefix, cancellationToken);
 
 			#endregion
 		}
diff --git a/Engine/Source/Programs/Horde/Horde.Server/Storage/StorageService.cs b/Engine/Source/Programs/Horde/Horde.Server/Storage/StorageService.cs
--- a/Engine/Source/Programs/Horde/Horde.Server/Storage/StorageService.cs
+++ b/Engine/Source/Programs/Horde/Horde.Server/Storage/StorageService.cs
@@ -137,7 +137,7 @@ namespace Horde.Server.Storage
 				}
 			}
 
-			string GetBlobPath(BundleLocator locator) => $"{_prefix}{locator.Path}.blob";
+			string GetBlobPath(BundleLocator locator) => $"{_prefix}{locator.Path}";
 
 			#region Blobs
 
@@ -164,16 +164,12 @@ namespace Horde.Server.Storage
 			/// <inheritdoc/>
 			public override async Task<BundleLocator> WriteBundleAsync(Bundle bundle, Utf8String prefix = default, CancellationToken cancellationToken = default)
 			{
-				// Add the blob record
-				BundleLocator locator = await _outer.AddBlobAsync(NamespaceId, prefix, null, cancellationToken);
-
-				using TelemetrySpan span = _tracer.StartActiveSpan($"{nameof(StorageService)}.{nameof(StorageClientImpl)}.{nameof(WriteBundleAsync)}");
-				span.SetAttribute("locator", locator.ToString());
-
-				// Write it to the backend
-				string path = GetBlobPath(locator);
 				using ReadOnlySequenceStream memoryStream = new ReadOnlySequenceStream(bundle.AsSequence());
-				await Backend.WriteAsync(path, memoryStream, cancellationToken);
+				string path = await Backend.WriteAsync(memoryStream, prefix.IsEmpty? null : prefix.ToString(), cancellationToken);
+
+				// Add the blob record
+				BundleLocator locator = new BundleLocator(path);
+				await _outer.AddBlobAsync(NamespaceId, locator, null, cancellationToken);
 
 				return locator;
 			}
@@ -186,16 +182,16 @@ namespace Horde.Server.Storage
 					return null;
 				}
 
-				BundleLocator locator = await _outer.AddBlobAsync(NamespaceId, prefix, null, cancellationToken);
-				string path = GetBlobPath(locator);
-
-				Uri? url = await Backend.TryGetWriteRedirectAsync(path, cancellationToken);
-				if (url == null)
+				(string Path, Uri Url)? redirect = await Backend.TryGetWriteRedirectAsync(prefix.IsEmpty? null : prefix.ToString(), cancellationToken);
+				if (redirect == null)
 				{
 					return null;
 				}
 
-				return (locator, url);
+				BundleLocator locator = new BundleLocator(redirect.Value.Path);
+				await _outer.AddBlobAsync(NamespaceId, locator, null, cancellationToken);
+
+				return (locator, redirect.Value.Url);
 			}
 
 			public async Task DeleteBlobAsync(BundleLocator locator, CancellationToken cancellationToken = default)
@@ -611,20 +607,16 @@ namespace Horde.Server.Storage
 		#region Blobs
 
 		/// <inheritdoc/>
-		async Task<BundleLocator> AddBlobAsync(NamespaceId namespaceId, Utf8String prefix = default, List<AliasInfo>? exports = null, CancellationToken cancellationToken = default)
+		async Task AddBlobAsync(NamespaceId namespaceId, BundleLocator locator, List<AliasInfo>? exports = null, CancellationToken cancellationToken = default)
 		{
 			ObjectId id = ObjectId.GenerateNewId(_clock.UtcNow);
-			BundleLocator blobId = (prefix.Length > 0) ? new BundleLocator($"{prefix}/{id}") : new BundleLocator(id.ToString());
-
-			BlobInfo blobInfo = new BlobInfo(id, namespaceId, blobId);
+			BlobInfo blobInfo = new BlobInfo(id, namespaceId, locator);
 			blobInfo.Aliases = exports;
 			await _blobCollection.InsertOneAsync(blobInfo, new InsertOneOptions { }, cancellationToken);
-
-			return blobId;
 		}
 
 		/// <inheritdoc/>
-		async Task<bool> IsBlobReferenced(ObjectId blobInfoId, CancellationToken cancellationToken = default)
+		async Task<bool> IsBlobReferencedAsync(ObjectId blobInfoId, CancellationToken cancellationToken = default)
 		{
 			FilterDefinition<BlobInfo> blobFilter = Builders<BlobInfo>.Filter.AnyEq(x => x.Imports, blobInfoId);
 			if (await _blobCollection.Find(blobFilter).AnyAsync(cancellationToken))
@@ -1035,7 +1027,7 @@ namespace Horde.Server.Storage
 				}
 
 				ObjectId blobInfoId = new ObjectId(((byte[]?)values[0])!);
-				if (blobInfoId < lastImportBlobInfoId && !await IsBlobReferenced(blobInfoId, cancellationToken))
+				if (blobInfoId < lastImportBlobInfoId && !await IsBlobReferencedAsync(blobInfoId, cancellationToken))
 				{
 					BlobInfo? info = await _blobCollection.FindOneAndDeleteAsync(x => x.Id == blobInfoId, cancellationToken: cancellationToken);
 					if (info != null)
diff --git a/Engine/Source/Programs/Horde/Horde.Server/Storage/TypedStorageBackend.cs b/Engine/Source/Programs/Horde/Horde.Server/Storage/TypedStorageBackend.cs
--- a/Engine/Source/Programs/Horde/Horde.Server/Storage/TypedStorageBackend.cs
+++ b/Engine/Source/Programs/Horde/Horde.Server/Storage/TypedStorageBackend.cs
@@ -47,7 +47,12 @@ namespace Horde.Server.Storage
 			public Task<Stream> ReadAsync(string path, int offset, int? length, CancellationToken cancellationToken) => _inner.ReadAsync(path, offset, length, cancellationToken);
 
 			/// <inheritdoc/>
-			public Task WriteAsync(string path, Stream stream, CancellationToken cancellationToken) => _inner.WriteAsync(path, stream, cancellationToken);
+			public Task<string> WriteAsync(Stream stream, string? prefix, CancellationToken cancellationToken) => _inner.WriteAsync(stream, prefix, cancellationToken);
+
+#pragma warning disable CS0618
+			/// <inheritdoc/>
+			public Task WriteExplicitPathAsync(string path, Stream stream, CancellationToken cancellationToken) => _inner.WriteExplicitPathAsync(path, stream, cancellationToken);
+#pragma warning restore CS0618
 
 			/// <inheritdoc/>
 			public Task DeleteAsync(string path, CancellationToken cancellationToken) => _inner.DeleteAsync(path, cancellationToken);
@@ -62,54 +67,7 @@ namespace Horde.Server.Storage
 			public ValueTask<Uri?> TryGetReadRedirectAsync(string path, CancellationToken cancellationToken = default) => _inner.TryGetReadRedirectAsync(path, cancellationToken);
 
 			/// <inheritdoc/>
-			public ValueTask<Uri?> TryGetWriteRedirectAsync(string path, CancellationToken cancellationToken = default) => _inner.TryGetWriteRedirectAsync(path, cancellationToken);
-		}
-
-		/// <summary>
-		/// Extension for blob files
-		/// </summary>
-		public const string BlobExtension = ".blob";
-
-		/// <summary>
-		/// Gets the blob id from a path
-		/// </summary>
-		/// <param name="path">Path to the blob</param>
-		/// <returns>Path to the blob</returns>
-		public static Utf8String GetBlobPathFromFileName(string path)
-		{
-			Utf8String blobPath;
-			if (!TryGetBlobPathFromFileName(path, out blobPath))
-			{
-				throw new ArgumentException("Path is not a valid blob identifier", nameof(path));
-			}
-			return blobPath;
-		}
-
-		/// <summary>
-		/// Gets the path to a blob
-		/// </summary>
-		/// <param name="blobPath">Blob identifier</param>
-		/// <returns>Path to the blob</returns>
-		public static string GetBlobFileName(Utf8String blobPath) => $"{blobPath}{BlobExtension}";
-
-		/// <summary>
-		/// Gets a blob id from a path within the storage backend
-		/// </summary>
-		/// <param name="path">Path to the file</param>
-		/// <param name="blobPath">Receives the blob id on success</param>
-		/// <returns>True on success</returns>
-		public static bool TryGetBlobPathFromFileName(string path, out Utf8String blobPath)
-		{
-			if (path.EndsWith(BlobExtension, StringComparison.Ordinal))
-			{
-				blobPath = new Utf8String(path.Substring(0, path.Length - BlobExtension.Length));
-				return true;
-			}
-			else
-			{
-				blobPath = default;
-				return false;
-			}
+			public ValueTask<(string, Uri)?> TryGetWriteRedirectAsync(string? prefix = null, CancellationToken cancellationToken = default) => _inner.TryGetWriteRedirectAsync(prefix, cancellationToken);
 		}
 
 		/// <summary>
diff --git a/Engine/Source/Programs/Shared/EpicGames.Horde/Storage/Backends/FileStorageBackend.cs b/Engine/Source/Programs/Shared/EpicGames.Horde/Storage/Backends/FileStorageBackend.cs
--- a/Engine/Source/Programs/Shared/EpicGames.Horde/Storage/Backends/FileStorageBackend.cs
+++ b/Engine/Source/Programs/Shared/EpicGames.Horde/Storage/Backends/FileStorageBackend.cs
@@ -20,16 +20,6 @@ namespace EpicGames.Horde.Storage.Backends
 		/// </summary>
 		private readonly DirectoryReference _baseDir;
 
-		/// <summary>
-		/// Unique identifier for this instance
-		/// </summary>
-		private readonly string _instanceId;
-
-		/// <summary>
-		/// Unique id for each write
-		/// </summary>
-		private int _uniqueId;
-
 		/// <inheritdoc/>
 		public bool SupportsRedirects => false;
 
@@ -40,7 +30,6 @@ namespace EpicGames.Horde.Storage.Backends
 		public FileStorageBackend(DirectoryReference baseDir)
 		{
 			_baseDir = baseDir;
-			_instanceId = Guid.NewGuid().ToString("N");
 			DirectoryReference.CreateDirectory(_baseDir);
 		}
 
@@ -50,52 +39,48 @@ namespace EpicGames.Horde.Storage.Backends
 		}
 
 		/// <inheritdoc/>
-		public Task<Stream> ReadAsync(string path, CancellationToken cancellationToken)
+		public Task<Stream> ReadAsync(string path, int offset, int? length, CancellationToken cancellationToken)
 		{
 			FileReference location = FileReference.Combine(_baseDir, path);
-			return Task.FromResult<Stream>(FileReference.Open(location, FileMode.Open, FileAccess.Read, FileShare.Read));
+			Stream stream = FileReference.Open(location, FileMode.Open, FileAccess.Read, FileShare.Read);
+			stream.Seek(offset, SeekOrigin.Begin);
+			return Task.FromResult(stream);
 		}
 
 		/// <inheritdoc/>
-		public async Task<Stream> ReadAsync(string path, int offset, int? length, CancellationToken cancellationToken)
+		public async Task<string> WriteAsync(Stream stream, string? prefix = null, CancellationToken cancellationToken = default)
 		{
-			Stream stream = await ReadAsync(path, cancellationToken);
-			stream.Seek(offset, SeekOrigin.Begin);
-			return stream;
+			string path = StorageHelpers.CreateUniqueName(prefix);
+			await WriteExplicitPathAsync(path, stream, cancellationToken);
+			return path;
 		}
 
 		/// <inheritdoc/>
-		public async Task WriteAsync(string path, Stream stream, CancellationToken cancellationToken)
+		public async Task WriteExplicitPathAsync(string path, Stream stream, CancellationToken cancellationToken = default)
 		{
 			FileReference finalLocation = FileReference.Combine(_baseDir, path);
-			if (!FileReference.Exists(finalLocation))
+			DirectoryReference.CreateDirectory(finalLocation.Directory);
+			FileReference tempLocation = new FileReference($"{finalLocation}.tmp");
+		
+			using (Stream outputStream = FileReference.Open(tempLocation, FileMode.Create, FileAccess.Write, FileShare.Read))
 			{
-				// Write to a temp file first
-				int newUniqueId = Interlocked.Increment(ref _uniqueId);
-
-				DirectoryReference.CreateDirectory(finalLocation.Directory);
-				FileReference tempLocation = new FileReference($"{finalLocation}.{_instanceId}.{newUniqueId:x8}");
-
-				using (Stream outputStream = FileReference.Open(tempLocation, FileMode.Create, FileAccess.Write, FileShare.Read))
-				{
-					await stream.CopyToAsync(outputStream, cancellationToken);
-				}
+				await stream.CopyToAsync(outputStream, cancellationToken);
+			}
 
-				// Move the temp file into place
-				try
+			// Move the temp file into place
+			try
+			{
+				FileReference.Move(tempLocation, finalLocation, true);
+			}
+			catch (IOException) // Already exists
+			{
+				if (FileReference.Exists(finalLocation))
 				{
-					FileReference.Move(tempLocation, finalLocation, true);
+					FileReference.Delete(tempLocation);
 				}
-				catch (IOException) // Already exists
+				else
 				{
-					if (FileReference.Exists(finalLocation))
-					{
-						FileReference.Delete(tempLocation);
-					}
-					else
-					{
-						throw;
-					}
+					throw;
 				}
 			}
 		}
@@ -158,6 +143,6 @@ namespace EpicGames.Horde.Storage.Backends
 		public ValueTask<Uri?> TryGetReadRedirectAsync(string path, CancellationToken cancellationToken = default) => default;
 
 		/// <inheritdoc/>
-		public ValueTask<Uri?> TryGetWriteRedirectAsync(string path, CancellationToken cancellationToken = default) => default;
+		public ValueTask<(string, Uri)?> TryGetWriteRedirectAsync(string? prefix = null, CancellationToken cancellationToken = default) => default;
 	}
 }
diff --git a/Engine/Source/Programs/Shared/EpicGames.Horde/Storage/Backends/MemoryStorageBackend.cs b/Engine/Source/Programs/Shared/EpicGames.Horde/Storage/Backends/MemoryStorageBackend.cs
--- a/Engine/Source/Programs/Shared/EpicGames.Horde/Storage/Backends/MemoryStorageBackend.cs
+++ b/Engine/Source/Programs/Shared/EpicGames.Horde/Storage/Backends/MemoryStorageBackend.cs
@@ -29,9 +29,6 @@ namespace EpicGames.Horde.Storage.Backends
 		{
 		}
 
-		/// <inheritdoc/>
-		public Task<Stream> ReadAsync(string path, CancellationToken cancellationToken) => ReadAsync(path, 0, null, cancellationToken);
-
 		/// <inheritdoc/>
 		public Task<Stream> ReadAsync(string path, int offset, int? length, CancellationToken cancellationToken)
 		{
@@ -44,7 +41,15 @@ namespace EpicGames.Horde.Storage.Backends
 		}
 
 		/// <inheritdoc/>
-		public async Task WriteAsync(string path, Stream stream, CancellationToken cancellationToken)
+		public async Task<string> WriteAsync(Stream stream, string? prefix = null, CancellationToken cancellationToken = default)
+		{
+			string path = StorageHelpers.CreateUniqueName(prefix);
+			await WriteExplicitPathAsync(path, stream, cancellationToken);
+			return path;
+		}
+
+		/// <inheritdoc/>
+		public async Task WriteExplicitPathAsync(string path, Stream stream, CancellationToken cancellationToken = default)
 		{
 			using (MemoryStream buffer = new MemoryStream())
 			{
@@ -81,6 +86,6 @@ namespace EpicGames.Horde.Storage.Backends
 		public ValueTask<Uri?> TryGetReadRedirectAsync(string path, CancellationToken cancellationToken = default) => default;
 
 		/// <inheritdoc/>
-		public ValueTask<Uri?> TryGetWriteRedirectAsync(string path, CancellationToken cancellationToken = default) => default;
+		public ValueTask<(string, Uri)?> TryGetWriteRedirectAsync(string? prefix = null, CancellationToken cancellationToken = default) => default;
 	}
 }
diff --git a/Engine/Source/Programs/Shared/EpicGames.Horde/Storage/IStorageBackend.cs b/Engine/Source/Programs/Shared/EpicGames.Horde/Storage/IStorageBackend.cs
--- a/Engine/Source/Programs/Shared/EpicGames.Horde/Storage/IStorageBackend.cs
+++ b/Engine/Source/Programs/Shared/EpicGames.Horde/Storage/IStorageBackend.cs
@@ -3,6 +3,7 @@
 using System;
 using System.Collections.Generic;
 using System.IO;
+using System.Text;
 using System.Threading;
 using System.Threading.Tasks;
 using EpicGames.Core;
@@ -30,12 +31,23 @@ namespace EpicGames.Horde.Storage
 		Task<Stream> ReadAsync(string path, int offset, int? length, CancellationToken cancellationToken = default);
 
 		/// <summary>
-		/// Writes a stream to the given path. If the stream throws an exception during read, the write will be aborted.
+		/// Writes a stream to the storage backend. If the stream throws an exception during read, the write will be aborted.
 		/// </summary>
-		/// <param name="path">Relative path within the bucket</param>
 		/// <param name="stream">Stream to write</param>
+		/// <param name="prefix">Path prefix for the uploaded data</param>
 		/// <param name="cancellationToken">Cancellation token for the operation</param>
-		Task WriteAsync(string path, Stream stream, CancellationToken cancellationToken = default);
+		/// <returns>Path to the uploaded object</returns>
+		Task<string> WriteAsync(Stream stream, string? prefix = null, CancellationToken cancellationToken = default);
+
+		/// <summary>
+		/// Writes a stream to the storage backend. This overload is deprecated; prefer passing a prefix to allow the server to determine a unique path.
+		/// </summary>
+		/// <param name="path"></param>
+		/// <param name="stream">Stream to write</param>
+		/// <param name="cancellationToken">Cancellation token for the operation</param>
+		/// <returns>Path to the uploaded object</returns>
+		[Obsolete("Use WriteAsync() instead. Ability to specify an explicit path is deprecated and will be removed in a future release.")]
+		Task WriteExplicitPathAsync(string path, Stream stream, CancellationToken cancellationToken = default);
 
 		/// <summary>
 		/// Tests whether the given path exists
@@ -71,10 +83,47 @@ namespace EpicGames.Horde.Storage
 		/// <summary>
 		/// Gets a HTTP redirect for a write request
 		/// </summary>
-		/// <param name="path">Path to write to</param>
+		/// <param name="prefix">Prefix for the uploaded data</param>
 		/// <param name="cancellationToken">Cancellation token for the operation</param>
-		/// <returns>Path to upload the data to</returns>
-		ValueTask<Uri?> TryGetWriteRedirectAsync(string path, CancellationToken cancellationToken = default);
+		/// <returns>Path for retrieval, and URI to upload the data to</returns>
+		ValueTask<(string, Uri)?> TryGetWriteRedirectAsync(string? prefix = null, CancellationToken cancellationToken = default);
+	}
+
+	/// <summary>
+	/// Utility methods for storage backend implementations
+	/// </summary>
+	public static class StorageHelpers
+	{
+		/// <summary>
+		/// Unique session id used for unique ids
+		/// </summary>
+		static readonly string s_sessionPrefix = $"{Guid.NewGuid():n}_";
+
+		/// <summary>
+		/// Incremented value used for each supplied id
+		/// </summary>
+		static int _increment;
+
+		/// <summary>
+		/// Creates a unique name with a given prefix
+		/// </summary>
+		/// <param name="prefix">The prefix to use</param>
+		/// <returns>Unique name generated with the given prefix</returns>
+		public static string CreateUniqueName(string? prefix)
+		{
+			StringBuilder builder = new StringBuilder(prefix);
+			if (!String.IsNullOrEmpty(prefix))
+			{
+				builder.Append(prefix);
+				if (!prefix.EndsWith('/'))
+				{
+					builder.Append('/');
+				}
+			}
+			builder.Append(s_sessionPrefix);
+			builder.Append(Interlocked.Increment(ref _increment));
+			return builder.ToString();
+		}
 	}
 
 	/// <summary>
@@ -114,14 +163,14 @@ namespace EpicGames.Horde.Storage
 		/// Writes a block of memory to storage
 		/// </summary>
 		/// <param name="storageBackend">Backend to read from</param>
-		/// <param name="path">Object name within the store</param>
 		/// <param name="data">Data to be written</param>
+		/// <param name="prefix">Prefix for the uploaded data</param>
 		/// <param name="cancellationToken">Cancellation token for the operation</param>
-		public static async Task WriteBytesAsync(this IStorageBackend storageBackend, string path, ReadOnlyMemory<byte> data, CancellationToken cancellationToken = default)
+		public static async Task<string> WriteBytesAsync(this IStorageBackend storageBackend, ReadOnlyMemory<byte> data, string? prefix = null, CancellationToken cancellationToken = default)
 		{
 			using (ReadOnlyMemoryStream stream = new ReadOnlyMemoryStream(data))
 			{
-				await storageBackend.WriteAsync(path, stream, cancellationToken);
+				return await storageBackend.WriteAsync(stream, prefix, cancellationToken);
 			}
 		}
 	}
