Horde: Add an HTTP storage backend, which uses the regular Horde endpoints for reading and writing blobs.

diff --git a/Engine/Source/Programs/Shared/EpicGames.Horde/Storage/Backends/HttpStorageBackend.cs b/Engine/Source/Programs/Shared/EpicGames.Horde/Storage/Backends/HttpStorageBackend.cs
new file mode 100644
--- /dev/null
+++ b/Engine/Source/Programs/Shared/EpicGames.Horde/Storage/Backends/HttpStorageBackend.cs
@@ -0,0 +1,185 @@
+// Copyright Epic Games, Inc. All Rights Reserved.
+
+using System;
+using System.Collections.Generic;
+using System.IO;
+using System.Net.Http;
+using System.Net.Http.Headers;
+using System.Net.Http.Json;
+using System.Threading;
+using System.Threading.Tasks;
+using Microsoft.Extensions.Logging;
+
+namespace EpicGames.Horde.Storage.Backends
+{
+	/// <summary>
+	/// Storage backend which communicates with the Horde server over HTTP for content
+	/// </summary>
+	public sealed class HttpStorageBackend : IStorageBackend
+	{
+		class WriteBlobResponse
+		{
+			public string Blob { get; set; }
+			public Uri? UploadUrl { get; set; }
+			public bool? SupportsRedirects { get; set; }
+		}
+
+		readonly Func<HttpClient> _createClient;
+		readonly Func<HttpClient> _createRedirectClient;
+		readonly ILogger _logger;
+		bool _supportsUploadRedirects = true;
+
+		/// <inheritdoc/>
+		public bool SupportsRedirects => _supportsUploadRedirects;
+
+		/// <summary>
+		/// Constructor
+		/// </summary>
+		public HttpStorageBackend(Func<HttpClient> createClient, Func<HttpClient> createRedirectClient, ILogger logger)
+		{
+			_createClient = createClient;
+			_createRedirectClient = createRedirectClient;
+			_logger = logger;
+		}
+
+		/// <inheritdoc/>
+		public void Dispose()
+		{
+		}
+
+		#region Blobs
+
+		/// <inheritdoc/>
+		public async Task<Stream> ReadAsync(string path, int offset, int? length, CancellationToken cancellationToken = default)
+		{
+			if (offset == 0 && length == null)
+			{
+				_logger.LogDebug("Reading {Locator}", path);
+			}
+			else if (length == null)
+			{
+				_logger.LogDebug("Reading {Locator} ({Offset}..)", path, offset);
+			}
+			else
+			{
+				_logger.LogDebug("Reading {Locator} ({Offset}+{Length})", path, offset, length);
+			}
+
+			if (length.HasValue && length.Value == 0)
+			{
+				return new MemoryStream(Array.Empty<byte>());
+			}
+
+			using (HttpClient httpClient = _createClient())
+			{
+				using (HttpRequestMessage request = new HttpRequestMessage(HttpMethod.Get, $"blobs/{path}"))
+				{
+					if (offset != 0 || length != null)
+					{
+						request.Headers.Range = new RangeHeaderValue(offset, (length == null) ? null : (offset + (length - 1)));
+					}
+
+					HttpResponseMessage response = await httpClient.SendAsync(request, cancellationToken);
+					response.EnsureSuccessStatusCode();
+					return await response.Content.ReadAsStreamAsync(cancellationToken);
+				}
+			}
+		}
+
+		/// <inheritdoc/>
+		public async Task<string> WriteAsync(Stream stream, string? prefix = null, CancellationToken cancellationToken = default)
+		{
+			using StreamContent streamContent = new StreamContent(stream);
+
+			if (_supportsUploadRedirects)
+			{
+				using (HttpClient redirectHttpClient = _createRedirectClient())
+				{
+					WriteBlobResponse redirectResponse = await SendWriteRequestAsync(null, prefix, cancellationToken);
+					if (redirectResponse.UploadUrl != null)
+					{
+						using (HttpResponseMessage uploadResponse = await redirectHttpClient.PutAsync(redirectResponse.UploadUrl, streamContent, cancellationToken))
+						{
+							if (!uploadResponse.IsSuccessStatusCode)
+							{
+								string body = await uploadResponse.Content.ReadAsStringAsync(cancellationToken);
+								throw new StorageException($"Unable to upload data to redirected URL: {body}");
+							}
+						}
+						_logger.LogDebug("Written {Locator} (using redirect)", redirectResponse.Blob);
+						return redirectResponse.Blob;
+					}
+				}
+			}
+
+			WriteBlobResponse response = await SendWriteRequestAsync(streamContent, prefix, cancellationToken);
+			_supportsUploadRedirects = response.SupportsRedirects ?? false;
+			_logger.LogDebug("Written {Locator} (direct)", response.Blob);
+			return response.Blob;
+		}
+
+		/// <inheritdoc/>
+		public Task WriteExplicitPathAsync(string path, Stream stream, CancellationToken cancellationToken = default) => throw new NotSupportedException();
+
+		async Task<WriteBlobResponse> SendWriteRequestAsync(StreamContent? streamContent, string? prefix = null, CancellationToken cancellationToken = default)
+		{
+			using (HttpClient httpClient = _createClient())
+			{
+				using (HttpRequestMessage request = new HttpRequestMessage(HttpMethod.Post, $"blobs"))
+				{
+					using StringContent stringContent = new StringContent(prefix ?? String.Empty);
+
+					MultipartFormDataContent form = new MultipartFormDataContent();
+					if (streamContent != null)
+					{
+						form.Add(streamContent, "file", "filename");
+					}
+					if (prefix != null)
+					{
+						form.Add(stringContent, "prefix");
+					}
+
+					request.Content = form;
+					using (HttpResponseMessage response = await httpClient.SendAsync(request, cancellationToken))
+					{
+						response.EnsureSuccessStatusCode();
+						WriteBlobResponse? data = await response.Content.ReadFromJsonAsync<WriteBlobResponse>(cancellationToken: cancellationToken);
+						return data!;
+					}
+				}
+			}
+		}
+
+		/// <inheritdoc/>
+		public Task<bool> ExistsAsync(string path, CancellationToken cancellationToken = default) => throw new NotSupportedException();
+
+		/// <inheritdoc/>
+		public Task DeleteAsync(string path, CancellationToken cancellationToken = default) => throw new NotSupportedException();
+
+		/// <inheritdoc/>
+		public IAsyncEnumerable<string> EnumerateAsync(CancellationToken cancellationToken = default) => throw new NotSupportedException();
+
+		/// <inheritdoc/>
+		public ValueTask<Uri?> TryGetReadRedirectAsync(string path, CancellationToken cancellationToken = default)
+		{
+			// We don't currently have any need for getting a read redirect explicitly, though ReadAsync() calls may redirect us automatically.
+			return default;
+		}
+
+		/// <inheritdoc/>
+		public async ValueTask<(string, Uri)?> TryGetWriteRedirectAsync(string? prefix = null, CancellationToken cancellationToken = default)
+		{
+			if (_supportsUploadRedirects)
+			{
+				WriteBlobResponse redirectResponse = await SendWriteRequestAsync(null, prefix, cancellationToken);
+				if (redirectResponse.UploadUrl != null)
+				{
+					return (redirectResponse.Blob, redirectResponse.UploadUrl);
+				}
+			}
+			return null;
+		}
+
+		#endregion
+	}
+}
