Fix null reference warning.

diff --git a/Engine/Source/Programs/Shared/EpicGames.Horde/Storage/Backends/HttpStorageBackend.cs b/Engine/Source/Programs/Shared/EpicGames.Horde/Storage/Backends/HttpStorageBackend.cs
--- a/Engine/Source/Programs/Shared/EpicGames.Horde/Storage/Backends/HttpStorageBackend.cs
+++ b/Engine/Source/Programs/Shared/EpicGames.Horde/Storage/Backends/HttpStorageBackend.cs
@@ -19,7 +19,7 @@ namespace EpicGames.Horde.Storage.Backends
 	{
 		class WriteBlobResponse
 		{
-			public string Blob { get; set; }
+			public string Blob { get; set; } = String.Empty;
 			public Uri? UploadUrl { get; set; }
 			public bool? SupportsRedirects { get; set; }
 		}
