Horde: Add a caching storage backend which uses the local filesystem.

diff --git a/Engine/Source/Programs/Shared/EpicGames.Horde.Tests/BackendTests.cs b/Engine/Source/Programs/Shared/EpicGames.Horde.Tests/BackendTests.cs
new file mode 100644
--- /dev/null
+++ b/Engine/Source/Programs/Shared/EpicGames.Horde.Tests/BackendTests.cs
@@ -0,0 +1,124 @@
+// Copyright Epic Games, Inc. All Rights Reserved.
+
+using System;
+using System.Collections.Generic;
+using System.IO;
+using System.Linq;
+using System.Security.Cryptography.X509Certificates;
+using System.Text;
+using System.Threading.Tasks;
+using EpicGames.Core;
+using EpicGames.Horde.Storage;
+using EpicGames.Horde.Storage.Backends;
+using JetBrains.Annotations;
+using Microsoft.VisualStudio.TestTools.UnitTesting;
+
+namespace EpicGames.Horde.Tests
+{
+	[TestClass]
+	public class BackendTests
+	{
+		class TempDir : IDisposable
+		{
+			readonly DirectoryReference _cacheDir;
+
+			public DirectoryReference Location => _cacheDir;
+
+			public TempDir(string name)
+			{
+				_cacheDir = new DirectoryReference(name);
+				DirectoryReference.CreateDirectory(_cacheDir);
+				FileUtils.ForceDeleteDirectoryContents(_cacheDir);
+			}
+
+			public void Dispose()
+			{
+				FileUtils.ForceDeleteDirectoryContents(_cacheDir);
+				DirectoryReference.Delete(_cacheDir);
+			}
+		}
+
+		[TestMethod]
+		public async Task TestMemoryBackend()
+		{
+			using MemoryStorageBackend backend = new MemoryStorageBackend();
+			await TestBackendAsync(backend);
+		}
+
+		[TestMethod]
+		public async Task TestFileBackend()
+		{
+			using (TempDir tempDir = new TempDir("Cache"))
+			{
+				using FileStorageBackend backend = new FileStorageBackend(tempDir.Location);
+				await TestBackendAsync(backend);
+			}
+		}
+
+		[TestMethod]
+		public async Task TestCacheBackend()
+		{
+			using (TempDir tempDir = new TempDir("Cache"))
+			{
+				using MemoryStorageBackend memoryBackend = new MemoryStorageBackend();
+
+				using CacheStorageBackend cacheBackend = new CacheStorageBackend(tempDir.Location, 12, memoryBackend);
+				await TestBackendAsync(cacheBackend);
+
+				Assert.AreEqual(1, cacheBackend.Items.Count());
+				byte[] value = await cacheBackend.ReadBytesAsync(cacheBackend.Items.First());
+				Assert.IsTrue(value.SequenceEqual(Encoding.UTF8.GetBytes("item 2")));
+
+				byte[] data3 = Encoding.UTF8.GetBytes("3");
+				string path3 = await cacheBackend.WriteBytesAsync(data3);
+				await cacheBackend.ReadBytesAsync(path3);
+
+				byte[] data4 = Encoding.UTF8.GetBytes("4");
+				string path4 = await cacheBackend.WriteBytesAsync(data4);
+				await cacheBackend.ReadBytesAsync(path4);
+
+				byte[] data5 = Encoding.UTF8.GetBytes("5");
+				string path5 = await cacheBackend.WriteBytesAsync(data5);
+				await cacheBackend.ReadBytesAsync(path5);
+
+				HashSet<string> paths = new HashSet<string>(cacheBackend.Items);
+				Assert.AreEqual(4, paths.Count);
+				Assert.IsTrue(paths.Contains(path3));
+				Assert.IsTrue(paths.Contains(path4));
+				Assert.IsTrue(paths.Contains(path5));
+
+				byte[] data6 = Encoding.UTF8.GetBytes("12345678901");
+				string path6 = await cacheBackend.WriteBytesAsync(data6);
+				await cacheBackend.ReadBytesAsync(path6);
+
+				paths = new HashSet<string>(cacheBackend.Items);
+				Assert.AreEqual(2, paths.Count);
+				Assert.IsTrue(paths.Contains(path5));
+				Assert.IsTrue(paths.Contains(path6));
+
+				await cacheBackend.ReadBytesAsync(path3);
+				await cacheBackend.ReadBytesAsync(path4);
+
+				paths = new HashSet<string>(cacheBackend.Items);
+				Assert.AreEqual(2, paths.Count);
+				Assert.IsTrue(paths.Contains(path3));
+				Assert.IsTrue(paths.Contains(path4));
+			}
+		}
+
+		static async Task TestBackendAsync(IStorageBackend backend)
+		{
+			byte[] data1 = Encoding.UTF8.GetBytes("hello world");
+			byte[] data2 = Encoding.UTF8.GetBytes("item 2");
+
+			string path1 = await backend.WriteBytesAsync(data1);
+			byte[] outputData1 = await backend.ReadBytesAsync(path1);
+
+			string path2 = await backend.WriteBytesAsync(data2);
+			byte[] outputData2 = await backend.ReadBytesAsync(path2);
+
+			Assert.IsTrue(data1.SequenceEqual(outputData1));
+			Assert.IsTrue(data2.SequenceEqual(outputData2));
+		}
+	}
+}
diff --git a/Engine/Source/Programs/Shared/EpicGames.Horde/Storage/Backends/CacheStorageBackend.cs b/Engine/Source/Programs/Shared/EpicGames.Horde/Storage/Backends/CacheStorageBackend.cs
new file mode 100644
--- /dev/null
+++ b/Engine/Source/Programs/Shared/EpicGames.Horde/Storage/Backends/CacheStorageBackend.cs
@@ -0,0 +1,200 @@
+// Copyright Epic Games, Inc. All Rights Reserved.
+
+using System;
+using System.Collections.Generic;
+using System.IO;
+using System.Linq;
+using System.Threading;
+using System.Threading.Tasks;
+using EpicGames.Core;
+
+namespace EpicGames.Horde.Storage.Backends
+{
+	/// <summary>
+	/// Storage backend that utilizes the local filesystem
+	/// </summary>
+	public sealed class CacheStorageBackend : IStorageBackend
+	{
+		class Item
+		{
+			public string Path { get; }
+			public long Length { get; }
+			public LinkedListNode<Item> ListNode { get; }
+
+			public Item(string path, long length)
+			{
+				Path = path;
+				Length = length;
+				ListNode = new LinkedListNode<Item>(this);
+			}
+		}
+
+		class PendingItem
+		{
+			int _refCount;
+			readonly BackgroundTask _readTask;
+
+			public PendingItem(BackgroundTask readTask) => _readTask = readTask;
+			public void AddRef() => Interlocked.Increment(ref _refCount);
+
+			public async ValueTask ReleaseAsync()
+			{
+				if (Interlocked.Decrement(ref _refCount) == 0)
+				{
+					await _readTask.DisposeAsync();
+				}
+			}
+
+			public Task WaitAsync(CancellationToken cancellationToken) => _readTask.Task.WaitAsync(cancellationToken);
+		}
+
+		object LockObject => _items;
+
+		readonly DirectoryReference _cacheDir;
+		readonly long _maxSize;
+		readonly IStorageBackend _inner;
+
+		readonly LinkedList<Item> _items = new LinkedList<Item>();
+		readonly Dictionary<string, Item> _pathToItem = new Dictionary<string, Item>(StringComparer.Ordinal);
+		readonly Dictionary<string, PendingItem> _pathToPendingItem = new Dictionary<string, PendingItem>(StringComparer.Ordinal);
+
+		long _size;
+
+		/// <inheritdoc/>
+		public bool SupportsRedirects => _inner.SupportsRedirects;
+
+		internal IEnumerable<string> Items => _items.Select(x => x.Path);
+
+		/// <summary>
+		/// Constructor
+		/// </summary>
+		public CacheStorageBackend(DirectoryReference cacheDir, long maxSize, IStorageBackend inner)
+		{
+			_cacheDir = cacheDir;
+			_maxSize = maxSize;
+			_inner = inner;
+		}
+
+		/// <inheritdoc/>
+		public void Dispose()
+		{
+			_inner.Dispose();
+		}
+
+		/// <inheritdoc/>
+		public async Task<Stream> ReadAsync(string path, int offset, int? length, CancellationToken cancellationToken = default)
+		{
+			Stream stream = await OpenAsync(path, cancellationToken);
+			if (offset != 0)
+			{
+				stream.Seek(0, SeekOrigin.Begin);
+			}
+			return stream;
+		}
+
+		async Task<Stream> OpenAsync(string path, CancellationToken cancellationToken = default)
+		{
+			for(; ;)
+			{
+				PendingItem? pendingItem;
+				lock (LockObject)
+				{
+					Item? item;
+					if (_pathToItem.TryGetValue(path, out item))
+					{
+						return OpenItem(item);
+					}
+
+					if (!_pathToPendingItem.TryGetValue(path, out pendingItem))
+					{
+						pendingItem = new PendingItem(BackgroundTask.StartNew(x => ReadIntoCacheAsync(path, x)));
+						_pathToPendingItem.Add(path, pendingItem);
+					}
+
+					pendingItem.AddRef();
+				}
+
+				try
+				{
+					await pendingItem.WaitAsync(cancellationToken);
+				}
+				finally
+				{
+					await pendingItem.ReleaseAsync();
+				}
+			}
+		}
+
+		Stream OpenItem(Item item)
+		{
+			_items.Remove(item.ListNode);
+			_items.AddFirst(item.ListNode);
+
+			FileReference file = FileReference.Combine(_cacheDir, item.Path);
+			return FileReference.Open(file, FileMode.Open, FileAccess.Read, FileShare.Read | FileShare.Delete);
+		}
+
+		async Task ReadIntoCacheAsync(string path, CancellationToken cancellationToken)
+		{
+			using Stream stream = await _inner.ReadAsync(path, cancellationToken);
+			long totalLength = stream.Length;
+
+			lock (LockObject)
+			{
+				_size += totalLength;
+
+				while (_size > _maxSize && _items.Count > 0)
+				{
+					LinkedListNode<Item> lastItem = _items.Last!;
+					_pathToItem.Remove(lastItem.Value.Path);
+					_items.RemoveLast();
+
+					_size -= lastItem.Value.Length;
+
+					FileReference file = FileReference.Combine(_cacheDir, lastItem.Value.Path);
+					FileReference.Delete(file);
+				}
+			}
+
+			FileReference cacheFile = FileReference.Combine(_cacheDir, path);
+			DirectoryReference.CreateDirectory(cacheFile.Directory);
+
+			using (FileStream outputStream = FileReference.Open(cacheFile, FileMode.Create, FileAccess.Write, FileShare.None))
+			{
+				await stream.CopyToAsync(outputStream, cancellationToken);
+			}
+
+			lock (LockObject)
+			{
+				Item item = new Item(path, totalLength);
+				_items.AddFirst(item.ListNode);
+				_pathToItem.Add(path, item);
+
+				_pathToPendingItem.Remove(path);
+			}
+		}
+
+		/// <inheritdoc/>
+		public Task<string> WriteAsync(Stream stream, string? prefix = null, CancellationToken cancellationToken = default) => _inner.WriteAsync(stream, prefix, cancellationToken);
+
+#pragma warning disable CS0618 // Type or member is obsolete
+		/// <inheritdoc/>
+		public Task WriteExplicitPathAsync(string path, Stream stream, CancellationToken cancellationToken = default) => _inner.WriteExplicitPathAsync(path, stream, cancellationToken);
+#pragma warning restore CS0618 // Type or member is obsolete
+
+		/// <inheritdoc/>
+		public Task DeleteAsync(string path, CancellationToken cancellationToken = default) => _inner.DeleteAsync(path, cancellationToken);
+
+		/// <inheritdoc/>
+		public IAsyncEnumerable<string> EnumerateAsync(CancellationToken cancellationToken = default) => _inner.EnumerateAsync(cancellationToken);
+
+		/// <inheritdoc/>
+		public Task<bool> ExistsAsync(string path, CancellationToken cancellationToken = default) => _inner.ExistsAsync(path, cancellationToken);
+
+		/// <inheritdoc/>
+		public ValueTask<Uri?> TryGetReadRedirectAsync(string path, CancellationToken cancellationToken = default) => _inner.TryGetReadRedirectAsync(path, cancellationToken);
+
+		/// <inheritdoc/>
+		public ValueTask<(string, Uri)?> TryGetWriteRedirectAsync(string? prefix = null, CancellationToken cancellationToken = default) => _inner.TryGetWriteRedirectAsync(prefix, cancellationToken);
+	}
+}
