Horde: Separate implementation of the local cache from the backend implementation, allowing multiple backends to share the same cache.

diff --git a/Engine/Source/Programs/Shared/EpicGames.Horde.Tests/BackendTests.cs b/Engine/Source/Programs/Shared/EpicGames.Horde.Tests/BackendTests.cs
--- a/Engine/Source/Programs/Shared/EpicGames.Horde.Tests/BackendTests.cs
+++ b/Engine/Source/Programs/Shared/EpicGames.Horde.Tests/BackendTests.cs
@@ -2,15 +2,12 @@
 
 using System;
 using System.Collections.Generic;
-using System.IO;
 using System.Linq;
-using System.Security.Cryptography.X509Certificates;
 using System.Text;
 using System.Threading.Tasks;
 using EpicGames.Core;
 using EpicGames.Horde.Storage;
 using EpicGames.Horde.Storage.Backends;
-using JetBrains.Annotations;
 using Microsoft.VisualStudio.TestTools.UnitTesting;
 
 namespace EpicGames.Horde.Tests
@@ -60,13 +57,14 @@ namespace EpicGames.Horde.Tests
 		{
 			using (TempDir tempDir = new TempDir("Cache"))
 			{
-				using MemoryStorageBackend memoryBackend = new MemoryStorageBackend();
+				CacheStorageBackendDetail cacheBackendDetail = new CacheStorageBackendDetail(tempDir.Location, 12);
 
-				using CacheStorageBackend cacheBackend = new CacheStorageBackend(tempDir.Location, 12, memoryBackend);
+				using MemoryStorageBackend memoryBackend = new MemoryStorageBackend();
+				using CacheStorageBackend cacheBackend = new CacheStorageBackend("", cacheBackendDetail, memoryBackend);
 				await TestBackendAsync(cacheBackend);
 
-				Assert.AreEqual(1, cacheBackend.Items.Count());
-				byte[] value = await cacheBackend.ReadBytesAsync(cacheBackend.Items.First());
+				Assert.AreEqual(1, cacheBackendDetail.Items.Count());
+				byte[] value = await cacheBackend.ReadBytesAsync(cacheBackendDetail.Items.First());
 				Assert.IsTrue(value.SequenceEqual(Encoding.UTF8.GetBytes("item 2")));
 
 				byte[] data3 = Encoding.UTF8.GetBytes("3");
@@ -81,7 +79,7 @@ namespace EpicGames.Horde.Tests
 				string path5 = await cacheBackend.WriteBytesAsync(data5);
 				await cacheBackend.ReadBytesAsync(path5);
 
-				HashSet<string> paths = new HashSet<string>(cacheBackend.Items);
+				HashSet<string> paths = new HashSet<string>(cacheBackendDetail.Items);
 				Assert.AreEqual(4, paths.Count);
 				Assert.IsTrue(paths.Contains(path3));
 				Assert.IsTrue(paths.Contains(path4));
@@ -91,7 +89,7 @@ namespace EpicGames.Horde.Tests
 				string path6 = await cacheBackend.WriteBytesAsync(data6);
 				await cacheBackend.ReadBytesAsync(path6);
 
-				paths = new HashSet<string>(cacheBackend.Items);
+				paths = new HashSet<string>(cacheBackendDetail.Items);
 				Assert.AreEqual(2, paths.Count);
 				Assert.IsTrue(paths.Contains(path5));
 				Assert.IsTrue(paths.Contains(path6));
@@ -99,7 +97,7 @@ namespace EpicGames.Horde.Tests
 				await cacheBackend.ReadBytesAsync(path3);
 				await cacheBackend.ReadBytesAsync(path4);
 
-				paths = new HashSet<string>(cacheBackend.Items);
+				paths = new HashSet<string>(cacheBackendDetail.Items);
 				Assert.AreEqual(2, paths.Count);
 				Assert.IsTrue(paths.Contains(path3));
 				Assert.IsTrue(paths.Contains(path4));
diff --git a/Engine/Source/Programs/Shared/EpicGames.Horde/Storage/Backends/CacheStorageBackend.cs b/Engine/Source/Programs/Shared/EpicGames.Horde/Storage/Backends/CacheStorageBackend.cs
--- a/Engine/Source/Programs/Shared/EpicGames.Horde/Storage/Backends/CacheStorageBackend.cs
+++ b/Engine/Source/Programs/Shared/EpicGames.Horde/Storage/Backends/CacheStorageBackend.cs
@@ -15,15 +15,70 @@
 	/// </summary>
 	public sealed class CacheStorageBackend : IStorageBackend
 	{
+		readonly string _keyPrefix;
+		readonly CacheStorageBackendDetail _cacheStorage;
+		readonly IStorageBackend _inner;
+
+		/// <inheritdoc/>
+		public bool SupportsRedirects => _inner.SupportsRedirects;
+
+		/// <summary>
+		/// Constructor
+		/// </summary>
+		public CacheStorageBackend(string keyPrefix, CacheStorageBackendDetail cacheStorage, IStorageBackend inner)
+		{
+			_keyPrefix = keyPrefix;
+			_cacheStorage = cacheStorage;
+			_inner = inner;
+		}
+
+		/// <inheritdoc/>
+		public void Dispose() => _inner.Dispose();
+
+		/// <inheritdoc/>
+		public async Task<Stream> ReadAsync(string path, int offset, int? length, CancellationToken cancellationToken = default)
+		{
+			return await _cacheStorage.ReadAsync($"{_keyPrefix}{path}", ctx => _inner.ReadAsync(path, ctx), cancellationToken);
+		}
+
+		/// <inheritdoc/>
+		public Task<string> WriteAsync(Stream stream, string? prefix = null, CancellationToken cancellationToken = default) => _inner.WriteAsync(stream, prefix, cancellationToken);
+
+#pragma warning disable CS0618 // Type or member is obsolete
+		/// <inheritdoc/>
+		public Task WriteExplicitPathAsync(string path, Stream stream, CancellationToken cancellationToken = default) => _inner.WriteExplicitPathAsync(path, stream, cancellationToken);
+#pragma warning restore CS0618 // Type or member is obsolete
+
+		/// <inheritdoc/>
+		public Task DeleteAsync(string path, CancellationToken cancellationToken = default) => _inner.DeleteAsync(path, cancellationToken);
+
+		/// <inheritdoc/>
+		public IAsyncEnumerable<string> EnumerateAsync(CancellationToken cancellationToken = default) => _inner.EnumerateAsync(cancellationToken);
+
+		/// <inheritdoc/>
+		public Task<bool> ExistsAsync(string path, CancellationToken cancellationToken = default) => _inner.ExistsAsync(path, cancellationToken);
+
+		/// <inheritdoc/>
+		public ValueTask<Uri?> TryGetReadRedirectAsync(string path, CancellationToken cancellationToken = default) => _inner.TryGetReadRedirectAsync(path, cancellationToken);
+
+		/// <inheritdoc/>
+		public ValueTask<(string, Uri)?> TryGetWriteRedirectAsync(string? prefix = null, CancellationToken cancellationToken = default) => _inner.TryGetWriteRedirectAsync(prefix, cancellationToken);
+	}
+
+	/// <summary>
+	/// Implementation of a local disk cache which can be shared by multiple backends
+	/// </summary>
+	public sealed class CacheStorageBackendDetail
+	{
 		class Item
 		{
-			public string Path { get; }
+			public string Key { get; }
 			public long Length { get; }
 			public LinkedListNode<Item> ListNode { get; }
 
-			public Item(string path, long length)
+			public Item(string key, long length)
 			{
-				Path = path;
+				Key = key;
 				Length = length;
 				ListNode = new LinkedListNode<Item>(this);
 			}
@@ -52,7 +107,6 @@
 
 		readonly DirectoryReference _cacheDir;
 		readonly long _maxSize;
-		readonly IStorageBackend _inner;
 
 		readonly LinkedList<Item> _items = new LinkedList<Item>();
 		readonly Dictionary<string, Item> _pathToItem = new Dictionary<string, Item>(StringComparer.Ordinal);
@@ -60,39 +114,19 @@
 
 		long _size;
 
-		/// <inheritdoc/>
-		public bool SupportsRedirects => _inner.SupportsRedirects;
+		internal IEnumerable<string> Items => _items.Select(x => x.Key);
 
-		internal IEnumerable<string> Items => _items.Select(x => x.Path);
-
 		/// <summary>
 		/// Constructor
 		/// </summary>
-		public CacheStorageBackend(DirectoryReference cacheDir, long maxSize, IStorageBackend inner)
+		public CacheStorageBackendDetail(DirectoryReference cacheDir, long maxSize)
 		{
 			_cacheDir = cacheDir;
 			_maxSize = maxSize;
-			_inner = inner;
 		}
 
 		/// <inheritdoc/>
-		public void Dispose()
+		public async Task<Stream> ReadAsync(string key, Func<CancellationToken, Task<Stream>> createStreamAsync, CancellationToken cancellationToken = default)
 		{
-			_inner.Dispose();
-		}
-
-		/// <inheritdoc/>
-		public async Task<Stream> ReadAsync(string path, int offset, int? length, CancellationToken cancellationToken = default)
-		{
-			Stream stream = await OpenAsync(path, cancellationToken);
-			if (offset != 0)
-			{
-				stream.Seek(0, SeekOrigin.Begin);
-			}
-			return stream;
-		}
-
-		async Task<Stream> OpenAsync(string path, CancellationToken cancellationToken = default)
-		{
 			for(; ;)
 			{
@@ -100,15 +134,15 @@
 				lock (LockObject)
 				{
 					Item? item;
-					if (_pathToItem.TryGetValue(path, out item))
+					if (_pathToItem.TryGetValue(key, out item))
 					{
 						return OpenItem(item);
 					}
 
-					if (!_pathToPendingItem.TryGetValue(path, out pendingItem))
+					if (!_pathToPendingItem.TryGetValue(key, out pendingItem))
 					{
-						pendingItem = new PendingItem(BackgroundTask.StartNew(x => ReadIntoCacheAsync(path, x)));
-						_pathToPendingItem.Add(path, pendingItem);
+						pendingItem = new PendingItem(BackgroundTask.StartNew(x => ReadIntoCacheAsync(key, createStreamAsync, x)));
+						_pathToPendingItem.Add(key, pendingItem);
 					}
 
 					pendingItem.AddRef();
@@ -130,13 +164,13 @@
 			_items.Remove(item.ListNode);
 			_items.AddFirst(item.ListNode);
 
-			FileReference file = FileReference.Combine(_cacheDir, item.Path);
+			FileReference file = FileReference.Combine(_cacheDir, item.Key);
 			return FileReference.Open(file, FileMode.Open, FileAccess.Read, FileShare.Read | FileShare.Delete);
 		}
 
-		async Task ReadIntoCacheAsync(string path, CancellationToken cancellationToken)
+		async Task ReadIntoCacheAsync(string key, Func<CancellationToken, Task<Stream>> createStreamAsync, CancellationToken cancellationToken)
 		{
-			using Stream stream = await _inner.ReadAsync(path, cancellationToken);
+			using Stream stream = await createStreamAsync(cancellationToken);
 			long totalLength = stream.Length;
 
 			lock (LockObject)
@@ -146,17 +180,17 @@
 				while (_size > _maxSize && _items.Count > 0)
 				{
 					LinkedListNode<Item> lastItem = _items.Last!;
-					_pathToItem.Remove(lastItem.Value.Path);
+					_pathToItem.Remove(lastItem.Value.Key);
 					_items.RemoveLast();
 
 					_size -= lastItem.Value.Length;
 
-					FileReference file = FileReference.Combine(_cacheDir, lastItem.Value.Path);
+					FileReference file = FileReference.Combine(_cacheDir, lastItem.Value.Key);
 					FileReference.Delete(file);
 				}
 			}
 
-			FileReference cacheFile = FileReference.Combine(_cacheDir, path);
+			FileReference cacheFile = FileReference.Combine(_cacheDir, key);
 			DirectoryReference.CreateDirectory(cacheFile.Directory);
 
 			using (FileStream outputStream = FileReference.Open(cacheFile, FileMode.Create, FileAccess.Write, FileShare.None))
@@ -166,35 +200,12 @@
 
 			lock (LockObject)
 			{
-				Item item = new Item(path, totalLength);
+				Item item = new Item(key, totalLength);
 				_items.AddFirst(item.ListNode);
-				_pathToItem.Add(path, item);
+				_pathToItem.Add(key, item);
 
-				_pathToPendingItem.Remove(path);
+				_pathToPendingItem.Remove(key);
 			}
 		}
-
-		/// <inheritdoc/>
-		public Task<string> WriteAsync(Stream stream, string? prefix = null, CancellationToken cancellationToken = default) => _inner.WriteAsync(stream, prefix, cancellationToken);
-
-#pragma warning disable CS0618 // Type or member is obsolete
-		/// <inheritdoc/>
-		public Task WriteExplicitPathAsync(string path, Stream stream, CancellationToken cancellationToken = default) => _inner.WriteExplicitPathAsync(path, stream, cancellationToken);
-#pragma warning restore CS0618 // Type or member is obsolete
-
-		/// <inheritdoc/>
-		public Task DeleteAsync(string path, CancellationToken cancellationToken = default) => _inner.DeleteAsync(path, cancellationToken);
-
-		/// <inheritdoc/>
-		public IAsyncEnumerable<string> EnumerateAsync(CancellationToken cancellationToken = default) => _inner.EnumerateAsync(cancellationToken);
-
-		/// <inheritdoc/>
-		public Task<bool> ExistsAsync(string path, CancellationToken cancellationToken = default) => _inner.ExistsAsync(path, cancellationToken);
-
-		/// <inheritdoc/>
-		public ValueTask<Uri?> TryGetReadRedirectAsync(string path, CancellationToken cancellationToken = default) => _inner.TryGetReadRedirectAsync(path, cancellationToken);
-
-		/// <inheritdoc/>
-		public ValueTask<(string, Uri)?> TryGetWriteRedirectAsync(string? prefix = null, CancellationToken cancellationToken = default) => _inner.TryGetWriteRedirectAsync(prefix, cancellationToken);
 	}
 }
