Horde: Fix incorrect handling of initial offsets for cached storage reads.

diff --git a/Engine/Source/Programs/Shared/EpicGames.Horde/Storage/Backends/CacheStorageBackend.cs b/Engine/Source/Programs/Shared/EpicGames.Horde/Storage/Backends/CacheStorageBackend.cs
--- a/Engine/Source/Programs/Shared/EpicGames.Horde/Storage/Backends/CacheStorageBackend.cs
+++ b/Engine/Source/Programs/Shared/EpicGames.Horde/Storage/Backends/CacheStorageBackend.cs
@@ -38,7 +38,12 @@ namespace EpicGames.Horde.Storage.Backends
 		/// <inheritdoc/>
 		public async Task<Stream> ReadAsync(string path, int offset, int? length, CancellationToken cancellationToken = default)
 		{
-			return await _cacheStorage.ReadAsync($"{_keyPrefix}{path}", ctx => _inner.ReadAsync(path, ctx), cancellationToken);
+			Stream stream = await _cacheStorage.ReadAsync($"{_keyPrefix}{path}", ctx => _inner.ReadAsync(path, ctx), cancellationToken);
+			if (offset != 0)
+			{
+				stream.Seek(offset, SeekOrigin.Begin);
+			}
+			return stream;
 		}
 
 		/// <inheritdoc/>
