Horde: Fix prefix being duplicated in bundle names.

diff --git a/Engine/Source/Programs/Shared/EpicGames.Horde/Storage/IStorageBackend.cs b/Engine/Source/Programs/Shared/EpicGames.Horde/Storage/IStorageBackend.cs
--- a/Engine/Source/Programs/Shared/EpicGames.Horde/Storage/IStorageBackend.cs
+++ b/Engine/Source/Programs/Shared/EpicGames.Horde/Storage/IStorageBackend.cs
@@ -112,13 +112,9 @@ namespace EpicGames.Horde.Storage
 		public static string CreateUniqueName(string? prefix)
 		{
 			StringBuilder builder = new StringBuilder(prefix);
-			if (!String.IsNullOrEmpty(prefix))
+			if (builder.Length > 0 && builder[^1] != '/')
 			{
-				builder.Append(prefix);
-				if (!prefix.EndsWith('/'))
-				{
-					builder.Append('/');
-				}
+				builder.Append('/');
 			}
 			builder.Append(s_sessionPrefix);
 			builder.Append(Interlocked.Increment(ref _increment));
