Horde: Rename shared backend cache implementation to StorageBackendCache.

diff --git a/Engine/Source/Programs/Shared/EpicGames.Horde.Tests/BackendTests.cs b/Engine/Source/Programs/Shared/EpicGames.Horde.Tests/BackendTests.cs
--- a/Engine/Source/Programs/Shared/EpicGames.Horde.Tests/BackendTests.cs
+++ b/Engine/Source/Programs/Shared/EpicGames.Horde.Tests/BackendTests.cs
@@ -57,14 +57,14 @@ namespace EpicGames.Horde.Tests
 		{
 			using (TempDir tempDir = new TempDir("Cache"))
 			{
-				CacheStorageBackendDetail cacheBackendDetail = new CacheStorageBackendDetail(tempDir.Location, 12);
+				StorageBackendCache cache = new StorageBackendCache(tempDir.Location, 12);
 
 				using MemoryStorageBackend memoryBackend = new MemoryStorageBackend();
-				using CacheStorageBackend cacheBackend = new CacheStorageBackend("", cacheBackendDetail, memoryBackend);
+				using IStorageBackend cacheBackend = StorageBackendCache.Wrap(memoryBackend, cache);
 				await TestBackendAsync(cacheBackend);
 
-				Assert.AreEqual(1, cacheBackendDetail.Items.Count());
-				byte[] value = await cacheBackend.ReadBytesAsync(cacheBackendDetail.Items.First());
+				Assert.AreEqual(1, cache.Items.Count());
+				byte[] value = await cacheBackend.ReadBytesAsync(cache.Items.First());
 				Assert.IsTrue(value.SequenceEqual(Encoding.UTF8.GetBytes("item 2")));
 
 				byte[] data3 = Encoding.UTF8.GetBytes("3");
@@ -79,7 +79,7 @@ namespace EpicGames.Horde.Tests
 				string path5 = await cacheBackend.WriteBytesAsync(data5);
 				await cacheBackend.ReadBytesAsync(path5);
 
-				HashSet<string> paths = new HashSet<string>(cacheBackendDetail.Items);
+				HashSet<string> paths = new HashSet<string>(cache.Items);
 				Assert.AreEqual(4, paths.Count);
 				Assert.IsTrue(paths.Contains(path3));
 				Assert.IsTrue(paths.Contains(path4));
@@ -89,7 +89,7 @@ namespace EpicGames.Horde.Tests
 				string path6 = await cacheBackend.WriteBytesAsync(data6);
 				await cacheBackend.ReadBytesAsync(path6);
 
-				paths = new HashSet<string>(cacheBackendDetail.Items);
+				paths = new HashSet<string>(cache.Items);
 				Assert.AreEqual(2, paths.Count);
 				Assert.IsTrue(paths.Contains(path5));
 				Assert.IsTrue(paths.Contains(path6));
@@ -97,7 +97,7 @@ namespace EpicGames.Horde.Tests
 				await cacheBackend.ReadBytesAsync(path3);
 				await cacheBackend.ReadBytesAsync(path4);
 
-				paths = new HashSet<string>(cacheBackendDetail.Items);
+				paths = new HashSet<string>(cache.Items);
 				Assert.AreEqual(2, paths.Count);
 				Assert.IsTrue(paths.Contains(path3));
 				Assert.IsTrue(paths.Contains(path4));
diff --git a/Engine/Source/Programs/Shared/EpicGames.Horde/Storage/Backends/CacheStorageBackend.cs b/Engine/Source/Programs/Shared/EpicGames.Horde/Storage/Backends/CacheStorageBackend.cs
deleted file mode 100644
diff --git a/Engine/Source/Programs/Shared/EpicGames.Horde/Storage/StorageBackendCache.cs b/Engine/Source/Programs/Shared/EpicGames.Horde/Storage/StorageBackendCache.cs
new file mode 100644
--- /dev/null
+++ b/Engine/Source/Programs/Shared/EpicGames.Horde/Storage/StorageBackendCache.cs
@@ -0,0 +1,213 @@
+// Copyright Epic Games, Inc. All Rights Reserved.
+
+using EpicGames.Core;
+using System;
+using System.Collections.Generic;
+using System.IO;
+using System.Linq;
+using System.Threading;
+using System.Threading.Tasks;
+
+namespace EpicGames.Horde.Storage
+{
+	/// <summary>
+	/// Implementation of a local disk cache which can be shared by multiple backends
+	/// </summary>
+	public sealed class StorageBackendCache
+	{
+		class Item
+		{
+			public string Key { get; }
+			public long Length { get; }
+			public LinkedListNode<Item> ListNode { get; }
+
+			public Item(string key, long length)
+			{
+				Key = key;
+				Length = length;
+				ListNode = new LinkedListNode<Item>(this);
+			}
+		}
+
+		class PendingItem
+		{
+			int _refCount;
+			readonly BackgroundTask _readTask;
+
+			public PendingItem(BackgroundTask readTask) => _readTask = readTask;
+			public void AddRef() => Interlocked.Increment(ref _refCount);
+
+			public async ValueTask ReleaseAsync()
+			{
+				if (Interlocked.Decrement(ref _refCount) == 0)
+				{
+					await _readTask.DisposeAsync();
+				}
+			}
+
+			public Task WaitAsync(CancellationToken cancellationToken) => _readTask.Task.WaitAsync(cancellationToken);
+		}
+
+		sealed class BackendWrapper : IStorageBackend
+		{
+			readonly string _keyPrefix;
+			readonly StorageBackendCache _cacheStorage;
+			readonly IStorageBackend _inner;
+
+			public bool SupportsRedirects => _inner.SupportsRedirects;
+
+			public BackendWrapper(string keyPrefix, StorageBackendCache cacheStorage, IStorageBackend inner)
+			{
+				_keyPrefix = keyPrefix;
+				_cacheStorage = cacheStorage;
+				_inner = inner;
+			}
+
+			public void Dispose() => _inner.Dispose();
+
+			public async Task<Stream> ReadAsync(string path, int offset, int? length, CancellationToken cancellationToken = default)
+			{
+				Stream stream = await _cacheStorage.ReadAsync($"{_keyPrefix}{path}", ctx => _inner.ReadAsync(path, ctx), cancellationToken);
+				if (offset != 0)
+				{
+					stream.Seek(offset, SeekOrigin.Begin);
+				}
+				return stream;
+			}
+
+			public Task<string> WriteAsync(Stream stream, string? prefix = null, CancellationToken cancellationToken = default) => _inner.WriteAsync(stream, prefix, cancellationToken);
+
+#pragma warning disable CS0618 // Type or member is obsolete
+			public Task WriteExplicitPathAsync(string path, Stream stream, CancellationToken cancellationToken = default) => _inner.WriteExplicitPathAsync(path, stream, cancellationToken);
+#pragma warning restore CS0618 // Type or member is obsolete
+
+			public Task DeleteAsync(string path, CancellationToken cancellationToken = default) => _inner.DeleteAsync(path, cancellationToken);
+			public IAsyncEnumerable<string> EnumerateAsync(CancellationToken cancellationToken = default) => _inner.EnumerateAsync(cancellationToken);
+			public Task<bool> ExistsAsync(string path, CancellationToken cancellationToken = default) => _inner.ExistsAsync(path, cancellationToken);
+			public ValueTask<Uri?> TryGetReadRedirectAsync(string path, CancellationToken cancellationToken = default) => _inner.TryGetReadRedirectAsync(path, cancellationToken);
+			public ValueTask<(string, Uri)?> TryGetWriteRedirectAsync(string? prefix = null, CancellationToken cancellationToken = default) => _inner.TryGetWriteRedirectAsync(prefix, cancellationToken);
+		}
+
+		object LockObject => _items;
+
+		readonly DirectoryReference _cacheDir;
+		readonly long _maxSize;
+
+		readonly LinkedList<Item> _items = new LinkedList<Item>();
+		readonly Dictionary<string, Item> _pathToItem = new Dictionary<string, Item>(StringComparer.Ordinal);
+		readonly Dictionary<string, PendingItem> _pathToPendingItem = new Dictionary<string, PendingItem>(StringComparer.Ordinal);
+
+		long _size;
+
+		internal IEnumerable<string> Items => _items.Select(x => x.Key);
+
+		/// <summary>
+		/// Constructor
+		/// </summary>
+		public StorageBackendCache(DirectoryReference cacheDir, long maxSize)
+		{
+			_cacheDir = cacheDir;
+			_maxSize = maxSize;
+		}
+
+		/// <summary>
+		/// Wraps a storage backend in another backend that routes requests through the cache
+		/// </summary>
+		/// <param name="backend">Backend to wrap</param>
+		/// <param name="cache">The cache instance. May be null.</param>
+		public static IStorageBackend Wrap(IStorageBackend backend, StorageBackendCache? cache)
+		{
+			if (cache == null)
+			{
+				return backend;
+			}
+			else
+			{
+				return new BackendWrapper("", cache, backend);
+			}
+		}
+
+		/// <inheritdoc/>
+		public async Task<Stream> ReadAsync(string key, Func<CancellationToken, Task<Stream>> createStreamAsync, CancellationToken cancellationToken = default)
+		{
+			for (; ; )
+			{
+				PendingItem? pendingItem;
+				lock (LockObject)
+				{
+					Item? item;
+					if (_pathToItem.TryGetValue(key, out item))
+					{
+						return OpenItem(item);
+					}
+
+					if (!_pathToPendingItem.TryGetValue(key, out pendingItem))
+					{
+						pendingItem = new PendingItem(BackgroundTask.StartNew(x => ReadIntoCacheAsync(key, createStreamAsync, x)));
+						_pathToPendingItem.Add(key, pendingItem);
+					}
+
+					pendingItem.AddRef();
+				}
+
+				try
+				{
+					await pendingItem.WaitAsync(cancellationToken);
+				}
+				finally
+				{
+					await pendingItem.ReleaseAsync();
+				}
+			}
+		}
+
+		Stream OpenItem(Item item)
+		{
+			_items.Remove(item.ListNode);
+			_items.AddFirst(item.ListNode);
+
+			FileReference file = FileReference.Combine(_cacheDir, item.Key);
+			return FileReference.Open(file, FileMode.Open, FileAccess.Read, FileShare.Read | FileShare.Delete);
+		}
+
+		async Task ReadIntoCacheAsync(string key, Func<CancellationToken, Task<Stream>> createStreamAsync, CancellationToken cancellationToken)
+		{
+			using Stream stream = await createStreamAsync(cancellationToken);
+			long totalLength = stream.Length;
+
+			lock (LockObject)
+			{
+				_size += totalLength;
+
+				while (_size > _maxSize && _items.Count > 0)
+				{
+					LinkedListNode<Item> lastItem = _items.Last!;
+					_pathToItem.Remove(lastItem.Value.Key);
+					_items.RemoveLast();
+
+					_size -= lastItem.Value.Length;
+
+					FileReference file = FileReference.Combine(_cacheDir, lastItem.Value.Key);
+					FileReference.Delete(file);
+				}
+			}
+
+			FileReference cacheFile = FileReference.Combine(_cacheDir, key);
+			DirectoryReference.CreateDirectory(cacheFile.Directory);
+
+			using (FileStream outputStream = FileReference.Open(cacheFile, FileMode.Create, FileAccess.Write, FileShare.None))
+			{
+				await stream.CopyToAsync(outputStream, cancellationToken);
+			}
+
+			lock (LockObject)
+			{
+				Item item = new Item(key, totalLength);
+				_items.AddFirst(item.ListNode);
+				_pathToItem.Add(key, item);
+
+				_pathToPendingItem.Remove(key);
+			}
+		}
+	}
+}
