Removed the more trivial C++ UHT compatibility shims in the code generator.

diff --git a/Engine/Source/Programs/Shared/EpicGames.UHT/Exporters/CodeGen/UhtCodeGenerator.cs b/Engine/Source/Programs/Shared/EpicGames.UHT/Exporters/CodeGen/UhtCodeGenerator.cs
--- a/Engine/Source/Programs/Shared/EpicGames.UHT/Exporters/CodeGen/UhtCodeGenerator.cs
+++ b/Engine/Source/Programs/Shared/EpicGames.UHT/Exporters/CodeGen/UhtCodeGenerator.cs
@@ -46,8 +46,6 @@ namespace EpicGames.UHT.Exporters.CodeGen
 		{
 			public string RegisteredSingletonName { get; set; }
 			public string UnregisteredSingletonName { get; set; }
-			public string RegsiteredCrossReference { get; set; }
-			public string UnregsiteredCrossReference { get; set; }
 			public string RegsiteredExternalDecl { get; set; }
 			public string UnregisteredExternalDecl { get; set; }
 			public UhtClass? NativeInterface { get; set; }
@@ -209,28 +207,6 @@ namespace EpicGames.UHT.Exporters.CodeGen
 		{
 			return registered ? ObjectInfos[objectIndex].RegsiteredExternalDecl : ObjectInfos[objectIndex].UnregisteredExternalDecl;
 		}
-
-		/// <summary>
-		/// Return the cross reference for an object
-		/// </summary>
-		/// <param name="obj">The object in question.</param>
-		/// <param name="registered">If true, return the registered cross reference.  Otherwise return the unregistered.</param>
-		/// <returns>Cross reference</returns>
-		public string GetCrossReference(UhtObject obj, bool registered)
-		{
-			return GetCrossReference(obj.ObjectTypeIndex, registered);
-		}
-
-		/// <summary>
-		/// Return the cross reference for an object
-		/// </summary>
-		/// <param name="objectIndex">The object in question.</param>
-		/// <param name="registered">If true, return the registered cross reference.  Otherwise return the unregistered.</param>
-		/// <returns>Cross reference</returns>
-		public string GetCrossReference(int objectIndex, bool registered)
-		{
-			return registered ? ObjectInfos[objectIndex].RegsiteredCrossReference : ObjectInfos[objectIndex].UnregsiteredCrossReference;
-		}
 		#endregion
 
 		#region Information initialization
@@ -247,8 +223,7 @@ namespace EpicGames.UHT.Exporters.CodeGen
 			builder.Append("Z_Construct_UPackage_");
 			builder.Append(packageInfo.StrippedName);
 			objectInfo.UnregisteredSingletonName = objectInfo.RegisteredSingletonName = builder.ToString();
-			objectInfo.UnregisteredExternalDecl = objectInfo.RegsiteredExternalDecl = $"\t{packageInfo.Api}_API UPackage* {objectInfo.RegisteredSingletonName}();\r\n"; //COMPATIBILITY-TODO remove the extra _API
-			objectInfo.UnregsiteredCrossReference = objectInfo.RegsiteredCrossReference = $"\tUPackage* {objectInfo.RegisteredSingletonName}();\r\n";
+			objectInfo.UnregisteredExternalDecl = objectInfo.RegsiteredExternalDecl = $"\tUPackage* {objectInfo.RegisteredSingletonName}();\r\n";
 
 			foreach (UhtType packageChild in package.Children)
 			{
@@ -392,10 +367,6 @@ namespace EpicGames.UHT.Exporters.CodeGen
 				objectInfo.UnregisteredExternalDecl = objectInfo.RegsiteredExternalDecl = $"\t{packageInfo.Api}U{engineClassName}* {objectInfo.RegisteredSingletonName}();\r\n";
 			}
 
-			//COMPATIBILITY-TODO - The cross reference string should match the extern decl string always.  But currently, it is different for packages.
-			objectInfo.UnregsiteredCrossReference = objectInfo.UnregisteredExternalDecl;
-			objectInfo.RegsiteredCrossReference = objectInfo.RegsiteredExternalDecl;
-
 			// Init the children
 			foreach (UhtType child in obj.Children)
 			{
diff --git a/Engine/Source/Programs/Shared/EpicGames.UHT/Exporters/CodeGen/UhtHeaderCodeGeneratorCppFile.cs b/Engine/Source/Programs/Shared/EpicGames.UHT/Exporters/CodeGen/UhtHeaderCodeGeneratorCppFile.cs
--- a/Engine/Source/Programs/Shared/EpicGames.UHT/Exporters/CodeGen/UhtHeaderCodeGeneratorCppFile.cs
+++ b/Engine/Source/Programs/Shared/EpicGames.UHT/Exporters/CodeGen/UhtHeaderCodeGeneratorCppFile.cs
@@ -136,7 +136,7 @@ namespace EpicGames.UHT.Exporters.CodeGen
 				if (HeaderFile.References.CrossModule.References.Count > 0)
 				{
 					ReadOnlyMemory<string> sorted = HeaderFile.References.CrossModule.GetSortedReferences(
-						(int objectIndex, bool registered) => GetCrossReference(objectIndex, registered));
+						(int objectIndex, bool registered) => GetExternalDecl(objectIndex, registered));
 					builder.Append("// Cross Module References\r\n");
 					foreach (string crossReference in sorted.Span)
 					{
@@ -808,8 +808,7 @@ namespace EpicGames.UHT.Exporters.CodeGen
 					}
 					if (wroteLine)
 					{
-						//COMPATIBILITY-TODO - Remove the tab
-						builder.Append("\t\r\n");
+						builder.Append("\r\n");
 					}
 					
 					foreach (UhtRigVMMethodInfo predicateInfo in scriptStruct.RigVMStructInfo.Methods)
@@ -820,7 +819,6 @@ namespace EpicGames.UHT.Exporters.CodeGen
 						}
 					}
 
-					//COMPATIBILITY-TODO - Replace spaces with \t
 					builder.Append('\t').Append(methodInfo.ReturnPrefix()).Append("Static").Append(methodInfo.Name).Append("(\r\n");
 					builder.Append("\t\tInExecuteContext");
 					builder.AppendParameterNames(scriptStruct.RigVMStructInfo.Members, true, ",\r\n\t\t", true);
diff --git a/Engine/Source/Programs/Shared/EpicGames.UHT/Exporters/CodeGen/UhtHeaderCodeGeneratorHFile.cs b/Engine/Source/Programs/Shared/EpicGames.UHT/Exporters/CodeGen/UhtHeaderCodeGeneratorHFile.cs
--- a/Engine/Source/Programs/Shared/EpicGames.UHT/Exporters/CodeGen/UhtHeaderCodeGeneratorHFile.cs
+++ b/Engine/Source/Programs/Shared/EpicGames.UHT/Exporters/CodeGen/UhtHeaderCodeGeneratorHFile.cs
@@ -1034,7 +1034,6 @@ namespace EpicGames.UHT.Exporters.CodeGen
 
 		private static StringBuilder AppendRpcWrappers(StringBuilder builder, List<UhtFunction> reversedFunctions, bool editorOnly)
 		{
-			bool first = true;
 			foreach (UhtFunction function in reversedFunctions)
 			{
 				if (!IsRpcFunction(function, editorOnly))
@@ -1045,12 +1044,6 @@ namespace EpicGames.UHT.Exporters.CodeGen
 				{
 					continue;
 				}
-				//COMPATIBILITY-TODO - Remove once we transition to C# version
-				if (first)
-				{
-					builder.Append(" \\\r\n");
-					first = false;
-				}
 				builder.Append("\tDECLARE_FUNCTION(").Append(function.UnMarshalAndCallName).Append("); \\\r\n");
 			}
 			return builder;
diff --git a/Engine/Source/Programs/Shared/EpicGames.UHT/Exporters/CodeGen/UhtPackageCodeGenerator.cs b/Engine/Source/Programs/Shared/EpicGames.UHT/Exporters/CodeGen/UhtPackageCodeGenerator.cs
--- a/Engine/Source/Programs/Shared/EpicGames.UHT/Exporters/CodeGen/UhtPackageCodeGenerator.cs
+++ b/Engine/Source/Programs/Shared/EpicGames.UHT/Exporters/CodeGen/UhtPackageCodeGenerator.cs
@@ -21,7 +21,7 @@ namespace EpicGames.UHT.Exporters.CodeGen
 		public static string DisableDeprecationWarnings = "PRAGMA_DISABLE_DEPRECATION_WARNINGS";
 
 		public static string BeginEditorOnlyGuard = "#if WITH_EDITOR\r\n";
-		public static string EndEditorOnlyGuard = "#endif //WITH_EDITOR\r\n"; //COMPATIBILITY-TODO - This does not match UhtMacroBlockEmitter
+		public static string EndEditorOnlyGuard = "#endif // WITH_EDITOR\r\n";
 
 		public readonly UhtCodeGenerator CodeGenerator;
 		public readonly UhtPackage Package;
@@ -75,28 +75,6 @@ namespace EpicGames.UHT.Exporters.CodeGen
 			return CodeGenerator.GetExternalDecl(objectIndex, registered);
 		}
 
-		/// <summary>
-		/// Return the cross reference for an object
-		/// </summary>
-		/// <param name="obj">The object in question.</param>
-		/// <param name="registered">If true, return the registered cross reference.  Otherwise return the unregistered.</param>
-		/// <returns>Cross reference</returns>
-		public string GetCrossReference(UhtObject obj, bool registered)
-		{
-			return CodeGenerator.GetCrossReference(obj, registered);
-		}
-
-		/// <summary>
-		/// Return the cross reference for an object
-		/// </summary>
-		/// <param name="objectIndex">The object in question.</param>
-		/// <param name="registered">If true, return the registered cross reference.  Otherwise return the unregistered.</param>
-		/// <returns>Cross reference</returns>
-		public string GetCrossReference(int objectIndex, bool registered)
-		{
-			return CodeGenerator.GetCrossReference(objectIndex, registered);
-		}
-
 		/// <summary>
 		/// Test to see if the given field is a delegate function
 		/// </summary>
diff --git a/Engine/Source/Programs/Shared/EpicGames.UHT/Types/Properties/UhtClassProperty.cs b/Engine/Source/Programs/Shared/EpicGames.UHT/Types/Properties/UhtClassProperty.cs
--- a/Engine/Source/Programs/Shared/EpicGames.UHT/Types/Properties/UhtClassProperty.cs
+++ b/Engine/Source/Programs/Shared/EpicGames.UHT/Types/Properties/UhtClassProperty.cs
@@ -49,7 +49,7 @@ namespace EpicGames.UHT.Types
 		/// <inheritdoc/>
 		private StringBuilder AppendSubClassText(StringBuilder builder)
 		{
-			builder.Append("TSubclassOf<").Append(MetaClass?.SourceName).Append("> "); //COMPATIBILITY-TODO - Extra space in old UHT
+			builder.Append("TSubclassOf<").Append(MetaClass?.SourceName).Append(">");
 			return builder;
 		}
 
diff --git a/Engine/Source/Programs/Shared/EpicGames.UHT/Types/Properties/UhtLazyObjectPtrProperty.cs b/Engine/Source/Programs/Shared/EpicGames.UHT/Types/Properties/UhtLazyObjectPtrProperty.cs
--- a/Engine/Source/Programs/Shared/EpicGames.UHT/Types/Properties/UhtLazyObjectPtrProperty.cs
+++ b/Engine/Source/Programs/Shared/EpicGames.UHT/Types/Properties/UhtLazyObjectPtrProperty.cs
@@ -49,10 +49,6 @@ namespace EpicGames.UHT.Types
 		{
 			switch (textType)
 			{
-				case UhtPropertyTextType.FunctionThunkParameterArgType:
-					builder.Append("TLazyObjectPtr<").Append(Class.SourceName).Append("> "); //COMPATIBILITY-TODO - Extra space
-					break;
-
 				default:
 					builder.Append("TLazyObjectPtr<").Append(Class.SourceName).Append('>');
 					break;
