Horde: Fix bundle tests.

diff --git a/Engine/Source/Programs/Shared/EpicGames.Horde.Tests/BundleStoreTests.cs b/Engine/Source/Programs/Shared/EpicGames.Horde.Tests/BundleStoreTests.cs
--- a/Engine/Source/Programs/Shared/EpicGames.Horde.Tests/BundleStoreTests.cs
+++ b/Engine/Source/Programs/Shared/EpicGames.Horde.Tests/BundleStoreTests.cs
@@ -95,7 +95,7 @@ namespace EpicGames.Horde.Tests
 			MemoryStorageClient blobStore = new MemoryStorageClient();
 			await TestTreeAsync(blobStore, new BundleOptions { MaxBlobSize = 1024 * 1024 });
 
-			Assert.AreEqual(1, blobStore.Bundles.Count);
+			Assert.AreEqual(1, blobStore.Blobs.Count);
 			Assert.AreEqual(1, blobStore.Refs.Count);
 		}
 
@@ -107,7 +107,7 @@ namespace EpicGames.Horde.Tests
 			MemoryStorageClient blobStore = new MemoryStorageClient();
 			await TestTreeAsync(blobStore, new BundleOptions { MaxBlobSize = 1 });
 
-			Assert.AreEqual(5, blobStore.Bundles.Count);
+			Assert.AreEqual(5, blobStore.Blobs.Count);
 			Assert.AreEqual(1, blobStore.Refs.Count);
 		}
 
@@ -388,7 +388,7 @@ namespace EpicGames.Horde.Tests
 
 				NodeRef<ChunkedDataNode> file = root.GetFileEntry("test");
 
-				long uniqueSize = store.Bundles.Values.SelectMany(x => x.Header.Packets).Sum(x => x.DecodedLength);
+				long uniqueSize = store.Blobs.Values.Select(x => Bundle.FromMemory(x)).SelectMany(x => x.Header.Packets).Sum(x => x.DecodedLength);
 				Assert.IsTrue(uniqueSize < data.Length / 3); // random fraction meaning "lots of dedupe happened"
 			}
 		}
diff --git a/Engine/Source/Programs/Shared/EpicGames.Horde.Tests/BundleTests.cs b/Engine/Source/Programs/Shared/EpicGames.Horde.Tests/BundleTests.cs
--- a/Engine/Source/Programs/Shared/EpicGames.Horde.Tests/BundleTests.cs
+++ b/Engine/Source/Programs/Shared/EpicGames.Horde.Tests/BundleTests.cs
@@ -214,7 +214,7 @@ namespace EpicGames.Horde.Tests
 
 			// Should be stored inline
 			Assert.AreEqual(1, store.Refs.Count);
-			Assert.AreEqual(1, store.Bundles.Count);
+			Assert.AreEqual(1, store.Blobs.Count);
 
 			// Check the ref
 			BundleNodeHandle refTarget =  await store.ReadRefTargetAsync(refName);
@@ -243,7 +243,7 @@ namespace EpicGames.Horde.Tests
 		public async Task DedupTests()
 		{
 			Assert.AreEqual(0, _storage.Refs.Count);
-			Assert.AreEqual(0, _storage.Bundles.Count);
+			Assert.AreEqual(0, _storage.Blobs.Count);
 
 			BundleOptions options = new BundleOptions();
 			options.MaxBlobSize = 1;
@@ -264,7 +264,7 @@ namespace EpicGames.Horde.Tests
 			}
 
 			Assert.AreEqual(1, _storage.Refs.Count);
-			Assert.AreEqual(2, _storage.Bundles.Count);
+			Assert.AreEqual(2, _storage.Blobs.Count);
 		}
 
 		[TestMethod]
@@ -289,7 +289,7 @@ namespace EpicGames.Horde.Tests
 				}
 
 				Assert.AreEqual(1, _storage.Refs.Count);
-				Assert.AreEqual(5, _storage.Bundles.Count);
+				Assert.AreEqual(5, _storage.Blobs.Count);
 			}
 
 			{
diff --git a/Engine/Source/Programs/Shared/EpicGames.Horde/Storage/Backends/MemoryStorageBackend.cs b/Engine/Source/Programs/Shared/EpicGames.Horde/Storage/Backends/MemoryStorageBackend.cs
--- a/Engine/Source/Programs/Shared/EpicGames.Horde/Storage/Backends/MemoryStorageBackend.cs
+++ b/Engine/Source/Programs/Shared/EpicGames.Horde/Storage/Backends/MemoryStorageBackend.cs
@@ -21,6 +21,11 @@ namespace EpicGames.Horde.Storage.Backends
 		/// </summary>
 		readonly ConcurrentDictionary<string, byte[]> _pathToData = new ConcurrentDictionary<string, byte[]>();
 
+		/// <summary>
+		/// Read only access to the stored blobs
+		/// </summary>
+		public IReadOnlyDictionary<string, byte[]> Blobs => _pathToData; 
+
 		/// <inheritdoc/>
 		public bool SupportsRedirects => false;
 
diff --git a/Engine/Source/Programs/Shared/EpicGames.Horde/Storage/Bundles/Bundle.cs b/Engine/Source/Programs/Shared/EpicGames.Horde/Storage/Bundles/Bundle.cs
--- a/Engine/Source/Programs/Shared/EpicGames.Horde/Storage/Bundles/Bundle.cs
+++ b/Engine/Source/Programs/Shared/EpicGames.Horde/Storage/Bundles/Bundle.cs
@@ -101,6 +101,17 @@ namespace EpicGames.Horde.Storage.Bundles
 			Packets = packets;
 		}
 
+		/// <summary>
+		/// Reads a bundle from the given memory
+		/// </summary>
+		/// <param name="data">Data to read from</param>
+		/// <returns>Bundle that was read</returns>
+		public static Bundle FromMemory(ReadOnlyMemory<byte> data)
+		{
+			using ReadOnlyMemoryStream stream = new ReadOnlyMemoryStream(data);
+			return FromStreamAsync(stream).Result;
+		}
+
 		/// <summary>
 		/// Reads a bundle from the given stream
 		/// </summary>
diff --git a/Engine/Source/Programs/Shared/EpicGames.Horde/Storage/Clients/MemoryStorageClient.cs b/Engine/Source/Programs/Shared/EpicGames.Horde/Storage/Clients/MemoryStorageClient.cs
--- a/Engine/Source/Programs/Shared/EpicGames.Horde/Storage/Clients/MemoryStorageClient.cs
+++ b/Engine/Source/Programs/Shared/EpicGames.Horde/Storage/Clients/MemoryStorageClient.cs
@@ -25,9 +25,9 @@ namespace EpicGames.Horde.Storage.Clients
 		record class ExportEntry(BundleNodeHandle Handle, int Rank, ExportEntry? Next);
 
 		/// <summary>
-		/// Map of blob id to blob data
+		/// Backend instance
 		/// </summary>
-		readonly ConcurrentDictionary<BundleLocator, Bundle> _bundles = new ConcurrentDictionary<BundleLocator, Bundle>();
+		readonly MemoryStorageBackend _backend;
 
 		/// <summary>
 		/// Map of ref name to ref data
@@ -39,8 +39,10 @@ namespace EpicGames.Horde.Storage.Clients
 		/// </summary>
 		readonly ConcurrentDictionary<Utf8String, ExportEntry> _exports = new ConcurrentDictionary<Utf8String, ExportEntry>();
 
-		/// <inheritdoc cref="_bundles"/>
-		public IReadOnlyDictionary<BundleLocator, Bundle> Bundles => _bundles;
+		/// <summary>
+		/// All data stored by the client
+		/// </summary>
+		public IReadOnlyDictionary<string, byte[]> Blobs => _backend.Blobs;
 
 		/// <inheritdoc cref="_refs"/>
 		public IReadOnlyDictionary<RefName, BundleNodeLocator> Refs => _refs;
@@ -49,8 +51,17 @@ namespace EpicGames.Horde.Storage.Clients
 		/// Constructor
 		/// </summary>
 		public MemoryStorageClient() 
-			: base(new MemoryStorageBackend(), StorageCache.None, NullLogger.Instance)
+			: this(new MemoryStorageBackend())
+		{
+		}
+
+		/// <summary>
+		/// Constructor
+		/// </summary>
+		private MemoryStorageClient(MemoryStorageBackend backend)
+			: base(backend, StorageCache.None, NullLogger.Instance)
 		{
+			_backend = backend;
 		}
 
 		#region Aliases
