Horde: Fix BadRequest responses when uploading a blob without a prefix. Not adding a prefix creates an empty form when a file stream is not given.

diff --git a/Engine/Source/Programs/Shared/EpicGames.Horde/Storage/Backends/HttpStorageBackend.cs b/Engine/Source/Programs/Shared/EpicGames.Horde/Storage/Backends/HttpStorageBackend.cs
--- a/Engine/Source/Programs/Shared/EpicGames.Horde/Storage/Backends/HttpStorageBackend.cs
+++ b/Engine/Source/Programs/Shared/EpicGames.Horde/Storage/Backends/HttpStorageBackend.cs
@@ -134,10 +134,7 @@ namespace EpicGames.Horde.Storage.Backends
 					{
 						form.Add(streamContent, "file", "filename");
 					}
-					if (prefix != null)
-					{
-						form.Add(stringContent, "prefix");
-					}
+					form.Add(stringContent, "prefix");
 
 					request.Content = form;
 					using (HttpResponseMessage response = await httpClient.SendAsync(request, cancellationToken))
