Horde: Prevent compute channel data being discarded from the remote until we've attached at least one receive buffer. Fixes a race where the remote can send attach notifications before the initiator is listening for them.

diff --git a/Engine/Source/Programs/Shared/EpicGames.Horde/Compute/ComputeSocket.cs b/Engine/Source/Programs/Shared/EpicGames.Horde/Compute/ComputeSocket.cs
--- a/Engine/Source/Programs/Shared/EpicGames.Horde/Compute/ComputeSocket.cs
+++ b/Engine/Source/Programs/Shared/EpicGames.Horde/Compute/ComputeSocket.cs
@@ -392,7 +392,7 @@ namespace EpicGames.Horde.Compute
 			_transport = transport;
 			_logger = logger;
 
-			_recvTask = BackgroundTask.StartNew(ctx => RunRecvTaskAsync(_transport, ctx));
+			_recvTask = new BackgroundTask(ctx => RunRecvTaskAsync(_transport, ctx));
 		}
 
 		/// <summary>
@@ -505,7 +505,7 @@ namespace EpicGames.Horde.Compute
 		{
 			if (!await TryReadPacketAsync(transport, id, size, cancellationToken))
 			{
-				_logger.LogDebug("Discarding {Size} bytes received on compute channel {Id}", size, id);
+				_logger.LogWarning("Discarding {Size} bytes received on compute channel {Id}; no buffer attached?", size, id);
 
 				int bufferSize = Math.Min(size, 65536);
 				using (IMemoryOwner<byte> buffer = MemoryPool<byte>.Shared.Rent(bufferSize))
@@ -637,6 +637,12 @@ namespace EpicGames.Horde.Compute
 				}
 
 				_recvBuffers.Add(channelId, new RecvBuffer(recvBuffer.CreateWriter()));
+
+				// Only start the receive task once we have a buffer to receive data, otherwise we discard data from the remote
+				if (_recvTask.Task.IsCompleted)
+				{
+					_recvTask.Start();
+				}
 			}
 		}
 
