EpicGames.Core: Enable code style warnings in Analyze configuration, and fix current violations.

diff --git a/Engine/Source/Programs/Shared/EpicGames.Core.Tests/JsonObjectTests.cs b/Engine/Source/Programs/Shared/EpicGames.Core.Tests/JsonObjectTests.cs
--- a/Engine/Source/Programs/Shared/EpicGames.Core.Tests/JsonObjectTests.cs
+++ b/Engine/Source/Programs/Shared/EpicGames.Core.Tests/JsonObjectTests.cs
@@ -22,6 +22,7 @@ namespace EpicGames.Core.Tests
 		[ClassInitialize]
 		public static void Setup(TestContext testContext)
 		{
+			_ = testContext;
 			JsonObjectTests.s_tempDirectory = CreateTempDir();
 		}
 
diff --git a/Engine/Source/Programs/Shared/EpicGames.Core/BinaryFormatterUtils.cs b/Engine/Source/Programs/Shared/EpicGames.Core/BinaryFormatterUtils.cs
--- a/Engine/Source/Programs/Shared/EpicGames.Core/BinaryFormatterUtils.cs
+++ b/Engine/Source/Programs/Shared/EpicGames.Core/BinaryFormatterUtils.cs
@@ -2,6 +2,7 @@
 
 using System.IO;
 using System.Runtime.Serialization.Formatters.Binary;
+using System.Threading.Tasks;
 
 #pragma warning disable SYSLIB0011
 #pragma warning disable CA2300 // Do not use insecure deserializer BinaryFormatter
diff --git a/Engine/Source/Programs/Shared/EpicGames.Core/DirectoryReference.cs b/Engine/Source/Programs/Shared/EpicGames.Core/DirectoryReference.cs
--- a/Engine/Source/Programs/Shared/EpicGames.Core/DirectoryReference.cs
+++ b/Engine/Source/Programs/Shared/EpicGames.Core/DirectoryReference.cs
@@ -51,6 +51,7 @@ namespace EpicGames.Core
 		public DirectoryReference(string fullName, Sanitize sanitize)
 			: base(fullName)
 		{
+			_ = sanitize;
 		}
 
 		/// <summary>
diff --git a/Engine/Source/Programs/Shared/EpicGames.Core/EpicGames.Core.csproj b/Engine/Source/Programs/Shared/EpicGames.Core/EpicGames.Core.csproj
--- a/Engine/Source/Programs/Shared/EpicGames.Core/EpicGames.Core.csproj
+++ b/Engine/Source/Programs/Shared/EpicGames.Core/EpicGames.Core.csproj
@@ -27,7 +27,8 @@
   </PropertyGroup>
   <PropertyGroup Condition=" '$(Configuration)|$(Platform)' == 'Analyze|AnyCPU' ">
     <RunAnalyzersDuringBuild>True</RunAnalyzersDuringBuild>
-    <TreatWarningsAsErrors>false</TreatWarningsAsErrors>
+    <TreatWarningsAsErrors>False</TreatWarningsAsErrors>
+    <EnforceCodeStyleInBuild>True</EnforceCodeStyleInBuild>
     <GenerateDocumentationFile>True</GenerateDocumentationFile>
     <DocumentationFile></DocumentationFile>
   </PropertyGroup>
diff --git a/Engine/Source/Programs/Shared/EpicGames.Core/JsonLogEvent.cs b/Engine/Source/Programs/Shared/EpicGames.Core/JsonLogEvent.cs
--- a/Engine/Source/Programs/Shared/EpicGames.Core/JsonLogEvent.cs
+++ b/Engine/Source/Programs/Shared/EpicGames.Core/JsonLogEvent.cs
@@ -242,7 +242,11 @@ namespace EpicGames.Core
 		/// <param name="state"></param>
 		/// <param name="ex"></param>
 		/// <returns></returns>
-		public static string Format(JsonLogEvent state, Exception? ex) => LogEvent.Read(state.Data.Span).ToString();
+		public static string Format(JsonLogEvent state, Exception? ex)
+		{
+			_ = ex;
+			return LogEvent.Read(state.Data.Span).ToString();
+		}
 
 		/// <inheritdoc/>
 		public override string ToString() => Encoding.UTF8.GetString(Data.ToArray());
diff --git a/Engine/Source/Programs/Shared/EpicGames.Core/JsonSchema.cs b/Engine/Source/Programs/Shared/EpicGames.Core/JsonSchema.cs
--- a/Engine/Source/Programs/Shared/EpicGames.Core/JsonSchema.cs
+++ b/Engine/Source/Programs/Shared/EpicGames.Core/JsonSchema.cs
@@ -1,13 +1,11 @@
 // Copyright Epic Games, Inc. All Rights Reserved.
 
 using System;
-using System.Collections.Concurrent;
 using System.Collections.Generic;
 using System.IO;
 using System.Reflection;
 using System.Text.Json;
 using System.Text.Json.Serialization;
-using System.Xml;
 
 namespace EpicGames.Core
 {
