Horde: Fix last CI warning.

diff --git a/Engine/Source/Programs/Shared/EpicGames.Horde.Tests/ChunkingTests.cs b/Engine/Source/Programs/Shared/EpicGames.Horde.Tests/ChunkingTests.cs
--- a/Engine/Source/Programs/Shared/EpicGames.Horde.Tests/ChunkingTests.cs
+++ b/Engine/Source/Programs/Shared/EpicGames.Horde.Tests/ChunkingTests.cs
@@ -89,7 +89,7 @@ namespace EpicGames.Horde.Tests
 			await TestChunkingAsync(options);
 		}
 
-		async Task TestChunkingAsync(ChunkingOptions options)
+		static async Task TestChunkingAsync(ChunkingOptions options)
 		{
 			using MemoryCache cache = new MemoryCache(new MemoryCacheOptions());
 
