Horde: Enable code style warnings in EpicGames.Slack, and fix current violations.

diff --git a/Engine/Source/Programs/Shared/EpicGames.Slack/Block.cs b/Engine/Source/Programs/Shared/EpicGames.Slack/Block.cs
--- a/Engine/Source/Programs/Shared/EpicGames.Slack/Block.cs
+++ b/Engine/Source/Programs/Shared/EpicGames.Slack/Block.cs
@@ -4,8 +4,6 @@ using System;
 using System.Collections.Generic;
 using System.Text.Json;
 using System.Text.Json.Serialization;
-using EpicGames.Slack.Blocks;
-using EpicGames.Slack.Elements;
 
 namespace EpicGames.Slack
 {
diff --git a/Engine/Source/Programs/Shared/EpicGames.Slack/Blocks/ContextBlock.cs b/Engine/Source/Programs/Shared/EpicGames.Slack/Blocks/ContextBlock.cs
--- a/Engine/Source/Programs/Shared/EpicGames.Slack/Blocks/ContextBlock.cs
+++ b/Engine/Source/Programs/Shared/EpicGames.Slack/Blocks/ContextBlock.cs
@@ -2,7 +2,6 @@
 
 using System;
 using System.Collections.Generic;
-using System.Text.Json;
 using System.Text.Json.Serialization;
 using EpicGames.Slack.Elements;
 
diff --git a/Engine/Source/Programs/Shared/EpicGames.Slack/Blocks/DividerBlock.cs b/Engine/Source/Programs/Shared/EpicGames.Slack/Blocks/DividerBlock.cs
--- a/Engine/Source/Programs/Shared/EpicGames.Slack/Blocks/DividerBlock.cs
+++ b/Engine/Source/Programs/Shared/EpicGames.Slack/Blocks/DividerBlock.cs
@@ -1,7 +1,5 @@
 // Copyright Epic Games, Inc. All Rights Reserved.
 
-using System.Collections.Generic;
-
 namespace EpicGames.Slack.Blocks
 {
 	/// <summary>
diff --git a/Engine/Source/Programs/Shared/EpicGames.Slack/Blocks/HeaderBlock.cs b/Engine/Source/Programs/Shared/EpicGames.Slack/Blocks/HeaderBlock.cs
--- a/Engine/Source/Programs/Shared/EpicGames.Slack/Blocks/HeaderBlock.cs
+++ b/Engine/Source/Programs/Shared/EpicGames.Slack/Blocks/HeaderBlock.cs
@@ -1,7 +1,5 @@
 // Copyright Epic Games, Inc. All Rights Reserved.
 
-using System.Collections.Generic;
-using System.Text.Json;
 using System.Text.Json.Serialization;
 using EpicGames.Slack.Elements;
 
diff --git a/Engine/Source/Programs/Shared/EpicGames.Slack/Blocks/ImageBlock.cs b/Engine/Source/Programs/Shared/EpicGames.Slack/Blocks/ImageBlock.cs
--- a/Engine/Source/Programs/Shared/EpicGames.Slack/Blocks/ImageBlock.cs
+++ b/Engine/Source/Programs/Shared/EpicGames.Slack/Blocks/ImageBlock.cs
@@ -1,8 +1,6 @@
 // Copyright Epic Games, Inc. All Rights Reserved.
 
 using System;
-using System.Collections.Generic;
-using System.Text.Json;
 using System.Text.Json.Serialization;
 using EpicGames.Slack.Elements;
 
diff --git a/Engine/Source/Programs/Shared/EpicGames.Slack/Blocks/InputBlock.cs b/Engine/Source/Programs/Shared/EpicGames.Slack/Blocks/InputBlock.cs
--- a/Engine/Source/Programs/Shared/EpicGames.Slack/Blocks/InputBlock.cs
+++ b/Engine/Source/Programs/Shared/EpicGames.Slack/Blocks/InputBlock.cs
@@ -1,8 +1,5 @@
 // Copyright Epic Games, Inc. All Rights Reserved.
 
-using System;
-using System.Collections.Generic;
-using System.Text.Json;
 using System.Text.Json.Serialization;
 using EpicGames.Slack.Elements;
 
diff --git a/Engine/Source/Programs/Shared/EpicGames.Slack/Blocks/SectionBlock.cs b/Engine/Source/Programs/Shared/EpicGames.Slack/Blocks/SectionBlock.cs
--- a/Engine/Source/Programs/Shared/EpicGames.Slack/Blocks/SectionBlock.cs
+++ b/Engine/Source/Programs/Shared/EpicGames.Slack/Blocks/SectionBlock.cs
@@ -1,8 +1,6 @@
 // Copyright Epic Games, Inc. All Rights Reserved.
 
-using System;
 using System.Collections.Generic;
-using System.Text.Json;
 using System.Text.Json.Serialization;
 using EpicGames.Slack.Elements;
 
diff --git a/Engine/Source/Programs/Shared/EpicGames.Slack/Elements/ButtonElement.cs b/Engine/Source/Programs/Shared/EpicGames.Slack/Elements/ButtonElement.cs
--- a/Engine/Source/Programs/Shared/EpicGames.Slack/Elements/ButtonElement.cs
+++ b/Engine/Source/Programs/Shared/EpicGames.Slack/Elements/ButtonElement.cs
@@ -1,8 +1,6 @@
 // Copyright Epic Games, Inc. All Rights Reserved.
 
 using System;
-using System.Collections.Generic;
-using System.Globalization;
 using System.Runtime.Serialization;
 using System.Text.Json;
 using System.Text.Json.Serialization;
diff --git a/Engine/Source/Programs/Shared/EpicGames.Slack/Elements/CheckboxGroupElement.cs b/Engine/Source/Programs/Shared/EpicGames.Slack/Elements/CheckboxGroupElement.cs
--- a/Engine/Source/Programs/Shared/EpicGames.Slack/Elements/CheckboxGroupElement.cs
+++ b/Engine/Source/Programs/Shared/EpicGames.Slack/Elements/CheckboxGroupElement.cs
@@ -1,10 +1,6 @@
 // Copyright Epic Games, Inc. All Rights Reserved.
 
-using System;
 using System.Collections.Generic;
-using System.Globalization;
-using System.Runtime.Serialization;
-using System.Text.Json;
 using System.Text.Json.Serialization;
 
 namespace EpicGames.Slack.Elements
diff --git a/Engine/Source/Programs/Shared/EpicGames.Slack/Elements/DatePickerElement.cs b/Engine/Source/Programs/Shared/EpicGames.Slack/Elements/DatePickerElement.cs
--- a/Engine/Source/Programs/Shared/EpicGames.Slack/Elements/DatePickerElement.cs
+++ b/Engine/Source/Programs/Shared/EpicGames.Slack/Elements/DatePickerElement.cs
@@ -1,10 +1,5 @@
 // Copyright Epic Games, Inc. All Rights Reserved.
 
-using System;
-using System.Collections.Generic;
-using System.Globalization;
-using System.Runtime.Serialization;
-using System.Text.Json;
 using System.Text.Json.Serialization;
 
 namespace EpicGames.Slack.Elements
diff --git a/Engine/Source/Programs/Shared/EpicGames.Slack/Elements/OverflowMenuElement.cs b/Engine/Source/Programs/Shared/EpicGames.Slack/Elements/OverflowMenuElement.cs
--- a/Engine/Source/Programs/Shared/EpicGames.Slack/Elements/OverflowMenuElement.cs
+++ b/Engine/Source/Programs/Shared/EpicGames.Slack/Elements/OverflowMenuElement.cs
@@ -1,10 +1,6 @@
 // Copyright Epic Games, Inc. All Rights Reserved.
 
-using System;
 using System.Collections.Generic;
-using System.Globalization;
-using System.Runtime.Serialization;
-using System.Text.Json;
 using System.Text.Json.Serialization;
 
 namespace EpicGames.Slack.Elements
diff --git a/Engine/Source/Programs/Shared/EpicGames.Slack/Elements/PlainTextInputElement.cs b/Engine/Source/Programs/Shared/EpicGames.Slack/Elements/PlainTextInputElement.cs
--- a/Engine/Source/Programs/Shared/EpicGames.Slack/Elements/PlainTextInputElement.cs
+++ b/Engine/Source/Programs/Shared/EpicGames.Slack/Elements/PlainTextInputElement.cs
@@ -1,9 +1,5 @@
 // Copyright Epic Games, Inc. All Rights Reserved.
 
-using System;
-using System.Collections.Generic;
-using System.Globalization;
-using System.Text.Json;
 using System.Text.Json.Serialization;
 
 namespace EpicGames.Slack.Elements
diff --git a/Engine/Source/Programs/Shared/EpicGames.Slack/Elements/RadioButtonGroupElement.cs b/Engine/Source/Programs/Shared/EpicGames.Slack/Elements/RadioButtonGroupElement.cs
--- a/Engine/Source/Programs/Shared/EpicGames.Slack/Elements/RadioButtonGroupElement.cs
+++ b/Engine/Source/Programs/Shared/EpicGames.Slack/Elements/RadioButtonGroupElement.cs
@@ -1,10 +1,6 @@
 // Copyright Epic Games, Inc. All Rights Reserved.
 
-using System;
 using System.Collections.Generic;
-using System.Globalization;
-using System.Runtime.Serialization;
-using System.Text.Json;
 using System.Text.Json.Serialization;
 
 namespace EpicGames.Slack.Elements
diff --git a/Engine/Source/Programs/Shared/EpicGames.Slack/Elements/SelectMenuElement.cs b/Engine/Source/Programs/Shared/EpicGames.Slack/Elements/SelectMenuElement.cs
--- a/Engine/Source/Programs/Shared/EpicGames.Slack/Elements/SelectMenuElement.cs
+++ b/Engine/Source/Programs/Shared/EpicGames.Slack/Elements/SelectMenuElement.cs
@@ -1,11 +1,7 @@
 // Copyright Epic Games, Inc. All Rights Reserved.
 
-using System;
 using System.Collections.Generic;
 using System.Diagnostics.CodeAnalysis;
-using System.Globalization;
-using System.Runtime.Serialization;
-using System.Text.Json;
 using System.Text.Json.Serialization;
 
 namespace EpicGames.Slack.Elements
diff --git a/Engine/Source/Programs/Shared/EpicGames.Slack/Elements/TextObject.cs b/Engine/Source/Programs/Shared/EpicGames.Slack/Elements/TextObject.cs
--- a/Engine/Source/Programs/Shared/EpicGames.Slack/Elements/TextObject.cs
+++ b/Engine/Source/Programs/Shared/EpicGames.Slack/Elements/TextObject.cs
@@ -1,9 +1,5 @@
 // Copyright Epic Games, Inc. All Rights Reserved.
 
-using System;
-using System.Collections.Generic;
-using System.Globalization;
-using System.Text.Json;
 using System.Text.Json.Serialization;
 
 namespace EpicGames.Slack.Elements
diff --git a/Engine/Source/Programs/Shared/EpicGames.Slack/Elements/TimePickerElement.cs b/Engine/Source/Programs/Shared/EpicGames.Slack/Elements/TimePickerElement.cs
--- a/Engine/Source/Programs/Shared/EpicGames.Slack/Elements/TimePickerElement.cs
+++ b/Engine/Source/Programs/Shared/EpicGames.Slack/Elements/TimePickerElement.cs
@@ -1,10 +1,5 @@
 // Copyright Epic Games, Inc. All Rights Reserved.
 
-using System;
-using System.Collections.Generic;
-using System.Globalization;
-using System.Runtime.Serialization;
-using System.Text.Json;
 using System.Text.Json.Serialization;
 
 namespace EpicGames.Slack.Elements
diff --git a/Engine/Source/Programs/Shared/EpicGames.Slack/EpicGames.Slack.csproj b/Engine/Source/Programs/Shared/EpicGames.Slack/EpicGames.Slack.csproj
--- a/Engine/Source/Programs/Shared/EpicGames.Slack/EpicGames.Slack.csproj
+++ b/Engine/Source/Programs/Shared/EpicGames.Slack/EpicGames.Slack.csproj
@@ -1,4 +1,4 @@
-﻿<Project Sdk="Microsoft.NET.Sdk">
+<Project Sdk="Microsoft.NET.Sdk">
   <PropertyGroup>
     <TargetFramework>net6.0</TargetFramework>
     <Configuration Condition=" '$(Configuration)' == '' ">Development</Configuration>
@@ -25,6 +25,7 @@
   </PropertyGroup>
   <PropertyGroup Condition=" '$(Configuration)|$(Platform)' == 'Analyze|AnyCPU' ">
     <DocumentationFile></DocumentationFile>
+    <EnforceCodeStyleInBuild>True</EnforceCodeStyleInBuild>
     <RunAnalyzersDuringBuild>True</RunAnalyzersDuringBuild>
     <GenerateDocumentationFile>True</GenerateDocumentationFile>
   </PropertyGroup>
diff --git a/Engine/Source/Programs/Shared/EpicGames.Slack/SlackAttachment.cs b/Engine/Source/Programs/Shared/EpicGames.Slack/SlackAttachment.cs
--- a/Engine/Source/Programs/Shared/EpicGames.Slack/SlackAttachment.cs
+++ b/Engine/Source/Programs/Shared/EpicGames.Slack/SlackAttachment.cs
@@ -1,12 +1,7 @@
 // Copyright Epic Games, Inc. All Rights Reserved.
 
-using System;
 using System.Collections.Generic;
-using System.Drawing;
-using System.Text.Json;
 using System.Text.Json.Serialization;
-using EpicGames.Slack.Blocks;
-using EpicGames.Slack.Elements;
 
 namespace EpicGames.Slack
 {
diff --git a/Engine/Source/Programs/Shared/EpicGames.Slack/SlackClient.cs b/Engine/Source/Programs/Shared/EpicGames.Slack/SlackClient.cs
--- a/Engine/Source/Programs/Shared/EpicGames.Slack/SlackClient.cs
+++ b/Engine/Source/Programs/Shared/EpicGames.Slack/SlackClient.cs
@@ -3,11 +3,9 @@
 using System;
 using System.Collections.Generic;
 using System.Net.Http;
-using System.Net.WebSockets;
 using System.Text;
 using System.Text.Json;
 using System.Text.Json.Serialization;
-using System.Threading;
 using System.Threading.Tasks;
 using Microsoft.Extensions.Logging;
 
diff --git a/Engine/Source/Programs/Shared/EpicGames.Slack/SlackMessage.cs b/Engine/Source/Programs/Shared/EpicGames.Slack/SlackMessage.cs
--- a/Engine/Source/Programs/Shared/EpicGames.Slack/SlackMessage.cs
+++ b/Engine/Source/Programs/Shared/EpicGames.Slack/SlackMessage.cs
@@ -1,10 +1,6 @@
 // Copyright Epic Games, Inc. All Rights Reserved.
 
-using System;
 using System.Collections.Generic;
-using System.Linq;
-using System.Text;
-using System.Threading.Tasks;
 
 namespace EpicGames.Slack
 {
diff --git a/Engine/Source/Programs/Shared/EpicGames.Slack/SlackOption.cs b/Engine/Source/Programs/Shared/EpicGames.Slack/SlackOption.cs
--- a/Engine/Source/Programs/Shared/EpicGames.Slack/SlackOption.cs
+++ b/Engine/Source/Programs/Shared/EpicGames.Slack/SlackOption.cs
@@ -2,7 +2,6 @@
 
 using System;
 using System.Collections.Generic;
-using System.Text.Json;
 using System.Text.Json.Serialization;
 using EpicGames.Slack.Elements;
 
diff --git a/Engine/Source/Programs/Shared/EpicGames.Slack/SlackUser.cs b/Engine/Source/Programs/Shared/EpicGames.Slack/SlackUser.cs
--- a/Engine/Source/Programs/Shared/EpicGames.Slack/SlackUser.cs
+++ b/Engine/Source/Programs/Shared/EpicGames.Slack/SlackUser.cs
@@ -1,11 +1,6 @@
 // Copyright Epic Games, Inc. All Rights Reserved.
 
-using System;
-using System.Collections.Generic;
-using System.Linq;
-using System.Text;
 using System.Text.Json.Serialization;
-using System.Threading.Tasks;
 
 namespace EpicGames.Slack
 {
diff --git a/Engine/Source/Programs/Shared/EpicGames.Slack/SlackView.cs b/Engine/Source/Programs/Shared/EpicGames.Slack/SlackView.cs
--- a/Engine/Source/Programs/Shared/EpicGames.Slack/SlackView.cs
+++ b/Engine/Source/Programs/Shared/EpicGames.Slack/SlackView.cs
@@ -3,10 +3,7 @@
 using System;
 using System.Collections.Generic;
 using System.Diagnostics.CodeAnalysis;
-using System.Linq;
-using System.Text;
 using System.Text.Json.Serialization;
-using System.Threading.Tasks;
 using EpicGames.Slack.Elements;
 
 namespace EpicGames.Slack
