Horde: Few more static analyzer fixes.

diff --git a/Engine/Source/Programs/Shared/EpicGames.BuildGraph/Expressions/BgNode.cs b/Engine/Source/Programs/Shared/EpicGames.BuildGraph/Expressions/BgNode.cs
--- a/Engine/Source/Programs/Shared/EpicGames.BuildGraph/Expressions/BgNode.cs
+++ b/Engine/Source/Programs/Shared/EpicGames.BuildGraph/Expressions/BgNode.cs
@@ -151,19 +151,19 @@ namespace EpicGames.BuildGraph.Expressions
 		internal BgNode Modify(BgList<BgFileSet>? inputs = null, BgList<BgNode>? fences = null, BgBool? runEarly = null, BgList<BgLabel>? labels = null)
 		{
 			BgNode node = Clone();
-			if (inputs is object)
+			if (inputs is not null)
 			{
 				node.Inputs = inputs;
 			}
-			if (fences is object)
+			if (fences is not null)
 			{
 				node.Fences = fences;
 			}
-			if (runEarly is object)
+			if (runEarly is not null)
 			{
 				node.RunEarly = runEarly;
 			}
-			if (labels is object)
+			if (labels is not null)
 			{
 				node.Labels = labels;
 			}
diff --git a/Engine/Source/Programs/Shared/EpicGames.Horde/Storage/Workspace.cs b/Engine/Source/Programs/Shared/EpicGames.Horde/Storage/Workspace.cs
--- a/Engine/Source/Programs/Shared/EpicGames.Horde/Storage/Workspace.cs
+++ b/Engine/Source/Programs/Shared/EpicGames.Horde/Storage/Workspace.cs
@@ -893,7 +893,7 @@ namespace EpicGames.Horde.Storage
 			}
 			hashInfo.Files.Add(file);
 		}
-
+		/*
 		void RemoveFileFromHashLookup(FileState file)
 		{
 			HashInfo? hashInfo;
@@ -902,7 +902,7 @@ namespace EpicGames.Horde.Storage
 				hashInfo.Files.Remove(file);
 			}
 		}
-
+		*/
 		void MoveFileToCache(FileState fileState, DirectoryState cacheDirState)
 		{
 			Utf8String name = fileState.Hash.ToUtf8String();
diff --git a/Engine/Source/Programs/Shared/EpicGames.Serialization/CbField.cs b/Engine/Source/Programs/Shared/EpicGames.Serialization/CbField.cs
--- a/Engine/Source/Programs/Shared/EpicGames.Serialization/CbField.cs
+++ b/Engine/Source/Programs/Shared/EpicGames.Serialization/CbField.cs
@@ -194,7 +194,7 @@ namespace EpicGames.Serialization
 	[DebuggerDisplay("{Hash}")]
 	[JsonConverter(typeof(CbBinaryAttachmentJsonConverter))]
 	[TypeConverter(typeof(CbBinaryAttachmentTypeConverter))]
-	public struct CbBinaryAttachment : IEquatable<CbBinaryAttachment>
+	public readonly struct CbBinaryAttachment : IEquatable<CbBinaryAttachment>
 	{
 		/// <summary>
 		/// Attachment with a hash of zero
@@ -276,7 +276,7 @@ namespace EpicGames.Serialization
 	[DebuggerDisplay("{Hash}")]
 	[JsonConverter(typeof(CbObjectAttachmentJsonConverter))]
 	[TypeConverter(typeof(CbObjectAttachmentTypeConverter))]
-	public struct CbObjectAttachment : IEquatable<CbObjectAttachment>
+	public readonly struct CbObjectAttachment : IEquatable<CbObjectAttachment>
 	{
 		/// <summary>
 		/// Attachment with a hash of zero
