Horde: Rename IStorageBackend.ReadAsync() to OpenAsync().

diff --git a/Engine/Source/Programs/Horde/Horde.Server/Jobs/Artifacts/ArtifactCollectionV1.cs b/Engine/Source/Programs/Horde/Horde.Server/Jobs/Artifacts/ArtifactCollectionV1.cs
--- a/Engine/Source/Programs/Horde/Horde.Server/Jobs/Artifacts/ArtifactCollectionV1.cs
+++ b/Engine/Source/Programs/Horde/Horde.Server/Jobs/Artifacts/ArtifactCollectionV1.cs
@@ -202,7 +202,7 @@ namespace Horde.Server.Jobs.Artifacts
 		/// <returns>The chunk data</returns>
 		public async Task<System.IO.Stream> OpenArtifactReadStreamAsync(IArtifactV1 artifact)
 		{
-			System.IO.Stream stream = await _storageBackend.ReadAsync(GetPath(artifact.JobId, artifact.StepId, artifact.Name));
+			System.IO.Stream stream = await _storageBackend.OpenAsync(GetPath(artifact.JobId, artifact.StepId, artifact.Name));
 			return stream;
 		}
 
diff --git a/Engine/Source/Programs/Horde/Horde.Server/Storage/Backends/AwsStorageBackend.cs b/Engine/Source/Programs/Horde/Horde.Server/Storage/Backends/AwsStorageBackend.cs
--- a/Engine/Source/Programs/Horde/Horde.Server/Storage/Backends/AwsStorageBackend.cs
+++ b/Engine/Source/Programs/Horde/Horde.Server/Storage/Backends/AwsStorageBackend.cs
@@ -262,13 +262,7 @@ namespace Horde.Server.Storage.Backends
 		string GetFullPath(string path) => _pathPrefix + path;
 
 		/// <inheritdoc/>
-		public Task<Stream> ReadAsync(string path, CancellationToken cancellationToken)
-		{
-			return ReadAsync(path, null, cancellationToken);
-		}
-
-		/// <inheritdoc/>
-		public Task<Stream> ReadAsync(string path, int offset, int? length, CancellationToken cancellationToken)
+		public Task<Stream> OpenAsync(string path, int offset, int? length, CancellationToken cancellationToken)
 		{
 			string range;
 			if (length == null)
diff --git a/Engine/Source/Programs/Horde/Horde.Server/Storage/StorageBackendProvider.cs b/Engine/Source/Programs/Horde/Horde.Server/Storage/StorageBackendProvider.cs
--- a/Engine/Source/Programs/Horde/Horde.Server/Storage/StorageBackendProvider.cs
+++ b/Engine/Source/Programs/Horde/Horde.Server/Storage/StorageBackendProvider.cs
@@ -79,13 +79,13 @@ namespace Horde.Server.Storage
 			public IAsyncEnumerable<string> EnumerateAsync(CancellationToken cancellationToken = default) => _backend.EnumerateAsync(cancellationToken);
 
 			/// <inheritdoc/>
-			public Task<Stream> ReadAsync(string path, int offset, int? length, CancellationToken cancellationToken = default)
+			public Task<Stream> OpenAsync(string path, int offset, int? length, CancellationToken cancellationToken = default)
 			{
 				if (_backend == null)
 				{
 					throw new InvalidOperationException("Backend has already been disposed");
 				}
-				return _backend.ReadAsync(path, offset, length, cancellationToken);
+				return _backend.OpenAsync(path, offset, length, cancellationToken);
 			}
 
 			/// <inheritdoc/>
diff --git a/Engine/Source/Programs/Horde/Horde.Server/Storage/StorageService.cs b/Engine/Source/Programs/Horde/Horde.Server/Storage/StorageService.cs
--- a/Engine/Source/Programs/Horde/Horde.Server/Storage/StorageService.cs
+++ b/Engine/Source/Programs/Horde/Horde.Server/Storage/StorageService.cs
@@ -126,11 +126,11 @@ namespace Horde.Server.Storage
 			public bool SupportsRedirects => _inner.SupportsRedirects;
 
 			/// <inheritdoc/>
-			public async Task<Stream> ReadAsync(string path, int offset, int? length, CancellationToken cancellationToken = default)
+			public async Task<Stream> OpenAsync(string path, int offset, int? length, CancellationToken cancellationToken = default)
 			{
 				string fullPath = $"{_prefix}{path}";
 
-				using TelemetrySpan span = _tracer.StartActiveSpan($"{nameof(StorageService)}.{nameof(StorageClientImpl)}.{nameof(ReadAsync)}");
+				using TelemetrySpan span = _tracer.StartActiveSpan($"{nameof(StorageService)}.{nameof(StorageClientImpl)}.{nameof(OpenAsync)}");
 				span.SetAttribute("path", fullPath);
 				span.SetAttribute("offset", offset);
 				span.SetAttribute("length", length);
@@ -140,7 +140,7 @@ namespace Horde.Server.Storage
 					return new MemoryStream(Array.Empty<byte>());
 				}
 
-				return await _inner.ReadAsync(fullPath, offset, length, cancellationToken);
+				return await _inner.OpenAsync(fullPath, offset, length, cancellationToken);
 			}
 
 			/// <inheritdoc/>
diff --git a/Engine/Source/Programs/Horde/Horde.Server/Storage/TypedStorageBackend.cs b/Engine/Source/Programs/Horde/Horde.Server/Storage/TypedStorageBackend.cs
--- a/Engine/Source/Programs/Horde/Horde.Server/Storage/TypedStorageBackend.cs
+++ b/Engine/Source/Programs/Horde/Horde.Server/Storage/TypedStorageBackend.cs
@@ -43,7 +43,7 @@ namespace Horde.Server.Storage
 			public void Dispose() => _inner.Dispose();
 
 			/// <inheritdoc/>
-			public Task<Stream> ReadAsync(string path, int offset, int? length, CancellationToken cancellationToken) => _inner.ReadAsync(path, offset, length, cancellationToken);
+			public Task<Stream> OpenAsync(string path, int offset, int? length, CancellationToken cancellationToken) => _inner.OpenAsync(path, offset, length, cancellationToken);
 
 			/// <inheritdoc/>
 			public Task<string> WriteAsync(Stream stream, string? prefix, CancellationToken cancellationToken) => _inner.WriteAsync(stream, prefix, cancellationToken);
diff --git a/Engine/Source/Programs/Shared/EpicGames.Horde/Compute/AgentStorageClient.cs b/Engine/Source/Programs/Shared/EpicGames.Horde/Compute/AgentStorageClient.cs
--- a/Engine/Source/Programs/Shared/EpicGames.Horde/Compute/AgentStorageClient.cs
+++ b/Engine/Source/Programs/Shared/EpicGames.Horde/Compute/AgentStorageClient.cs
@@ -96,7 +96,7 @@ namespace EpicGames.Horde.Compute
 		public bool SupportsRedirects => throw new NotImplementedException();
 
 		/// <inheritdoc/>
-		public async Task<Stream> ReadAsync(string path, int offset, int? length, CancellationToken cancellationToken = default)
+		public async Task<Stream> OpenAsync(string path, int offset, int? length, CancellationToken cancellationToken = default)
 		{
 			await _semaphore.WaitAsync(cancellationToken);
 			try
diff --git a/Engine/Source/Programs/Shared/EpicGames.Horde/Storage/Backends/FileStorageBackend.cs b/Engine/Source/Programs/Shared/EpicGames.Horde/Storage/Backends/FileStorageBackend.cs
--- a/Engine/Source/Programs/Shared/EpicGames.Horde/Storage/Backends/FileStorageBackend.cs
+++ b/Engine/Source/Programs/Shared/EpicGames.Horde/Storage/Backends/FileStorageBackend.cs
@@ -44,7 +44,7 @@ namespace EpicGames.Horde.Storage.Backends
 		FileReference GetBlobFile(string path) => FileReference.Combine(_baseDir, $"{path}.blob");
 	
 		/// <inheritdoc/>
-		public Task<Stream> ReadAsync(string path, int offset, int? length, CancellationToken cancellationToken)
+		public Task<Stream> OpenAsync(string path, int offset, int? length, CancellationToken cancellationToken)
 		{
 			FileReference location = GetBlobFile(path);
 			Stream stream = FileReference.Open(location, FileMode.Open, FileAccess.Read, FileShare.Read);
diff --git a/Engine/Source/Programs/Shared/EpicGames.Horde/Storage/Backends/HttpStorageBackend.cs b/Engine/Source/Programs/Shared/EpicGames.Horde/Storage/Backends/HttpStorageBackend.cs
--- a/Engine/Source/Programs/Shared/EpicGames.Horde/Storage/Backends/HttpStorageBackend.cs
+++ b/Engine/Source/Programs/Shared/EpicGames.Horde/Storage/Backends/HttpStorageBackend.cs
@@ -50,7 +50,7 @@ namespace EpicGames.Horde.Storage.Backends
 		#region Blobs
 
 		/// <inheritdoc/>
-		public async Task<Stream> ReadAsync(string path, int offset, int? length, CancellationToken cancellationToken = default)
+		public async Task<Stream> OpenAsync(string path, int offset, int? length, CancellationToken cancellationToken = default)
 		{
 			if (offset == 0 && length == null)
 			{
diff --git a/Engine/Source/Programs/Shared/EpicGames.Horde/Storage/Backends/MemoryStorageBackend.cs b/Engine/Source/Programs/Shared/EpicGames.Horde/Storage/Backends/MemoryStorageBackend.cs
--- a/Engine/Source/Programs/Shared/EpicGames.Horde/Storage/Backends/MemoryStorageBackend.cs
+++ b/Engine/Source/Programs/Shared/EpicGames.Horde/Storage/Backends/MemoryStorageBackend.cs
@@ -35,7 +35,7 @@ namespace EpicGames.Horde.Storage.Backends
 		}
 
 		/// <inheritdoc/>
-		public Task<Stream> ReadAsync(string path, int offset, int? length, CancellationToken cancellationToken)
+		public Task<Stream> OpenAsync(string path, int offset, int? length, CancellationToken cancellationToken)
 		{
 			ReadOnlyMemory<byte> data = _pathToData[path].AsMemory(offset);
 			if (length != null && length.Value < data.Length)
diff --git a/Engine/Source/Programs/Shared/EpicGames.Horde/Storage/Clients/BundleStorageClient.cs b/Engine/Source/Programs/Shared/EpicGames.Horde/Storage/Clients/BundleStorageClient.cs
--- a/Engine/Source/Programs/Shared/EpicGames.Horde/Storage/Clients/BundleStorageClient.cs
+++ b/Engine/Source/Programs/Shared/EpicGames.Horde/Storage/Clients/BundleStorageClient.cs
@@ -45,7 +45,7 @@ namespace EpicGames.Horde.Storage.Clients
 		/// <inheritdoc/>
 		public async Task<Stream> OpenAsync(BundleLocator locator, int offset, int? length = null, CancellationToken cancellationToken = default)
 		{
-			return await _backend.ReadAsync(locator.Path.ToString(), offset, length, cancellationToken);
+			return await _backend.OpenAsync(locator.Path.ToString(), offset, length, cancellationToken);
 		}
 
 		/// <summary>
diff --git a/Engine/Source/Programs/Shared/EpicGames.Horde/Storage/IStorageBackend.cs b/Engine/Source/Programs/Shared/EpicGames.Horde/Storage/IStorageBackend.cs
--- a/Engine/Source/Programs/Shared/EpicGames.Horde/Storage/IStorageBackend.cs
+++ b/Engine/Source/Programs/Shared/EpicGames.Horde/Storage/IStorageBackend.cs
@@ -28,7 +28,7 @@ namespace EpicGames.Horde.Storage
 		/// <param name="length">Length of data to read</param>
 		/// <param name="cancellationToken">Cancellation token for the operation</param>
 		/// <returns></returns>
-		Task<Stream> ReadAsync(string path, int offset, int? length, CancellationToken cancellationToken = default);
+		Task<Stream> OpenAsync(string path, int offset, int? length, CancellationToken cancellationToken = default);
 
 		/// <summary>
 		/// Writes a stream to the storage backend. If the stream throws an exception during read, the write will be aborted.
@@ -134,7 +134,7 @@ namespace EpicGames.Horde.Storage
 		/// <param name="path">Object name within the store</param>
 		/// <param name="cancellationToken">Cancellation token for the operation</param>
 		/// <returns>Stream for the object</returns>
-		public static Task<Stream> ReadAsync(this IStorageBackend storageBackend, string path, CancellationToken cancellationToken = default) => storageBackend.ReadAsync(path, 0, null, cancellationToken);
+		public static Task<Stream> OpenAsync(this IStorageBackend storageBackend, string path, CancellationToken cancellationToken = default) => storageBackend.OpenAsync(path, 0, null, cancellationToken);
 
 		/// <summary>
 		/// Reads an object as an array of bytes
@@ -145,7 +145,7 @@ namespace EpicGames.Horde.Storage
 		/// <returns>Contents of the object</returns>
 		public static async Task<byte[]> ReadBytesAsync(this IStorageBackend storageBackend, string path, CancellationToken cancellationToken = default)
 		{
-			using (Stream inputStream = await storageBackend.ReadAsync(path, cancellationToken))
+			using (Stream inputStream = await storageBackend.OpenAsync(path, cancellationToken))
 			{
 				using (MemoryStream outputStream = new MemoryStream())
 				{
diff --git a/Engine/Source/Programs/Shared/EpicGames.Horde/Storage/Nodes/DirectoryNode.cs b/Engine/Source/Programs/Shared/EpicGames.Horde/Storage/Nodes/DirectoryNode.cs
--- a/Engine/Source/Programs/Shared/EpicGames.Horde/Storage/Nodes/DirectoryNode.cs
+++ b/Engine/Source/Programs/Shared/EpicGames.Horde/Storage/Nodes/DirectoryNode.cs
@@ -530,7 +530,7 @@ namespace EpicGames.Horde.Storage.Nodes
 		/// <inheritdoc cref="AddFilesAsync(DirectoryReference, IReadOnlyList{FileInfo}, ChunkingOptions, IStorageWriter, IProgress{ICopyStats}?, CancellationToken)"/>
 		public Task AddFilesAsync(DirectoryInfo baseDir, IEnumerable<FileInfo> files, ChunkingOptions options, IStorageWriter writer, IProgress<ICopyStats>? progress, CancellationToken cancellationToken)
 		{
-			return AddFilesAsync(baseDir, files.ToList(), options, writer, progress, cancellationToken);
+			return AddFilesAsync(new DirectoryReference(baseDir), files.ToList(), options, writer, progress, cancellationToken);
 		}
 
 		/// <summary>
diff --git a/Engine/Source/Programs/Shared/EpicGames.Horde/Storage/StorageBackendCache.cs b/Engine/Source/Programs/Shared/EpicGames.Horde/Storage/StorageBackendCache.cs
--- a/Engine/Source/Programs/Shared/EpicGames.Horde/Storage/StorageBackendCache.cs
+++ b/Engine/Source/Programs/Shared/EpicGames.Horde/Storage/StorageBackendCache.cs
@@ -65,9 +65,9 @@ namespace EpicGames.Horde.Storage
 
 			public void Dispose() => _inner.Dispose();
 
-			public async Task<Stream> ReadAsync(string path, int offset, int? length, CancellationToken cancellationToken = default)
+			public async Task<Stream> OpenAsync(string path, int offset, int? length, CancellationToken cancellationToken = default)
 			{
-				Stream stream = await _cacheStorage.ReadAsync($"{_keyPrefix}{path}", ctx => _inner.ReadAsync(path, ctx), cancellationToken);
+				Stream stream = await _cacheStorage.ReadAsync($"{_keyPrefix}{path}", ctx => _inner.OpenAsync(path, ctx), cancellationToken);
 				if (offset != 0)
 				{
 					stream.Seek(offset, SeekOrigin.Begin);
