Horde: Add optional request ID to compute allocation requests

diff --git a/Engine/Source/Programs/Horde/Horde.Agent/Commands/Compute/ComputeCommand.cs b/Engine/Source/Programs/Horde/Horde.Agent/Commands/Compute/ComputeCommand.cs
--- a/Engine/Source/Programs/Horde/Horde.Agent/Commands/Compute/ComputeCommand.cs
+++ b/Engine/Source/Programs/Horde/Horde.Agent/Commands/Compute/ComputeCommand.cs
@@ -58,7 +58,7 @@ namespace Horde.Agent.Commands.Compute
 				requirements = new Requirements(Condition.Parse(Requirements));
 			}
 
-			await using IComputeLease? lease = await client.TryAssignWorkerAsync(new ClusterId(ClusterId), requirements, CancellationToken.None);
+			await using IComputeLease? lease = await client.TryAssignWorkerAsync(new ClusterId(ClusterId), requirements, null, CancellationToken.None);
 			if (lease == null)
 			{
 				throw new Exception("Unable to create lease");
diff --git a/Engine/Source/Programs/Horde/Horde/Commands/Compute/ComputeRun.cs b/Engine/Source/Programs/Horde/Horde/Commands/Compute/ComputeRun.cs
--- a/Engine/Source/Programs/Horde/Horde/Commands/Compute/ComputeRun.cs
+++ b/Engine/Source/Programs/Horde/Horde/Commands/Compute/ComputeRun.cs
@@ -72,7 +72,7 @@ namespace Horde.Commands.Compute
 				requirements = new Requirements(Condition.Parse(Requirements));
 			}
 
-			await using IComputeLease? lease = await client.TryAssignWorkerAsync(new ClusterId(ClusterId), requirements, CancellationToken.None);
+			await using IComputeLease? lease = await client.TryAssignWorkerAsync(new ClusterId(ClusterId), requirements, null, CancellationToken.None);
 			if (lease == null)
 			{
 				throw new Exception("Unable to create lease");
diff --git a/Engine/Source/Programs/Horde/Samples/RemoteClient/RemoteClientApp.cs b/Engine/Source/Programs/Horde/Samples/RemoteClient/RemoteClientApp.cs
--- a/Engine/Source/Programs/Horde/Samples/RemoteClient/RemoteClientApp.cs
+++ b/Engine/Source/Programs/Horde/Samples/RemoteClient/RemoteClientApp.cs
@@ -60,7 +60,7 @@ namespace RemoteClient
 				requirements = new Requirements(Condition.Parse(options.Condition));
 			}
 
-			await using IComputeLease? lease = await client.TryAssignWorkerAsync(new ClusterId("default"), requirements);
+			await using IComputeLease? lease = await client.TryAssignWorkerAsync(new ClusterId("default"), requirements, null);
 			if (lease == null)
 			{
 				logger.LogInformation("Unable to connect to remote");
diff --git a/Engine/Source/Programs/Shared/EpicGames.Horde/Api/ComputeMessages.cs b/Engine/Source/Programs/Shared/EpicGames.Horde/Api/ComputeMessages.cs
--- a/Engine/Source/Programs/Shared/EpicGames.Horde/Api/ComputeMessages.cs
+++ b/Engine/Source/Programs/Shared/EpicGames.Horde/Api/ComputeMessages.cs
@@ -17,6 +17,12 @@ namespace EpicGames.Horde.Api
 		/// Condition to identify machines that can execute the request
 		/// </summary>
 		public Requirements? Requirements { get; set; }
+		
+		/// <summary>
+		/// Arbitrary ID to correlate the same request over multiple calls.
+		/// It's recommended to pick something globally unique, such as a UUID.
+		/// </summary>
+		public string? RequestId { get; set; }
 	}
 
 	/// <summary>
diff --git a/Engine/Source/Programs/Shared/EpicGames.Horde/Compute/Clients/AgentComputeClient.cs b/Engine/Source/Programs/Shared/EpicGames.Horde/Compute/Clients/AgentComputeClient.cs
--- a/Engine/Source/Programs/Shared/EpicGames.Horde/Compute/Clients/AgentComputeClient.cs
+++ b/Engine/Source/Programs/Shared/EpicGames.Horde/Compute/Clients/AgentComputeClient.cs
@@ -67,7 +67,13 @@ namespace EpicGames.Horde.Compute.Clients
 		public ValueTask DisposeAsync() => new ValueTask();
 
 		/// <inheritdoc/>
-		public async Task<IComputeLease?> TryAssignWorkerAsync(ClusterId clusterId, Requirements? requirements, CancellationToken cancellationToken)
+		public Task<IComputeLease?> TryAssignWorkerAsync(ClusterId clusterId, Requirements? requirements, CancellationToken cancellationToken)
+		{
+			return TryAssignWorkerAsync(clusterId, requirements, null, cancellationToken);
+		}
+		
+		/// <inheritdoc/>
+		public async Task<IComputeLease?> TryAssignWorkerAsync(ClusterId clusterId, Requirements? requirements, string? requestId, CancellationToken cancellationToken)
 		{
 			_logger.LogInformation("** CLIENT **");
 			_logger.LogInformation("Launching {Path} to handle remote", _hordeAgentAssembly);
diff --git a/Engine/Source/Programs/Shared/EpicGames.Horde/Compute/Clients/LocalComputeClient.cs b/Engine/Source/Programs/Shared/EpicGames.Horde/Compute/Clients/LocalComputeClient.cs
--- a/Engine/Source/Programs/Shared/EpicGames.Horde/Compute/Clients/LocalComputeClient.cs
+++ b/Engine/Source/Programs/Shared/EpicGames.Horde/Compute/Clients/LocalComputeClient.cs
@@ -120,10 +120,16 @@ namespace EpicGames.Horde.Compute.Clients
 				await socket.CloseAsync(cancellationToken);
 			}
 		}
-
+		
 		/// <inheritdoc/>
 		public Task<IComputeLease?> TryAssignWorkerAsync(ClusterId clusterId, Requirements? requirements, CancellationToken cancellationToken)
 		{
+			return TryAssignWorkerAsync(clusterId, requirements, null, cancellationToken);
+		}
+
+		/// <inheritdoc/>
+		public Task<IComputeLease?> TryAssignWorkerAsync(ClusterId clusterId, Requirements? requirements, string? requestId, CancellationToken cancellationToken)
+		{
 #pragma warning disable CA2000 // Dispose objects before losing scope
 			RemoteComputeSocket socket = new RemoteComputeSocket(new TcpTransport(_socket), new PrefixLogger("CLIENT", _logger));
 			return Task.FromResult<IComputeLease?>(new LeaseImpl(socket));
diff --git a/Engine/Source/Programs/Shared/EpicGames.Horde/Compute/Clients/ServerComputeClient.cs b/Engine/Source/Programs/Shared/EpicGames.Horde/Compute/Clients/ServerComputeClient.cs
--- a/Engine/Source/Programs/Shared/EpicGames.Horde/Compute/Clients/ServerComputeClient.cs
+++ b/Engine/Source/Programs/Shared/EpicGames.Horde/Compute/Clients/ServerComputeClient.cs
@@ -141,9 +141,15 @@ namespace EpicGames.Horde.Compute.Clients
 		}
 
 		/// <inheritdoc/>
-		public async Task<IComputeLease?> TryAssignWorkerAsync(ClusterId clusterId, Requirements? requirements, CancellationToken cancellationToken)
+		public Task<IComputeLease?> TryAssignWorkerAsync(ClusterId clusterId, Requirements? requirements, CancellationToken cancellationToken)
 		{
-			IAsyncEnumerator<LeaseInfo> source = ConnectAsync(clusterId, requirements, cancellationToken).GetAsyncEnumerator(cancellationToken);
+			return TryAssignWorkerAsync(clusterId, requirements, null, cancellationToken);
+		}
+		
+		/// <inheritdoc/>
+		public async Task<IComputeLease?> TryAssignWorkerAsync(ClusterId clusterId, Requirements? requirements, string? requestId, CancellationToken cancellationToken)
+		{
+			IAsyncEnumerator<LeaseInfo> source = ConnectAsync(clusterId, requirements, requestId, cancellationToken).GetAsyncEnumerator(cancellationToken);
 			if (!await source.MoveNextAsync())
 			{
 				await source.DisposeAsync();
@@ -153,7 +159,7 @@ namespace EpicGames.Horde.Compute.Clients
 		}
 
 		/// <inheritdoc/>
-		async IAsyncEnumerable<LeaseInfo> ConnectAsync(ClusterId clusterId, Requirements? requirements, [EnumeratorCancellation] CancellationToken cancellationToken)
+		async IAsyncEnumerable<LeaseInfo> ConnectAsync(ClusterId clusterId, Requirements? requirements, string? requestId, [EnumeratorCancellation] CancellationToken cancellationToken)
 		{
 			_logger.LogDebug("Requesting compute resource");
 
@@ -165,6 +171,7 @@ namespace EpicGames.Horde.Compute.Clients
 
 			AssignComputeRequest request = new AssignComputeRequest();
 			request.Requirements = requirements;
+			request.RequestId = requestId;
 
 			AssignComputeResponse? responseMessage;
 			using (HttpResponseMessage response = await client.PostAsync($"api/v2/compute/{clusterId}", request, _cancellationSource.Token))
diff --git a/Engine/Source/Programs/Shared/EpicGames.Horde/Compute/IComputeClient.cs b/Engine/Source/Programs/Shared/EpicGames.Horde/Compute/IComputeClient.cs
--- a/Engine/Source/Programs/Shared/EpicGames.Horde/Compute/IComputeClient.cs
+++ b/Engine/Source/Programs/Shared/EpicGames.Horde/Compute/IComputeClient.cs
@@ -12,11 +12,20 @@ namespace EpicGames.Horde.Compute
 	public interface IComputeClient : IAsyncDisposable
 	{
 		/// <summary>
-		/// Adds a new remote request
+		/// Adds a new remote request (deprecated, use method below instead)
 		/// </summary>
 		/// <param name="clusterId">Cluster to execute the request</param>
 		/// <param name="requirements">Requirements for the agent</param>
 		/// <param name="cancellationToken">Cancellation token for the operation</param>
 		public Task<IComputeLease?> TryAssignWorkerAsync(ClusterId clusterId, Requirements? requirements, CancellationToken cancellationToken = default);
+		
+		/// <summary>
+		/// Adds a new remote request
+		/// </summary>
+		/// <param name="clusterId">Cluster to execute the request</param>
+		/// <param name="requirements">Requirements for the agent</param>
+		/// <param name="requestId">Optional ID identifying the request over multiple calls, such as retrying the same request</param>
+		/// <param name="cancellationToken">Cancellation token for the operation</param>
+		public Task<IComputeLease?> TryAssignWorkerAsync(ClusterId clusterId, Requirements? requirements, string? requestId, CancellationToken cancellationToken = default);
 	}
 }
