Horde: Enable code style warnings in EpicGames.Horde.

diff --git a/Engine/Source/Programs/Shared/EpicGames.Horde/Api/ArtifactId.cs b/Engine/Source/Programs/Shared/EpicGames.Horde/Api/ArtifactId.cs
--- a/Engine/Source/Programs/Shared/EpicGames.Horde/Api/ArtifactId.cs
+++ b/Engine/Source/Programs/Shared/EpicGames.Horde/Api/ArtifactId.cs
@@ -1,7 +1,6 @@
 // Copyright Epic Games, Inc. All Rights Reserved.
 
 using System.ComponentModel;
-using System.Text.Json.Serialization;
 
 namespace EpicGames.Horde.Api
 {
diff --git a/Engine/Source/Programs/Shared/EpicGames.Horde/Api/ArtifactMessages.cs b/Engine/Source/Programs/Shared/EpicGames.Horde/Api/ArtifactMessages.cs
--- a/Engine/Source/Programs/Shared/EpicGames.Horde/Api/ArtifactMessages.cs
+++ b/Engine/Source/Programs/Shared/EpicGames.Horde/Api/ArtifactMessages.cs
@@ -1,9 +1,6 @@
 ﻿// Copyright Epic Games, Inc. All Rights Reserved.
 
-using System;
 using System.Collections.Generic;
-using System.Security.Cryptography;
-using System.Text;
 using System.Threading;
 using System.Threading.Tasks;
 using EpicGames.Core;
diff --git a/Engine/Source/Programs/Shared/EpicGames.Horde/Api/SecretMessages.cs b/Engine/Source/Programs/Shared/EpicGames.Horde/Api/SecretMessages.cs
--- a/Engine/Source/Programs/Shared/EpicGames.Horde/Api/SecretMessages.cs
+++ b/Engine/Source/Programs/Shared/EpicGames.Horde/Api/SecretMessages.cs
@@ -1,6 +1,5 @@
 // Copyright Epic Games, Inc. All Rights Reserved.
 
-using System;
 using System.Collections.Generic;
 using System.Diagnostics;
 using System.Threading;
diff --git a/Engine/Source/Programs/Shared/EpicGames.Horde/Api/UserId.cs b/Engine/Source/Programs/Shared/EpicGames.Horde/Api/UserId.cs
--- a/Engine/Source/Programs/Shared/EpicGames.Horde/Api/UserId.cs
+++ b/Engine/Source/Programs/Shared/EpicGames.Horde/Api/UserId.cs
@@ -2,7 +2,6 @@
 
 using System.ComponentModel;
 using EpicGames.Core;
-using EpicGames.Horde;
 
 namespace EpicGames.Horde.Api
 {
diff --git a/Engine/Source/Programs/Shared/EpicGames.Horde/Compute/AgentMessageChannel.cs b/Engine/Source/Programs/Shared/EpicGames.Horde/Compute/AgentMessageChannel.cs
--- a/Engine/Source/Programs/Shared/EpicGames.Horde/Compute/AgentMessageChannel.cs
+++ b/Engine/Source/Programs/Shared/EpicGames.Horde/Compute/AgentMessageChannel.cs
@@ -1,7 +1,6 @@
 // Copyright Epic Games, Inc. All Rights Reserved.
 
 using System;
-using System.Buffers;
 using System.Buffers.Binary;
 using System.Text;
 using System.Threading;
diff --git a/Engine/Source/Programs/Shared/EpicGames.Horde/Compute/Clients/AgentComputeClient.cs b/Engine/Source/Programs/Shared/EpicGames.Horde/Compute/Clients/AgentComputeClient.cs
--- a/Engine/Source/Programs/Shared/EpicGames.Horde/Compute/Clients/AgentComputeClient.cs
+++ b/Engine/Source/Programs/Shared/EpicGames.Horde/Compute/Clients/AgentComputeClient.cs
@@ -1,6 +1,5 @@
 // Copyright Epic Games, Inc. All Rights Reserved.
 
-using System;
 using System.Collections.Generic;
 using System.Diagnostics;
 using System.Net;
diff --git a/Engine/Source/Programs/Shared/EpicGames.Horde/Compute/Clients/ServerComputeClient.cs b/Engine/Source/Programs/Shared/EpicGames.Horde/Compute/Clients/ServerComputeClient.cs
--- a/Engine/Source/Programs/Shared/EpicGames.Horde/Compute/Clients/ServerComputeClient.cs
+++ b/Engine/Source/Programs/Shared/EpicGames.Horde/Compute/Clients/ServerComputeClient.cs
@@ -5,7 +5,6 @@ using EpicGames.Horde.Api;
 using EpicGames.Horde.Compute.Transports;
 using Microsoft.Extensions.Logging;
 using System;
-using System.Buffers;
 using System.Collections.Generic;
 using System.Net;
 using System.Net.Http;
diff --git a/Engine/Source/Programs/Shared/EpicGames.Horde/Compute/ComputeBuffer.cs b/Engine/Source/Programs/Shared/EpicGames.Horde/Compute/ComputeBuffer.cs
--- a/Engine/Source/Programs/Shared/EpicGames.Horde/Compute/ComputeBuffer.cs
+++ b/Engine/Source/Programs/Shared/EpicGames.Horde/Compute/ComputeBuffer.cs
@@ -2,7 +2,6 @@
 
 using System;
 using System.Diagnostics;
-using System.Linq;
 using System.Threading;
 using System.Threading.Tasks;
 
diff --git a/Engine/Source/Programs/Shared/EpicGames.Horde/Compute/Native.cs b/Engine/Source/Programs/Shared/EpicGames.Horde/Compute/Native.cs
--- a/Engine/Source/Programs/Shared/EpicGames.Horde/Compute/Native.cs
+++ b/Engine/Source/Programs/Shared/EpicGames.Horde/Compute/Native.cs
@@ -3,7 +3,6 @@
 using System;
 using System.ComponentModel;
 using System.IO;
-using System.Net.Http.Headers;
 using System.Runtime.InteropServices;
 using System.Threading;
 using Microsoft.Win32.SafeHandles;
diff --git a/Engine/Source/Programs/Shared/EpicGames.Horde/EpicGames.Horde.csproj b/Engine/Source/Programs/Shared/EpicGames.Horde/EpicGames.Horde.csproj
--- a/Engine/Source/Programs/Shared/EpicGames.Horde/EpicGames.Horde.csproj
+++ b/Engine/Source/Programs/Shared/EpicGames.Horde/EpicGames.Horde.csproj
@@ -28,6 +28,7 @@
   </PropertyGroup>
   <PropertyGroup Condition=" '$(Configuration)|$(Platform)' == 'Analyze|AnyCPU' ">
     <DocumentationFile></DocumentationFile>
+    <EnforceCodeStyleInBuild>True</EnforceCodeStyleInBuild>
     <GenerateDocumentationFile>True</GenerateDocumentationFile>
     <RunAnalyzersDuringBuild>True</RunAnalyzersDuringBuild>
   </PropertyGroup>
diff --git a/Engine/Source/Programs/Shared/EpicGames.Horde/HordeHttpClient.cs b/Engine/Source/Programs/Shared/EpicGames.Horde/HordeHttpClient.cs
--- a/Engine/Source/Programs/Shared/EpicGames.Horde/HordeHttpClient.cs
+++ b/Engine/Source/Programs/Shared/EpicGames.Horde/HordeHttpClient.cs
@@ -105,16 +105,6 @@ namespace EpicGames.Horde
 			options.Converters.Add(new BinaryIdJsonConverterFactory());
 		}
 
-		static HttpMessageHandler CreateDefaultHttpMessageHandler()
-		{
-			AsyncRetryPolicy<HttpResponseMessage> retryPolicy = HttpPolicyExtensions
-				.HandleTransientHttpError()
-				.WaitAndRetryAsync(new[] { TimeSpan.FromSeconds(2.0), TimeSpan.FromSeconds(5.0), TimeSpan.FromSeconds(10), TimeSpan.FromSeconds(30) });
-
-			SocketsHttpHandler socketsHandler = new SocketsHttpHandler { PooledConnectionLifetime = TimeSpan.FromMinutes(15) };
-			return new PolicyHttpMessageHandler(retryPolicy) { InnerHandler = socketsHandler };
-		}
-
 		/// <summary>
 		/// Creates a storage client for this Horde instance
 		/// </summary>
diff --git a/Engine/Source/Programs/Shared/EpicGames.Horde/Logs/LogChunkNode.cs b/Engine/Source/Programs/Shared/EpicGames.Horde/Logs/LogChunkNode.cs
--- a/Engine/Source/Programs/Shared/EpicGames.Horde/Logs/LogChunkNode.cs
+++ b/Engine/Source/Programs/Shared/EpicGames.Horde/Logs/LogChunkNode.cs
@@ -3,8 +3,6 @@
 using System;
 using System.Collections;
 using System.Collections.Generic;
-using System.IO;
-using System.Linq;
 using System.Text.Json;
 using System.Text;
 using EpicGames.Core;
diff --git a/Engine/Source/Programs/Shared/EpicGames.Horde/Logs/LogNode.cs b/Engine/Source/Programs/Shared/EpicGames.Horde/Logs/LogNode.cs
--- a/Engine/Source/Programs/Shared/EpicGames.Horde/Logs/LogNode.cs
+++ b/Engine/Source/Programs/Shared/EpicGames.Horde/Logs/LogNode.cs
@@ -2,14 +2,11 @@
 
 using System;
 using System.Collections.Generic;
-using System.Data.Common;
 using System.Linq;
-using System.Reflection;
 using System.Threading.Tasks;
 using System.Threading;
 using EpicGames.Core;
 using EpicGames.Horde.Storage;
-using EpicGames.Serialization;
 using Microsoft.Extensions.Logging;
 using System.Diagnostics;
 using System.Runtime.CompilerServices;
diff --git a/Engine/Source/Programs/Shared/EpicGames.Horde/Storage/BlobType.cs b/Engine/Source/Programs/Shared/EpicGames.Horde/Storage/BlobType.cs
--- a/Engine/Source/Programs/Shared/EpicGames.Horde/Storage/BlobType.cs
+++ b/Engine/Source/Programs/Shared/EpicGames.Horde/Storage/BlobType.cs
@@ -2,7 +2,6 @@
 
 using System;
 using System.Buffers.Binary;
-using EpicGames.Core;
 
 namespace EpicGames.Horde.Storage
 {
diff --git a/Engine/Source/Programs/Shared/EpicGames.Horde/Storage/Bundles/BundleNodeHandle.cs b/Engine/Source/Programs/Shared/EpicGames.Horde/Storage/Bundles/BundleNodeHandle.cs
--- a/Engine/Source/Programs/Shared/EpicGames.Horde/Storage/Bundles/BundleNodeHandle.cs
+++ b/Engine/Source/Programs/Shared/EpicGames.Horde/Storage/Bundles/BundleNodeHandle.cs
@@ -1,9 +1,5 @@
 // Copyright Epic Games, Inc. All Rights Reserved.
 
-using System;
-using System.Collections.Generic;
-using System.Linq;
-using System.Text;
 using System.Threading;
 using System.Threading.Tasks;
 using EpicGames.Core;
diff --git a/Engine/Source/Programs/Shared/EpicGames.Horde/Storage/Clients/FileStorageClient.cs b/Engine/Source/Programs/Shared/EpicGames.Horde/Storage/Clients/FileStorageClient.cs
--- a/Engine/Source/Programs/Shared/EpicGames.Horde/Storage/Clients/FileStorageClient.cs
+++ b/Engine/Source/Programs/Shared/EpicGames.Horde/Storage/Clients/FileStorageClient.cs
@@ -1,18 +1,12 @@
 // Copyright Epic Games, Inc. All Rights Reserved.
 
 using System;
-using System.Buffers;
 using System.Collections.Generic;
-using System.Diagnostics;
-using System.Globalization;
 using System.IO;
-using System.Linq;
-using System.Runtime.Serialization;
 using System.Threading;
 using System.Threading.Tasks;
 using EpicGames.Core;
 using EpicGames.Horde.Storage.Backends;
-using Microsoft.Extensions.Caching.Memory;
 using Microsoft.Extensions.Logging;
 using EpicGames.Horde.Storage.Bundles;
 
diff --git a/Engine/Source/Programs/Shared/EpicGames.Horde/Storage/Clients/MemoryStorageClient.cs b/Engine/Source/Programs/Shared/EpicGames.Horde/Storage/Clients/MemoryStorageClient.cs
--- a/Engine/Source/Programs/Shared/EpicGames.Horde/Storage/Clients/MemoryStorageClient.cs
+++ b/Engine/Source/Programs/Shared/EpicGames.Horde/Storage/Clients/MemoryStorageClient.cs
@@ -1,16 +1,12 @@
 // Copyright Epic Games, Inc. All Rights Reserved.
 
 using System;
-using System.Buffers;
 using System.Collections.Concurrent;
 using System.Collections.Generic;
-using System.IO;
 using System.Runtime.CompilerServices;
 using System.Threading;
 using System.Threading.Tasks;
 using EpicGames.Core;
-using Microsoft.Extensions.Caching.Memory;
-using Microsoft.Extensions.Logging;
 using Microsoft.Extensions.Logging.Abstractions;
 using EpicGames.Horde.Storage.Bundles;
 using EpicGames.Horde.Storage.Backends;
diff --git a/Engine/Source/Programs/Shared/EpicGames.Horde/Storage/Git/GitCommit.cs b/Engine/Source/Programs/Shared/EpicGames.Horde/Storage/Git/GitCommit.cs
--- a/Engine/Source/Programs/Shared/EpicGames.Horde/Storage/Git/GitCommit.cs
+++ b/Engine/Source/Programs/Shared/EpicGames.Horde/Storage/Git/GitCommit.cs
@@ -4,7 +4,6 @@ using System;
 using System.Collections.Generic;
 using System.Diagnostics;
 using System.Linq;
-using System.Security.Cryptography;
 using EpicGames.Core;
 
 namespace EpicGames.Horde.Storage.Git
diff --git a/Engine/Source/Programs/Shared/EpicGames.Horde/Storage/IStorageClient.cs b/Engine/Source/Programs/Shared/EpicGames.Horde/Storage/IStorageClient.cs
--- a/Engine/Source/Programs/Shared/EpicGames.Horde/Storage/IStorageClient.cs
+++ b/Engine/Source/Programs/Shared/EpicGames.Horde/Storage/IStorageClient.cs
@@ -2,7 +2,6 @@
 
 using System;
 using System.Collections.Generic;
-using System.IO;
 using System.Threading;
 using System.Threading.Tasks;
 using EpicGames.Core;
diff --git a/Engine/Source/Programs/Shared/EpicGames.Horde/Storage/Nodes/CbNode.cs b/Engine/Source/Programs/Shared/EpicGames.Horde/Storage/Nodes/CbNode.cs
--- a/Engine/Source/Programs/Shared/EpicGames.Horde/Storage/Nodes/CbNode.cs
+++ b/Engine/Source/Programs/Shared/EpicGames.Horde/Storage/Nodes/CbNode.cs
@@ -2,10 +2,7 @@
 
 using EpicGames.Core;
 using EpicGames.Serialization;
-using System;
 using System.Collections.Generic;
-using System.Linq;
-using System.Reflection.Metadata;
 
 namespace EpicGames.Horde.Storage.Nodes
 {
diff --git a/Engine/Source/Programs/Shared/EpicGames.Horde/Storage/Nodes/ChunkedDataNode.cs b/Engine/Source/Programs/Shared/EpicGames.Horde/Storage/Nodes/ChunkedDataNode.cs
--- a/Engine/Source/Programs/Shared/EpicGames.Horde/Storage/Nodes/ChunkedDataNode.cs
+++ b/Engine/Source/Programs/Shared/EpicGames.Horde/Storage/Nodes/ChunkedDataNode.cs
@@ -6,7 +6,6 @@ using System.Buffers.Binary;
 using System.Collections.Generic;
 using System.Diagnostics;
 using System.IO;
-using System.Linq;
 using System.Threading;
 using System.Threading.Tasks;
 using EpicGames.Core;
diff --git a/Engine/Source/Programs/Shared/EpicGames.Horde/Storage/Nodes/ChunkedDataWriter.cs b/Engine/Source/Programs/Shared/EpicGames.Horde/Storage/Nodes/ChunkedDataWriter.cs
--- a/Engine/Source/Programs/Shared/EpicGames.Horde/Storage/Nodes/ChunkedDataWriter.cs
+++ b/Engine/Source/Programs/Shared/EpicGames.Horde/Storage/Nodes/ChunkedDataWriter.cs
@@ -1,7 +1,6 @@
 // Copyright Epic Games, Inc. All Rights Reserved.
 
 using System;
-using System.Buffers.Binary;
 using System.Collections.Generic;
 using System.Diagnostics;
 using System.IO;
diff --git a/Engine/Source/Programs/Shared/EpicGames.Horde/Storage/Workspace.cs b/Engine/Source/Programs/Shared/EpicGames.Horde/Storage/Workspace.cs
--- a/Engine/Source/Programs/Shared/EpicGames.Horde/Storage/Workspace.cs
+++ b/Engine/Source/Programs/Shared/EpicGames.Horde/Storage/Workspace.cs
@@ -2,7 +2,6 @@
 
 using System;
 using System.Collections.Generic;
-using System.ComponentModel;
 using System.Diagnostics;
 using System.Diagnostics.CodeAnalysis;
 using System.IO;
