Horde: Fix remaining style analyzer messages in EpicGames.Core and EpicGames.Horde.

diff --git a/Engine/Source/Programs/Shared/.editorconfig b/Engine/Source/Programs/Shared/.editorconfig
--- a/Engine/Source/Programs/Shared/.editorconfig
+++ b/Engine/Source/Programs/Shared/.editorconfig
@@ -8,6 +8,7 @@
 csharp_style_prefer_switch_expression = false
 csharp_style_implicit_object_creation_when_type_is_apparent = false
 dotnet_remove_unnecessary_suppression_exclusions = false
+dotnet_diagnostic.IDE0079.severity = silent
 
 # Sort using and Import directives with System.* appearing first
 dotnet_sort_system_directives_first = true
diff --git a/Engine/Source/Programs/Shared/EpicGames.Core/FileReference.cs b/Engine/Source/Programs/Shared/EpicGames.Core/FileReference.cs
--- a/Engine/Source/Programs/Shared/EpicGames.Core/FileReference.cs
+++ b/Engine/Source/Programs/Shared/EpicGames.Core/FileReference.cs
@@ -248,16 +248,16 @@ namespace EpicGames.Core
 		public int CompareTo(FileReference? other) => Comparer.Compare(FullName, other?.FullName);
 
 		/// <inheritdoc/>
-		public static bool operator <(FileReference left, FileReference right) => ReferenceEquals(left, null) ? !ReferenceEquals(right, null) : left.CompareTo(right) < 0;
+		public static bool operator <(FileReference left, FileReference right) => left is null ? right is not null : left.CompareTo(right) < 0;
 
 		/// <inheritdoc/>
-		public static bool operator <=(FileReference left, FileReference right) => ReferenceEquals(left, null) || left.CompareTo(right) <= 0;
+		public static bool operator <=(FileReference left, FileReference right) => left is null || left.CompareTo(right) <= 0;
 
 		/// <inheritdoc/>
-		public static bool operator >(FileReference left, FileReference right) => !ReferenceEquals(left, null) && left.CompareTo(right) > 0;
+		public static bool operator >(FileReference left, FileReference right) => left is not null && left.CompareTo(right) > 0;
 
 		/// <inheritdoc/>
-		public static bool operator >=(FileReference left, FileReference right) => ReferenceEquals(left, null) ? ReferenceEquals(right, null) : left.CompareTo(right) >= 0;
+		public static bool operator >=(FileReference left, FileReference right) => left is null ? right is null : left.CompareTo(right) >= 0;
 
 		/// <summary>
 		/// Helper function to create a remote file reference. Unlike normal FileReference objects, these aren't converted to a full path in the local filesystem, but are
@@ -464,20 +464,19 @@ namespace EpicGames.Core
 			{
 				using (StreamReader sr = new StreamReader(fs, Encoding.UTF8, true))
 				{
-
 					// Try to read the whole file into a buffer created by hand.  This avoids a LOT of memory allocations which in turn reduces the
 					// GC stress on the system.  Removing the StreamReader would be nice in the future.
-					long RawFileLength = fs.Length;
-					char[] InitialBuffer = new char[RawFileLength];
-					int ReadLength = sr.Read(InitialBuffer, 0, (int)RawFileLength);
+					long rawFileLength = fs.Length;
+					char[] initialBuffer = new char[rawFileLength];
+					int readLength = sr.Read(initialBuffer, 0, (int)rawFileLength);
 					if (sr.EndOfStream)
 					{
-						return new String(InitialBuffer, 0, ReadLength);
+						return new string(initialBuffer, 0, readLength);
 					}
 					else
 					{
-						string Remaining = sr.ReadToEnd();
-						return String.Concat(new ReadOnlySpan<char>(InitialBuffer, 0, ReadLength), Remaining);
+						string remaining = sr.ReadToEnd();
+						return String.Concat(new ReadOnlySpan<char>(initialBuffer, 0, readLength), remaining);
 					}
 				}
 			}
diff --git a/Engine/Source/Programs/Shared/EpicGames.Horde/Api/ArtifactId.cs b/Engine/Source/Programs/Shared/EpicGames.Horde/Api/ArtifactId.cs
--- a/Engine/Source/Programs/Shared/EpicGames.Horde/Api/ArtifactId.cs
+++ b/Engine/Source/Programs/Shared/EpicGames.Horde/Api/ArtifactId.cs
@@ -16,7 +16,7 @@ namespace EpicGames.Horde.Api
 		public static ArtifactId Parse(string text) => new ArtifactId(BinaryId.Parse(text));
 
 		/// <inheritdoc/>
-		public override string ToString() => Id.ToString();
+		public override readonly string ToString() => Id.ToString();
 	}
 
 	/// <summary>
diff --git a/Engine/Source/Programs/Shared/EpicGames.Horde/Api/UserId.cs b/Engine/Source/Programs/Shared/EpicGames.Horde/Api/UserId.cs
--- a/Engine/Source/Programs/Shared/EpicGames.Horde/Api/UserId.cs
+++ b/Engine/Source/Programs/Shared/EpicGames.Horde/Api/UserId.cs
@@ -43,7 +43,7 @@ namespace EpicGames.Horde.Api
 		}
 
 		/// <inheritdoc/>
-		public override string ToString() => Id.ToString();
+		public override readonly string ToString() => Id.ToString();
 	}
 
 	/// <summary>
diff --git a/Engine/Source/Programs/Shared/EpicGames.Horde/Compute/AgentMessageChannel.cs b/Engine/Source/Programs/Shared/EpicGames.Horde/Compute/AgentMessageChannel.cs
--- a/Engine/Source/Programs/Shared/EpicGames.Horde/Compute/AgentMessageChannel.cs
+++ b/Engine/Source/Programs/Shared/EpicGames.Horde/Compute/AgentMessageChannel.cs
@@ -80,9 +80,7 @@ namespace EpicGames.Horde.Compute
 		// Can lock chunked memory writer to acquire pointer
 		readonly ILogger _logger;
 
-#pragma warning disable CA2213
 		MessageBuilder? _currentBuilder;
-#pragma warning restore CA2213
 
 		/// <summary>
 		/// Constructor
diff --git a/Engine/Source/Programs/Shared/EpicGames.Horde/Compute/ComputeBuffer.cs b/Engine/Source/Programs/Shared/EpicGames.Horde/Compute/ComputeBuffer.cs
--- a/Engine/Source/Programs/Shared/EpicGames.Horde/Compute/ComputeBuffer.cs
+++ b/Engine/Source/Programs/Shared/EpicGames.Horde/Compute/ComputeBuffer.cs
@@ -279,22 +279,22 @@ namespace EpicGames.Horde.Compute
 		protected internal record struct ChunkState(ulong Value)
 		{
 			// Written length of this chunk
-			public int Length => (int)(Value & 0x7fffffff);
+			public readonly int Length => (int)(Value & 0x7fffffff);
 
 			// Set of flags which are set for each reader that still has to read from a chunk
-			public int ReaderFlags => (int)((Value >> 31) & 0x7fffffff);
+			public readonly int ReaderFlags => (int)((Value >> 31) & 0x7fffffff);
 
 			// State of the writer
-			public WriteState WriteState => (WriteState)(Value >> 62);
+			public readonly WriteState WriteState => (WriteState)(Value >> 62);
 
 			// Constructor
 			public ChunkState(WriteState writerState, int readerFlags, int length) : this(((ulong)writerState << 62) | ((ulong)readerFlags << 31) | (uint)length) { }
 
 			// Test whether a particular reader is still referencing the chunk
-			public bool HasReaderFlag(int readerIdx) => (Value & (1UL << (31 + readerIdx))) != 0;
+			public readonly bool HasReaderFlag(int readerIdx) => (Value & (1UL << (31 + readerIdx))) != 0;
 
 			/// <inheritdoc/>
-			public override string ToString() => $"{WriteState}, Length: {Length}, Readers: {ReaderFlags}";
+			public override readonly string ToString() => $"{WriteState}, Length: {Length}, Readers: {ReaderFlags}";
 		}
 
 		/// <summary>
@@ -343,13 +343,13 @@ namespace EpicGames.Horde.Compute
 				: this((ulong)(uint)offset | ((ulong)(uint)chunkIdx << 32) | ((ulong)(uint)refCount << 40) | ((ulong)((detached? (1UL << 63) : 0))))
 			{ }
 
-			public int Offset => (int)(Value & 0xffffffff);
-			public int ChunkIdx => (int)((Value >> 32) & 0xff);
-			public int RefCount => (int)((Value >> 40) & 0x7fff);
-			public bool Detached => (Value & (1UL << 63)) != 0;
+			public readonly int Offset => (int)(Value & 0xffffffff);
+			public readonly int ChunkIdx => (int)((Value >> 32) & 0xff);
+			public readonly int RefCount => (int)((Value >> 40) & 0x7fff);
+			public readonly bool Detached => (Value & (1UL << 63)) != 0;
 
 			/// <inheritdoc/>
-			public override string ToString() => $"Chunk: {ChunkIdx}, Offset: {Offset}, RefCount: {RefCount}, Detached: {Detached}";
+			public override readonly string ToString() => $"Chunk: {ChunkIdx}, Offset: {Offset}, RefCount: {RefCount}, Detached: {Detached}";
 		}
 
 		/// <summary>
@@ -380,13 +380,13 @@ namespace EpicGames.Horde.Compute
 				: this((ulong)(uint)chunkIdx | ((ulong)(uint)readerFlags << 32) | ((ulong)(uint)refCount << 48) | (hasWrapped ? (1UL << 63) : 0))
 			{ }
 
-			public int ChunkIdx => (int)(Value & 0x7fffffff);
-			public int ReaderFlags => (int)(Value >> 32) & 0xffff;
-			public int RefCount => (int)(Value >> 48) & 0x7fff;
-			public bool HasWrapped => (Value & (1UL << 63)) != 0;
+			public readonly int ChunkIdx => (int)(Value & 0x7fffffff);
+			public readonly int ReaderFlags => (int)(Value >> 32) & 0xffff;
+			public readonly int RefCount => (int)(Value >> 48) & 0x7fff;
+			public readonly bool HasWrapped => (Value & (1UL << 63)) != 0;
 
 			/// <inheritdoc/>
-			public override string ToString() => $"Chunk: {ChunkIdx}, ReaderFlags: {ReaderFlags}, RefCount: {RefCount}, HasWrapped: {HasWrapped}";
+			public override readonly string ToString() => $"Chunk: {ChunkIdx}, ReaderFlags: {ReaderFlags}, RefCount: {RefCount}, HasWrapped: {HasWrapped}";
 		}
 
 		/// <summary>
@@ -402,7 +402,7 @@ namespace EpicGames.Horde.Compute
 			public WriterState Get() => new WriterState(Interlocked.CompareExchange(ref *_data, 0, 0));
 
 			// Set the current value
-			public void Set(WriterState State) => Interlocked.Exchange(ref *_data, State.Value);
+			public void Set(WriterState state) => Interlocked.Exchange(ref *_data, state.Value);
 
 			// Compare and swap
 			public bool TryUpdate(WriterState prevState, WriterState nextState) => Interlocked.CompareExchange(ref *_data, nextState.Value, prevState.Value) == prevState.Value;
@@ -573,9 +573,9 @@ namespace EpicGames.Horde.Compute
 							}
 							if (writerStatePtr.TryUpdate(writerState, new WriterState(writerState.ChunkIdx, writerState.ReaderFlags | (1 << readerIdx), writerState.RefCount, writerState.HasWrapped)))
 							{
-								for (int WriteChunkIdx = 0; WriteChunkIdx <= writerState.ChunkIdx; WriteChunkIdx++)
+								for (int writeChunkIdx = 0; writeChunkIdx <= writerState.ChunkIdx; writeChunkIdx++)
 								{
-									_headerPtr.GetChunkStatePtr(WriteChunkIdx).StartReading(readerIdx);
+									_headerPtr.GetChunkStatePtr(writeChunkIdx).StartReading(readerIdx);
 								}
 								return readerIdx;
 							}
@@ -616,9 +616,9 @@ namespace EpicGames.Horde.Compute
 
 				if (readerState.RefCount == 1)
 				{
-					for (int Idx = 0; Idx < _headerPtr.NumChunks; Idx++)
+					for (int idx = 0; idx < _headerPtr.NumChunks; idx++)
 					{
-						_headerPtr.GetChunkStatePtr(Idx).FinishReading(readerIdx);
+						_headerPtr.GetChunkStatePtr(idx).FinishReading(readerIdx);
 					}
 				}
 
diff --git a/Engine/Source/Programs/Shared/EpicGames.Horde/Compute/Native.cs b/Engine/Source/Programs/Shared/EpicGames.Horde/Compute/Native.cs
--- a/Engine/Source/Programs/Shared/EpicGames.Horde/Compute/Native.cs
+++ b/Engine/Source/Programs/Shared/EpicGames.Horde/Compute/Native.cs
@@ -11,6 +11,7 @@ namespace EpicGames.Horde.Compute
 {
 	internal class Native
 	{
+#pragma warning disable IDE1006 // Naming Styles
 		[StructLayout(LayoutKind.Sequential)]
 		public class SECURITY_ATTRIBUTES
 		{
@@ -18,6 +19,7 @@ namespace EpicGames.Horde.Compute
 			public IntPtr lpSecurityDescriptor;
 			public int bInheritHandle;
 		}
+#pragma warning restore IDE1006 // Naming Styles
 
 		[DllImport("kernel32.dll")]
 		public static extern SafeWaitHandle CreateEvent(SECURITY_ATTRIBUTES lpEventAttributes, bool bManualReset, bool bInitialState, string? lpName);
diff --git a/Engine/Source/Programs/Shared/EpicGames.Horde/HordeHttpClient.cs b/Engine/Source/Programs/Shared/EpicGames.Horde/HordeHttpClient.cs
--- a/Engine/Source/Programs/Shared/EpicGames.Horde/HordeHttpClient.cs
+++ b/Engine/Source/Programs/Shared/EpicGames.Horde/HordeHttpClient.cs
@@ -16,8 +16,6 @@ using Polly;
 using Polly.Extensions.Http;
 using Polly.Retry;
 
-#pragma warning disable CA2234 // Pass system uri objects instead of strings
-
 namespace EpicGames.Horde
 {
 	/// <summary>
diff --git a/Engine/Source/Programs/Shared/EpicGames.Horde/Storage/BlobType.cs b/Engine/Source/Programs/Shared/EpicGames.Horde/Storage/BlobType.cs
--- a/Engine/Source/Programs/Shared/EpicGames.Horde/Storage/BlobType.cs
+++ b/Engine/Source/Programs/Shared/EpicGames.Horde/Storage/BlobType.cs
@@ -31,13 +31,13 @@ namespace EpicGames.Horde.Storage
 		/// <summary>
 		/// Serialize to a byte span
 		/// </summary>
-		public void Write(Span<byte> data)
+		public readonly void Write(Span<byte> data)
 		{
 			Guid.TryWriteBytes(data.Slice(0, 16));
 			BinaryPrimitives.WriteInt32LittleEndian(data.Slice(16), Version);
 		}
 
 		/// <inheritdoc/>
-		public override string ToString() => $"{Guid}#{Version}";
+		public override readonly string ToString() => $"{Guid}#{Version}";
 	}
 }
diff --git a/Engine/Source/Programs/Shared/EpicGames.Horde/Storage/Clients/BundleStorageClient.cs b/Engine/Source/Programs/Shared/EpicGames.Horde/Storage/Clients/BundleStorageClient.cs
--- a/Engine/Source/Programs/Shared/EpicGames.Horde/Storage/Clients/BundleStorageClient.cs
+++ b/Engine/Source/Programs/Shared/EpicGames.Horde/Storage/Clients/BundleStorageClient.cs
@@ -125,11 +125,7 @@ namespace EpicGames.Horde.Storage.Clients
 		public async Task<BundleNodeHandle> ReadRefTargetAsync(RefName name, RefCacheTime cacheTime = default, CancellationToken cancellationToken = default)
 		{
 			BundleNodeHandle? refTarget = await TryReadRefTargetAsync(name, cacheTime, cancellationToken);
-			if (refTarget == null)
-			{
-				throw new RefNameNotFoundException(name);
-			}
-			return refTarget;
+			return refTarget ?? throw new RefNameNotFoundException(name);
 		}
 
 		/// <inheritdoc/>
diff --git a/Engine/Source/Programs/Shared/EpicGames.Horde/Storage/IStorageClient.cs b/Engine/Source/Programs/Shared/EpicGames.Horde/Storage/IStorageClient.cs
--- a/Engine/Source/Programs/Shared/EpicGames.Horde/Storage/IStorageClient.cs
+++ b/Engine/Source/Programs/Shared/EpicGames.Horde/Storage/IStorageClient.cs
@@ -210,7 +210,7 @@ namespace EpicGames.Horde.Storage
 		/// <summary>
 		/// Maximum age for a cached value to be returned
 		/// </summary>
-		public TimeSpan MaxAge => DateTime.UtcNow - Utc;
+		public readonly TimeSpan MaxAge => DateTime.UtcNow - Utc;
 
 		/// <summary>
 		/// Sets the earliest time at which the entry must have been valid
@@ -221,7 +221,7 @@ namespace EpicGames.Horde.Storage
 		/// <summary>
 		/// Tests whether this value is set
 		/// </summary>
-		public bool IsSet() => Utc != default;
+		public readonly bool IsSet() => Utc != default;
 
 		/// <summary>
 		/// Determines if this cache time deems a particular cache entry stale
@@ -273,11 +273,7 @@ namespace EpicGames.Horde.Storage
 		public static async Task<BlobHandle> ReadRefTargetAsync(this IStorageClient store, RefName name, RefCacheTime cacheTime = default, CancellationToken cancellationToken = default)
 		{
 			BlobHandle? refTarget = await store.TryReadRefTargetAsync(name, cacheTime, cancellationToken);
-			if (refTarget == null)
-			{
-				throw new RefNameNotFoundException(name);
-			}
-			return refTarget;
+			return refTarget ?? throw new RefNameNotFoundException(name);
 		}
 
 		#endregion
diff --git a/Engine/Source/Programs/Shared/EpicGames.Horde/Storage/Workspace.cs b/Engine/Source/Programs/Shared/EpicGames.Horde/Storage/Workspace.cs
--- a/Engine/Source/Programs/Shared/EpicGames.Horde/Storage/Workspace.cs
+++ b/Engine/Source/Programs/Shared/EpicGames.Horde/Storage/Workspace.cs
@@ -584,11 +584,7 @@ namespace EpicGames.Horde.Storage
 		/// <param name="cancellationToken">Cancellation token for the operation</param>
 		public async Task SyncAsync(WorkspaceLayerId layerId, DirectoryNode? contents, CancellationToken cancellationToken = default)
 		{
-			LayerState? layerState = GetLayerState(layerId);
-			if (layerState == null)
-			{
-				throw new InvalidOperationException($"Layer '{layerId}' does not exist");
-			}
+			LayerState layerState = GetLayerState(layerId) ?? throw new InvalidOperationException($"Layer '{layerId}' does not exist");
 
 			DirectoryState hordeDirState = _rootDirState.FindOrAddDirectory(HordeDirName);
 			hordeDirState.LayerFlags |= CacheLayerFlag;
