[EpicGames.MsBuild]
* Fixed rethrow warning

diff --git a/Engine/Source/Programs/Shared/EpicGames.MsBuild/CsProjBuilder.cs b/Engine/Source/Programs/Shared/EpicGames.MsBuild/CsProjBuilder.cs
--- a/Engine/Source/Programs/Shared/EpicGames.MsBuild/CsProjBuilder.cs
+++ b/Engine/Source/Programs/Shared/EpicGames.MsBuild/CsProjBuilder.cs
@@ -246,7 +246,7 @@ namespace EpicGames.MsBuild
 									Logger.LogError("  {Path}", Path);
 								}
 							}
-							throw IPFEx;
+							throw;
 						}
 
 						if (!OperatingSystem.IsWindows())
