Refactored aspects of code generation to use more consistent patterns.

diff --git a/Engine/Source/Programs/Shared/EpicGames.UHT/Exporters/CodeGen/UhtHeaderCodeGenerator.cs b/Engine/Source/Programs/Shared/EpicGames.UHT/Exporters/CodeGen/UhtHeaderCodeGenerator.cs
--- a/Engine/Source/Programs/Shared/EpicGames.UHT/Exporters/CodeGen/UhtHeaderCodeGenerator.cs
+++ b/Engine/Source/Programs/Shared/EpicGames.UHT/Exporters/CodeGen/UhtHeaderCodeGenerator.cs
@@ -16,10 +16,10 @@ namespace EpicGames.UHT.Exporters.CodeGen
 		public const string CallbackWrappersMacroSuffix = "CALLBACK_WRAPPERS";
 		public const string SparseDataMacroSuffix = "SPARSE_DATA";
 		public const string SparseDataPropertyAccessorsMacroSuffix = "SPARSE_DATA_PROPERTY_ACCESSORS";
-		public const string EditorOnlySparseDataPropertyAccessorsMacroSuffix = "EDITOR_ONLY_SPARSE_DATA_PROPERTY_ACCESSORS";
-		public const string EditorOnlyRpcWrappersMacroSuffix = "EDITOR_ONLY_RPC_WRAPPERS";
+		public const string EditorOnlySparseDataPropertyAccessorsMacroSuffix = "SPARSE_DATA_PROPERTY_ACCESSORS_EOD";
+		public const string EditorOnlyRpcWrappersMacroSuffix = "RPC_WRAPPERS_EOD";
 		public const string RpcWrappersMacroSuffix = "RPC_WRAPPERS";
-		public const string EditorOnlyRpcWrappersNoPureDeclsMacroSuffix = "EDITOR_ONLY_RPC_WRAPPERS_NO_PURE_DECLS";
+		public const string EditorOnlyRpcWrappersNoPureDeclsMacroSuffix = "RPC_WRAPPERS_NO_PURE_DECLS_EOD";
 		public const string RpcWrappersNoPureDeclsMacroSuffix = "RPC_WRAPPERS_NO_PURE_DECLS";
 		public const string AccessorsMacroSuffix = "ACCESSORS";
 		public const string ArchiveSerializerMacroSuffix = "ARCHIVESERIALIZER";
@@ -36,6 +36,7 @@ namespace EpicGames.UHT.Exporters.CodeGen
 		public const string PrologMacroSuffix = "PROLOG";
 		public const string DelegateMacroSuffix = "DELEGATE";
 		public const string AutoGettersSettersMacroSuffix = "AUTOGETTERSETTER_DECLS";
+		public const string EditorOnlyAutoGettersSettersMacroSuffix = "AUTOGETTERSETTER_DECLS_EOD";
 		#endregion
 
 		public readonly UhtHeaderFile HeaderFile;
@@ -696,84 +697,6 @@ namespace EpicGames.UHT.Exporters.CodeGen
 			return builder;
 		}
 		#endregion
-
-		#region AutoGettersSetters support
-		protected static bool NeedAutoGetterSetterCodeGen(UhtClass classObj)
-		{
-			return classObj.ClassExportFlags.HasAnyFlags(UhtClassExportFlags.HasAutoGettersSetters);
-		}
-		
-		protected static void GetAutoGetterSetterStats(UhtClass classObj, out bool hasProperties, out bool hasEditorFields, out bool allEditorFields)
-		{
-			// Scan the children to see what we have
-			hasProperties = false;
-			hasEditorFields = false;
-			allEditorFields = true;
-			foreach (UhtType type in classObj.Children)
-			{
-				if (type is UhtProperty property)
-				{
-					if (property.PropertyExportFlags.HasAnyFlags(UhtPropertyExportFlags.GetterSpecifiedAuto | UhtPropertyExportFlags.SetterSpecifiedAuto))
-					{
-						hasProperties = true;
-						hasEditorFields |= property.IsEditorOnlyProperty;
-						allEditorFields &= property.IsEditorOnlyProperty;
-					}
-				}
-			}
-
-			// If we have no editor fields, then by definition, all fields can't be editor fields
-			allEditorFields &= hasEditorFields;
-		}
-		
-		protected static StringBuilder AppendAutoGettersSetters(StringBuilder builder, UhtClass classObj,
-			UhtPropertyExportFlags specifierFlag,
-			bool hasProperties, bool hasEditorFields, bool allEditorFields,
-			bool includeEditorOnlyFields, bool appendDefine, Action<StringBuilder, UhtClass, UhtProperty> appendAction)
-		{
-			if (hasProperties && !allEditorFields)
-			{
-				foreach (UhtType child in classObj.Children)
-				{
-					if (child is UhtProperty property)
-					{
-						if (property.PropertyExportFlags.HasAnyFlags(specifierFlag) && !property.IsEditorOnlyProperty)
-						{
-							appendAction(builder, classObj, property);
-						}
-					}
-				}
-			}
-
-			if (hasEditorFields && includeEditorOnlyFields)
-			{
-				if (!allEditorFields && appendDefine)
-				{
-					builder.Append("#if WITH_EDITORONLY_DATA\r\n");
-				}
-
-				if (hasProperties)
-				{
-					foreach (UhtType child in classObj.Children)
-					{
-						if (child is UhtProperty property)
-						{
-							if (property.PropertyExportFlags.HasAnyFlags(specifierFlag) && property.IsEditorOnlyProperty)
-							{
-								appendAction(builder, classObj, property);
-							}
-						}
-					}
-				}
-
-				if (!allEditorFields && appendDefine)
-				{
-					builder.Append("#endif // WITH_EDITORONLY_DATA\r\n");
-				}
-			}
-			return builder;
-		}
-		#endregion
 	}
 
 	internal static class UhtHaederCodeGeneratorStringBuilderExtensions
diff --git a/Engine/Source/Programs/Shared/EpicGames.UHT/Exporters/CodeGen/UhtHeaderCodeGeneratorCppFile.cs b/Engine/Source/Programs/Shared/EpicGames.UHT/Exporters/CodeGen/UhtHeaderCodeGeneratorCppFile.cs
--- a/Engine/Source/Programs/Shared/EpicGames.UHT/Exporters/CodeGen/UhtHeaderCodeGeneratorCppFile.cs
+++ b/Engine/Source/Programs/Shared/EpicGames.UHT/Exporters/CodeGen/UhtHeaderCodeGeneratorCppFile.cs
@@ -191,15 +191,8 @@ namespace EpicGames.UHT.Exporters.CodeGen
 
 					if (hasRegisteredEnums)
 					{
-						if (allEnumsEditorOnly)
-						{
-							builder.Append("#if WITH_EDITORONLY_DATA\r\n");
-						}
+						using UhtMacroBlockEmitter blockEmitter = new(builder, "WITH_EDITORONLY_DATA", allEnumsEditorOnly);
 						builder.Append("\t\tstatic const FEnumRegisterCompiledInInfo EnumInfo[];\r\n");
-						if (allEnumsEditorOnly)
-						{
-							builder.Append("#endif\r\n");
-						}
 					}
 					if (hasRegisteredScriptStructs)
 					{
@@ -216,19 +209,13 @@ namespace EpicGames.UHT.Exporters.CodeGen
 
 					if (hasRegisteredEnums)
 					{
-						if (allEnumsEditorOnly)
-						{
-							builder.Append("#if WITH_EDITORONLY_DATA\r\n");
-						}
+						using UhtMacroBlockEmitter blockEmitter = new(builder, "WITH_EDITORONLY_DATA", allEnumsEditorOnly);
 						builder.Append("\tconst FEnumRegisterCompiledInInfo ").Append(staticsName).Append("::EnumInfo[] = {\r\n");
 						foreach (UhtObject obj in HeaderFile.References.ExportTypes)
 						{
 							if (obj is UhtEnum enumObj)
 							{
-								if (!allEnumsEditorOnly && enumObj.IsEditorOnly)
-								{
-									builder.Append("#if WITH_EDITORONLY_DATA\r\n");
-								}
+								blockEmitter.Set(enumObj.IsEditorOnly);
 								uint hash = ObjectInfos[enumObj.ObjectTypeIndex].Hash;
 								builder
 									.Append("\t\t{ ")
@@ -238,18 +225,11 @@ namespace EpicGames.UHT.Exporters.CodeGen
 									.Append("\"), &Z_Registration_Info_UEnum_")
 									.Append(enumObj.EngineName)
 									.Append($", CONSTRUCT_RELOAD_VERSION_INFO(FEnumReloadVersionInfo, {hash}U) }},\r\n");
-								if (!allEnumsEditorOnly && enumObj.IsEditorOnly)
-								{
-									builder.Append("#endif\r\n");
-								}
 								combinedHash = HashCombine(combinedHash, hash);
 							}
 						}
+						blockEmitter.Set(allEnumsEditorOnly);
 						builder.Append("\t};\r\n");
-						if (allEnumsEditorOnly)
-						{
-							builder.Append("#endif\r\n");
-						}
 					}
 
 					if (hasRegisteredScriptStructs)
@@ -1041,32 +1021,19 @@ namespace EpicGames.UHT.Exporters.CodeGen
 
 			PropertyMemberContextImpl context = new(CodeGenerator, structObj, structSourceName, staticsName);
 
-			using (UhtMacroBlockEmitter emitter = new(builder, "WITH_EDITORONLY_DATA"))
+			bool hasAllEditorOnlyDataProperties = structObj.Properties.All(x => x.IsEditorOnlyProperty);
+
+			using (UhtMacroBlockEmitter emitter = new(builder, "WITH_EDITORONLY_DATA", hasAllEditorOnlyDataProperties))
 			{
-				bool hasAllEditorOnlyDataProperties = true;
-				foreach (UhtType type in structObj.Children)
+				foreach (UhtProperty property in structObj.Properties)
 				{
-					if (type is UhtProperty property)
-					{
-						emitter.Set(property.IsEditorOnlyProperty);
-						hasAllEditorOnlyDataProperties &= property.IsEditorOnlyProperty;
-						builder.AppendMemberDecl(property, context, property.EngineName, "", tabs);
-					}
+					emitter.Set(property.IsEditorOnlyProperty);
+					hasAllEditorOnlyDataProperties &= property.IsEditorOnlyProperty;
+					builder.AppendMemberDecl(property, context, property.EngineName, "", tabs);
 				}
 
-				// This will force it off if the last one was editor only but we has some that weren't
 				emitter.Set(hasAllEditorOnlyDataProperties);
-
 				builder.AppendTabs(tabs).Append("static const UECodeGen_Private::FPropertyParamsBase* const PropPointers[];\r\n");
-
-				foreach (UhtType type in structObj.Children)
-				{
-					if (type is UhtProperty property)
-					{
-						emitter.Set(property.IsEditorOnlyProperty);
-					}
-				}
-				emitter.Set(hasAllEditorOnlyDataProperties);
 			}
 			return builder;
 		}
@@ -1080,35 +1047,26 @@ namespace EpicGames.UHT.Exporters.CodeGen
 
 			PropertyMemberContextImpl context = new(CodeGenerator, structObj, structSourceName, staticsName);
 
-			using (UhtMacroBlockEmitter emitter = new(builder, "WITH_EDITORONLY_DATA"))
+			bool hasAllEditorOnlyDataProperties = structObj.Properties.All(x => x.IsEditorOnlyProperty);
+
+			using (UhtMacroBlockEmitter emitter = new(builder, "WITH_EDITORONLY_DATA", hasAllEditorOnlyDataProperties))
 			{
-				bool hasAllEditorOnlyDataProperties = true;
-				foreach (UhtType type in structObj.Children)
+				foreach (UhtProperty property in structObj.Properties)
 				{
-					if (type is UhtProperty property)
-					{
-						emitter.Set(property.IsEditorOnlyProperty);
-						hasAllEditorOnlyDataProperties &= property.IsEditorOnlyProperty;
-						builder.AppendMemberDef(property, context, property.EngineName, "", null, tabs);
-					}
+					emitter.Set(property.IsEditorOnlyProperty);
+					builder.AppendMemberDef(property, context, property.EngineName, "", null, tabs);
 				}
+			}
 
-				// This will force it off if the last one was editor only but we has some that weren't
-				emitter.Set(hasAllEditorOnlyDataProperties);
-
+			using (UhtMacroBlockEmitter emitter = new(builder, "WITH_EDITORONLY_DATA", hasAllEditorOnlyDataProperties))
+			{
 				builder.AppendTabs(tabs).Append("const UECodeGen_Private::FPropertyParamsBase* const ").Append(staticsName).Append("::PropPointers[] = {\r\n");
-				foreach (UhtType type in structObj.Children)
+				foreach (UhtProperty property in structObj.Properties)
 				{
-					if (type is UhtProperty property)
-					{
-						emitter.Set(property.IsEditorOnlyProperty);
-						builder.AppendMemberPtr(property, context, property.EngineName, "", tabs + 1);
-					}
+					emitter.Set(property.IsEditorOnlyProperty);
+					builder.AppendMemberPtr(property, context, property.EngineName, "", tabs + 1);
 				}
-
-				// This will force it off if the last one was editor only but we has some that weren't
 				emitter.Set(hasAllEditorOnlyDataProperties);
-
 				builder.AppendTabs(tabs).Append("};\r\n");
 			}
 			return builder;
@@ -1226,6 +1184,12 @@ namespace EpicGames.UHT.Exporters.CodeGen
 					builder.Append($"\tUE_FIELD_NOTIFICATION_IMPLEMENT_FIELD({classObj.SourceName}, {name})\r\n");
 				});
 
+			if (allEditorFields)
+			{
+				builder.Append("#endif // WITH_EDITORONLY_DATA\r\n");
+				builder.Append("#if WITH_EDITORONLY_DATA\r\n");
+			}
+
 			//UE_FIELD_NOTIFICATION_DECLARE_ENUM_FIELD
 			builder.Append($"\tUE_FIELD_NOTIFICATION_IMPLEMENTATION_BEGIN({classObj.SourceName})\r\n");
 			AppendFieldNotify(builder, classObj, hasProperties, hasFunctions, hasEditorFields, allEditorFields, true, true,
@@ -1244,66 +1208,48 @@ namespace EpicGames.UHT.Exporters.CodeGen
 
 		private static StringBuilder AppendAutoGettersSetters(StringBuilder builder, UhtClass classObj)
 		{
-			if (!NeedAutoGetterSetterCodeGen(classObj))
+			IEnumerable<UhtProperty> properties = classObj.Properties.Where(x => x.PropertyExportFlags.HasAnyFlags(UhtPropertyExportFlags.GetterSpecifiedAuto | UhtPropertyExportFlags.SetterSpecifiedAuto));
+			if (!properties.Any())
 			{
 				return builder;
 			}
-			
-			// Scan the children to see what we have
-			GetAutoGetterSetterStats(classObj, out bool hasProperties, out bool hasEditorFields, out bool allEditorFields);
-			
-			// If we only have editor fields or no editor fields, then we only emit one block
-			if (hasEditorFields)
-			{
-				builder.Append("#if WITH_EDITORONLY_DATA\r\n");
-				AppendAutoGettersSetters(builder, classObj, hasProperties, hasEditorFields, allEditorFields, true);
-				builder.Append("#else //WITH_EDITORONLY_DATA\r\n");
-				AppendAutoGettersSetters(builder, classObj, hasProperties, hasEditorFields, allEditorFields, false);
-				builder.Append("#endif // WITH_EDITORONLY_DATA\r\n");
-			}
-			else
+
+			using (UhtMacroBlockEmitter blockEmitter = new(builder, "WITH_EDITORONLY_DATA", properties.Any(x => x.IsEditorOnlyProperty)))
 			{
-				AppendAutoGettersSetters(builder, classObj, hasProperties, hasEditorFields, allEditorFields, false);
-			}
-			return builder;
-		}
-		
-		private static StringBuilder AppendAutoGettersSetters(StringBuilder builder, UhtClass classObj,
-			bool hasProperties, bool hasEditorFields, bool allEditorFields,
-			bool includeEditorOnlyFields)
-		{
-			AppendAutoGettersSetters(builder, classObj, UhtPropertyExportFlags.GetterSpecifiedAuto, hasProperties, hasEditorFields, allEditorFields,
-				includeEditorOnlyFields, false, (StringBuilder parentBuilder, UhtClass classObj, UhtProperty property) =>
-				{
-					using BorrowStringBuilder borrower = new(StringBuilderCache.Small);
-					StringBuilder propertyStringBuilder = borrower.StringBuilder;
-					propertyStringBuilder.AppendPropertyText(property, UhtPropertyTextType.GetterRetVal);
-					string getterRetText = propertyStringBuilder.ToString();
-					string getterCallText = property.Getter ?? "Get" + property.SourceName;
-					parentBuilder.Append('\t').Append(getterRetText).Append(classObj.SourceName).Append("::").Append(getterCallText).Append("() const\r\n");
-					parentBuilder.Append("\t{\r\n");
-					parentBuilder.Append('\t', 2).Append("return ").Append(property.SourceName).Append(";\r\n");
-					parentBuilder.Append("\t}\r\n");
-				});
-			
-			AppendAutoGettersSetters(builder, classObj, UhtPropertyExportFlags.SetterSpecifiedAuto, hasProperties, hasEditorFields, allEditorFields,
-				includeEditorOnlyFields, false, (StringBuilder parentBuilder, UhtClass classObj, UhtProperty property) =>
-				{
-					using BorrowStringBuilder borrower = new(StringBuilderCache.Small);
-					StringBuilder propertyStringBuilder = borrower.StringBuilder;
-					propertyStringBuilder.AppendPropertyText(property, UhtPropertyTextType.SetterParameterArgType);
-					string setterArgText = propertyStringBuilder + "InValue";
-					string setterCallText = property.Setter ?? "Set" + property.SourceName;
-					parentBuilder.Append("\tvoid ").Append(classObj.SourceName).Append("::").Append(setterCallText).Append('(').Append(setterArgText).Append(")\r\n");
-					parentBuilder.Append("\t{\r\n");
-					
-					// @todo: setter defn
-					
-					parentBuilder.Append("\t}\r\n");
-					
-					
-				});
+				foreach (UhtProperty property in classObj.Properties)
+				{
+					if (property.PropertyExportFlags.HasAnyFlags(UhtPropertyExportFlags.GetterSpecifiedAuto | UhtPropertyExportFlags.SetterSpecifiedAuto))
+					{
+						blockEmitter.Set(property.IsEditorOnlyProperty);
+						if (property.PropertyExportFlags.HasAnyFlags(UhtPropertyExportFlags.GetterSpecifiedAuto))
+						{
+							using BorrowStringBuilder borrower = new(StringBuilderCache.Small);
+							StringBuilder propertyStringBuilder = borrower.StringBuilder;
+							propertyStringBuilder.AppendPropertyText(property, UhtPropertyTextType.GetterRetVal);
+							string getterRetText = propertyStringBuilder.ToString();
+							string getterCallText = property.Getter ?? "Get" + property.SourceName;
+							builder.Append('\t').Append(getterRetText).Append(classObj.SourceName).Append("::").Append(getterCallText).Append("() const\r\n");
+							builder.Append("\t{\r\n");
+							builder.Append('\t', 2).Append("return ").Append(property.SourceName).Append(";\r\n");
+							builder.Append("\t}\r\n");
+						}
+						if (property.PropertyExportFlags.HasAnyFlags(UhtPropertyExportFlags.SetterSpecifiedAuto))
+						{
+							using BorrowStringBuilder borrower = new(StringBuilderCache.Small);
+							StringBuilder propertyStringBuilder = borrower.StringBuilder;
+							propertyStringBuilder.AppendPropertyText(property, UhtPropertyTextType.SetterParameterArgType);
+							string setterArgText = propertyStringBuilder + "InValue";
+							string setterCallText = property.Setter ?? "Set" + property.SourceName;
+							builder.Append("\tvoid ").Append(classObj.SourceName).Append("::").Append(setterCallText).Append('(').Append(setterArgText).Append(")\r\n");
+							builder.Append("\t{\r\n");
+
+							// @todo: setter defn
 
+							builder.Append("\t}\r\n");
+						}
+					}
+				}
+			}
 			return builder;
 		}
 
@@ -1320,7 +1266,7 @@ namespace EpicGames.UHT.Exporters.CodeGen
 						builder.Append("\t{\r\n");
 						if (editorOnlyProperty)
 						{
-							builder.Append("\t#if WITH_EDITORONLY_DATA\r\n");
+							builder.Append("#if WITH_EDITORONLY_DATA\r\n");
 						}
 						builder.Append("\t\tconst ").Append(classObj.SourceName).Append("* Obj = (const ").Append(classObj.SourceName).Append("*)Object;\r\n");
 						if (property.IsStaticArray)
@@ -1379,7 +1325,7 @@ namespace EpicGames.UHT.Exporters.CodeGen
 						}
 						if (editorOnlyProperty)
 						{
-							builder.Append("\t#endif // WITH_EDITORONLY_DATA\r\n");
+							builder.Append("#endif // WITH_EDITORONLY_DATA\r\n");
 						}
 						builder.Append("\t}\r\n");
 					}
@@ -1389,7 +1335,7 @@ namespace EpicGames.UHT.Exporters.CodeGen
 						builder.Append("\t{\r\n");
 						if (editorOnlyProperty)
 						{
-							builder.Append("\t#if WITH_EDITORONLY_DATA\r\n");
+							builder.Append("#if WITH_EDITORONLY_DATA\r\n");
 						}
 						builder.Append("\t\t").Append(classObj.SourceName).Append("* Obj = (").Append(classObj.SourceName).Append("*)Object;\r\n");
 						if (property.IsStaticArray)
@@ -1428,7 +1374,7 @@ namespace EpicGames.UHT.Exporters.CodeGen
 							.Append("(Value);\r\n");
 						if (editorOnlyProperty)
 						{
-							builder.Append("\t#endif // WITH_EDITORONLY_DATA\r\n");
+							builder.Append("#endif // WITH_EDITORONLY_DATA\r\n");
 						}
 						builder.Append("\t}\r\n");
 					}
@@ -1581,7 +1527,7 @@ namespace EpicGames.UHT.Exporters.CodeGen
 			List<UhtFunction> sortedFunctions = new();
 			foreach (UhtFunction function in classObj.Functions)
 			{
-				if (!function.FunctionFlags.HasAnyFlags(EFunctionFlags.Delegate))
+				//if (!function.FunctionFlags.HasAnyFlags(EFunctionFlags.Delegate))
 				{
 					allEditorOnlyFunctions &= function.FunctionFlags.HasAnyFlags(EFunctionFlags.EditorOnly);
 				}
@@ -1654,13 +1600,12 @@ namespace EpicGames.UHT.Exporters.CodeGen
 
 				if (sortedFunctions.Count > 0)
 				{
-					builder.AppendBeginEditorOnlyGuard(allEditorOnlyFunctions);
+					using UhtMacroBlockEmitter blockEmitter = new(builder, "WITH_EDITOR", allEditorOnlyFunctions);
 					builder.Append("\tconst FClassFunctionLinkInfo ").Append(staticsName).Append("::FuncInfo[] = {\r\n");
 
 					foreach (UhtFunction function in sortedFunctions)
 					{
-						bool isEditorOnlyFunction = function.FunctionFlags.HasAnyFlags(EFunctionFlags.EditorOnly);
-						builder.AppendBeginEditorOnlyGuard(isEditorOnlyFunction);
+						blockEmitter.Set(function.FunctionFlags.HasAnyFlags(EFunctionFlags.EditorOnly));
 						builder
 							.Append("\t\t{ &")
 							.Append(GetSingletonName(function, true))
@@ -1669,12 +1614,11 @@ namespace EpicGames.UHT.Exporters.CodeGen
 							.Append(" },")
 							.AppendObjectHash(classObj, context, function)
 							.Append("\r\n");
-						builder.AppendEndEditorOnlyGuard(isEditorOnlyFunction);
 					}
 
+					blockEmitter.Set(allEditorOnlyFunctions);
 					builder.Append("\t};\r\n");
 					builder.Append("\tstatic_assert(UE_ARRAY_COUNT(").Append(staticsName).Append("::FuncInfo) < 2048);\r\n");
-					builder.AppendEndEditorOnlyGuard(allEditorOnlyFunctions);
 				}
 
 				builder.AppendMetaDataDef(classObj, staticsName, MetaDataParamsName, 1);
@@ -2288,7 +2232,7 @@ namespace EpicGames.UHT.Exporters.CodeGen
 			{
 				builder.Append("#if WITH_EDITORONLY_DATA\r\n")
 					   .Append("\tstatic_assert(UE_ARRAY_COUNT(").Append(staticsName).Append("::PropPointers) < 2048);\r\n")
-					   .Append("#endif\r\n");
+					   .Append("#endif // WITH_EDITORONLY_DATA\r\n");
 			}
 			return builder;
 		}
diff --git a/Engine/Source/Programs/Shared/EpicGames.UHT/Exporters/CodeGen/UhtHeaderCodeGeneratorHFile.cs b/Engine/Source/Programs/Shared/EpicGames.UHT/Exporters/CodeGen/UhtHeaderCodeGeneratorHFile.cs
--- a/Engine/Source/Programs/Shared/EpicGames.UHT/Exporters/CodeGen/UhtHeaderCodeGeneratorHFile.cs
+++ b/Engine/Source/Programs/Shared/EpicGames.UHT/Exporters/CodeGen/UhtHeaderCodeGeneratorHFile.cs
@@ -521,8 +521,12 @@ namespace EpicGames.UHT.Exporters.CodeGen
 			// With interfaces, there are cases where we use the native interface class export flags to check for legacy macro
 			bool alternateUsesLegacy = (classObj.ClassFlags.HasAnyFlags(EClassFlags.Interface) ? nativeInterface! : classObj).ClassExportFlags.HasAnyFlags(UhtClassExportFlags.UsesGeneratedBodyLegacy);
 
+			// Collect sparse data
+			IEnumerable<UhtScriptStruct> sparseScriptStructs = GetSparseDataStructsToExport(classObj);
+			CalculateSparsePropertyDeclarations(sparseScriptStructs, out Dictionary<UhtScriptStruct, List<UhtProperty>> commonSparseProperties, out Dictionary<UhtScriptStruct, List<UhtProperty>> editorOnlySparseProperties);
+
 			// Write the spare declarations
-			AppendSparseDeclarations(builder, classObj);
+			AppendSparseDeclarations(builder, classObj, sparseScriptStructs, commonSparseProperties, editorOnlySparseProperties);
 
 			// Collect the functions in reversed order
 			List<UhtFunction> reversedFunctions = new(classObj.Functions);
@@ -548,6 +552,7 @@ namespace EpicGames.UHT.Exporters.CodeGen
 			
 			// Determine if auto getters/setters have been generated
 			bool hasAutoGettersSetters = classObj.Properties.Any(x => x.PropertyExportFlags.HasAnyFlags(UhtPropertyExportFlags.GetterSpecifiedAuto | UhtPropertyExportFlags.SetterSpecifiedAuto));
+			bool hasEditorOnlyAutoGettersSetters = classObj.Properties.Any(x => x.IsEditorOnlyProperty && x.PropertyExportFlags.HasAnyFlags(UhtPropertyExportFlags.GetterSpecifiedAuto | UhtPropertyExportFlags.SetterSpecifiedAuto));
 
 			// Generate the RPC wrappers for the callbacks
 			AppendCallbackRpcWrapperDecls(builder, classObj, callbackFunctions);
@@ -608,7 +613,8 @@ namespace EpicGames.UHT.Exporters.CodeGen
 				}
 
 				AppendProlog(builder, classObj);
-				AppendGeneratedBodyMacroBlock(builder, classObj, nativeInterface!, alternateUsesLegacy, hasEditorRpc, hasCallbacks, hasAutoGettersSetters, null);
+				AppendGeneratedBodyMacroBlock(builder, classObj, nativeInterface!, alternateUsesLegacy, hasEditorRpc, hasCallbacks, hasAutoGettersSetters, 
+					hasEditorOnlyAutoGettersSetters, sparseScriptStructs.Any(), editorOnlySparseProperties.Count > 0, null);
 			}
 			else
 			{
@@ -632,7 +638,8 @@ namespace EpicGames.UHT.Exporters.CodeGen
 				AppendFieldNotify(builder, classObj);
 				AppendAutoGettersSetters(builder, classObj);
 				AppendProlog(builder, classObj);
-				AppendGeneratedBodyMacroBlock(builder, classObj, classObj, usesLegacy, hasEditorRpc, hasCallbacks, hasAutoGettersSetters, usesLegacy ? "GENERATED_UCLASS_BODY" : null);
+				AppendGeneratedBodyMacroBlock(builder, classObj, classObj, usesLegacy, hasEditorRpc, hasCallbacks, hasAutoGettersSetters, hasEditorOnlyAutoGettersSetters,
+					sparseScriptStructs.Any(), editorOnlySparseProperties.Count > 0, usesLegacy ? "GENERATED_UCLASS_BODY" : null);
 			}
 
 			// Forward declare the StaticClass specialization in the header
@@ -671,13 +678,13 @@ namespace EpicGames.UHT.Exporters.CodeGen
 				{
 					AppendFieldNotify(builder, classObj, hasProperties, hasFunctions, hasEditorFields, allEditorFields, true);
 				}
-				if (!allEditorFields)
+				//if (!allEditorFields)
 				{
-					builder.Append("#else //WITH_EDITORONLY_DATA\r\n");
+					builder.Append("#else\r\n");
 					using UhtMacroCreator macro = new(builder, this, classObj, FieldNotifyMacroSuffix);
 					AppendFieldNotify(builder, classObj, hasProperties, hasFunctions, hasEditorFields, allEditorFields, false);
 				}
-				builder.Append("#endif // WITH_EDITORONLY_DATA\r\n");
+				builder.Append("#endif\r\n");
 			}
 			else
 			{
@@ -723,77 +730,91 @@ namespace EpicGames.UHT.Exporters.CodeGen
 
 		private StringBuilder AppendAutoGettersSetters(StringBuilder builder, UhtClass classObj)
 		{
-			if (!NeedAutoGetterSetterCodeGen(classObj))
+			// Scan the children to see what we have
+			IEnumerable<UhtProperty> properties = classObj.Properties.Where(x => x.PropertyExportFlags.HasAnyFlags(UhtPropertyExportFlags.GetterSpecifiedAuto | UhtPropertyExportFlags.SetterSpecifiedAuto));
+			if (properties.Any())
 			{
-				return builder;
+				AppendAutoGettersSetters(builder, classObj, properties, false);
+				if (properties.Any(x => x.IsEditorOnlyProperty))
+				{
+					AppendAutoGettersSetters(builder, classObj, properties, true);
+				}
+			}
+			return builder;
+		}
+		
+		private StringBuilder AppendAutoGettersSetters(StringBuilder builder, UhtClass classObj, IEnumerable<UhtProperty> properties, bool isEditorOnly)
+		{
+			// We can skip writing the macros if there are no properties to declare, as the 'if' and 'else' would be the same
+			if (isEditorOnly)
+			{
+				builder.Append("#if WITH_EDITORONLY_DATA\r\n");
 			}
 
-			// Scan the children to see what we have
-			GetAutoGetterSetterStats(classObj, out bool hasProperties, out bool hasEditorFields, out bool allEditorFields);
+			using (UhtMacroCreator macro = new(builder, this, classObj, isEditorOnly ? EditorOnlyAutoGettersSettersMacroSuffix : AutoGettersSettersMacroSuffix))
+			{
+				foreach (UhtProperty property in properties)
+				{
+					if (property.IsEditorOnlyProperty == isEditorOnly)
+					{
+						if (property.PropertyExportFlags.HasAnyFlags(UhtPropertyExportFlags.GetterSpecifiedAuto))
+						{
+							using BorrowStringBuilder borrower = new(StringBuilderCache.Small);
+							StringBuilder propertyStringBuilder = borrower.StringBuilder;
+							propertyStringBuilder.AppendPropertyText(property, UhtPropertyTextType.GetterRetVal);
+							string getterRetText = propertyStringBuilder.ToString();
+							string getterCallText = property.Getter ?? "Get" + property.SourceName;
+							builder.Append('\t').Append(getterRetText).Append(getterCallText).Append("() const; \\\r\n");
+						}
+						if (property.PropertyExportFlags.HasAnyFlags(UhtPropertyExportFlags.SetterSpecifiedAuto))
+						{
+							using BorrowStringBuilder borrower = new(StringBuilderCache.Small);
+							StringBuilder propertyStringBuilder = borrower.StringBuilder;
+							propertyStringBuilder.AppendPropertyText(property, UhtPropertyTextType.SetterParameterArgType);
+							string setterArgText = propertyStringBuilder + "InValue";
+							string setterCallText = property.Setter ?? "Set" + property.SourceName;
+							builder.Append("\tvoid ").Append(setterCallText).Append('(').Append(setterArgText).Append("); \\\r\n");
+						}
+					}
+				}
+			}
 
-			// If we only have editor fields or no editor fields, then we only emit one block
-			if (hasEditorFields)
+			// We can skip writing the macros if there are no properties to declare, as the 'if' and 'else' would be the same
+			if (isEditorOnly)
 			{
-				builder.Append("#if WITH_EDITORONLY_DATA\r\n");
-				using (UhtMacroCreator macro = new(builder, this, classObj, AutoGettersSettersMacroSuffix))
+				// Trim the extra newlines added after the macro generator
+				if (builder.Length > 4 &&
+					builder[^4] == '\r' &&
+					builder[^3] == '\n' &&
+					builder[^2] == '\r' &&
+					builder[^1] == '\n')
 				{
-					AppendAutoGetterSetter(builder, classObj, hasProperties, hasEditorFields, allEditorFields, true);
+					builder.Length -= 4;
 				}
-				builder.Append("#else //WITH_EDITORONLY_DATA\r\n");
-				using (UhtMacroCreator macro = new(builder, this, classObj, AutoGettersSettersMacroSuffix))
+
+				builder.Append("#else // WITH_EDITORONLY_DATA\r\n");
+				using (UhtMacroCreator macro = new(builder, this, classObj, EditorOnlyAutoGettersSettersMacroSuffix))
 				{
-					AppendAutoGetterSetter(builder, classObj, hasProperties, hasEditorFields, allEditorFields, false);
+					// Empty macro when not compiling WITH_EDITORONLY_DATA
 				}
-				builder.Append("#endif // WITH_EDITORONLY_DATA\r\n");
-			}
-			else
-			{
-				using UhtMacroCreator macro = new(builder, this, classObj, AutoGettersSettersMacroSuffix);
-				AppendAutoGetterSetter(builder, classObj, hasProperties, hasEditorFields, allEditorFields, false);
+				builder.Append("#endif // WITH_EDITORONLY_DATA\r\n\r\n\r\n");
 			}
+
 			return builder;
 		}
 		
-		private static StringBuilder AppendAutoGetterSetter(StringBuilder builder, UhtClass classObj,
-			bool hasProperties, bool hasEditorFields, bool allEditorFields,
-			bool includeEditorOnlyFields)
+		private StringBuilder AppendSparseDeclarations(StringBuilder builder, UhtClass classObj, IEnumerable<UhtScriptStruct> sparseScriptStructs, 
+			Dictionary<UhtScriptStruct, List<UhtProperty>> commonProperties, Dictionary<UhtScriptStruct, List<UhtProperty>> editorOnlyProperties)
 		{
-			builder.Append("public: \\\r\n");
-            AppendAutoGettersSetters(builder, classObj, UhtPropertyExportFlags.GetterSpecifiedAuto, hasProperties, hasEditorFields, allEditorFields,
-				includeEditorOnlyFields, false, (StringBuilder parentBuilder, UhtClass classObj, UhtProperty property) =>
-				{
-					using BorrowStringBuilder borrower = new(StringBuilderCache.Small);
-					StringBuilder propertyStringBuilder = borrower.StringBuilder;
-					propertyStringBuilder.AppendPropertyText(property, UhtPropertyTextType.GetterRetVal);
-					string getterRetText = propertyStringBuilder.ToString();
-					string getterCallText = property.Getter ?? "Get" + property.SourceName;
-					parentBuilder.Append('\t').Append(getterRetText).Append(getterCallText).Append("() const; \\\r\n");
-				});
-            
-            AppendAutoGettersSetters(builder, classObj, UhtPropertyExportFlags.SetterSpecifiedAuto, hasProperties, hasEditorFields, allEditorFields,
-	            includeEditorOnlyFields, false, (StringBuilder parentBuilder, UhtClass classObj, UhtProperty property) =>
-	            {
-		            using BorrowStringBuilder borrower = new(StringBuilderCache.Small);
-		            StringBuilder propertyStringBuilder = borrower.StringBuilder;
-		            propertyStringBuilder.AppendPropertyText(property, UhtPropertyTextType.SetterParameterArgType);
-		            string setterArgText = propertyStringBuilder + "InValue";
-		            string setterCallText = property.Setter ?? "Set" + property.SourceName;
-		            parentBuilder.Append("\tvoid ").Append(setterCallText).Append('(').Append(setterArgText).Append("); \\\r\n");
-	            });
-
-            return builder;
-		}
-		
-		private StringBuilder AppendSparseDeclarations(StringBuilder builder, UhtClass classObj)
-		{
-			IEnumerable<UhtScriptStruct> sparseScriptStructs = GetSparseDataStructsToExport(classObj);
-
-			AppendSparseStructDeclarations(builder, classObj, sparseScriptStructs);
-
-			CalculateSparsePropertyDeclarations(sparseScriptStructs, out Dictionary<UhtScriptStruct, List<UhtProperty>> commonProperties, out Dictionary<UhtScriptStruct, List<UhtProperty>> editorOnlyProperties);
-			AppendSparsePropertyDeclarations(builder, classObj, sparseScriptStructs, commonProperties, /*isEditorOnly*/false);
-			AppendSparsePropertyDeclarations(builder, classObj, sparseScriptStructs, editorOnlyProperties, /*isEditorOnly*/true);
-			
+			if (sparseScriptStructs.Any())
+			{
+				AppendSparseStructDeclarations(builder, classObj, sparseScriptStructs);
+				AppendSparsePropertyDeclarations(builder, classObj, sparseScriptStructs, commonProperties, /*isEditorOnly*/false);
+				if (editorOnlyProperties.Count > 0)
+				{
+					AppendSparsePropertyDeclarations(builder, classObj, sparseScriptStructs, editorOnlyProperties, /*isEditorOnly*/true);
+				}
+			}
 			return builder;
 		}
 
@@ -871,12 +892,12 @@ namespace EpicGames.UHT.Exporters.CodeGen
 					builder.Length -= 4;
 				}
 
-				builder.Append("#else //WITH_EDITORONLY_DATA\r\n");
+				builder.Append("#else // WITH_EDITORONLY_DATA\r\n");
 				using (UhtMacroCreator macro = new(builder, this, classObj, EditorOnlySparseDataPropertyAccessorsMacroSuffix))
 				{
 					// Empty macro when not compiling WITH_EDITORONLY_DATA
 				}
-				builder.Append("#endif //WITH_EDITORONLY_DATA\r\n\r\n\r\n");
+				builder.Append("#endif // WITH_EDITORONLY_DATA\r\n\r\n\r\n");
 			}
 
 			return builder;
@@ -940,7 +961,7 @@ namespace EpicGames.UHT.Exporters.CodeGen
 
 			if (editorOnly)
 			{
-				builder.Append("#else\r\n");
+				builder.Append("#else // WITH_EDITOR\r\n");
 				if (usesLegacy)
 				{
 					using UhtMacroCreator macro = new(builder, this, classObj, editorOnly ? EditorOnlyRpcWrappersMacroSuffix : RpcWrappersMacroSuffix);
@@ -1610,7 +1631,8 @@ namespace EpicGames.UHT.Exporters.CodeGen
 			return builder;
 		}
 
-		private StringBuilder AppendGeneratedBodyMacroBlock(StringBuilder builder, UhtClass classObj, UhtClass bodyClassObj, bool isLegacy, bool hasEditorRpc, bool hasCallbacks, bool hasAutoGettersSetters, string? deprecatedMacroName)
+		private StringBuilder AppendGeneratedBodyMacroBlock(StringBuilder builder, UhtClass classObj, UhtClass bodyClassObj, bool isLegacy, bool hasEditorRpc, bool hasCallbacks, 
+			bool hasAutoGettersSetters, bool hasEditorOnlyAutoGettersSetters, bool hasSparseProperties, bool hasEditorOnlySparseProperties, string? deprecatedMacroName)
 		{
 			bool isInterface = classObj.ClassFlags.HasAnyFlags(EClassFlags.Interface);
 			using (UhtMacroCreator macro = new(builder, this, bodyClassObj, isLegacy ? GeneratedBodyLegacyMacroSuffix : GeneratedBodyMacroSuffix))
@@ -1621,9 +1643,15 @@ namespace EpicGames.UHT.Exporters.CodeGen
 				}
 				builder.Append(DisableDeprecationWarnings).Append(" \\\r\n");
 				builder.Append("public: \\\r\n");
-				builder.Append('\t').AppendMacroName(this, classObj, SparseDataMacroSuffix).Append(" \\\r\n");
-				builder.Append('\t').AppendMacroName(this, classObj, SparseDataPropertyAccessorsMacroSuffix).Append(" \\\r\n");
-				builder.Append('\t').AppendMacroName(this, classObj, EditorOnlySparseDataPropertyAccessorsMacroSuffix).Append(" \\\r\n");
+				if (hasSparseProperties)
+				{
+					builder.Append('\t').AppendMacroName(this, classObj, SparseDataMacroSuffix).Append(" \\\r\n");
+					builder.Append('\t').AppendMacroName(this, classObj, SparseDataPropertyAccessorsMacroSuffix).Append(" \\\r\n");
+					if (hasEditorOnlySparseProperties)
+					{
+						builder.Append('\t').AppendMacroName(this, classObj, EditorOnlySparseDataPropertyAccessorsMacroSuffix).Append(" \\\r\n");
+					}
+				}
 				if (isLegacy)
 				{
 					builder.Append('\t').AppendMacroName(this, classObj, RpcWrappersMacroSuffix).Append(" \\\r\n");
@@ -1647,6 +1675,10 @@ namespace EpicGames.UHT.Exporters.CodeGen
 				if (hasAutoGettersSetters)
 				{
 					builder.Append('\t').AppendMacroName(this, classObj, AutoGettersSettersMacroSuffix).Append(" \\\r\n");
+					if (hasEditorOnlyAutoGettersSetters)
+					{
+						builder.Append('\t').AppendMacroName(this, classObj, EditorOnlyAutoGettersSettersMacroSuffix).Append(" \\\r\n");
+					}
 				}
 				if (hasCallbacks)
 				{
diff --git a/Engine/Source/Programs/Shared/EpicGames.UHT/Types/UhtClass.cs b/Engine/Source/Programs/Shared/EpicGames.UHT/Types/UhtClass.cs
--- a/Engine/Source/Programs/Shared/EpicGames.UHT/Types/UhtClass.cs
+++ b/Engine/Source/Programs/Shared/EpicGames.UHT/Types/UhtClass.cs
@@ -94,11 +94,6 @@ namespace EpicGames.UHT.Types
 		/// legacy generated body.  If this flag is set, generate the legacy instead of the GENERATED_BODY macros.
 		/// </summary>
 		UsesGeneratedBodyLegacy = 1 << 12,
-		
-		/// <summary>
-		/// Class has one or more auto getters/setters
-		/// </summary>
-		HasAutoGettersSetters = 1 << 13,
 	}
 
 	/// <summary>
@@ -880,7 +875,6 @@ namespace EpicGames.UHT.Types
 									// if it's Auto, it will be generated so no need to look for a user implementation
 									if (property.PropertyExportFlags.HasAnyFlags(UhtPropertyExportFlags.GetterSpecifiedAuto))
 									{
-										ClassExportFlags |= UhtClassExportFlags.HasAutoGettersSetters;
 										property.PropertyExportFlags |= UhtPropertyExportFlags.GetterFound;
 										property.Getter = "Get" + property.SourceName;
 									}
@@ -898,7 +892,6 @@ namespace EpicGames.UHT.Types
 									// if it's Auto, it will be generated so no need to look for a user implementation
 									if (property.PropertyExportFlags.HasAnyFlags(UhtPropertyExportFlags.SetterSpecifiedAuto))
 									{
-										ClassExportFlags |= UhtClassExportFlags.HasAutoGettersSetters;
 										property.PropertyExportFlags |= UhtPropertyExportFlags.SetterFound;
 										property.Setter = "Set" + property.SourceName;
 									}
