Adding support for more complex #define scopes

NOTE: This change results in no differences in UHT generated code.  It is for future expansion.

diff --git a/Engine/Source/Programs/Shared/EpicGames.UHT/Exporters/CodeGen/UhtHeaderCodeGenerator.cs b/Engine/Source/Programs/Shared/EpicGames.UHT/Exporters/CodeGen/UhtHeaderCodeGenerator.cs
--- a/Engine/Source/Programs/Shared/EpicGames.UHT/Exporters/CodeGen/UhtHeaderCodeGenerator.cs
+++ b/Engine/Source/Programs/Shared/EpicGames.UHT/Exporters/CodeGen/UhtHeaderCodeGenerator.cs
@@ -2,6 +2,7 @@
 
 using System;
 using System.Collections.Generic;
+using System.Linq;
 using System.Text;
 using EpicGames.Core;
 using EpicGames.UHT.Types;
@@ -16,10 +17,7 @@ namespace EpicGames.UHT.Exporters.CodeGen
 		public const string CallbackWrappersMacroSuffix = "CALLBACK_WRAPPERS";
 		public const string SparseDataMacroSuffix = "SPARSE_DATA";
 		public const string SparseDataPropertyAccessorsMacroSuffix = "SPARSE_DATA_PROPERTY_ACCESSORS";
-		public const string EditorOnlySparseDataPropertyAccessorsMacroSuffix = "SPARSE_DATA_PROPERTY_ACCESSORS_EOD";
-		public const string EditorOnlyRpcWrappersMacroSuffix = "RPC_WRAPPERS_EOD";
 		public const string RpcWrappersMacroSuffix = "RPC_WRAPPERS";
-		public const string EditorOnlyRpcWrappersNoPureDeclsMacroSuffix = "RPC_WRAPPERS_NO_PURE_DECLS_EOD";
 		public const string RpcWrappersNoPureDeclsMacroSuffix = "RPC_WRAPPERS_NO_PURE_DECLS";
 		public const string AccessorsMacroSuffix = "ACCESSORS";
 		public const string ArchiveSerializerMacroSuffix = "ARCHIVESERIALIZER";
@@ -36,7 +34,6 @@ namespace EpicGames.UHT.Exporters.CodeGen
 		public const string PrologMacroSuffix = "PROLOG";
 		public const string DelegateMacroSuffix = "DELEGATE";
 		public const string AutoGettersSettersMacroSuffix = "AUTOGETTERSETTER_DECLS";
-		public const string EditorOnlyAutoGettersSettersMacroSuffix = "AUTOGETTERSETTER_DECLS_EOD";
 		#endregion
 
 		public readonly UhtHeaderFile HeaderFile;
@@ -430,11 +427,6 @@ namespace EpicGames.UHT.Exporters.CodeGen
 				!function.FunctionExportFlags.HasAnyFlags(UhtFunctionExportFlags.CustomThunk);
 		}
 
-		protected static bool IsRpcFunction(UhtFunction function, bool editorOnly)
-		{
-			return IsRpcFunction(function) && function.FunctionFlags.HasAnyFlags(EFunctionFlags.EditorOnly) == editorOnly;
-		}
-
 		/// <summary>
 		/// Determines whether the glue version of the specified native function should be exported.
 		/// </summary>
@@ -589,141 +581,92 @@ namespace EpicGames.UHT.Exporters.CodeGen
 				classObj.ClassExportFlags.HasAnyFlags(UhtClassExportFlags.HasFieldNotify);
 		}
 
-		protected static void GetFieldNotifyStats(UhtClass classObj, out bool hasProperties, out bool hasFunctions, out bool hasEditorFields, out bool allEditorFields)
+		protected static string GetNotifyTypeName(UhtType notifyType)
 		{
-			// Scan the children to see what we have
-			hasProperties = false;
-			hasFunctions = false;
-			hasEditorFields = false;
-			allEditorFields = true;
-			foreach (UhtType type in classObj.Children)
-			{
-				if (type is UhtProperty property)
-				{
-					if (property.PropertyExportFlags.HasAnyFlags(UhtPropertyExportFlags.FieldNotify))
-					{
-						hasProperties = true;
-						hasEditorFields |= property.IsEditorOnlyProperty;
-						allEditorFields &= property.IsEditorOnlyProperty;
-					}
-				}
-				else if (type is UhtFunction function)
-				{
-					if (function.FunctionExportFlags.HasAnyFlags(UhtFunctionExportFlags.FieldNotify))
-					{
-						hasFunctions = true;
-						hasEditorFields |= function.FunctionFlags.HasAnyFlags(EFunctionFlags.EditorOnly);
-						allEditorFields &= function.FunctionFlags.HasAnyFlags(EFunctionFlags.EditorOnly);
-					}
-				}
+			if (notifyType is UhtProperty property)
+			{
+				return property.SourceName;
 			}
-
-			// If we have no editor fields, then by definition, all fields can't be editor fields
-			allEditorFields &= hasEditorFields;
-		}
-
-		protected static StringBuilder AppendFieldNotify(StringBuilder builder, UhtClass classObj,
-			bool hasProperties, bool hasFunctions, bool hasEditorFields, bool allEditorFields,
-			bool includeEditorOnlyFields, bool appendDefine, Action<StringBuilder, UhtClass, string> appendAction)
-		{
-			if (hasProperties && !allEditorFields)
+			else if (notifyType is UhtFunction function)
 			{
-				foreach (UhtType child in classObj.Children)
-				{
-					if (child is UhtProperty property)
-					{
-						if (property.PropertyExportFlags.HasAnyFlags(UhtPropertyExportFlags.FieldNotify) && !property.IsEditorOnlyProperty)
-						{
-							appendAction(builder, classObj, property.SourceName);
-						}
-					}
-				}
+				return function.CppImplName;
 			}
-
-			if (hasFunctions && !allEditorFields)
+			else
 			{
-				foreach (UhtType child in classObj.Children)
-				{
-					if (child is UhtFunction function)
-					{
-						if (function.FunctionExportFlags.HasAnyFlags(UhtFunctionExportFlags.FieldNotify) && !function.FunctionFlags.HasAnyFlags(EFunctionFlags.EditorOnly))
-						{
-							appendAction(builder, classObj, function.CppImplName);
-						}
-					}
-				}
+				throw new UhtIceException("Unexpected type in notification code generation");
 			}
+		}
 
-			if (hasEditorFields && includeEditorOnlyFields)
+		protected static UhtUsedDefineScopes<UhtType> GetFieldNotifyTypes(UhtClass classObj)
+		{
+			UhtUsedDefineScopes<UhtType> notifyTypes = new(classObj.Children.Where(x =>
 			{
-				if (!allEditorFields && appendDefine)
+				if (x is UhtProperty property)
 				{
-					builder.Append("#if WITH_EDITORONLY_DATA\r\n");
+					return property.PropertyExportFlags.HasAnyFlags(UhtPropertyExportFlags.FieldNotify);
 				}
-
-				if (hasProperties)
+				else if (x is UhtFunction function)
 				{
-					foreach (UhtType child in classObj.Children)
-					{
-						if (child is UhtProperty property)
-						{
-							if (property.PropertyExportFlags.HasAnyFlags(UhtPropertyExportFlags.FieldNotify) && property.IsEditorOnlyProperty)
-							{
-								appendAction(builder, classObj, property.SourceName);
-							}
-						}
-					}
+					return function.FunctionExportFlags.HasAnyFlags(UhtFunctionExportFlags.FieldNotify);
 				}
-
-				if (hasFunctions)
+				else
 				{
-					foreach (UhtType child in classObj.Children)
-					{
-						if (child is UhtFunction function)
-						{
-							if (function.FunctionExportFlags.HasAnyFlags(UhtFunctionExportFlags.FieldNotify) && function.FunctionFlags.HasAnyFlags(EFunctionFlags.EditorOnly))
-							{
-								appendAction(builder, classObj, function.CppImplName);
-							}
-						}
-					}
+					return false;
 				}
+			}));
 
-				if (!allEditorFields && appendDefine)
-				{
-					builder.Append("#endif // WITH_EDITORONLY_DATA\r\n");
-				}
-			}
-			return builder;
+			// We want properties followed by functions
+			notifyTypes.Instances = notifyTypes.Instances.OrderBy(x => (x is UhtProperty ? 0 : 1) * (int)UhtDefineScope.ScopeCount + x.DefineScope).ToList();
+			return notifyTypes;
 		}
 		#endregion
-	}
 
-	internal static class UhtHaederCodeGeneratorStringBuilderExtensions
-	{
-		public static StringBuilder AppendMacroName(this StringBuilder builder, string fileId, int lineNumber, string macroSuffix)
+		#region AutoGettersSetters support
+		protected static UhtUsedDefineScopes<UhtProperty> GetAutoGetterSetterProperties(UhtClass classObj)
 		{
-			return builder.Append(fileId).Append('_').Append(lineNumber).Append('_').Append(macroSuffix);
-		}
-
-		public static StringBuilder AppendMacroName(this StringBuilder builder, UhtHeaderCodeGenerator generator, int lineNumber, string macroSuffix)
-		{
-			return builder.AppendMacroName(generator.FileId, lineNumber, macroSuffix);
+			UhtUsedDefineScopes<UhtProperty> properties =  new(classObj.Properties.Where(
+				x => x.PropertyExportFlags.HasAnyFlags(UhtPropertyExportFlags.GetterSpecifiedAuto | UhtPropertyExportFlags.SetterSpecifiedAuto
+				)));
+			return properties;
 		}
+		#endregion
+	}
 
-		public static StringBuilder AppendMacroName(this StringBuilder builder, UhtHeaderCodeGenerator generator, UhtClass classObj, string macroSuffix)
+	internal static class UhtHaederCodeGeneratorStringBuilderExtensions
+	{
+		public static StringBuilder AppendMacroName(this StringBuilder builder, string fileId, int lineNumber, string macroSuffix, UhtDefineScope defineScope = UhtDefineScope.None)
 		{
-			return builder.AppendMacroName(generator, classObj.GeneratedBodyLineNumber, macroSuffix);
+			builder.Append(fileId).Append('_').Append(lineNumber).Append('_').Append(macroSuffix);
+			if (defineScope.HasAnyFlags(UhtDefineScope.EditorOnlyData))
+			{
+				builder.Append("_EOD");
+			}
+			return builder;
 		}
 
-		public static StringBuilder AppendMacroName(this StringBuilder builder, UhtHeaderCodeGenerator generator, UhtScriptStruct scriptStruct, string macroSuffix)
+		public static StringBuilder AppendMacroName(this StringBuilder builder, UhtHeaderCodeGenerator generator, int lineNumber, string macroSuffix, UhtDefineScope defineScope = UhtDefineScope.None)
 		{
-			return builder.AppendMacroName(generator, scriptStruct.MacroDeclaredLineNumber, macroSuffix);
+			return builder.AppendMacroName(generator.FileId, lineNumber, macroSuffix, defineScope);
 		}
 
-		public static StringBuilder AppendMacroName(this StringBuilder builder, UhtHeaderCodeGenerator generator, UhtFunction function, string macroSuffix)
+		public static StringBuilder AppendMacroName(this StringBuilder builder, UhtHeaderCodeGenerator generator, UhtType type, string macroSuffix, UhtDefineScope defineScope = UhtDefineScope.None)
 		{
-			return builder.AppendMacroName(generator, function.MacroLineNumber, macroSuffix);
+			if (type is UhtClass classObj)
+			{
+				return builder.AppendMacroName(generator, classObj.GeneratedBodyLineNumber, macroSuffix, defineScope);
+			}
+			else if (type is UhtScriptStruct scriptStruct)
+			{
+				return builder.AppendMacroName(generator, scriptStruct.MacroDeclaredLineNumber, macroSuffix, defineScope);
+			}
+			else if (type is UhtFunction function)
+			{
+				return builder.AppendMacroName(generator, function.MacroLineNumber, macroSuffix, defineScope);
+			}
+			else
+			{
+				throw new UhtException(type, "Can not use given type to create a macro");
+			}
 		}
 	}
 }
diff --git a/Engine/Source/Programs/Shared/EpicGames.UHT/Exporters/CodeGen/UhtHeaderCodeGeneratorCppFile.cs b/Engine/Source/Programs/Shared/EpicGames.UHT/Exporters/CodeGen/UhtHeaderCodeGeneratorCppFile.cs
--- a/Engine/Source/Programs/Shared/EpicGames.UHT/Exporters/CodeGen/UhtHeaderCodeGeneratorCppFile.cs
+++ b/Engine/Source/Programs/Shared/EpicGames.UHT/Exporters/CodeGen/UhtHeaderCodeGeneratorCppFile.cs
@@ -147,24 +147,22 @@ namespace EpicGames.UHT.Exporters.CodeGen
 
 				int generatedBodyStart = builder.Length;
 
-				bool hasRegisteredEnums = false;
-				bool hasRegisteredScriptStructs = false;
-				bool hasRegisteredClasses = false;
-				bool allEnumsEditorOnly = true;
+				UhtUsedDefineScopes<UhtEnum> enums = new();
+				UhtUsedDefineScopes<UhtScriptStruct> scriptStructs = new();
+				UhtUsedDefineScopes<UhtClass> classes = new();
 				foreach (UhtField field in HeaderFile.References.ExportTypes)
 				{
 					if (field is UhtEnum enumObj)
 					{
 						AppendEnum(builder, enumObj);
-						hasRegisteredEnums = true;
-						allEnumsEditorOnly &= enumObj.IsEditorOnly;
+						enums.Add(enumObj);
 					}
 					else if (field is UhtScriptStruct scriptStruct)
 					{
 						AppendScriptStruct(builder, scriptStruct);
 						if (scriptStruct.ScriptStructFlags.HasAnyFlags(EStructFlags.Native))
 						{
-							hasRegisteredScriptStructs = true;
+							scriptStructs.Add(scriptStruct);
 						}
 					}
 					else if (field is UhtFunction function)
@@ -176,12 +174,12 @@ namespace EpicGames.UHT.Exporters.CodeGen
 						if (!classObj.ClassFlags.HasAnyFlags(EClassFlags.Intrinsic))
 						{
 							AppendClass(builder, classObj);
-							hasRegisteredClasses = true;
+							classes.Add(classObj);
 						}
 					}
 				}
 
-				if (hasRegisteredEnums || hasRegisteredScriptStructs || hasRegisteredClasses)
+				if (!enums.IsEmpty || !scriptStructs.IsEmpty || !classes.IsEmpty)
 				{
 					string name = $"Z_CompiledInDeferFile_{headerInfo.FileId}";
 					string staticsName = $"{name}_Statics";
@@ -189,106 +187,72 @@ namespace EpicGames.UHT.Exporters.CodeGen
 					builder.Append("\tstruct ").Append(staticsName).Append("\r\n");
 					builder.Append("\t{\r\n");
 
-					if (hasRegisteredEnums)
-					{
-						using UhtMacroBlockEmitter blockEmitter = new(builder, "WITH_EDITORONLY_DATA", allEnumsEditorOnly);
-						builder.Append("\t\tstatic const FEnumRegisterCompiledInInfo EnumInfo[];\r\n");
-					}
-					if (hasRegisteredScriptStructs)
-					{
-						builder.Append("\t\tstatic const FStructRegisterCompiledInInfo ScriptStructInfo[];\r\n");
-					}
-					if (hasRegisteredClasses)
-					{
-						builder.Append("\t\tstatic const FClassRegisterCompiledInInfo ClassInfo[];\r\n");
-					}
+					builder.AppendIfInstances(enums, UhtDefineScopeNames.Standard, builder => builder.Append("\t\tstatic const FEnumRegisterCompiledInInfo EnumInfo[];\r\n"));
+					builder.AppendIfInstances(scriptStructs, UhtDefineScopeNames.Standard, builder => builder.Append("\t\tstatic const FStructRegisterCompiledInInfo ScriptStructInfo[];\r\n"));
+					builder.AppendIfInstances(classes, UhtDefineScopeNames.Standard, builder => builder.Append("\t\tstatic const FClassRegisterCompiledInInfo ClassInfo[];\r\n"));
 
 					builder.Append("\t};\r\n");
 
 					uint combinedHash = UInt32.MaxValue;
 
-					if (hasRegisteredEnums)
-					{
-						using UhtMacroBlockEmitter blockEmitter = new(builder, "WITH_EDITORONLY_DATA", allEnumsEditorOnly);
-						builder.Append("\tconst FEnumRegisterCompiledInInfo ").Append(staticsName).Append("::EnumInfo[] = {\r\n");
-						foreach (UhtObject obj in HeaderFile.References.ExportTypes)
+					builder.AppendInstances(enums, UhtDefineScopeNames.Standard,
+						builder => builder.Append("\tconst FEnumRegisterCompiledInInfo ").Append(staticsName).Append("::EnumInfo[] = {\r\n"),
+						(builder, enumObj) =>
 						{
-							if (obj is UhtEnum enumObj)
-							{
-								blockEmitter.Set(enumObj.IsEditorOnly);
-								uint hash = ObjectInfos[enumObj.ObjectTypeIndex].Hash;
-								builder
-									.Append("\t\t{ ")
-									.Append(enumObj.SourceName)
-									.Append("_StaticEnum, TEXT(\"")
-									.Append(enumObj.EngineName)
-									.Append("\"), &Z_Registration_Info_UEnum_")
-									.Append(enumObj.EngineName)
-									.Append($", CONSTRUCT_RELOAD_VERSION_INFO(FEnumReloadVersionInfo, {hash}U) }},\r\n");
-								combinedHash = HashCombine(combinedHash, hash);
-							}
-						}
-						blockEmitter.Set(allEnumsEditorOnly);
-						builder.Append("\t};\r\n");
-					}
-
-					if (hasRegisteredScriptStructs)
-					{
-						builder.Append("\tconst FStructRegisterCompiledInInfo ").Append(staticsName).Append("::ScriptStructInfo[] = {\r\n");
-						foreach (UhtObject obj in HeaderFile.References.ExportTypes)
+							uint hash = ObjectInfos[enumObj.ObjectTypeIndex].Hash;
+							builder
+								.Append("\t\t{ ")
+								.Append(enumObj.SourceName)
+								.Append("_StaticEnum, TEXT(\"")
+								.Append(enumObj.EngineName)
+								.Append("\"), &Z_Registration_Info_UEnum_")
+								.Append(enumObj.EngineName)
+								.Append($", CONSTRUCT_RELOAD_VERSION_INFO(FEnumReloadVersionInfo, {hash}U) }},\r\n");
+							combinedHash = HashCombine(combinedHash, hash);
+						},
+						builder => builder.Append("\t};\r\n"));
+
+					builder.AppendInstances(scriptStructs, UhtDefineScopeNames.Standard,
+						builder => builder.Append("\tconst FStructRegisterCompiledInInfo ").Append(staticsName).Append("::ScriptStructInfo[] = {\r\n"),
+						(builder, scriptStruct) =>
 						{
-							if (obj is UhtScriptStruct scriptStruct)
-							{
-								if (scriptStruct.ScriptStructFlags.HasAnyFlags(EStructFlags.Native))
-								{
-									uint hash = ObjectInfos[scriptStruct.ObjectTypeIndex].Hash;
-									builder
-										.Append("\t\t{ ")
-										.Append(scriptStruct.SourceName)
-										.Append("::StaticStruct, Z_Construct_UScriptStruct_")
-										.Append(scriptStruct.SourceName)
-										.Append("_Statics::NewStructOps, TEXT(\"")
-										.Append(scriptStruct.EngineName)
-										.Append("\"), &Z_Registration_Info_UScriptStruct_")
-										.Append(scriptStruct.EngineName)
-										.Append(", CONSTRUCT_RELOAD_VERSION_INFO(FStructReloadVersionInfo, sizeof(")
-										.Append(scriptStruct.SourceName)
-										.Append($"), {hash}U) }},\r\n");
-									combinedHash = HashCombine(combinedHash, hash);
-								}
-							}
-						}
-						builder.Append("\t};\r\n");
-					}
-
-					if (hasRegisteredClasses)
-					{
-						builder.Append("\tconst FClassRegisterCompiledInInfo ").Append(staticsName).Append("::ClassInfo[] = {\r\n");
-						foreach (UhtObject obj in HeaderFile.References.ExportTypes)
+							uint hash = ObjectInfos[scriptStruct.ObjectTypeIndex].Hash;
+							builder
+								.Append("\t\t{ ")
+								.Append(scriptStruct.SourceName)
+								.Append("::StaticStruct, Z_Construct_UScriptStruct_")
+								.Append(scriptStruct.SourceName)
+								.Append("_Statics::NewStructOps, TEXT(\"")
+								.Append(scriptStruct.EngineName)
+								.Append("\"), &Z_Registration_Info_UScriptStruct_")
+								.Append(scriptStruct.EngineName)
+								.Append(", CONSTRUCT_RELOAD_VERSION_INFO(FStructReloadVersionInfo, sizeof(")
+								.Append(scriptStruct.SourceName)
+								.Append($"), {hash}U) }},\r\n");
+							combinedHash = HashCombine(combinedHash, hash);
+						},
+						builder => builder.Append("\t};\r\n"));
+
+					builder.AppendInstances(classes, UhtDefineScopeNames.Standard,
+						builder => builder.Append("\tconst FClassRegisterCompiledInInfo ").Append(staticsName).Append("::ClassInfo[] = {\r\n"),
+						(builder, classObj) =>
 						{
-							if (obj is UhtClass classObj)
-							{
-								if (!classObj.ClassFlags.HasAnyFlags(EClassFlags.Intrinsic))
-								{
-									uint hash = ObjectInfos[classObj.ObjectTypeIndex].Hash;
-									builder
-										.Append("\t\t{ Z_Construct_UClass_")
-										.Append(classObj.SourceName)
-										.Append(", ")
-										.Append(classObj.SourceName)
-										.Append("::StaticClass, TEXT(\"")
-										.Append(classObj.SourceName)
-										.Append("\"), &Z_Registration_Info_UClass_")
-										.Append(classObj.SourceName)
-										.Append(", CONSTRUCT_RELOAD_VERSION_INFO(FClassReloadVersionInfo, sizeof(")
-										.Append(classObj.SourceName)
-										.Append($"), {hash}U) }},\r\n");
-									combinedHash = HashCombine(combinedHash, hash);
-								}
-							}
-						}
-						builder.Append("\t};\r\n");
-					}
+							uint hash = ObjectInfos[classObj.ObjectTypeIndex].Hash;
+							builder
+								.Append("\t\t{ Z_Construct_UClass_")
+								.Append(classObj.SourceName)
+								.Append(", ")
+								.Append(classObj.SourceName)
+								.Append("::StaticClass, TEXT(\"")
+								.Append(classObj.SourceName)
+								.Append("\"), &Z_Registration_Info_UClass_")
+								.Append(classObj.SourceName)
+								.Append(", CONSTRUCT_RELOAD_VERSION_INFO(FClassReloadVersionInfo, sizeof(")
+								.Append(classObj.SourceName)
+								.Append($"), {hash}U) }},\r\n");
+							combinedHash = HashCombine(combinedHash, hash);
+						},
+						builder => builder.Append("\t};\r\n"));
 
 					builder
 						.Append("\tstatic FRegisterCompiledInInfo ")
@@ -297,9 +261,9 @@ namespace EpicGames.UHT.Exporters.CodeGen
 						.Append(Package.EngineName)
 						.Append("\"),\r\n");
 
-					builder.Append("\t\t").AppendArray(!hasRegisteredClasses, false, staticsName, "ClassInfo").Append(",\r\n");
-					builder.Append("\t\t").AppendArray(!hasRegisteredScriptStructs, false, staticsName, "ScriptStructInfo").Append(",\r\n");
-					builder.Append("\t\t").AppendArray(!hasRegisteredEnums, allEnumsEditorOnly, staticsName, "EnumInfo").Append(");\r\n");
+					builder.AppendArrayPtrAndCountLine(classes, UhtDefineScopeNames.Standard, staticsName, "ClassInfo", 2, ",\r\n");
+					builder.AppendArrayPtrAndCountLine(scriptStructs, UhtDefineScopeNames.Standard, staticsName, "ScriptStructInfo", 2, ",\r\n");
+					builder.AppendArrayPtrAndCountLine(enums, UhtDefineScopeNames.Standard, staticsName, "EnumInfo", 2, ");\r\n");
 				}
 
 				if (Session.IncludeDebugOutput)
@@ -382,7 +346,7 @@ namespace EpicGames.UHT.Exporters.CodeGen
 				enumDisplayNameFn = "nullptr";
 			}
 
-			using (UhtMacroBlockEmitter macroBlockEmitter = new(builder, "WITH_EDITORONLY_DATA", enumObj.IsEditorOnly))
+			using (UhtMacroBlockEmitter macroBlockEmitter = new(builder, UhtDefineScopeNames.Standard, enumObj.DefineScope))
 			{
 
 				// If we don't have a zero 0 then we emit a static assert to verify we have one
@@ -625,6 +589,9 @@ namespace EpicGames.UHT.Exporters.CodeGen
 			// Everything from this point on will be part of the definition hash
 			int hashCodeBlockStart = builder.Length;
 
+			// Collect the properties
+			UhtUsedDefineScopes<UhtProperty> properties = new(scriptStruct.Properties);
+
 			// Declare the statics structure
 			{
 				builder.Append("\tstruct ").Append(staticsName).Append("\r\n");
@@ -646,7 +613,7 @@ namespace EpicGames.UHT.Exporters.CodeGen
 					builder.Append("\t\tstatic void* NewStructOps();\r\n");
 				}
 
-				AppendPropertiesDecl(builder, scriptStruct, scriptStruct.SourceName, staticsName, 2);
+				AppendPropertiesDecl(builder, scriptStruct, scriptStruct.SourceName, properties, staticsName, 2);
 
 				builder.Append("\t\tstatic const UECodeGen_Private::FStructParams ReturnStructParams;\r\n");
 				builder.Append("\t};\r\n");
@@ -664,7 +631,7 @@ namespace EpicGames.UHT.Exporters.CodeGen
 					builder.Append("\t}\r\n");
 				}
 
-				AppendPropertiesDefs(builder, scriptStruct, scriptStruct.SourceName, staticsName, 1);
+				AppendPropertiesDefs(builder, scriptStruct, scriptStruct.SourceName, properties, staticsName, 1);
 
 				builder.Append("\tconst UECodeGen_Private::FStructParams ").Append(staticsName).Append("::ReturnStructParams = {\r\n");
 				builder.Append("\t\t(UObject* (*)())").Append(PackageSingletonName).Append(",\r\n");
@@ -678,7 +645,8 @@ namespace EpicGames.UHT.Exporters.CodeGen
 					builder.Append("\t\tnullptr,\r\n");
 				}
 				builder.Append("\t\t").AppendUTF8LiteralString(scriptStruct.EngineName).Append(",\r\n");
-				builder.AppendPropertiesParams(scriptStruct, staticsName, 2, "\r\n");
+				builder.AppendArrayPtrLine(properties, UhtDefineScopeNames.Standard, staticsName, "PropPointers", 2, ",\r\n");
+				builder.AppendArrayCountLine(properties, UhtDefineScopeNames.Standard, staticsName, "PropPointers", 2, ",\r\n");
 				builder.Append("\t\tsizeof(").Append(scriptStruct.SourceName).Append("),\r\n");
 				builder.Append("\t\talignof(").Append(scriptStruct.SourceName).Append("),\r\n");
 				builder.Append("\t\tRF_Public|RF_Transient|RF_MarkAsNative,\r\n");
@@ -686,7 +654,7 @@ namespace EpicGames.UHT.Exporters.CodeGen
 				builder.Append("\t\t").AppendMetaDataParams(scriptStruct, staticsName, MetaDataParamsName).Append("\r\n");
 				builder.Append("\t};\r\n");
 
-				builder.AppendPropertiesParamsCountStaticAssert(scriptStruct, staticsName);
+				builder.AppendPropertiesParamsCountStaticAssert(properties, staticsName);
 			}
 
 			// Generate the registration function
@@ -858,106 +826,111 @@ namespace EpicGames.UHT.Exporters.CodeGen
 			// Everything from this point on will be part of the definition hash
 			int hashCodeBlockStart = builder.Length;
 
-			builder.AppendBeginEditorOnlyGuard(function.FunctionFlags.HasAnyFlags(EFunctionFlags.EditorOnly));
-
-			// Statics declaration
+			// Note: This is fairly agressive 
+			using (UhtMacroBlockEmitter blockEmitter = new(builder, UhtDefineScopeNames.WithEditor, function.DefineScope))
 			{
-				builder.Append("\tstruct ").Append(staticsName).Append("\r\n");
-				builder.Append("\t{\r\n");
 
-				if (paramsInStatic)
+				// Collect the properties
+				UhtUsedDefineScopes<UhtProperty> properties = new(function.Properties);
+
+				// Statics declaration
 				{
-					List<UhtScriptStruct> noExportStructs = FindNoExportStructs(function);
-					foreach (UhtScriptStruct noExportStruct in noExportStructs)
+					builder.Append("\tstruct ").Append(staticsName).Append("\r\n");
+					builder.Append("\t{\r\n");
+
+					if (paramsInStatic)
 					{
-						AppendMirrorsForNoexportStruct(builder, noExportStruct, 2);
+						List<UhtScriptStruct> noExportStructs = FindNoExportStructs(function);
+						foreach (UhtScriptStruct noExportStruct in noExportStructs)
+						{
+							AppendMirrorsForNoexportStruct(builder, noExportStruct, 2);
+						}
+						AppendEventParameter(builder, function, strippedFunctionName, UhtPropertyTextType.EventParameterFunctionMember, false, 2, "\r\n");
 					}
-					AppendEventParameter(builder, function, strippedFunctionName, UhtPropertyTextType.EventParameterFunctionMember, false, 2, "\r\n");
-				}
-
-				AppendPropertiesDecl(builder, function, eventParameters, staticsName, 2);
-
-				builder.AppendMetaDataDecl(function, MetaDataParamsName, 2);
 
-				builder.Append("\t\tstatic const UECodeGen_Private::FFunctionParams FuncParams;\r\n");
+					AppendPropertiesDecl(builder, function, eventParameters, properties, staticsName, 2);
 
-				builder.Append("\t};\r\n");
-			}
+					builder.AppendMetaDataDecl(function, MetaDataParamsName, 2);
 
-			// Statics definition
-			{
-				AppendPropertiesDefs(builder, function, eventParameters, staticsName, 1);
+					builder.Append("\t\tstatic const UECodeGen_Private::FFunctionParams FuncParams;\r\n");
 
-				builder.AppendMetaDataDef(function, staticsName, MetaDataParamsName, 1);
+					builder.Append("\t};\r\n");
+				}
 
-				builder
-					.Append("\tconst UECodeGen_Private::FFunctionParams ")
-					.Append(staticsName).Append("::FuncParams = { ")
-					.Append("(UObject*(*)())").Append(GetFunctionOuterFunc(function)).Append(", ")
-					.Append(GetSingletonName(function.SuperFunction, true)).Append(", ")
-					.AppendUTF8LiteralString(function.EngineName).Append(", ")
-					.AppendUTF8LiteralString(function.FunctionType == UhtFunctionType.SparseDelegate, function.SparseOwningClassName).Append(", ")
-					.AppendUTF8LiteralString(function.FunctionType == UhtFunctionType.SparseDelegate, function.SparseDelegateName).Append(", ");
+				// Statics definition
+				{
+					AppendPropertiesDefs(builder, function, eventParameters, properties, staticsName, 1);
 
-				builder.AppendPropertiesParams(function, staticsName, 0, " ");
+					builder.AppendMetaDataDef(function, staticsName, MetaDataParamsName, 1);
 
-				string sizeOfStatic = "";
-				if (function.Children.Count > 0)
-				{
-					UhtFunction tempFunction = function;
-					while (tempFunction.SuperFunction != null)
+					builder
+						.Append("\tconst UECodeGen_Private::FFunctionParams ")
+						.Append(staticsName).Append("::FuncParams = { ")
+						.Append("(UObject*(*)())").Append(GetFunctionOuterFunc(function)).Append(", ")
+						.Append(GetSingletonName(function.SuperFunction, true)).Append(", ")
+						.AppendUTF8LiteralString(function.EngineName).Append(", ")
+						.AppendUTF8LiteralString(function.FunctionType == UhtFunctionType.SparseDelegate, function.SparseOwningClassName).Append(", ")
+						.AppendUTF8LiteralString(function.FunctionType == UhtFunctionType.SparseDelegate, function.SparseDelegateName).Append(", ");
+
+					builder.AppendArrayPtrLine(properties, UhtDefineScopeNames.Standard, staticsName, "PropPointers", 0, ", ");
+					builder.AppendArrayCountLine(properties, UhtDefineScopeNames.Standard, staticsName, "PropPointers", 0, ", ");
+
+					string sizeOfStatic = "";
+					if (function.Children.Count > 0)
 					{
-						tempFunction = tempFunction.SuperFunction;
-					}
+						UhtFunction tempFunction = function;
+						while (tempFunction.SuperFunction != null)
+						{
+							tempFunction = tempFunction.SuperFunction;
+						}
 
-					if (paramsInStatic)
-					{
-						sizeOfStatic = "sizeof(" + staticsName + "::" + GetEventStructParametersName(tempFunction.Outer, tempFunction.StrippedFunctionName) + ")";
+						if (paramsInStatic)
+						{
+							sizeOfStatic = "sizeof(" + staticsName + "::" + GetEventStructParametersName(tempFunction.Outer, tempFunction.StrippedFunctionName) + ")";
+						}
+						else
+						{
+							sizeOfStatic = "sizeof(" + GetEventStructParametersName(tempFunction.Outer, tempFunction.StrippedFunctionName) + ")";
+						}
+
+						builder.Append(sizeOfStatic);
 					}
 					else
 					{
-						sizeOfStatic = "sizeof(" + GetEventStructParametersName(tempFunction.Outer, tempFunction.StrippedFunctionName) + ")";
+						builder.Append('0');
 					}
+					builder.Append(", ");
 
-					builder.Append(sizeOfStatic);
-				}
-				else
-				{
-					builder.Append('0');
-				}
-				builder.Append(", ");
+					builder
+						.Append("RF_Public|RF_Transient|RF_MarkAsNative, ")
+						.Append($"(EFunctionFlags)0x{(uint)function.FunctionFlags:X8}, ")
+						.Append(isNet ? function.RPCId : 0).Append(", ")
+						.Append(isNet ? function.RPCResponseId : 0).Append(", ")
+						.AppendMetaDataParams(function, staticsName, MetaDataParamsName)
+						.Append(" };\r\n");
 
-				builder
-					.Append("RF_Public|RF_Transient|RF_MarkAsNative, ")
-					.Append($"(EFunctionFlags)0x{(uint)function.FunctionFlags:X8}, ")
-					.Append(isNet ? function.RPCId : 0).Append(", ")
-					.Append(isNet ? function.RPCResponseId : 0).Append(", ")
-					.AppendMetaDataParams(function, staticsName, MetaDataParamsName)
-					.Append(" };\r\n");
+					builder.AppendPropertiesParamsCountStaticAssert(properties, staticsName);
 
-				builder.AppendPropertiesParamsCountStaticAssert(function, staticsName);
+					if (!String.IsNullOrEmpty(sizeOfStatic))
+					{
+						builder.Append("\tstatic_assert(").Append(sizeOfStatic).Append(" < MAX_uint16);\r\n");
+					}
+				}
 
-				if (!String.IsNullOrEmpty(sizeOfStatic))
+				// Registration function
 				{
-					builder.Append("\tstatic_assert(").Append(sizeOfStatic).Append(" < MAX_uint16);\r\n");
+					builder.Append("\tUFunction* ").Append(singletonName).Append("()\r\n");
+					builder.Append("\t{\r\n");
+					builder.Append("\t\tstatic UFunction* ReturnFunction = nullptr;\r\n");
+					builder.Append("\t\tif (!ReturnFunction)\r\n");
+					builder.Append("\t\t{\r\n");
+					builder.Append("\t\t\tUECodeGen_Private::ConstructUFunction(&ReturnFunction, ").Append(staticsName).Append("::FuncParams);\r\n");
+					builder.Append("\t\t}\r\n");
+					builder.Append("\t\treturn ReturnFunction;\r\n");
+					builder.Append("\t}\r\n");
 				}
 			}
 
-			// Registration function
-			{
-				builder.Append("\tUFunction* ").Append(singletonName).Append("()\r\n");
-				builder.Append("\t{\r\n");
-				builder.Append("\t\tstatic UFunction* ReturnFunction = nullptr;\r\n");
-				builder.Append("\t\tif (!ReturnFunction)\r\n");
-				builder.Append("\t\t{\r\n");
-				builder.Append("\t\t\tUECodeGen_Private::ConstructUFunction(&ReturnFunction, ").Append(staticsName).Append("::FuncParams);\r\n");
-				builder.Append("\t\t}\r\n");
-				builder.Append("\t\treturn ReturnFunction;\r\n");
-				builder.Append("\t}\r\n");
-			}
-
-			builder.AppendEndEditorOnlyGuard(function.FunctionFlags.HasAnyFlags(EFunctionFlags.EditorOnly));
-
 			using (UhtRentedPoolBuffer<char> borrowBuffer = builder.RentPoolBuffer(hashCodeBlockStart, builder.Length - hashCodeBlockStart))
 			{
 				ObjectInfos[function.ObjectTypeIndex].Hash = UhtHash.GenenerateTextHash(borrowBuffer.Buffer.Memory.Span);
@@ -1001,95 +974,57 @@ namespace EpicGames.UHT.Exporters.CodeGen
 
 		private static StringBuilder AppendExportProperties(StringBuilder builder, UhtScriptStruct scriptStruct, int tabs)
 		{
-			using (UhtMacroBlockEmitter emitter = new(builder, "WITH_EDITORONLY_DATA"))
+			using (UhtMacroBlockEmitter emitter = new(builder, UhtDefineScopeNames.Standard, UhtDefineScope.None))
 			{
 				foreach (UhtProperty property in scriptStruct.Properties)
 				{
-					emitter.Set(property.IsEditorOnlyProperty);
+					emitter.Set(property.DefineScope);
 					builder.AppendTabs(tabs).AppendFullDecl(property, UhtPropertyTextType.ExportMember, false).Append(";\r\n");
 				}
 			}
 			return builder;
 		}
 
-		private StringBuilder AppendPropertiesDecl(StringBuilder builder, UhtStruct structObj, string structSourceName, string staticsName, int tabs)
+		private StringBuilder AppendPropertiesDecl(StringBuilder builder, UhtStruct structObj, string structSourceName, UhtUsedDefineScopes<UhtProperty> properties, string staticsName, int tabs)
 		{
-			if (!structObj.Children.Any(x => x is UhtProperty))
+			if (properties.IsEmpty)
 			{
 				return builder;
 			}
 
 			PropertyMemberContextImpl context = new(CodeGenerator, structObj, structSourceName, staticsName);
 
-			bool hasAllEditorOnlyDataProperties = structObj.Properties.All(x => x.IsEditorOnlyProperty);
-
-			using (UhtMacroBlockEmitter emitter = new(builder, "WITH_EDITORONLY_DATA", hasAllEditorOnlyDataProperties))
-			{
-				foreach (UhtProperty property in structObj.Properties)
-				{
-					emitter.Set(property.IsEditorOnlyProperty);
-					hasAllEditorOnlyDataProperties &= property.IsEditorOnlyProperty;
-					builder.AppendMemberDecl(property, context, property.EngineName, "", tabs);
-				}
-
-				emitter.Set(hasAllEditorOnlyDataProperties);
-				builder.AppendTabs(tabs).Append("static const UECodeGen_Private::FPropertyParamsBase* const PropPointers[];\r\n");
-			}
+			builder.AppendInstances(properties, UhtDefineScopeNames.Standard,
+				builder => { },
+				(builder, property) => builder.AppendMemberDecl(property, context, property.EngineName, "", tabs),
+				builder => builder.AppendTabs(tabs).Append("static const UECodeGen_Private::FPropertyParamsBase* const PropPointers[];\r\n"));
 			return builder;
 		}
 
-		private StringBuilder AppendPropertiesDefs(StringBuilder builder, UhtStruct structObj, string structSourceName, string staticsName, int tabs)
+		private StringBuilder AppendPropertiesDefs(StringBuilder builder, UhtStruct structObj, string structSourceName, UhtUsedDefineScopes<UhtProperty> properties, string staticsName, int tabs)
 		{
-			if (!structObj.Children.Any(x => x is UhtProperty))
+			if (properties.IsEmpty)
 			{
 				return builder;
 			}
 
 			PropertyMemberContextImpl context = new(CodeGenerator, structObj, structSourceName, staticsName);
 
-			bool hasAllEditorOnlyDataProperties = structObj.Properties.All(x => x.IsEditorOnlyProperty);
-
-			using (UhtMacroBlockEmitter emitter = new(builder, "WITH_EDITORONLY_DATA", hasAllEditorOnlyDataProperties))
-			{
-				foreach (UhtProperty property in structObj.Properties)
-				{
-					emitter.Set(property.IsEditorOnlyProperty);
-					builder.AppendMemberDef(property, context, property.EngineName, "", null, tabs);
-				}
-			}
+			builder.AppendInstances(properties, UhtDefineScopeNames.Standard,
+				(builder, property) => builder.AppendMemberDef(property, context, property.EngineName, "", null, tabs));
 
-			using (UhtMacroBlockEmitter emitter = new(builder, "WITH_EDITORONLY_DATA", hasAllEditorOnlyDataProperties))
-			{
-				builder.AppendTabs(tabs).Append("const UECodeGen_Private::FPropertyParamsBase* const ").Append(staticsName).Append("::PropPointers[] = {\r\n");
-				foreach (UhtProperty property in structObj.Properties)
-				{
-					emitter.Set(property.IsEditorOnlyProperty);
-					builder.AppendMemberPtr(property, context, property.EngineName, "", tabs + 1);
-				}
-				emitter.Set(hasAllEditorOnlyDataProperties);
-				builder.AppendTabs(tabs).Append("};\r\n");
-			}
+			builder.AppendInstances(properties, UhtDefineScopeNames.Standard,
+				builder => builder.AppendTabs(tabs).Append("const UECodeGen_Private::FPropertyParamsBase* const ").Append(staticsName).Append("::PropPointers[] = {\r\n"),
+				(builder, property) => builder.AppendMemberPtr(property, context, property.EngineName, "", tabs + 1),
+				builder => builder.AppendTabs(tabs).Append("};\r\n"));
 			return builder;
 		}
 
 		private StringBuilder AppendClass(StringBuilder builder, UhtClass classObj)
 		{
-			// Collect the functions in reversed order
-			List<UhtFunction> reversedFunctions = new(classObj.Functions.Where(x => IsRpcFunction(x) && ShouldExportFunction(x)));
-			reversedFunctions.Reverse();
-
-			// Check to see if we have any RPC functions for the editor
-			bool hasEditorRpc = reversedFunctions.Any(x => x.FunctionFlags.HasAnyFlags(EFunctionFlags.EditorOnly));
-
 			// Output the RPC methods
-			AppendRpcFunctions(builder, classObj, reversedFunctions, false);
-			if (hasEditorRpc)
-			{
-				builder.AppendBeginEditorOnlyGuard();
-				AppendRpcFunctions(builder, classObj, reversedFunctions, true);
-				builder.AppendEndEditorOnlyGuard();
-			}
-			
+			AppendRpcFunctions(builder, classObj);
+
 			// Add the auto getters/setters
 			AppendAutoGettersSetters(builder, classObj);
 
@@ -1169,163 +1104,124 @@ namespace EpicGames.UHT.Exporters.CodeGen
 				return builder;
 			}
 
-			// Scan the children to see what we have
-			GetFieldNotifyStats(classObj, out bool hasProperties, out bool hasFunctions, out bool hasEditorFields, out bool allEditorFields);
-
-			if (allEditorFields)
-			{
-				builder.Append("#if WITH_EDITORONLY_DATA\r\n");
-			}
+			UhtUsedDefineScopes<UhtType> notifyTypes = GetFieldNotifyTypes(classObj);
 
 			//UE_FIELD_NOTIFICATION_DECLARE_FIELD
-			AppendFieldNotify(builder, classObj, hasProperties, hasFunctions, hasEditorFields, allEditorFields, true, true,
-				(StringBuilder builder, UhtClass classObj, string name) =>
+			builder.AppendInstances(notifyTypes, UhtDefineScopeNames.Standard,
+				(builder, notifyType) =>
 				{
-					builder.Append($"\tUE_FIELD_NOTIFICATION_IMPLEMENT_FIELD({classObj.SourceName}, {name})\r\n");
+					builder.Append($"\tUE_FIELD_NOTIFICATION_IMPLEMENT_FIELD({classObj.SourceName}, {GetNotifyTypeName(notifyType)})\r\n");
 				});
 
-			if (allEditorFields)
-			{
-				builder.Append("#endif // WITH_EDITORONLY_DATA\r\n");
-				builder.Append("#if WITH_EDITORONLY_DATA\r\n");
-			}
-
 			//UE_FIELD_NOTIFICATION_DECLARE_ENUM_FIELD
-			builder.Append($"\tUE_FIELD_NOTIFICATION_IMPLEMENTATION_BEGIN({classObj.SourceName})\r\n");
-			AppendFieldNotify(builder, classObj, hasProperties, hasFunctions, hasEditorFields, allEditorFields, true, true,
-				(StringBuilder builder, UhtClass classObj, string name) =>
+			builder.AppendInstances(notifyTypes, UhtDefineScopeNames.Standard,
+				builder => builder.Append($"\tUE_FIELD_NOTIFICATION_IMPLEMENTATION_BEGIN({classObj.SourceName})\r\n"),
+				(builder, notifyType) =>
 				{
-					builder.Append($"\tUE_FIELD_NOTIFICATION_IMPLEMENT_ENUM_FIELD({classObj.SourceName}, {name})\r\n");
-				});
-			builder.Append($"\tUE_FIELD_NOTIFICATION_IMPLEMENTATION_END({classObj.SourceName});\r\n");
+					builder.Append($"\tUE_FIELD_NOTIFICATION_IMPLEMENT_ENUM_FIELD({classObj.SourceName}, {GetNotifyTypeName(notifyType)})\r\n");
+				}, 
+				builder => builder.Append($"\tUE_FIELD_NOTIFICATION_IMPLEMENTATION_END({classObj.SourceName});\r\n"));
 
-			if (allEditorFields)
-			{
-				builder.Append("#endif // WITH_EDITORONLY_DATA\r\n");
-			}
 			return builder;
 		}
 
 		private static StringBuilder AppendAutoGettersSetters(StringBuilder builder, UhtClass classObj)
 		{
-			IEnumerable<UhtProperty> properties = classObj.Properties.Where(x => x.PropertyExportFlags.HasAnyFlags(UhtPropertyExportFlags.GetterSpecifiedAuto | UhtPropertyExportFlags.SetterSpecifiedAuto));
-			if (!properties.Any())
+			UhtUsedDefineScopes<UhtProperty> autoGetterSetterProperties = GetAutoGetterSetterProperties(classObj);
+			if (autoGetterSetterProperties.IsEmpty)
 			{
 				return builder;
 			}
 
-			using (UhtMacroBlockEmitter blockEmitter = new(builder, "WITH_EDITORONLY_DATA", properties.Any(x => x.IsEditorOnlyProperty)))
-			{
-				foreach (UhtProperty property in classObj.Properties)
+			return builder.AppendInstances(autoGetterSetterProperties, UhtDefineScopeNames.Standard,
+				(builder, property) =>
 				{
-					if (property.PropertyExportFlags.HasAnyFlags(UhtPropertyExportFlags.GetterSpecifiedAuto | UhtPropertyExportFlags.SetterSpecifiedAuto))
+					if (property.PropertyExportFlags.HasAnyFlags(UhtPropertyExportFlags.GetterSpecifiedAuto))
 					{
-						blockEmitter.Set(property.IsEditorOnlyProperty);
-						if (property.PropertyExportFlags.HasAnyFlags(UhtPropertyExportFlags.GetterSpecifiedAuto))
-						{
-							using BorrowStringBuilder borrower = new(StringBuilderCache.Small);
-							StringBuilder propertyStringBuilder = borrower.StringBuilder;
-							propertyStringBuilder.AppendPropertyText(property, UhtPropertyTextType.GetterRetVal);
-							string getterRetText = propertyStringBuilder.ToString();
-							string getterCallText = property.Getter ?? "Get" + property.SourceName;
-							builder.Append('\t').Append(getterRetText).Append(classObj.SourceName).Append("::").Append(getterCallText).Append("() const\r\n");
-							builder.Append("\t{\r\n");
-							builder.Append('\t', 2).Append("return ").Append(property.SourceName).Append(";\r\n");
-							builder.Append("\t}\r\n");
-						}
-						if (property.PropertyExportFlags.HasAnyFlags(UhtPropertyExportFlags.SetterSpecifiedAuto))
-						{
-							using BorrowStringBuilder borrower = new(StringBuilderCache.Small);
-							StringBuilder propertyStringBuilder = borrower.StringBuilder;
-							propertyStringBuilder.AppendPropertyText(property, UhtPropertyTextType.SetterParameterArgType);
-							string setterArgText = propertyStringBuilder + "InValue";
-							string setterCallText = property.Setter ?? "Set" + property.SourceName;
-							builder.Append("\tvoid ").Append(classObj.SourceName).Append("::").Append(setterCallText).Append('(').Append(setterArgText).Append(")\r\n");
-							builder.Append("\t{\r\n");
-
-							// @todo: setter defn
-
-							builder.Append("\t}\r\n");
-						}
+						string getterCallText = property.Getter ?? "Get" + property.SourceName;
+						builder.Append('\t').AppendPropertyText(property, UhtPropertyTextType.GetterRetVal).Append(classObj.SourceName).Append("::").Append(getterCallText).Append("() const\r\n");
+						builder.Append("\t{\r\n");
+						builder.Append('\t', 2).Append("return ").Append(property.SourceName).Append(";\r\n");
+						builder.Append("\t}\r\n");
 					}
-				}
-			}
-			return builder;
+					if (property.PropertyExportFlags.HasAnyFlags(UhtPropertyExportFlags.SetterSpecifiedAuto))
+					{
+						string setterCallText = property.Setter ?? "Set" + property.SourceName;
+						builder.Append("\tvoid ").Append(classObj.SourceName).Append("::").Append(setterCallText).Append('(').AppendPropertyText(property, UhtPropertyTextType.SetterParameterArgType).Append("InValue").Append(")\r\n");
+						builder.Append("\t{\r\n");
+						// @todo: setter defn
+						builder.Append("\t}\r\n");
+					}
+				});
 		}
-
+		
 		private static StringBuilder AppendPropertyAccessors(StringBuilder builder, UhtClass classObj)
 		{
 			foreach (UhtType type in classObj.Children)
 			{
 				if (type is UhtProperty property)
 				{
-					bool editorOnlyProperty = property.IsEditorOnlyProperty;
 					if (property.PropertyExportFlags.HasAnyFlags(UhtPropertyExportFlags.GetterFound))
 					{
 						builder.Append("\tvoid ").Append(classObj.SourceName).Append("::").AppendPropertyGetterWrapperName(property).Append("(const void* Object, void* OutValue)\r\n");
 						builder.Append("\t{\r\n");
-						if (editorOnlyProperty)
-						{
-							builder.Append("#if WITH_EDITORONLY_DATA\r\n");
-						}
-						builder.Append("\t\tconst ").Append(classObj.SourceName).Append("* Obj = (const ").Append(classObj.SourceName).Append("*)Object;\r\n");
-						if (property.IsStaticArray)
-						{
-							builder
-								.Append("\t\t")
-								.AppendPropertyText(property, UhtPropertyTextType.GetterSetterArg)
-								.Append("* Source = (")
-								.AppendPropertyText(property, UhtPropertyTextType.GetterSetterArg)
-								.Append("*)Obj->")
-								.Append(property.Getter!)
-								.Append("();\r\n");
-							builder
-								.Append("\t\t")
-								.AppendPropertyText(property, UhtPropertyTextType.GetterSetterArg)
-								.Append("* Result = (")
-								.AppendPropertyText(property, UhtPropertyTextType.GetterSetterArg)
-								.Append("*)OutValue;\r\n");
-							builder
-								.Append("\t\tCopyAssignItems(Result, Source, ")
-								.Append(property.ArrayDimensions)
-								.Append(");\r\n");
-                        }
-                        else if (property is UhtByteProperty byteProperty && byteProperty.Enum != null)
-                        {
-                            // If someone passed in a TEnumAsByte instead of the actual enum value, the cast in the else clause would cause an issue.
-                            // Since this is known to be a TEnumAsByte, we just fetch the first byte.  *HOWEVER* on MSB machines where 
-                            // the actual enum value is passed in this will fail and return zero if the native size of the enum > 1 byte.
-                            builder
-                                .Append("\t\t")
-                                .Append("uint8")
-                                .Append("& Result = *(")
-                                .Append("uint8")
-                                .Append("*)OutValue;\r\n");
-                            builder
-                                .Append("\t\tResult = (")
-                                .Append("uint8")
-                                .Append(")Obj->")
-                                .Append(property.Getter!)
-                                .Append("();\r\n");
-                        }
-                        else
-						{
-							builder
-								.Append("\t\t")
-								.AppendPropertyText(property, UhtPropertyTextType.GetterSetterArg)
-								.Append("& Result = *(")
-								.AppendPropertyText(property, UhtPropertyTextType.GetterSetterArg)
-								.Append("*)OutValue;\r\n");
-							builder
-								.Append("\t\tResult = (")
-								.AppendPropertyText(property, UhtPropertyTextType.GetterSetterArg)
-								.Append(")Obj->")
-								.Append(property.Getter!)
-								.Append("();\r\n");
-						}
-						if (editorOnlyProperty)
+						using (UhtMacroBlockEmitter blockEmitter = new(builder, UhtDefineScopeNames.Standard, property.DefineScope))
 						{
-							builder.Append("#endif // WITH_EDITORONLY_DATA\r\n");
+							builder.Append("\t\tconst ").Append(classObj.SourceName).Append("* Obj = (const ").Append(classObj.SourceName).Append("*)Object;\r\n");
+							if (property.IsStaticArray)
+							{
+								builder
+									.Append("\t\t")
+									.AppendPropertyText(property, UhtPropertyTextType.GetterSetterArg)
+									.Append("* Source = (")
+									.AppendPropertyText(property, UhtPropertyTextType.GetterSetterArg)
+									.Append("*)Obj->")
+									.Append(property.Getter!)
+									.Append("();\r\n");
+								builder
+									.Append("\t\t")
+									.AppendPropertyText(property, UhtPropertyTextType.GetterSetterArg)
+									.Append("* Result = (")
+									.AppendPropertyText(property, UhtPropertyTextType.GetterSetterArg)
+									.Append("*)OutValue;\r\n");
+								builder
+									.Append("\t\tCopyAssignItems(Result, Source, ")
+									.Append(property.ArrayDimensions)
+									.Append(");\r\n");
+							}
+							else if (property is UhtByteProperty byteProperty && byteProperty.Enum != null)
+							{
+								// If someone passed in a TEnumAsByte instead of the actual enum value, the cast in the else clause would cause an issue.
+								// Since this is known to be a TEnumAsByte, we just fetch the first byte.  *HOWEVER* on MSB machines where 
+								// the actual enum value is passed in this will fail and return zero if the native size of the enum > 1 byte.
+								builder
+									.Append("\t\t")
+									.Append("uint8")
+									.Append("& Result = *(")
+									.Append("uint8")
+									.Append("*)OutValue;\r\n");
+								builder
+									.Append("\t\tResult = (")
+									.Append("uint8")
+									.Append(")Obj->")
+									.Append(property.Getter!)
+									.Append("();\r\n");
+							}
+							else
+							{
+								builder
+									.Append("\t\t")
+									.AppendPropertyText(property, UhtPropertyTextType.GetterSetterArg)
+									.Append("& Result = *(")
+									.AppendPropertyText(property, UhtPropertyTextType.GetterSetterArg)
+									.Append("*)OutValue;\r\n");
+								builder
+									.Append("\t\tResult = (")
+									.AppendPropertyText(property, UhtPropertyTextType.GetterSetterArg)
+									.Append(")Obj->")
+									.Append(property.Getter!)
+									.Append("();\r\n");
+							}
 						}
 						builder.Append("\t}\r\n");
 					}
@@ -1333,48 +1229,43 @@ namespace EpicGames.UHT.Exporters.CodeGen
 					{
 						builder.Append("\tvoid ").Append(classObj.SourceName).Append("::").AppendPropertySetterWrapperName(property).Append("(void* Object, const void* InValue)\r\n");
 						builder.Append("\t{\r\n");
-						if (editorOnlyProperty)
-						{
-							builder.Append("#if WITH_EDITORONLY_DATA\r\n");
-						}
-						builder.Append("\t\t").Append(classObj.SourceName).Append("* Obj = (").Append(classObj.SourceName).Append("*)Object;\r\n");
-						if (property.IsStaticArray)
-						{
-							builder
-								.Append("\t\t")
-								.AppendPropertyText(property, UhtPropertyTextType.GetterSetterArg)
-								.Append("* Value = (")
-								.AppendPropertyText(property, UhtPropertyTextType.GetterSetterArg)
-								.Append("*)InValue;\r\n");
-						}
-						else if (property is UhtByteProperty byteProperty && byteProperty.Enum != null)
-						{
-							// If someone passed in a TEnumAsByte instead of the actual enum value, the cast in the else clause would cause an issue.
-							// Since this is known to be a TEnumAsByte, we just fetch the first byte.  *HOWEVER* on MSB machines where 
-							// the actual enum value is passed in this will fail and return zero if the native size of the enum > 1 byte.
-							builder
-								.Append("\t\t")
-								.AppendPropertyText(property, UhtPropertyTextType.GetterSetterArg)
-								.Append(" Value = (")
-								.AppendPropertyText(property, UhtPropertyTextType.GetterSetterArg)
-								.Append(")*(uint8*)InValue;\r\n");
-						}
-						else
+						using (UhtMacroBlockEmitter blockEmitter = new(builder, UhtDefineScopeNames.Standard, property.DefineScope))
 						{
+							builder.Append("\t\t").Append(classObj.SourceName).Append("* Obj = (").Append(classObj.SourceName).Append("*)Object;\r\n");
+							if (property.IsStaticArray)
+							{
+								builder
+									.Append("\t\t")
+									.AppendPropertyText(property, UhtPropertyTextType.GetterSetterArg)
+									.Append("* Value = (")
+									.AppendPropertyText(property, UhtPropertyTextType.GetterSetterArg)
+									.Append("*)InValue;\r\n");
+							}
+							else if (property is UhtByteProperty byteProperty && byteProperty.Enum != null)
+							{
+								// If someone passed in a TEnumAsByte instead of the actual enum value, the cast in the else clause would cause an issue.
+								// Since this is known to be a TEnumAsByte, we just fetch the first byte.  *HOWEVER* on MSB machines where 
+								// the actual enum value is passed in this will fail and return zero if the native size of the enum > 1 byte.
+								builder
+									.Append("\t\t")
+									.AppendPropertyText(property, UhtPropertyTextType.GetterSetterArg)
+									.Append(" Value = (")
+									.AppendPropertyText(property, UhtPropertyTextType.GetterSetterArg)
+									.Append(")*(uint8*)InValue;\r\n");
+							}
+							else
+							{
+								builder
+									.Append("\t\t")
+									.AppendPropertyText(property, UhtPropertyTextType.GetterSetterArg)
+									.Append("& Value = *(")
+									.AppendPropertyText(property, UhtPropertyTextType.GetterSetterArg)
+									.Append("*)InValue;\r\n");
+							}
 							builder
-								.Append("\t\t")
-								.AppendPropertyText(property, UhtPropertyTextType.GetterSetterArg)
-								.Append("& Value = *(")
-								.AppendPropertyText(property, UhtPropertyTextType.GetterSetterArg)
-								.Append("*)InValue;\r\n");
-						}
-						builder
-							.Append("\t\tObj->")
-							.Append(property.Setter!)
-							.Append("(Value);\r\n");
-						if (editorOnlyProperty)
-						{
-							builder.Append("#endif // WITH_EDITORONLY_DATA\r\n");
+								.Append("\t\tObj->")
+								.Append(property.Setter!)
+								.Append("(Value);\r\n");
 						}
 						builder.Append("\t}\r\n");
 					}
@@ -1523,20 +1414,11 @@ namespace EpicGames.UHT.Exporters.CodeGen
 			bool hasInterfaces = classObj.Bases.Any(x => x is UhtClass baseClass && baseClass.ClassFlags.HasAnyFlags(EClassFlags.Interface));
 
 			// Collect the functions to be exported
-			bool allEditorOnlyFunctions = true;
-			List<UhtFunction> sortedFunctions = new();
-			foreach (UhtFunction function in classObj.Functions)
-			{
-				//if (!function.FunctionFlags.HasAnyFlags(EFunctionFlags.Delegate))
-				{
-					allEditorOnlyFunctions &= function.FunctionFlags.HasAnyFlags(EFunctionFlags.EditorOnly);
-				}
-				sortedFunctions.Add(function);
-			}
-			sortedFunctions.Sort((x, y) => StringComparerUE.OrdinalIgnoreCase.Compare(x.EngineName, y.EngineName));
+			UhtUsedDefineScopes<UhtFunction> functions = new(classObj.Functions);
+			functions.Instances.Sort((x, y) => StringComparerUE.OrdinalIgnoreCase.Compare(x.EngineName, y.EngineName));
 
 			// Output any function
-			foreach (UhtFunction function in sortedFunctions)
+			foreach (UhtFunction function in functions.Instances)
 			{
 				if (!function.FunctionFlags.HasAnyFlags(EFunctionFlags.Delegate))
 				{
@@ -1557,22 +1439,20 @@ namespace EpicGames.UHT.Exporters.CodeGen
 				builder.Append("\t}\r\n");
 			}
 
+			// Collect the properties
+			UhtUsedDefineScopes<UhtProperty> properties = new(classObj.Properties);
+
 			// Declare the statics object
 			{
 				builder.Append("\tstruct ").Append(staticsName).Append("\r\n");
 				builder.Append("\t{\r\n");
 				builder.Append("\t\tstatic UObject* (*const DependentSingletons[])();\r\n");
 
-				if (sortedFunctions.Count != 0)
-				{
-					builder.AppendBeginEditorOnlyGuard(allEditorOnlyFunctions);
-					builder.Append("\t\tstatic const FClassFunctionLinkInfo FuncInfo[];\r\n");
-					builder.AppendEndEditorOnlyGuard(allEditorOnlyFunctions);
-				}
+				builder.AppendIfInstances(functions, UhtDefineScopeNames.WithEditor, builder => builder.Append("\t\tstatic const FClassFunctionLinkInfo FuncInfo[];\r\n"));
 
 				builder.AppendMetaDataDecl(classObj, MetaDataParamsName, 2);
 
-				AppendPropertiesDecl(builder, classObj, classObj.SourceName, staticsName, 2);
+				AppendPropertiesDecl(builder, classObj, classObj.SourceName, properties, staticsName, 2);
 
 				if (hasInterfaces)
 				{
@@ -1598,14 +1478,13 @@ namespace EpicGames.UHT.Exporters.CodeGen
 
 				builder.Append("\tstatic_assert(UE_ARRAY_COUNT(").Append(staticsName).Append("::DependentSingletons) < 16);\r\n");
 
-				if (sortedFunctions.Count > 0)
-				{
-					using UhtMacroBlockEmitter blockEmitter = new(builder, "WITH_EDITOR", allEditorOnlyFunctions);
-					builder.Append("\tconst FClassFunctionLinkInfo ").Append(staticsName).Append("::FuncInfo[] = {\r\n");
-
-					foreach (UhtFunction function in sortedFunctions)
+				builder.AppendInstances(functions, UhtDefineScopeNames.WithEditor,
+					builder =>
+					{
+						builder.Append("\tconst FClassFunctionLinkInfo ").Append(staticsName).Append("::FuncInfo[] = {\r\n");
+					},
+					(builder, function) =>
 					{
-						blockEmitter.Set(function.FunctionFlags.HasAnyFlags(EFunctionFlags.EditorOnly));
 						builder
 							.Append("\t\t{ &")
 							.Append(GetSingletonName(function, true))
@@ -1614,16 +1493,16 @@ namespace EpicGames.UHT.Exporters.CodeGen
 							.Append(" },")
 							.AppendObjectHash(classObj, context, function)
 							.Append("\r\n");
-					}
-
-					blockEmitter.Set(allEditorOnlyFunctions);
-					builder.Append("\t};\r\n");
-					builder.Append("\tstatic_assert(UE_ARRAY_COUNT(").Append(staticsName).Append("::FuncInfo) < 2048);\r\n");
-				}
+					},
+					builder =>
+					{
+						builder.Append("\t};\r\n");
+						builder.Append("\tstatic_assert(UE_ARRAY_COUNT(").Append(staticsName).Append("::FuncInfo) < 2048);\r\n");
+					});
 
 				builder.AppendMetaDataDef(classObj, staticsName, MetaDataParamsName, 1);
 
-				AppendPropertiesDefs(builder, classObj, classObj.SourceName, staticsName, 1);
+				AppendPropertiesDefs(builder, classObj, classObj.SourceName, properties, staticsName, 1);
 
 				if (hasInterfaces)
 				{
@@ -1679,40 +1558,18 @@ namespace EpicGames.UHT.Exporters.CodeGen
 				}
 				builder.Append("\t\t&StaticCppClassTypeInfo,\r\n");
 				builder.Append("\t\tDependentSingletons,\r\n");
-				if (sortedFunctions.Count == 0)
-				{
-					builder.Append("\t\tnullptr,\r\n");
-				}
-				else if (allEditorOnlyFunctions)
-				{
-					builder.Append("\t\tIF_WITH_EDITOR(FuncInfo, nullptr),\r\n");
-				}
-				else
-				{
-					builder.Append("\t\tFuncInfo,\r\n");
-				}
-				builder.AppendPropertiesParamsList(classObj, staticsName, 2, "\r\n");
+				builder.AppendArrayPtrLine(functions, UhtDefineScopeNames.WithEditor, null, "FuncInfo", 2, ",\r\n");
+				builder.AppendArrayPtrLine(properties, UhtDefineScopeNames.Standard, staticsName, "PropPointers", 2, ",\r\n");
 				builder.Append("\t\t").Append(hasInterfaces ? "InterfaceParams" : "nullptr").Append(",\r\n");
 				builder.Append("\t\tUE_ARRAY_COUNT(DependentSingletons),\r\n");
-				if (sortedFunctions.Count == 0)
-				{
-					builder.Append("\t\t0,\r\n");
-				}
-				else if (allEditorOnlyFunctions)
-				{
-					builder.Append("\t\tIF_WITH_EDITOR(UE_ARRAY_COUNT(FuncInfo), 0),\r\n");
-				}
-				else
-				{
-					builder.Append("\t\tUE_ARRAY_COUNT(FuncInfo),\r\n");
-				}
-				builder.AppendPropertiesParamsCount(classObj, staticsName, 2, "\r\n");
+				builder.AppendArrayCountLine(functions, UhtDefineScopeNames.WithEditor, null, "FuncInfo", 2, ",\r\n");
+				builder.AppendArrayCountLine(properties, UhtDefineScopeNames.Standard, staticsName, "PropPointers", 2, ",\r\n");
 				builder.Append("\t\t").Append(hasInterfaces ? "UE_ARRAY_COUNT(InterfaceParams)" : "0").Append(",\r\n");
 				builder.Append($"\t\t0x{(uint)classFlags:X8}u,\r\n");
 				builder.Append("\t\t").AppendMetaDataParams(classObj, staticsName, MetaDataParamsName).Append("\r\n");
 				builder.Append("\t};\r\n");
 
-				builder.AppendPropertiesParamsCountStaticAssert(classObj, staticsName);
+				builder.AppendPropertiesParamsCountStaticAssert(properties, staticsName);
 			}
 
 			// Class registration
@@ -1840,29 +1697,19 @@ namespace EpicGames.UHT.Exporters.CodeGen
 			builder.Append("\tvoid ").Append(classObj.SourceName).Append("::StaticRegisterNatives").Append(classObj.SourceName).Append("()\r\n");
 			builder.Append("\t{\r\n");
 
-			bool allEditorOnly = true;
-
-			List<UhtFunction> sortedFunctions = new();
-			foreach (UhtFunction function in classObj.Functions)
-			{
-				if (function.FunctionFlags.HasExactFlags(EFunctionFlags.Native | EFunctionFlags.NetRequest, EFunctionFlags.Native))
-				{
-					sortedFunctions.Add(function);
-					allEditorOnly &= function.FunctionFlags.HasAnyFlags(EFunctionFlags.EditorOnly);
-				}
-			}
-			sortedFunctions.Sort((x, y) => StringComparerUE.OrdinalIgnoreCase.Compare(x.EngineName, y.EngineName));
+			UhtUsedDefineScopes<UhtFunction> functions = new(classObj.Functions.Where(x => x.FunctionFlags.HasExactFlags(EFunctionFlags.Native | EFunctionFlags.NetRequest, EFunctionFlags.Native)));
+			functions.Instances.Sort((x, y) => StringComparerUE.OrdinalIgnoreCase.Compare(x.EngineName, y.EngineName));
 
-			if (sortedFunctions.Count != 0)
+			if (!functions.IsEmpty)
 			{
 				{
-					using UhtMacroBlockEmitter blockEmitter = new(builder, "WITH_EDITOR", allEditorOnly);
+					using UhtMacroBlockEmitter blockEmitter = new(builder, UhtDefineScopeNames.WithEditor, functions.SoleScope);
 					builder.Append("\t\tUClass* Class = ").Append(classObj.SourceName).Append("::StaticClass();\r\n");
 					builder.Append("\t\tstatic const FNameNativePtrPair Funcs[] = {\r\n");
 
-					foreach (UhtFunction function in sortedFunctions)
+					foreach (UhtFunction function in functions.Instances)
 					{
-						blockEmitter.Set(function.FunctionFlags.HasAnyFlags(EFunctionFlags.EditorOnly));
+						blockEmitter.Set(function.DefineScope);
 						builder
 							.Append("\t\t\t{ ")
 							.AppendUTF8LiteralString(function.EngineName)
@@ -1874,7 +1721,7 @@ namespace EpicGames.UHT.Exporters.CodeGen
 					}
 
 					// This will close the block if we have one that isn't editor only
-					blockEmitter.Set(allEditorOnly);
+					blockEmitter.Set(functions.SoleScope);
 
 					builder.Append("\t\t};\r\n");
 					builder.Append("\t\tFNativeFunctionRegistrar::RegisterFunctions(Class, Funcs, UE_ARRAY_COUNT(Funcs));\r\n");
@@ -1900,7 +1747,7 @@ namespace EpicGames.UHT.Exporters.CodeGen
 			{
 				bool isInterfaceClass = classObj.ClassFlags.HasAnyFlags(EClassFlags.Interface);
 				{
-					using UhtMacroBlockEmitter blockEmitter = new(builder, "WITH_EDITOR");
+					using UhtMacroBlockEmitter blockEmitter = new(builder, UhtDefineScopeNames.WithEditor, UhtDefineScope.None);
 					foreach (UhtFunction function in callbackFunctions)
 					{
 						// Net response functions don't go into the VM
@@ -1909,7 +1756,7 @@ namespace EpicGames.UHT.Exporters.CodeGen
 							continue;
 						}
 
-						blockEmitter.Set(function.FunctionFlags.HasAnyFlags(EFunctionFlags.EditorOnly));
+						blockEmitter.Set(function.DefineScope);
 
 						if (!isInterfaceClass)
 						{
@@ -1965,18 +1812,21 @@ namespace EpicGames.UHT.Exporters.CodeGen
 			return builder;
 		}
 
-		private static StringBuilder AppendRpcFunctions(StringBuilder builder, UhtClass classObj, List<UhtFunction> reversedFunctions, bool editorOnly)
+		private static StringBuilder AppendRpcFunctions(StringBuilder builder, UhtClass classObj)
 		{
-			foreach (UhtFunction function in reversedFunctions)
+			// Collect the functions in reversed order
+			UhtUsedDefineScopes<UhtFunction> functions = new(classObj.Functions.Where(x => IsRpcFunction(x) && ShouldExportFunction(x)));
+			functions.Instances.Reverse();
+			functions.OrderByDefineScope();
+
+			// Output the RPC methods
+			builder.AppendInstances(functions, UhtDefineScopeNames.WithEditor, (builder, function) => 
 			{
-				if (function.FunctionFlags.HasAnyFlags(EFunctionFlags.EditorOnly) == editorOnly)
-				{
-					builder.Append("\tDEFINE_FUNCTION(").AppendClassSourceNameOrInterfaceName(classObj).Append("::").Append(function.UnMarshalAndCallName).Append(")\r\n");
-					builder.Append("\t{\r\n");
-					AppendFunctionThunk(builder, function);
-					builder.Append("\t}\r\n");
-				}
-			}
+				builder.Append("\tDEFINE_FUNCTION(").AppendClassSourceNameOrInterfaceName(classObj).Append("::").Append(function.UnMarshalAndCallName).Append(")\r\n");
+				builder.Append("\t{\r\n");
+				AppendFunctionThunk(builder, function);
+				builder.Append("\t}\r\n");
+			});
 			return builder;
 		}
 
@@ -2138,137 +1988,26 @@ namespace EpicGames.UHT.Exporters.CodeGen
 					builder.AppendFunctionThunkParameterArg(property);
 					if (needsGCBarrier)
 					{
-						builder.Append(")");
+						builder.Append(')');
 					}
 				}
 			}
 			return builder;
 		}
 
-		/// <summary>
-		/// Append the parameter list and count as arguments to a structure constructor
-		/// </summary>
-		/// <param name="builder">Destination builder</param>
-		/// <param name="structObj">Structure in question</param>
-		/// <param name="staticsName">Name of the statics section for this structure</param>
-		/// <param name="tabs">Number of tabs to precede the line</param>
-		/// <param name="endl">String used to terminate a line</param>
-		/// <returns>Destination builder</returns>
-		public static StringBuilder AppendPropertiesParams(this StringBuilder builder, UhtStruct structObj, string staticsName, int tabs, string endl)
-		{
-			return builder.AppendPropertiesParamsList(structObj, staticsName, tabs, endl).AppendPropertiesParamsCount(structObj, staticsName, tabs, endl);
-		}
-
-		/// <summary>
-		/// Append the parameter list as an argument to a structure constructor
-		/// </summary>
-		/// <param name="builder">Destination builder</param>
-		/// <param name="structObj">Structure in question</param>
-		/// <param name="staticsName">Name of the statics section for this structure</param>
-		/// <param name="tabs">Number of tabs to precede the line</param>
-		/// <param name="endl">String used to terminate a line</param>
-		/// <returns>Destination builder</returns>
-		public static StringBuilder AppendPropertiesParamsList(this StringBuilder builder, UhtStruct structObj, string staticsName, int tabs, string endl)
-		{
-			if (!structObj.Children.Any(x => x is UhtProperty))
-			{
-				builder.AppendTabs(tabs).Append("nullptr,").Append(endl);
-			}
-			else if (structObj.Properties.Any(x => !x.IsEditorOnlyProperty))
-			{
-				builder.AppendTabs(tabs).Append(staticsName).Append("::PropPointers,").Append(endl);
-			}
-			else
-			{
-				builder.AppendTabs(tabs).Append("IF_WITH_EDITORONLY_DATA(").Append(staticsName).Append("::PropPointers, nullptr),").Append(endl);
-			}
-			return builder;
-		}
-
-		/// <summary>
-		/// Append the parameter count as an argument to a structure constructor
-		/// </summary>
-		/// <param name="builder">Destination builder</param>
-		/// <param name="structObj">Structure in question</param>
-		/// <param name="staticsName">Name of the statics section for this structure</param>
-		/// <param name="tabs">Number of tabs to precede the line</param>
-		/// <param name="endl">String used to terminate a line</param>
-		/// <returns>Destination builder</returns>
-		public static StringBuilder AppendPropertiesParamsCount(this StringBuilder builder, UhtStruct structObj, string staticsName, int tabs, string endl)
-		{
-			if (!structObj.Children.Any(x => x is UhtProperty))
-			{
-				builder.AppendTabs(tabs).Append("0,").Append(endl);
-			}
-			else if (structObj.Properties.Any(x => !x.IsEditorOnlyProperty))
-			{
-				builder.AppendTabs(tabs).Append("UE_ARRAY_COUNT(").Append(staticsName).Append("::PropPointers),").Append(endl);
-			}
-			else
-			{
-				builder.AppendTabs(tabs).Append("IF_WITH_EDITORONLY_DATA(UE_ARRAY_COUNT(").Append(staticsName).Append("::PropPointers), 0),").Append(endl);
-			}
-			return builder;
-		}
-
 		/// <summary>
 		/// Append static_assert to check that properties params count fit within struct
 		/// </summary>
 		/// <param name="builder">Destination builder</param>
-		/// <param name="structObj">Structure in question</param>
+		/// <param name="properties">Collection of properties</param>
 		/// <param name="staticsName">Name of the statics section for this structure</param>
 		/// <returns>Destination builder</returns>
-		public static StringBuilder AppendPropertiesParamsCountStaticAssert(this StringBuilder builder, UhtStruct structObj, string staticsName)
+		public static StringBuilder AppendPropertiesParamsCountStaticAssert(this StringBuilder builder, UhtUsedDefineScopes<UhtProperty> properties, string staticsName)
 		{
-			if (!structObj.Children.Any(x => x is UhtProperty))
-			{
-				// do nothing here
-			}
-			else if (structObj.Properties.Any(x => !x.IsEditorOnlyProperty))
+			builder.AppendIfInstances(properties, UhtDefineScopeNames.Standard, builder =>
 			{
 				builder.Append("\tstatic_assert(UE_ARRAY_COUNT(").Append(staticsName).Append("::PropPointers) < 2048);\r\n");
-			}
-			else
-			{
-				builder.Append("#if WITH_EDITORONLY_DATA\r\n")
-					   .Append("\tstatic_assert(UE_ARRAY_COUNT(").Append(staticsName).Append("::PropPointers) < 2048);\r\n")
-					   .Append("#endif // WITH_EDITORONLY_DATA\r\n");
-			}
-			return builder;
-		}
-
-		/// <summary>
-		/// Append the given array list and count as arguments to a structure constructor
-		/// </summary>
-		/// <param name="builder">Destination builder</param>
-		/// <param name="isEmpty">If true, the array list is empty</param>
-		/// <param name="allEditorOnlyData">If true, the array is all editor only data</param>
-		/// <param name="staticsName">Name of the statics section</param>
-		/// <param name="arrayName">The name of the arrray</param>
-		/// <returns>Destination builder</returns>
-		public static StringBuilder AppendArray(this StringBuilder builder, bool isEmpty, bool allEditorOnlyData, string staticsName, string arrayName)
-		{
-			if (isEmpty)
-			{
-				builder.Append("nullptr, 0");
-			}
-			else if (allEditorOnlyData)
-			{
-				builder
-					.Append("IF_WITH_EDITORONLY_DATA(")
-					.Append(staticsName).Append("::").Append(arrayName)
-					.Append(", nullptr), IF_WITH_EDITORONLY_DATA(UE_ARRAY_COUNT(")
-					.Append(staticsName).Append("::").Append(arrayName)
-					.Append("), 0)");
-			}
-			else
-			{
-				builder
-					.Append(staticsName).Append("::").Append(arrayName)
-					.Append(", UE_ARRAY_COUNT(")
-					.Append(staticsName).Append("::").Append(arrayName)
-					.Append(')');
-			}
+			});
 			return builder;
 		}
 	}
diff --git a/Engine/Source/Programs/Shared/EpicGames.UHT/Exporters/CodeGen/UhtHeaderCodeGeneratorHFile.cs b/Engine/Source/Programs/Shared/EpicGames.UHT/Exporters/CodeGen/UhtHeaderCodeGeneratorHFile.cs
--- a/Engine/Source/Programs/Shared/EpicGames.UHT/Exporters/CodeGen/UhtHeaderCodeGeneratorHFile.cs
+++ b/Engine/Source/Programs/Shared/EpicGames.UHT/Exporters/CodeGen/UhtHeaderCodeGeneratorHFile.cs
@@ -521,26 +521,19 @@ namespace EpicGames.UHT.Exporters.CodeGen
 			// With interfaces, there are cases where we use the native interface class export flags to check for legacy macro
 			bool alternateUsesLegacy = (classObj.ClassFlags.HasAnyFlags(EClassFlags.Interface) ? nativeInterface! : classObj).ClassExportFlags.HasAnyFlags(UhtClassExportFlags.UsesGeneratedBodyLegacy);
 
-			// Collect sparse data
+			// Collect sparse data information
 			IEnumerable<UhtScriptStruct> sparseScriptStructs = GetSparseDataStructsToExport(classObj);
-			CalculateSparsePropertyDeclarations(sparseScriptStructs, out Dictionary<UhtScriptStruct, List<UhtProperty>> commonSparseProperties, out Dictionary<UhtScriptStruct, List<UhtProperty>> editorOnlySparseProperties);
+			UhtUsedDefineScopes<UhtProperty> sparseProperties = new(EnumerateSparseDataStructProperties(sparseScriptStructs));
 
 			// Write the spare declarations
-			AppendSparseDeclarations(builder, classObj, sparseScriptStructs, commonSparseProperties, editorOnlySparseProperties);
+			AppendSparseDeclarations(builder, classObj, sparseScriptStructs, sparseProperties);
 
-			// Collect the functions in reversed order
-			List<UhtFunction> reversedFunctions = new(classObj.Functions);
-			reversedFunctions.Reverse();
-
-			// Check to see if we have any RPC functions for the editor
-			bool hasEditorRpc = reversedFunctions.Any(x => IsRpcFunction(x, true));
+			// Collect the rpc functions in reversed order
+			UhtUsedDefineScopes<UhtFunction> rpcFunctions = new(classObj.Functions.Where(x => IsRpcFunction(x)));
+			rpcFunctions.Instances.Reverse();
 
 			// Output the RPC methods
-			AppendRpcFunctions(builder, classObj, alternateUsesLegacy, reversedFunctions, false);
-			if (hasEditorRpc)
-			{
-				AppendRpcFunctions(builder, classObj, alternateUsesLegacy, reversedFunctions, true);
-			}
+			AppendRpcFunctions(builder, classObj, alternateUsesLegacy, rpcFunctions);
 
 			// Output property accessors
 			AppendPropertyAccessors(builder, classObj);
@@ -549,10 +542,9 @@ namespace EpicGames.UHT.Exporters.CodeGen
 			List<UhtFunction> callbackFunctions = new(classObj.Functions.Where(x => x.FunctionFlags.HasAnyFlags(EFunctionFlags.Event) && x.SuperFunction == null));
 			callbackFunctions.Sort((x, y) => StringComparerUE.OrdinalIgnoreCase.Compare(x.EngineName, y.EngineName));
 			bool hasCallbacks = callbackFunctions.Count > 0;
-			
-			// Determine if auto getters/setters have been generated
-			bool hasAutoGettersSetters = classObj.Properties.Any(x => x.PropertyExportFlags.HasAnyFlags(UhtPropertyExportFlags.GetterSpecifiedAuto | UhtPropertyExportFlags.SetterSpecifiedAuto));
-			bool hasEditorOnlyAutoGettersSetters = classObj.Properties.Any(x => x.IsEditorOnlyProperty && x.PropertyExportFlags.HasAnyFlags(UhtPropertyExportFlags.GetterSpecifiedAuto | UhtPropertyExportFlags.SetterSpecifiedAuto));
+
+			// Collect the auto getter setter properties
+			UhtUsedDefineScopes<UhtProperty> autoGetterSetterProperties = GetAutoGetterSetterProperties(classObj);
 
 			// Generate the RPC wrappers for the callbacks
 			AppendCallbackRpcWrapperDecls(builder, classObj, callbackFunctions);
@@ -613,8 +605,8 @@ namespace EpicGames.UHT.Exporters.CodeGen
 				}
 
 				AppendProlog(builder, classObj);
-				AppendGeneratedBodyMacroBlock(builder, classObj, nativeInterface!, alternateUsesLegacy, hasEditorRpc, hasCallbacks, hasAutoGettersSetters, 
-					hasEditorOnlyAutoGettersSetters, sparseScriptStructs.Any(), editorOnlySparseProperties.Count > 0, null);
+				AppendGeneratedBodyMacroBlock(builder, classObj, nativeInterface!, alternateUsesLegacy, rpcFunctions, autoGetterSetterProperties, 
+					sparseScriptStructs.Any(), sparseProperties, hasCallbacks, null);
 			}
 			else
 			{
@@ -636,10 +628,10 @@ namespace EpicGames.UHT.Exporters.CodeGen
 				}
 
 				AppendFieldNotify(builder, classObj);
-				AppendAutoGettersSetters(builder, classObj);
+				AppendAutoGettersSetters(builder, classObj, autoGetterSetterProperties);
 				AppendProlog(builder, classObj);
-				AppendGeneratedBodyMacroBlock(builder, classObj, classObj, usesLegacy, hasEditorRpc, hasCallbacks, hasAutoGettersSetters, hasEditorOnlyAutoGettersSetters,
-					sparseScriptStructs.Any(), editorOnlySparseProperties.Count > 0, usesLegacy ? "GENERATED_UCLASS_BODY" : null);
+				AppendGeneratedBodyMacroBlock(builder, classObj, classObj, usesLegacy, rpcFunctions, autoGetterSetterProperties, 
+					sparseScriptStructs.Any(), sparseProperties, hasCallbacks, usesLegacy ? "GENERATED_UCLASS_BODY" : null);
 			}
 
 			// Forward declare the StaticClass specialization in the header
@@ -667,154 +659,109 @@ namespace EpicGames.UHT.Exporters.CodeGen
 				return builder;
 			}
 
-			// Scan the children to see what we have
-			GetFieldNotifyStats(classObj, out bool hasProperties, out bool hasFunctions, out bool hasEditorFields, out bool allEditorFields);
-
-			// If we have editor fields, then we must place them in #if blocks
-			if (hasEditorFields)
-			{
-				builder.Append("#if WITH_EDITORONLY_DATA\r\n");
-				using (UhtMacroCreator macro = new(builder, this, classObj, FieldNotifyMacroSuffix))
+			UhtUsedDefineScopes<UhtType> notifyTypes = GetFieldNotifyTypes(classObj);
+			return builder.AppendSingleMacro(notifyTypes, UhtDefineScopeNames.Standard, this, classObj, FieldNotifyMacroSuffix, false,
+				(builder, filteredTypes) =>
 				{
-					AppendFieldNotify(builder, classObj, hasProperties, hasFunctions, hasEditorFields, allEditorFields, true);
-				}
-				//if (!allEditorFields)
-				{
-					builder.Append("#else\r\n");
-					using UhtMacroCreator macro = new(builder, this, classObj, FieldNotifyMacroSuffix);
-					AppendFieldNotify(builder, classObj, hasProperties, hasFunctions, hasEditorFields, allEditorFields, false);
-				}
-				builder.Append("#endif\r\n");
-			}
-			else
-			{
-				using UhtMacroCreator macro = new(builder, this, classObj, FieldNotifyMacroSuffix);
-				AppendFieldNotify(builder, classObj, hasProperties, hasFunctions, hasEditorFields, allEditorFields, false);
-			}
-			return builder;
-		}
-
-		private StringBuilder AppendFieldNotify(StringBuilder builder, UhtClass classObj,
-			bool hasProperties, bool hasFunctions, bool hasEditorFields, bool allEditorFields,
-			bool includeEditorOnlyFields)
-		{
-			builder.Append("\tUE_FIELD_NOTIFICATION_DECLARE_CLASS_DESCRIPTOR_BEGIN(").Append(PackageApi).Append(") \\\r\n");
-
-			//UE_FIELD_NOTIFICATION_DECLARE_FIELD
-			AppendFieldNotify(builder, classObj, hasProperties, hasFunctions, hasEditorFields, allEditorFields,
-				includeEditorOnlyFields, false, (StringBuilder builder, UhtClass classObj, string name) =>
-			{
-				builder.Append("\tUE_FIELD_NOTIFICATION_DECLARE_FIELD(").Append(name).Append(") \\\r\n");
-			});
+					builder.Append("\tUE_FIELD_NOTIFICATION_DECLARE_CLASS_DESCRIPTOR_BEGIN(").Append(PackageApi).Append(") \\\r\n");
 
-			//UE_FIELD_NOTIFICATION_DECLARE_ENUM_FIELD
-			bool isFirst = true;
-			AppendFieldNotify(builder, classObj, hasProperties, hasFunctions, hasEditorFields, allEditorFields,
-				includeEditorOnlyFields, false, (StringBuilder builder, UhtClass classObj, string name) =>
-			{
-				if (isFirst)
-				{
-					isFirst = false;
-					builder.Append("\tUE_FIELD_NOTIFICATION_DECLARE_ENUM_FIELD_BEGIN(").Append(name).Append(") \\\r\n");
-				}
-				else
-				{
-					builder.Append("\tUE_FIELD_NOTIFICATION_DECLARE_ENUM_FIELD(").Append(name).Append(") \\\r\n");
-				}
-			});
+					// UE_FIELD_NOTIFICATION_DECLARE_FIELD
+					foreach (UhtType notifyType in filteredTypes)
+					{
+						builder.Append("\tUE_FIELD_NOTIFICATION_DECLARE_FIELD(").Append(GetNotifyTypeName(notifyType)).Append(") \\\r\n");
+					}
 
-			builder.Append("\tUE_FIELD_NOTIFICATION_DECLARE_ENUM_FIELD_END() \\\r\n");
-			builder.Append("\tUE_FIELD_NOTIFICATION_DECLARE_CLASS_DESCRIPTOR_END(); \\\r\n");
-			return builder;
+					// UE_FIELD_NOTIFICATION_DECLARE_ENUM_FIELD
+					bool isFirst = true;
+					foreach (UhtType notifyType in filteredTypes)
+					{
+						if (isFirst)
+						{
+							isFirst = false;
+							builder.Append("\tUE_FIELD_NOTIFICATION_DECLARE_ENUM_FIELD_BEGIN(").Append(GetNotifyTypeName(notifyType)).Append(") \\\r\n");
+						}
+						else
+						{
+							builder.Append("\tUE_FIELD_NOTIFICATION_DECLARE_ENUM_FIELD(").Append(GetNotifyTypeName(notifyType)).Append(") \\\r\n");
+						}
+					}
+					builder.Append("\tUE_FIELD_NOTIFICATION_DECLARE_ENUM_FIELD_END() \\\r\n");
+					builder.Append("\tUE_FIELD_NOTIFICATION_DECLARE_CLASS_DESCRIPTOR_END(); \\\r\n");
+				});
 		}
 
-		private StringBuilder AppendAutoGettersSetters(StringBuilder builder, UhtClass classObj)
-		{
-			// Scan the children to see what we have
-			IEnumerable<UhtProperty> properties = classObj.Properties.Where(x => x.PropertyExportFlags.HasAnyFlags(UhtPropertyExportFlags.GetterSpecifiedAuto | UhtPropertyExportFlags.SetterSpecifiedAuto));
-			if (properties.Any())
-			{
-				AppendAutoGettersSetters(builder, classObj, properties, false);
-				if (properties.Any(x => x.IsEditorOnlyProperty))
-				{
-					AppendAutoGettersSetters(builder, classObj, properties, true);
-				}
-			}
-			return builder;
-		}
-		
-		private StringBuilder AppendAutoGettersSetters(StringBuilder builder, UhtClass classObj, IEnumerable<UhtProperty> properties, bool isEditorOnly)
+		private StringBuilder AppendAutoGettersSetters(StringBuilder builder, UhtClass classObj, UhtUsedDefineScopes<UhtProperty> autoGetterSetterProperties)
 		{
-			// We can skip writing the macros if there are no properties to declare, as the 'if' and 'else' would be the same
-			if (isEditorOnly)
+			if (autoGetterSetterProperties.IsEmpty)
 			{
-				builder.Append("#if WITH_EDITORONLY_DATA\r\n");
+				return builder;
 			}
 
-			using (UhtMacroCreator macro = new(builder, this, classObj, isEditorOnly ? EditorOnlyAutoGettersSettersMacroSuffix : AutoGettersSettersMacroSuffix))
-			{
-				foreach (UhtProperty property in properties)
+			return builder.AppendMultiMacros(autoGetterSetterProperties, UhtDefineScopeNames.Standard, this, classObj, AutoGettersSettersMacroSuffix, true,
+				(builder, properties) =>
 				{
-					if (property.IsEditorOnlyProperty == isEditorOnly)
+					foreach (UhtProperty property in properties)
 					{
 						if (property.PropertyExportFlags.HasAnyFlags(UhtPropertyExportFlags.GetterSpecifiedAuto))
 						{
-							using BorrowStringBuilder borrower = new(StringBuilderCache.Small);
-							StringBuilder propertyStringBuilder = borrower.StringBuilder;
-							propertyStringBuilder.AppendPropertyText(property, UhtPropertyTextType.GetterRetVal);
-							string getterRetText = propertyStringBuilder.ToString();
 							string getterCallText = property.Getter ?? "Get" + property.SourceName;
-							builder.Append('\t').Append(getterRetText).Append(getterCallText).Append("() const; \\\r\n");
+							builder.Append('\t').AppendPropertyText(property, UhtPropertyTextType.GetterRetVal).Append(getterCallText).Append("() const; \\\r\n");
 						}
 						if (property.PropertyExportFlags.HasAnyFlags(UhtPropertyExportFlags.SetterSpecifiedAuto))
 						{
-							using BorrowStringBuilder borrower = new(StringBuilderCache.Small);
-							StringBuilder propertyStringBuilder = borrower.StringBuilder;
-							propertyStringBuilder.AppendPropertyText(property, UhtPropertyTextType.SetterParameterArgType);
-							string setterArgText = propertyStringBuilder + "InValue";
 							string setterCallText = property.Setter ?? "Set" + property.SourceName;
-							builder.Append("\tvoid ").Append(setterCallText).Append('(').Append(setterArgText).Append("); \\\r\n");
+							builder.Append("\tvoid ").Append(setterCallText).Append('(').AppendPropertyText(property, UhtPropertyTextType.SetterParameterArgType).Append("InValue").Append("); \\\r\n");
 						}
 					}
-				}
-			}
+				});
+		}
 
-			// We can skip writing the macros if there are no properties to declare, as the 'if' and 'else' would be the same
-			if (isEditorOnly)
+		private static IEnumerable<UhtProperty> EnumerateSparseDataStructProperties(IEnumerable<UhtScriptStruct> sparseScriptStructs)
+		{
+			foreach (UhtScriptStruct sparseScriptStruct in sparseScriptStructs)
 			{
-				// Trim the extra newlines added after the macro generator
-				if (builder.Length > 4 &&
-					builder[^4] == '\r' &&
-					builder[^3] == '\n' &&
-					builder[^2] == '\r' &&
-					builder[^1] == '\n')
+				foreach (UhtProperty property in sparseScriptStruct.Properties)
 				{
-					builder.Length -= 4;
-				}
-
-				builder.Append("#else // WITH_EDITORONLY_DATA\r\n");
-				using (UhtMacroCreator macro = new(builder, this, classObj, EditorOnlyAutoGettersSettersMacroSuffix))
-				{
-					// Empty macro when not compiling WITH_EDITORONLY_DATA
+					if (!property.MetaData.ContainsKey(UhtNames.NoGetter))
+					{
+						yield return property;
+					}
 				}
-				builder.Append("#endif // WITH_EDITORONLY_DATA\r\n\r\n\r\n");
 			}
-
-			return builder;
 		}
 		
-		private StringBuilder AppendSparseDeclarations(StringBuilder builder, UhtClass classObj, IEnumerable<UhtScriptStruct> sparseScriptStructs, 
-			Dictionary<UhtScriptStruct, List<UhtProperty>> commonProperties, Dictionary<UhtScriptStruct, List<UhtProperty>> editorOnlyProperties)
+		private StringBuilder AppendSparseDeclarations(StringBuilder builder, UhtClass classObj, IEnumerable<UhtScriptStruct> sparseScriptStructs, UhtUsedDefineScopes<UhtProperty> sparseProperties)
 		{
-			if (sparseScriptStructs.Any())
+			if (!sparseScriptStructs.Any())
 			{
-				AppendSparseStructDeclarations(builder, classObj, sparseScriptStructs);
-				AppendSparsePropertyDeclarations(builder, classObj, sparseScriptStructs, commonProperties, /*isEditorOnly*/false);
-				if (editorOnlyProperties.Count > 0)
-				{
-					AppendSparsePropertyDeclarations(builder, classObj, sparseScriptStructs, editorOnlyProperties, /*isEditorOnly*/true);
-				}
+				return builder;
 			}
+
+			AppendSparseStructDeclarations(builder, classObj, sparseScriptStructs);
+
+			builder.AppendMultiMacros(sparseProperties, UhtDefineScopeNames.Standard, this, classObj, SparseDataPropertyAccessorsMacroSuffix, 
+				true, (builder, properties) =>
+				{
+					foreach (UhtProperty property in properties)
+					{
+						string propertyName = property.SourceName;
+						ReadOnlySpan<char> cleanPropertyName = propertyName.AsSpan();
+						if (property is UhtBoolProperty && propertyName.StartsWith("b", StringComparison.Ordinal))
+						{
+							cleanPropertyName = cleanPropertyName[1..];
+						}
+
+						if (property.MetaData.ContainsKey(UhtNames.GetByRef))
+						{
+							builder.Append("const ").AppendSparse(property).Append("& Get").Append(cleanPropertyName).Append("() const");
+						}
+						else
+						{
+							builder.AppendSparse(property).Append(" Get").Append(cleanPropertyName).Append("() const");
+						}
+						builder.Append(" { return Get").Append(property.Outer?.EngineName).Append("(EGetSparseClassDataMethod::ArchetypeIfNull)->").Append(propertyName).Append("; } \\\r\n");
+					}
+				});	
 			return builder;
 		}
 
@@ -840,138 +787,25 @@ namespace EpicGames.UHT.Exporters.CodeGen
 			return builder;
 		}
 
-		private StringBuilder AppendSparsePropertyDeclarations(StringBuilder builder, UhtClass classObj, IEnumerable<UhtScriptStruct> sparseScriptStructs, Dictionary<UhtScriptStruct, List<UhtProperty>> propertiesToDeclare, bool isEditorOnly)
+		private StringBuilder AppendRpcFunctions(StringBuilder builder, UhtClass classObj, bool usesLegacy, UhtUsedDefineScopes<UhtFunction> functions)
 		{
-			// We can skip writing the macros if there are no properties to declare, as the 'if' and 'else' would be the same
-			if (isEditorOnly && propertiesToDeclare.Count > 0)
-			{
-				builder.Append("#if WITH_EDITORONLY_DATA\r\n");
-			}
-
-			using (UhtMacroCreator macro = new(builder, this, classObj, isEditorOnly ? EditorOnlySparseDataPropertyAccessorsMacroSuffix : SparseDataPropertyAccessorsMacroSuffix))
-			{
-				foreach (UhtScriptStruct sparseScriptStruct in sparseScriptStructs)
+			builder.AppendMultiMacros(functions, UhtDefineScopeNames.WithEditor, this, classObj, usesLegacy ? RpcWrappersMacroSuffix : RpcWrappersNoPureDeclsMacroSuffix, false,
+				(builder, functions) =>
 				{
-					if (propertiesToDeclare.TryGetValue(sparseScriptStruct, out List<UhtProperty>? sparseProperties))
+					if (usesLegacy)
 					{
-						string sparseDataType = sparseScriptStruct.EngineName;
-
-						foreach (UhtProperty sparseProperty in sparseProperties)
-						{
-							string propertyName = sparseProperty.SourceName;
-							string cleanPropertyName = propertyName;
-							if (sparseProperty is UhtBoolProperty && propertyName.StartsWith("b", StringComparison.Ordinal))
-							{
-								cleanPropertyName = propertyName[1..];
-							}
-
-							if (sparseProperty.MetaData.ContainsKey(UhtNames.GetByRef))
-							{
-								builder.Append("const ").AppendSparse(sparseProperty).Append("& Get").Append(cleanPropertyName).Append("() const");
-							}
-							else
-							{
-								builder.AppendSparse(sparseProperty).Append(" Get").Append(cleanPropertyName).Append("() const");
-							}
-							builder.Append(" { return Get").Append(sparseDataType).Append("(EGetSparseClassDataMethod::ArchetypeIfNull)->").Append(propertyName).Append("; } \\\r\n");
-						}
+						AppendAutogeneratedBlueprintFunctionDeclarations(builder, classObj, functions);
+						AppendRpcWrappers(builder, functions);
 					}
-				}
-			}
-
-			// We can skip writing the macros if there are no properties to declare, as the 'if' and 'else' would be the same
-			if (isEditorOnly && propertiesToDeclare.Count > 0)
-			{
-				// Trim the extra newlines added after the macro generator
-				if (builder.Length > 4 &&
-					builder[^4] == '\r' &&
-					builder[^3] == '\n' &&
-					builder[^2] == '\r' &&
-					builder[^1] == '\n')
-				{
-					builder.Length -= 4;
-				}
-
-				builder.Append("#else // WITH_EDITORONLY_DATA\r\n");
-				using (UhtMacroCreator macro = new(builder, this, classObj, EditorOnlySparseDataPropertyAccessorsMacroSuffix))
-				{
-					// Empty macro when not compiling WITH_EDITORONLY_DATA
-				}
-				builder.Append("#endif // WITH_EDITORONLY_DATA\r\n\r\n\r\n");
-			}
-
-			return builder;
-		}
-
-		private static void CalculateSparsePropertyDeclarations(IEnumerable<UhtScriptStruct> sparseScriptStructs, out Dictionary<UhtScriptStruct, List<UhtProperty>> commonProperties, out Dictionary<UhtScriptStruct, List<UhtProperty>> editorOnlyProperties)
-		{
-			commonProperties = new();
-			editorOnlyProperties = new();
-
-			foreach (UhtScriptStruct sparseScriptStruct in sparseScriptStructs)
-			{
-				List<UhtProperty> commonPropertiesForStruct = new();
-				List<UhtProperty> editorOnlyPropertiesForStruct = new();
-
-				foreach (UhtProperty sparseProperty in sparseScriptStruct.Properties)
-				{
-					if (!sparseProperty.MetaData.ContainsKey(UhtNames.NoGetter))
+					else
 					{
-						if (sparseProperty.IsEditorOnlyProperty)
-						{
-							editorOnlyPropertiesForStruct.Add(sparseProperty);
-						}
-						else
+						if (classObj.GeneratedCodeVersion <= EGeneratedCodeVersion.V1)
 						{
-							commonPropertiesForStruct.Add(sparseProperty);
+							AppendAutogeneratedBlueprintFunctionDeclarationsOnlyNotDeclared(builder, classObj, functions);
 						}
+						AppendRpcWrappers(builder, functions);
 					}
-				}
-
-				if (commonPropertiesForStruct.Count > 0)
-				{
-					commonProperties.Add(sparseScriptStruct, commonPropertiesForStruct);
-				}
-				if (editorOnlyPropertiesForStruct.Count > 0)
-				{
-					editorOnlyProperties.Add(sparseScriptStruct, editorOnlyPropertiesForStruct);
-				}
-			}
-		}
-
-		private StringBuilder AppendRpcFunctions(StringBuilder builder, UhtClass classObj, bool usesLegacy, List<UhtFunction> reversedFunctions, bool editorOnly)
-		{
-			builder.AppendBeginEditorOnlyGuard(editorOnly);
-
-			if (usesLegacy)
-			{
-				using UhtMacroCreator macro = new(builder, this, classObj, editorOnly ? EditorOnlyRpcWrappersMacroSuffix : RpcWrappersMacroSuffix);
-				AppendAutogeneratedBlueprintFunctionDeclarations(builder, classObj, reversedFunctions, editorOnly);
-				AppendRpcWrappers(builder, reversedFunctions, editorOnly);
-			}
-			else
-			{
-				using UhtMacroCreator macro = new(builder, this, classObj, editorOnly ? EditorOnlyRpcWrappersNoPureDeclsMacroSuffix : RpcWrappersNoPureDeclsMacroSuffix);
-				if (classObj.GeneratedCodeVersion <= EGeneratedCodeVersion.V1)
-				{
-					AppendAutogeneratedBlueprintFunctionDeclarationsOnlyNotDeclared(builder, classObj, reversedFunctions, editorOnly);
-				}
-				AppendRpcWrappers(builder, reversedFunctions, editorOnly);
-			}
-
-			if (editorOnly)
-			{
-				builder.Append("#else // WITH_EDITOR\r\n");
-				if (usesLegacy)
-				{
-					using UhtMacroCreator macro = new(builder, this, classObj, editorOnly ? EditorOnlyRpcWrappersMacroSuffix : RpcWrappersMacroSuffix);
-				}
-				else
-				{
-					using UhtMacroCreator macro = new(builder, this, classObj, editorOnly ? EditorOnlyRpcWrappersNoPureDeclsMacroSuffix : RpcWrappersNoPureDeclsMacroSuffix);
-				}
-				builder.AppendEndEditorOnlyGuard();
-			}
+				});
 			return builder;
 		}
 
@@ -997,11 +831,11 @@ namespace EpicGames.UHT.Exporters.CodeGen
 			return builder;
 		}
 
-		private StringBuilder AppendAutogeneratedBlueprintFunctionDeclarations(StringBuilder builder, UhtClass classObj, List<UhtFunction> reversedFunctions, bool editorOnly)
+		private StringBuilder AppendAutogeneratedBlueprintFunctionDeclarations(StringBuilder builder, UhtClass classObj, IEnumerable<UhtFunction> functions)
 		{
-			foreach (UhtFunction function in reversedFunctions)
+			foreach (UhtFunction function in functions)
 			{
-				if (!IsRpcFunction(function, editorOnly) || function.CppImplName == function.SourceName)
+				if (function.CppImplName == function.SourceName)
 				{
 					continue;
 				}
@@ -1011,11 +845,11 @@ namespace EpicGames.UHT.Exporters.CodeGen
 			return builder;
 		}
 
-		private StringBuilder AppendAutogeneratedBlueprintFunctionDeclarationsOnlyNotDeclared(StringBuilder builder, UhtClass classObj, List<UhtFunction> reversedFunctions, bool editorOnly)
+		private StringBuilder AppendAutogeneratedBlueprintFunctionDeclarationsOnlyNotDeclared(StringBuilder builder, UhtClass classObj, IEnumerable<UhtFunction> functions)
 		{
-			foreach (UhtFunction function in reversedFunctions)
+			foreach (UhtFunction function in functions)
 			{
-				if (!IsRpcFunction(function, editorOnly) || function.CppImplName == function.SourceName)
+				if (function.CppImplName == function.SourceName)
 				{
 					continue;
 				}
@@ -1053,14 +887,10 @@ namespace EpicGames.UHT.Exporters.CodeGen
 			return builder;
 		}
 
-		private static StringBuilder AppendRpcWrappers(StringBuilder builder, List<UhtFunction> reversedFunctions, bool editorOnly)
+		private static StringBuilder AppendRpcWrappers(StringBuilder builder, IEnumerable<UhtFunction> functions)
 		{
-			foreach (UhtFunction function in reversedFunctions)
+			foreach (UhtFunction function in functions)
 			{
-				if (!IsRpcFunction(function, editorOnly))
-				{
-					continue;
-				}
 				if (!ShouldExportFunction(function))
 				{
 					continue;
@@ -1631,8 +1461,9 @@ namespace EpicGames.UHT.Exporters.CodeGen
 			return builder;
 		}
 
-		private StringBuilder AppendGeneratedBodyMacroBlock(StringBuilder builder, UhtClass classObj, UhtClass bodyClassObj, bool isLegacy, bool hasEditorRpc, bool hasCallbacks, 
-			bool hasAutoGettersSetters, bool hasEditorOnlyAutoGettersSetters, bool hasSparseProperties, bool hasEditorOnlySparseProperties, string? deprecatedMacroName)
+		private StringBuilder AppendGeneratedBodyMacroBlock(StringBuilder builder, UhtClass classObj, UhtClass bodyClassObj, bool isLegacy, 
+			UhtUsedDefineScopes<UhtFunction> rpcFunctions, UhtUsedDefineScopes<UhtProperty> autoGetterSetterProperties, bool hasSparseStructs,
+			UhtUsedDefineScopes<UhtProperty> sparseProperties, bool hasCallbacks, string? deprecatedMacroName)
 		{
 			bool isInterface = classObj.ClassFlags.HasAnyFlags(EClassFlags.Interface);
 			using (UhtMacroCreator macro = new(builder, this, bodyClassObj, isLegacy ? GeneratedBodyLegacyMacroSuffix : GeneratedBodyMacroSuffix))
@@ -1643,42 +1474,16 @@ namespace EpicGames.UHT.Exporters.CodeGen
 				}
 				builder.Append(DisableDeprecationWarnings).Append(" \\\r\n");
 				builder.Append("public: \\\r\n");
-				if (hasSparseProperties)
+				if (hasSparseStructs)
 				{
 					builder.Append('\t').AppendMacroName(this, classObj, SparseDataMacroSuffix).Append(" \\\r\n");
-					builder.Append('\t').AppendMacroName(this, classObj, SparseDataPropertyAccessorsMacroSuffix).Append(" \\\r\n");
-					if (hasEditorOnlySparseProperties)
-					{
-						builder.Append('\t').AppendMacroName(this, classObj, EditorOnlySparseDataPropertyAccessorsMacroSuffix).Append(" \\\r\n");
-					}
-				}
-				if (isLegacy)
-				{
-					builder.Append('\t').AppendMacroName(this, classObj, RpcWrappersMacroSuffix).Append(" \\\r\n");
-				}
-				else
-				{
-					builder.Append('\t').AppendMacroName(this, classObj, RpcWrappersNoPureDeclsMacroSuffix).Append(" \\\r\n");
-				}
-				if (hasEditorRpc)
-				{
-					if (isLegacy)
-					{
-						builder.Append('\t').AppendMacroName(this, classObj, EditorOnlyRpcWrappersMacroSuffix).Append(" \\\r\n");
-					}
-					else
-					{
-						builder.Append('\t').AppendMacroName(this, classObj, EditorOnlyRpcWrappersNoPureDeclsMacroSuffix).Append(" \\\r\n");
-					}
+					builder.AppendMultiMacroRefs(sparseProperties, this, classObj, SparseDataPropertyAccessorsMacroSuffix);
 				}
+				builder.AppendMultiMacroRefs(rpcFunctions, this, classObj, isLegacy ? RpcWrappersMacroSuffix : RpcWrappersNoPureDeclsMacroSuffix);
 				builder.Append('\t').AppendMacroName(this, classObj, AccessorsMacroSuffix).Append(" \\\r\n");
-				if (hasAutoGettersSetters)
+				if (!autoGetterSetterProperties.IsEmpty)
 				{
-					builder.Append('\t').AppendMacroName(this, classObj, AutoGettersSettersMacroSuffix).Append(" \\\r\n");
-					if (hasEditorOnlyAutoGettersSetters)
-					{
-						builder.Append('\t').AppendMacroName(this, classObj, EditorOnlyAutoGettersSettersMacroSuffix).Append(" \\\r\n");
-					}
+					builder.AppendMultiMacroRefs(autoGetterSetterProperties, this, classObj, AutoGettersSettersMacroSuffix);
 				}
 				if (hasCallbacks)
 				{
diff --git a/Engine/Source/Programs/Shared/EpicGames.UHT/Exporters/CodeGen/UhtMacroBlockEmitter.cs b/Engine/Source/Programs/Shared/EpicGames.UHT/Exporters/CodeGen/UhtMacroBlockEmitter.cs
--- a/Engine/Source/Programs/Shared/EpicGames.UHT/Exporters/CodeGen/UhtMacroBlockEmitter.cs
+++ b/Engine/Source/Programs/Shared/EpicGames.UHT/Exporters/CodeGen/UhtMacroBlockEmitter.cs
@@ -2,43 +2,41 @@
 
 using System;
 using System.Text;
+using EpicGames.UHT.Types;
 
 namespace EpicGames.UHT.Exporters.CodeGen
 {
 	internal struct UhtMacroBlockEmitter : IDisposable
 	{
 		private readonly StringBuilder _builder;
-		private readonly string _macroName;
-		private bool _bEmitted;
+		private readonly UhtDefineScopeNames _defineScopeNames;
+		private UhtDefineScope _current = UhtDefineScope.None;
 
-		public UhtMacroBlockEmitter(StringBuilder builder, string macroName, bool initialState = false)
+		public UhtMacroBlockEmitter(StringBuilder builder, UhtDefineScopeNames defineScopeNames, UhtDefineScope initialState)
 		{
 			_builder = builder;
-			_macroName = macroName;
-			_bEmitted = false;
+			_defineScopeNames = defineScopeNames;
 			Set(initialState);
 		}
 
-		public void Set(bool emit)
+		public void Set(UhtDefineScope defineScope)
 		{
-			if (_bEmitted == emit)
+			if (defineScope == UhtDefineScope.Invalid)
 			{
-				return;
-			}
-			_bEmitted = emit;
-			if (_bEmitted)
-			{
-				_builder.Append("#if ").Append(_macroName).Append("\r\n");
+				defineScope = UhtDefineScope.None;
 			}
-			else
+			if (_current == defineScope)
 			{
-				_builder.Append("#endif // ").Append(_macroName).Append("\r\n");
+				return;
 			}
+			_builder.AppendEndIfPreprocessor(_current, _defineScopeNames);
+			_builder.AppendIfPreprocessor(defineScope, _defineScopeNames);
+			_current = defineScope;
 		}
 
 		public void Dispose()
 		{
-			Set(false);
+			Set(UhtDefineScope.None);
 		}
 	}
 }
diff --git a/Engine/Source/Programs/Shared/EpicGames.UHT/Exporters/CodeGen/UhtMacroCreator.cs b/Engine/Source/Programs/Shared/EpicGames.UHT/Exporters/CodeGen/UhtMacroCreator.cs
--- a/Engine/Source/Programs/Shared/EpicGames.UHT/Exporters/CodeGen/UhtMacroCreator.cs
+++ b/Engine/Source/Programs/Shared/EpicGames.UHT/Exporters/CodeGen/UhtMacroCreator.cs
@@ -12,30 +12,16 @@ namespace EpicGames.UHT.Exporters.CodeGen
 		private readonly StringBuilder _builder;
 		private readonly int _startingLength;
 
-		public UhtMacroCreator(StringBuilder builder, UhtHeaderCodeGenerator generator, int lineNumber, string macroSuffix)
+		public UhtMacroCreator(StringBuilder builder, UhtHeaderCodeGenerator generator, int lineNumber, string macroSuffix, UhtDefineScope defineScope = UhtDefineScope.None)
 		{
-			builder.Append("#define ").AppendMacroName(generator, lineNumber, macroSuffix).Append(" \\\r\n");
+			builder.Append("#define ").AppendMacroName(generator, lineNumber, macroSuffix, defineScope).Append(" \\\r\n");
 			_builder = builder;
 			_startingLength = builder.Length;
 		}
 
-		public UhtMacroCreator(StringBuilder builder, UhtHeaderCodeGenerator generator, UhtClass classObj, string macroSuffix)
+		public UhtMacroCreator(StringBuilder builder, UhtHeaderCodeGenerator generator, UhtType type, string macroSuffix, UhtDefineScope defineScope = UhtDefineScope.None)
 		{
-			builder.Append("#define ").AppendMacroName(generator, classObj, macroSuffix).Append(" \\\r\n");
-			_builder = builder;
-			_startingLength = builder.Length;
-		}
-
-		public UhtMacroCreator(StringBuilder builder, UhtHeaderCodeGenerator generator, UhtScriptStruct scriptStruct, string macroSuffix)
-		{
-			builder.Append("#define ").AppendMacroName(generator, scriptStruct, macroSuffix).Append(" \\\r\n");
-			_builder = builder;
-			_startingLength = builder.Length;
-		}
-
-		public UhtMacroCreator(StringBuilder builder, UhtHeaderCodeGenerator generator, UhtFunction function, string macroSuffix)
-		{
-			builder.Append("#define ").AppendMacroName(generator, function, macroSuffix).Append(" \\\r\n");
+			builder.Append("#define ").AppendMacroName(generator, type, macroSuffix, defineScope).Append(" \\\r\n");
 			_builder = builder;
 			_startingLength = builder.Length;
 		}
diff --git a/Engine/Source/Programs/Shared/EpicGames.UHT/Exporters/CodeGen/UhtPackageCodeGenerator.cs b/Engine/Source/Programs/Shared/EpicGames.UHT/Exporters/CodeGen/UhtPackageCodeGenerator.cs
--- a/Engine/Source/Programs/Shared/EpicGames.UHT/Exporters/CodeGen/UhtPackageCodeGenerator.cs
+++ b/Engine/Source/Programs/Shared/EpicGames.UHT/Exporters/CodeGen/UhtPackageCodeGenerator.cs
@@ -1,6 +1,5 @@
 // Copyright Epic Games, Inc. All Rights Reserved.
 
-using System.Text;
 using EpicGames.UHT.Types;
 
 namespace EpicGames.UHT.Exporters.CodeGen
@@ -20,9 +19,6 @@ namespace EpicGames.UHT.Exporters.CodeGen
 		public static string EnableDeprecationWarnings = "PRAGMA_ENABLE_DEPRECATION_WARNINGS";
 		public static string DisableDeprecationWarnings = "PRAGMA_DISABLE_DEPRECATION_WARNINGS";
 
-		public static string BeginEditorOnlyGuard = "#if WITH_EDITOR\r\n";
-		public static string EndEditorOnlyGuard = "#endif // WITH_EDITOR\r\n";
-
 		public readonly UhtCodeGenerator CodeGenerator;
 		public readonly UhtPackage Package;
 		public bool SaveExportedHeaders => Package.Module.SaveExportedHeaders;
@@ -122,25 +118,4 @@ namespace EpicGames.UHT.Exporters.CodeGen
 		}
 		#endregion
 	}
-
-	internal static class UhtPackageCodeGeneratorStringBuilderExtensions
-	{
-		public static StringBuilder AppendBeginEditorOnlyGuard(this StringBuilder builder, bool enable = true)
-		{
-			if (enable)
-			{
-				builder.Append(UhtPackageCodeGenerator.BeginEditorOnlyGuard);
-			}
-			return builder;
-		}
-
-		public static StringBuilder AppendEndEditorOnlyGuard(this StringBuilder builder, bool enable = true)
-		{
-			if (enable)
-			{
-				builder.Append(UhtPackageCodeGenerator.EndEditorOnlyGuard);
-			}
-			return builder;
-		}
-	}
 }
diff --git a/Engine/Source/Programs/Shared/EpicGames.UHT/Exporters/CodeGen/UhtUsedDefineScopes.cs b/Engine/Source/Programs/Shared/EpicGames.UHT/Exporters/CodeGen/UhtUsedDefineScopes.cs
new file mode 100644
--- /dev/null
+++ b/Engine/Source/Programs/Shared/EpicGames.UHT/Exporters/CodeGen/UhtUsedDefineScopes.cs
@@ -0,0 +1,761 @@
+// Copyright Epic Games, Inc. All Rights Reserved.
+
+using System;
+using System.Collections.Generic;
+using System.Linq;
+using System.Text;
+using EpicGames.UHT.Types;
+using EpicGames.UHT.Utils;
+
+namespace EpicGames.UHT.Exporters.CodeGen
+{
+
+	/// <summary>
+	/// Helper class for collection where types exist in different define scope blocks
+	/// </summary>
+	/// <typeparam name="T"></typeparam>
+	public class UhtUsedDefineScopes<T> where T : UhtType
+	{
+		private bool _first = true;
+		private UhtDefineScope _soleScope = UhtDefineScope.Invalid;
+		private UhtDefineScope _allScopes = UhtDefineScope.None;
+		private uint _present = 0;
+
+		/// <summary>
+		/// Collection of instances
+		/// </summary>
+		public List<T> Instances { get; set; } = new();
+
+		/// <summary>
+		/// If true, there are no instances
+		/// </summary>
+		public bool IsEmpty => Instances.Count == 0;
+
+		/// <summary>
+		/// If true, at least one instance had no scope
+		/// </summary>
+		public bool HasNoneScope => HasScope(UhtDefineScope.None);
+
+		/// <summary>
+		/// If all types share the same scope, then the sole scope is that scope.  The value will be Invalid if the
+		/// instances have different scopes.
+		/// </summary>
+		public UhtDefineScope SoleScope => _soleScope;
+
+		/// <summary>
+		/// Collection of all scopes found in the types.  If 
+		/// </summary>
+		public UhtDefineScope AllScopes => _allScopes;
+
+		/// <summary>
+		/// If any instance has no scope, then None is returned.  Otherwise AllScopes is returned.
+		/// </summary>
+		public UhtDefineScope NoneAwareScopes => HasNoneScope ? UhtDefineScope.None : AllScopes;
+
+		/// <summary>
+		/// Constructor with no initial instances
+		/// </summary>
+		public UhtUsedDefineScopes()
+		{
+		}
+
+		/// <summary>
+		/// Constructor with initial range of types
+		/// </summary>
+		/// <param name="instances">Instances to initially add</param>
+		public UhtUsedDefineScopes(IEnumerable<T> instances)
+		{
+			AddRange(instances);
+		}
+
+		/// <summary>
+		/// Add an instance to the collection
+		/// </summary>
+		/// <param name="instance">Instance to be added</param>
+		public void Add(T instance)
+		{
+			UhtDefineScope defineScope = instance.DefineScope;
+			if (_first)
+			{
+				_soleScope = defineScope;
+				_first = false;
+			}
+			else if (_soleScope != UhtDefineScope.Invalid && _soleScope != defineScope)
+			{
+				_soleScope = UhtDefineScope.Invalid;
+			}
+			_allScopes |= defineScope;
+			_present |= ScopeToMask(defineScope);
+			Instances.Add(instance);
+		}
+
+		/// <summary>
+		/// Add a range of instances to the collection
+		/// </summary>
+		/// <param name="instances">Collection of instances</param>
+		public void AddRange(IEnumerable<T> instances)
+		{
+			foreach (T instance in instances)
+			{
+				Add(instance);
+			}
+		}
+
+		/// <summary>
+		/// Check to see if the given scope has instances
+		/// </summary>
+		/// <param name="defineScope">Scope to test</param>
+		/// <returns>True if the scope has elements</returns>
+		public bool HasScope(UhtDefineScope defineScope)
+		{
+			return (_present & ScopeToMask(defineScope)) != 0;
+		}
+
+		/// <summary>
+		/// Update the list to by ordered by the define scope
+		/// </summary>
+		public void OrderByDefineScope()
+		{
+			Instances = Instances.OrderBy(x => x.DefineScope).ToList();
+		}
+
+		/// <summary>
+		/// Enumerate all of the used defined scopes
+		/// </summary>
+		/// <returns>Enumeration</returns>
+		public IEnumerable<UhtDefineScope> EnumerateDefinedScopes()
+		{
+			// At this time, we alway return None
+			yield return UhtDefineScope.None;
+
+			ulong present = _present >> 1;
+			for (int index = 1; present != 0; ++index, present >>= 1)
+			{
+				if ((present & 1) != 0)
+				{
+					yield return (UhtDefineScope)index;
+				}
+			}
+		}
+
+		/// <summary>
+		/// Enumerate all of the used defined scopes
+		/// </summary>
+		/// <returns>Enumeration</returns>
+		public IEnumerable<UhtDefineScope> EnumerateDefinedScopesNoneAtEnd()
+		{
+			ulong present = _present >> 1;
+			for (int index = 1; present != 0; ++index, present >>= 1)
+			{
+				if ((present & 1) != 0)
+				{
+					yield return (UhtDefineScope)index;
+				}
+			}
+
+			// At this time, we alway return None
+			yield return UhtDefineScope.None;
+		}
+
+		private static uint ScopeToMask(UhtDefineScope defineScope)
+		{
+			return (uint)1 << (int)defineScope;
+		}
+	}
+
+	/// <summary>
+	/// Defines which set of names will be used for each scope
+	/// </summary>
+	internal enum UhtDefineScopeNames
+	{
+
+		/// <summary>
+		/// Names match the define
+		/// </summary>
+		Standard,
+
+		/// <summary>
+		/// Use WITH_EDITOR in place of WITH_EDITORONLY_DATA
+		/// </summary>
+		WithEditor,
+	}
+
+	/// <summary>
+	/// Helper extensions when working with defined scope collections.
+	/// 
+	/// This class helps with producing consistent and correct code that works with the DefineScope element in UhtType.
+	/// 
+	/// In generated header files, it supports two flavors of macro definitions used to generate the body macro.
+	/// 
+	/// The multi macro style will generate one macro for each unique combination of DefineScope found.  Only the instances
+	/// that match the DefineScope will be placed in that macro.  The macros will be generated in such a way that they 
+	/// will always be present and each much be included in the generated body macro.  This style is used to generate
+	/// declarations as needed.
+	/// 
+	/// The single macro style will generate a single macro, but that macro will appear multiple times depending on
+	/// each combination of the DefineScope combinations required.  Each macro will have a complete set of instances
+	/// filtered by the DefineScope.  This style is used to populate such things as enum definitions needed by the engine.
+	/// 
+	/// In generated cpp files, there is support for enumerating through all the instances and emitting the appropriate
+	/// #if block to include instances based on their DefineScope.
+	/// </summary>
+	internal static class UhtUsedDefineScopesExtensions
+	{
+		/// <summary>
+		/// Append multi macros for the given collection of scopes
+		/// </summary>
+		/// <param name="builder">Destination builder</param>
+		/// <param name="instances">Defined scope instances</param>
+		/// <param name="defineScopeNames">Which set of scope names will be used</param>
+		/// <param name="generator">Header code generator</param>
+		/// <param name="outerType">Output type owning the instances</param>
+		/// <param name="macroSuffix">Macro being created</param>
+		/// <param name="trimMacro">If true, trim the end of the prior macro</param>
+		/// <param name="appendAction">Action to invoke to append an instance</param>
+		/// <returns>String builder</returns>
+		public static StringBuilder AppendMultiMacros<T>(this StringBuilder builder, UhtUsedDefineScopes<T> instances, UhtDefineScopeNames defineScopeNames, 
+			UhtHeaderCodeGenerator generator, UhtType outerType, string macroSuffix, bool trimMacro, Action<StringBuilder, IEnumerable<T>> appendAction) where T : UhtType
+		{
+			foreach (UhtDefineScope defineScope in instances.EnumerateDefinedScopes())
+			{
+				using (UhtMacroBlockEmitter blockEmitter = new(builder, defineScopeNames, defineScope))
+				{
+					using (UhtMacroCreator macro = new(builder, generator, outerType, macroSuffix, defineScope))
+					{
+						appendAction(builder, instances.Instances.Where(x => x.DefineScope == defineScope));
+					}
+
+					// We can skip writing the macros if there are no properties to declare, as the 'if' and 'else' would be the same
+					if (defineScope != UhtDefineScope.None)
+					{
+						if (trimMacro)
+						{
+							// Trim the extra newlines added after the macro generator
+							if (builder.Length > 4 &&
+								builder[^4] == '\r' &&
+								builder[^3] == '\n' &&
+								builder[^2] == '\r' &&
+								builder[^1] == '\n')
+							{
+								builder.Length -= 4;
+							}
+						}
+
+						builder.AppendElsePreprocessor(defineScope, defineScopeNames);
+						using UhtMacroCreator macro = new(builder, generator, outerType, macroSuffix, defineScope); // Empty macro
+					}
+				}
+
+				if (trimMacro && defineScope != UhtDefineScope.None)
+				{
+					builder.Append("\r\n\r\n");
+				}
+			}
+			return builder;
+		}
+
+		/// <summary>
+		/// Append the macro definitions requested
+		/// </summary>
+		/// <param name="builder">Destination builder</param>
+		/// <param name="instances">Defined scope instances</param>
+		/// <param name="generator">Header code generator</param>
+		/// <param name="outerType">Output type owning the instances</param>
+		/// <param name="macroSuffix">Macro being created</param>
+		/// <returns>String builder</returns>
+		public static StringBuilder AppendMultiMacroRefs<T>(this StringBuilder builder, UhtUsedDefineScopes<T> instances, 
+			UhtHeaderCodeGenerator generator, UhtType outerType, string macroSuffix) where T : UhtType
+		{
+			foreach (UhtDefineScope defineScope in instances.EnumerateDefinedScopes())
+			{
+				builder.Append('\t').AppendMacroName(generator, outerType, macroSuffix, defineScope).Append(" \\\r\n");
+			}
+			return builder;
+		}
+
+		/// <summary>
+		/// Append a single macro, where the macro -same- definition can exist inside of define scopes where
+		/// the instances contained are all instances that would be satisfied by a scope.
+		/// </summary>
+		/// <param name="builder">Destination builder</param>
+		/// <param name="instances">Defined scope instances</param>
+		/// <param name="defineScopeNames">Which set of scope names will be used</param>
+		/// <param name="generator">Header code generator</param>
+		/// <param name="outerType">Output type owning the instances</param>
+		/// <param name="macroSuffix">Macro being created</param>
+		/// <param name="trimMacro">If true, trim the end of the prior macro</param>
+		/// <param name="appendAction">Action to invoke to append an instance</param>
+		/// <returns>String builder</returns>
+		public static StringBuilder AppendSingleMacro<T>(this StringBuilder builder, UhtUsedDefineScopes<T> instances, UhtDefineScopeNames defineScopeNames, 
+			UhtHeaderCodeGenerator generator, UhtType outerType, string macroSuffix, bool trimMacro, Action<StringBuilder, IEnumerable<T>> appendAction) where T : UhtType
+		{
+			if (instances.SoleScope == UhtDefineScope.None)
+			{
+				using UhtMacroCreator macro = new(builder, generator, outerType, macroSuffix, UhtDefineScope.None);
+				appendAction(builder, instances.Instances);
+			}
+			else
+			{
+				bool first = true;
+				foreach (UhtDefineScope defineScope in instances.EnumerateDefinedScopesNoneAtEnd())
+				{
+					if (first)
+					{
+						builder.AppendIfPreprocessor(defineScope, defineScopeNames);
+						first = false;
+					}
+					else if (defineScope != UhtDefineScope.None)
+					{
+						builder.AppendElseIfPreprocessor(defineScope, defineScopeNames);
+					}
+					else
+					{
+						builder.Append("#else\r\n");
+					}
+					using (UhtMacroCreator macro = new(builder, generator, outerType, macroSuffix, UhtDefineScope.None))
+					{
+						appendAction(builder, instances.Instances.Where(x => (x.DefineScope & ~defineScope) == 0));
+					}
+					if (trimMacro)
+					{
+						if (builder.Length > 4 &&
+							builder[^4] == '\r' &&
+							builder[^3] == '\n' &&
+							builder[^2] == '\r' &&
+							builder[^1] == '\n')
+						{
+							builder.Length -= 4;
+						}
+					}
+				}
+				builder.Append("#endif\r\n");
+			}
+			return builder;
+		}
+
+		/// <summary>
+		/// Invoke the append action if any types are present.  If all types are from the same define scope, then
+		/// it will be wrapped with and #if block.
+		/// </summary>
+		/// <param name="builder">Destination builder</param>
+		/// <param name="instances">Defined scope instances</param>
+		/// <param name="defineScopeNames">Names to use when outputting the scope</param>
+		/// <param name="appendAction">Action to invoke</param>
+		/// <returns>String builder</returns>
+		public static StringBuilder AppendIfInstances<T>(this StringBuilder builder, UhtUsedDefineScopes<T> instances, UhtDefineScopeNames defineScopeNames, Action<StringBuilder> appendAction) where T : UhtType
+		{
+			if (!instances.IsEmpty)
+			{
+				using UhtMacroBlockEmitter blockEmitter = new(builder, defineScopeNames, instances.NoneAwareScopes);
+				appendAction(builder);
+			}
+			return builder;
+		}
+
+		/// <summary>
+		/// Append each instance to the builder
+		/// </summary>
+		/// <param name="builder">Destination builder</param>
+		/// <param name="instances">Defined scope instances</param>
+		/// <param name="defineScopeNames">Names to use when outputting the scope</param>
+		/// <param name="appendAction">Action to invoke</param>
+		/// <returns>String builder</returns>
+		public static StringBuilder AppendInstances<T>(this StringBuilder builder, UhtUsedDefineScopes<T> instances, UhtDefineScopeNames defineScopeNames, Action<StringBuilder, T> appendAction) where T : UhtType
+		{
+			return builder.AppendInstances(instances, defineScopeNames, null, appendAction, null);
+		}
+
+		/// <summary>
+		/// Append each instance to the builder
+		/// </summary>
+		/// <param name="builder">Destination builder</param>
+		/// <param name="instances">Defined scope instances</param>
+		/// <param name="defineScopeNames">Names to use when outputting the scope</param>
+		/// <param name="preambleAction">Action to invoke prior to first instance</param>
+		/// <param name="appendAction">Action to invoke</param>
+		/// <param name="postambleAction">Action to invoke following all instances</param>
+		/// <returns>String builder</returns>
+		public static StringBuilder AppendInstances<T>(this StringBuilder builder, UhtUsedDefineScopes<T> instances, UhtDefineScopeNames defineScopeNames,
+			Action<StringBuilder>? preambleAction, Action<StringBuilder, T> appendAction, Action<StringBuilder>? postambleAction) where T : UhtType
+		{
+			if (!instances.IsEmpty)
+			{
+				using UhtMacroBlockEmitter blockEmitter = new(builder, defineScopeNames, instances.SoleScope);
+				preambleAction?.Invoke(builder);
+				foreach (T instance in instances.Instances)
+				{
+					blockEmitter.Set(instance.DefineScope);
+					appendAction(builder, instance);
+				}
+
+				// Make sure the postable run with the proper initial scope
+				blockEmitter.Set(instances.SoleScope);
+				postambleAction?.Invoke(builder);
+			}
+			return builder;
+		}
+
+		/// <summary>
+		/// Append the given array list and count as arguments to a structure constructor
+		/// </summary>
+		/// <param name="builder">Destination builder</param>
+		/// <param name="instances">Collected instances</param>
+		/// <param name="defineScopeNames">Names to use when outputting the scope</param>
+		/// <param name="staticsName">Name of the statics section</param>
+		/// <param name="arrayName">The name of the arrray</param>
+		/// <param name="tabs">Number of tabs to start the line</param>
+		/// <param name="endl">Text to end the line</param>
+		/// <returns>Destination builder</returns>
+		public static StringBuilder AppendArrayPtrAndCountLine<T>(this StringBuilder builder, UhtUsedDefineScopes<T> instances, UhtDefineScopeNames defineScopeNames, string? staticsName, string arrayName, int tabs, string endl) where T : UhtType
+		{
+			UhtDefineScope noneAwareScopes = instances.NoneAwareScopes;
+
+			// We have no instances in the collection, the list will always be empty
+			if (instances.IsEmpty)
+			{
+				builder
+					.AppendTabs(tabs)
+					.Append("nullptr, 0")
+					.Append(endl);
+			}
+
+			// If we have any instances that have no scope, then we can always just reference the array
+			else if (noneAwareScopes == UhtDefineScope.None)
+			{
+				builder
+					.AppendTabs(tabs)
+					.AppendCppName(staticsName, arrayName)
+					.Append(", UE_ARRAY_COUNT(")
+					.AppendCppName(staticsName, arrayName)
+					.Append(')')
+					.Append(endl);
+			}
+
+			// If the only scope we have is WITH_EDITORONLY_DATA, then we can use the existing macros
+			else if (noneAwareScopes == UhtDefineScope.EditorOnlyData)
+			{
+				switch (defineScopeNames)
+				{
+					case UhtDefineScopeNames.Standard:
+						builder
+							.AppendTabs(tabs)
+							.Append("IF_WITH_EDITORONLY_DATA(")
+							.AppendCppName(staticsName, arrayName)
+							.Append(", nullptr), IF_WITH_EDITORONLY_DATA(UE_ARRAY_COUNT(")
+							.AppendCppName(staticsName, arrayName)
+							.Append("), 0)")
+							.Append(endl);
+						break;
+					case UhtDefineScopeNames.WithEditor:
+						builder
+							.AppendTabs(tabs)
+							.Append("IF_WITH_EDITOR(")
+							.AppendCppName(staticsName, arrayName)
+							.Append(", nullptr), IF_WITH_EDITOR(UE_ARRAY_COUNT(")
+							.AppendCppName(staticsName, arrayName)
+							.Append("), 0)")
+							.Append(endl);
+						break;
+					default:
+						throw new UhtIceException("Unexpected define scope names value");
+				}
+			}
+
+			// Otherwise we have many different scopes but no none scope, must generate #if block
+			else
+			{
+				builder.AppendIfPreprocessor(instances.AllScopes, defineScopeNames);
+				builder
+					.AppendTabs(tabs)
+					.AppendCppName(staticsName, arrayName)
+					.Append(", UE_ARRAY_COUNT(")
+					.AppendCppName(staticsName, arrayName)
+					.Append(')')
+					.Append(endl);
+				builder.AppendElsePreprocessor(instances.AllScopes, defineScopeNames);
+				builder
+					.AppendTabs(tabs)
+					.Append("nullptr, 0")
+					.Append(endl);
+				builder.AppendEndIfPreprocessor(instances.AllScopes, defineScopeNames);
+			}
+			return builder;
+		}
+
+		/// <summary>
+		/// Append the given array list as arguments to a structure constructor
+		/// </summary>
+		/// <param name="builder">Destination builder</param>
+		/// <param name="instances">Collected instances</param>
+		/// <param name="defineScopeNames">Names to use when outputting the scope</param>
+		/// <param name="staticsName">Name of the statics section</param>
+		/// <param name="arrayName">The name of the arrray</param>
+		/// <param name="tabs">Number of tabs to start the line</param>
+		/// <param name="endl">Text to end the line</param>
+		/// <returns>Destination builder</returns>
+		public static StringBuilder AppendArrayPtrLine<T>(this StringBuilder builder, UhtUsedDefineScopes<T> instances, UhtDefineScopeNames defineScopeNames, string? staticsName, string arrayName, int tabs, string endl) where T : UhtType
+		{
+			UhtDefineScope noneAwareScopes = instances.NoneAwareScopes;
+
+			// We have no instances in the collection, the list will always be empty
+			if (instances.IsEmpty)
+			{
+				builder
+					.AppendTabs(tabs)
+					.Append("nullptr")
+					.Append(endl);
+			}
+
+			// If we have any instances that have no scope, then we can always just reference the array
+			else if (noneAwareScopes == UhtDefineScope.None)
+			{
+				builder
+					.AppendTabs(tabs)
+					.AppendCppName(staticsName, arrayName)
+					.Append(endl);
+			}
+
+			// If the only scope we have is WITH_EDITORONLY_DATA, then we can use the existing macros
+			else if (noneAwareScopes == UhtDefineScope.EditorOnlyData)
+			{
+				switch (defineScopeNames)
+				{
+					case UhtDefineScopeNames.Standard:
+						builder
+							.AppendTabs(tabs)
+							.Append("IF_WITH_EDITORONLY_DATA(")
+							.AppendCppName(staticsName, arrayName)
+							.Append(", nullptr)")
+							.Append(endl);
+						break;
+					case UhtDefineScopeNames.WithEditor:
+						builder
+							.AppendTabs(tabs)
+							.Append("IF_WITH_EDITOR(")
+							.AppendCppName(staticsName, arrayName)
+							.Append(", nullptr)")
+							.Append(endl);
+						break;
+					default:
+						throw new UhtIceException("Unexpected define scope names value");
+				}
+			}
+
+			// Otherwise we have many different scopes but no none scope, must generate #if block
+			else
+			{
+				builder.AppendIfPreprocessor(instances.AllScopes, defineScopeNames);
+				builder
+					.AppendTabs(tabs)
+					.AppendCppName(staticsName, arrayName)
+					.Append(endl);
+				builder.AppendElsePreprocessor(instances.AllScopes, defineScopeNames);
+				builder
+					.AppendTabs(tabs)
+					.Append("nullptr")
+					.Append(endl);
+				builder.AppendEndIfPreprocessor(instances.AllScopes, defineScopeNames);
+			}
+			return builder;
+		}
+
+		/// <summary>
+		/// Append the given array count as arguments to a structure constructor
+		/// </summary>
+		/// <param name="builder">Destination builder</param>
+		/// <param name="instances">Collected instances</param>
+		/// <param name="defineScopeNames">Names to use when outputting the scope</param>
+		/// <param name="staticsName">Name of the statics section</param>
+		/// <param name="arrayName">The name of the arrray</param>
+		/// <param name="tabs">Number of tabs to start the line</param>
+		/// <param name="endl">Text to end the line</param>
+		/// <returns>Destination builder</returns>
+		public static StringBuilder AppendArrayCountLine<T>(this StringBuilder builder, UhtUsedDefineScopes<T> instances, UhtDefineScopeNames defineScopeNames, string? staticsName, string arrayName, int tabs, string endl) where T : UhtType
+		{
+			UhtDefineScope noneAwareScopes = instances.NoneAwareScopes;
+
+			// We have no instances in the collection, the list will always be empty
+			if (instances.IsEmpty)
+			{
+				builder
+					.AppendTabs(tabs)
+					.Append('0')
+					.Append(endl);
+			}
+
+			// If we have any instances that have no scope, then we can always just reference the array
+			else if (noneAwareScopes == UhtDefineScope.None)
+			{
+				builder
+					.AppendTabs(tabs)
+					.Append("UE_ARRAY_COUNT(")
+					.AppendCppName(staticsName, arrayName)
+					.Append(')')
+					.Append(endl);
+			}
+
+			// If the only scope we have is WITH_EDITORONLY_DATA, then we can use the existing macros
+			else if (noneAwareScopes == UhtDefineScope.EditorOnlyData)
+			{
+				switch (defineScopeNames)
+				{
+					case UhtDefineScopeNames.Standard:
+						builder
+							.AppendTabs(tabs)
+							.Append("IF_WITH_EDITORONLY_DATA(UE_ARRAY_COUNT(")
+							.AppendCppName(staticsName, arrayName)
+							.Append("), 0)")
+							.Append(endl);
+						break;
+					case UhtDefineScopeNames.WithEditor:
+						builder
+							.AppendTabs(tabs)
+							.Append("IF_WITH_EDITOR(UE_ARRAY_COUNT(")
+							.AppendCppName(staticsName, arrayName)
+							.Append("), 0)")
+							.Append(endl);
+						break;
+					default:
+						throw new UhtIceException("Unexpected define scope names value");
+				}
+			}
+
+			// Otherwise we have many different scopes but no none scope, must generate #if block
+			else
+			{
+				builder.AppendIfPreprocessor(instances.AllScopes, defineScopeNames);
+				builder
+					.AppendTabs(tabs)
+					.Append("UE_ARRAY_COUNT(")
+					.AppendCppName(staticsName, arrayName)
+					.Append(')')
+					.Append(endl);
+				builder.AppendElsePreprocessor(instances.AllScopes, defineScopeNames);
+				builder
+					.AppendTabs(tabs)
+					.Append('0')
+					.Append(endl);
+				builder.AppendEndIfPreprocessor(instances.AllScopes, defineScopeNames);
+			}
+			return builder;
+		}
+
+		/// <summary>
+		/// Start an #if block with the given scope
+		/// </summary>
+		/// <param name="builder">String builder</param>
+		/// <param name="defineScope">Scope</param>
+		/// <param name="defineScopeNames">Which set of scope names will be used</param>
+		/// <returns>String builder</returns>
+		public static StringBuilder AppendIfPreprocessor(this StringBuilder builder, UhtDefineScope defineScope, UhtDefineScopeNames defineScopeNames = UhtDefineScopeNames.Standard)
+		{
+			if (defineScope != UhtDefineScope.None && defineScope != UhtDefineScope.Invalid)
+			{
+				builder.Append("#if ").AppendScopeExpression(defineScope, defineScopeNames).Append("\r\n");
+			}
+			return builder;
+		}
+
+		/// <summary>
+		/// Start an #else block with the given scope
+		/// </summary>
+		/// <param name="builder">String builder</param>
+		/// <param name="defineScope">Scope</param>
+		/// <param name="defineScopeNames">Which set of scope names will be used</param>
+		/// <returns>String builder</returns>
+		public static StringBuilder AppendElsePreprocessor(this StringBuilder builder, UhtDefineScope defineScope, UhtDefineScopeNames defineScopeNames = UhtDefineScopeNames.Standard)
+		{
+			if (defineScope != UhtDefineScope.None && defineScope != UhtDefineScope.Invalid)
+			{
+				builder.Append("#else // ").AppendScopeExpression(defineScope, defineScopeNames).Append("\r\n");
+			}
+			return builder;
+		}
+
+		/// <summary>
+		/// Start an #elif block with the given scope
+		/// </summary>
+		/// <param name="builder">String builder</param>
+		/// <param name="defineScope">Scope</param>
+		/// <param name="defineScopeNames">Which set of scope names will be used</param>
+		/// <returns>String builder</returns>
+		public static StringBuilder AppendElseIfPreprocessor(this StringBuilder builder, UhtDefineScope defineScope, UhtDefineScopeNames defineScopeNames = UhtDefineScopeNames.Standard)
+		{
+			if (defineScope != UhtDefineScope.None && defineScope != UhtDefineScope.Invalid)
+			{
+				builder.Append("#elif ").AppendScopeExpression(defineScope, defineScopeNames).Append("\r\n");
+			}
+			return builder;
+		}
+
+		/// <summary>
+		/// Start an #endif block with the given scope
+		/// </summary>
+		/// <param name="builder">String builder</param>
+		/// <param name="defineScope">Scope</param>
+		/// <param name="defineScopeNames">Which set of scope names will be used</param>
+		/// <returns>String builder</returns>
+		public static StringBuilder AppendEndIfPreprocessor(this StringBuilder builder, UhtDefineScope defineScope, UhtDefineScopeNames defineScopeNames = UhtDefineScopeNames.Standard)
+		{
+			if (defineScope != UhtDefineScope.None && defineScope != UhtDefineScope.Invalid)
+			{
+				builder.Append("#endif // ").AppendScopeExpression(defineScope, defineScopeNames).Append("\r\n");
+			}
+			return builder;
+		}
+
+		/// <summary>
+		/// Append scope expression (i.e. WITH_X || WITH_Y || ...)
+		/// </summary>
+		/// <param name="builder">String builder</param>
+		/// <param name="defineScope">Scope</param>
+		/// <param name="defineScopeNames">Which set of scope names will be used</param>
+		/// <returns>String builder</returns>
+		public static StringBuilder AppendScopeExpression(this StringBuilder builder, UhtDefineScope defineScope, UhtDefineScopeNames defineScopeNames)
+		{
+			if (defineScope != UhtDefineScope.None && defineScope != UhtDefineScope.Invalid)
+			{
+				bool needOr = false;
+				switch (defineScopeNames)
+				{
+					case UhtDefineScopeNames.Standard:
+						builder
+							.AppendName(defineScope, UhtDefineScope.EditorOnlyData, "WITH_EDITORONLY_DATA", ref needOr)
+							;
+						break;
+					case UhtDefineScopeNames.WithEditor:
+						builder
+							.AppendName(defineScope, UhtDefineScope.EditorOnlyData, "WITH_EDITOR", ref needOr)
+							;
+						break;
+					default:
+						throw new UhtIceException("Unexpected define scope names value");
+				}
+			}
+			return builder;
+		}
+
+		private static StringBuilder AppendCppName(this StringBuilder builder, string? staticsName, string arrayName)
+		{
+			if (!String.IsNullOrEmpty(staticsName))
+			{
+				builder.Append(staticsName).Append("::");
+			}
+			return builder.Append(arrayName);
+		}
+
+		private static StringBuilder AppendName(this StringBuilder builder, UhtDefineScope defineScope, UhtDefineScope testScope, string name, ref bool needOr)
+		{
+			if (defineScope.HasAnyFlags(testScope))
+			{
+				if (needOr)
+				{
+					builder.Append("|| ");
+				}
+				builder.Append(name);
+				needOr = true;
+			}
+			return builder;
+		}
+	}
+}
diff --git a/Engine/Source/Programs/Shared/EpicGames.UHT/Parsers/UhtEnumParser.cs b/Engine/Source/Programs/Shared/EpicGames.UHT/Parsers/UhtEnumParser.cs
--- a/Engine/Source/Programs/Shared/EpicGames.UHT/Parsers/UhtEnumParser.cs
+++ b/Engine/Source/Programs/Shared/EpicGames.UHT/Parsers/UhtEnumParser.cs
@@ -70,7 +70,7 @@ namespace EpicGames.UHT.Parsers
 
 					if ((topScope.HeaderParser.GetCurrentCompositeCompilerDirective() & UhtCompilerDirective.WithEditorOnlyData) != 0)
 					{
-						enumObject.IsEditorOnly = true;
+						enumObject.DefineScope |= UhtDefineScope.EditorOnlyData;
 					}
 
 					// Read base for enum class
diff --git a/Engine/Source/Programs/Shared/EpicGames.UHT/Parsers/UhtFunctionParser.cs b/Engine/Source/Programs/Shared/EpicGames.UHT/Parsers/UhtFunctionParser.cs
--- a/Engine/Source/Programs/Shared/EpicGames.UHT/Parsers/UhtFunctionParser.cs
+++ b/Engine/Source/Programs/Shared/EpicGames.UHT/Parsers/UhtFunctionParser.cs
@@ -356,6 +356,7 @@ namespace EpicGames.UHT.Parsers
 					if ((topScope.HeaderParser.GetCurrentCompositeCompilerDirective() & UhtCompilerDirective.WithEditor) != 0)
 					{
 						function.FunctionFlags |= EFunctionFlags.EditorOnly;
+						function.DefineScope |= UhtDefineScope.EditorOnlyData;
 					}
 
 					specifierParser.ParseDeferred();
diff --git a/Engine/Source/Programs/Shared/EpicGames.UHT/Parsers/UhtHeaderFileParser.cs b/Engine/Source/Programs/Shared/EpicGames.UHT/Parsers/UhtHeaderFileParser.cs
--- a/Engine/Source/Programs/Shared/EpicGames.UHT/Parsers/UhtHeaderFileParser.cs
+++ b/Engine/Source/Programs/Shared/EpicGames.UHT/Parsers/UhtHeaderFileParser.cs
@@ -1074,7 +1074,7 @@ namespace EpicGames.UHT.Parsers
 			}
 		}
 
-		private bool IsRestrictedDirective(UhtCompilerDirective compilerDirective)
+		private static bool IsRestrictedDirective(UhtCompilerDirective compilerDirective)
 		{
 			switch (compilerDirective)
 			{
@@ -1240,7 +1240,7 @@ namespace EpicGames.UHT.Parsers
 				return false;
 			}
 
-			IUhtTokenReader replayReader = UhtTokenReplayReader.GetThreadInstance(this.HeaderFile, this.HeaderFile.Data.Memory, new ReadOnlyMemory<UhtToken>(declaration.Tokens), UhtTokenType.EndOfDeclaration);
+			IUhtTokenReader replayReader = UhtTokenReplayReader.GetThreadInstance(HeaderFile, HeaderFile.Data.Memory, new ReadOnlyMemory<UhtToken>(declaration.Tokens), UhtTokenType.EndOfDeclaration);
 
 			SkipVirtualAndAPI(replayReader);
 
diff --git a/Engine/Source/Programs/Shared/EpicGames.UHT/Parsers/UhtPropertyParser.cs b/Engine/Source/Programs/Shared/EpicGames.UHT/Parsers/UhtPropertyParser.cs
--- a/Engine/Source/Programs/Shared/EpicGames.UHT/Parsers/UhtPropertyParser.cs
+++ b/Engine/Source/Programs/Shared/EpicGames.UHT/Parsers/UhtPropertyParser.cs
@@ -797,6 +797,7 @@ namespace EpicGames.UHT.Parsers
 				if (compilerDirective.HasAnyFlags(UhtCompilerDirective.WithEditorOnlyData))
 				{
 					propertySettings.PropertyFlags |= EPropertyFlags.EditorOnly;
+					propertySettings.DefineScope |= UhtDefineScope.EditorOnlyData;
 				}
 				else if (compilerDirective.HasAnyFlags(UhtCompilerDirective.WithEditor))
 				{
@@ -845,6 +846,7 @@ namespace EpicGames.UHT.Parsers
 						if (layoutMacroType.IsEditorOnly())
 						{
 							propertySettings.PropertyFlags |= EPropertyFlags.EditorOnly;
+							propertySettings.DefineScope |= UhtDefineScope.EditorOnlyData;
 						}
 					}
 				}
@@ -1305,9 +1307,8 @@ namespace EpicGames.UHT.Parsers
 		[SuppressMessage("Style", "IDE0060:Remove unused parameter", Justification = "Attribute accessed method")]
 		private static UhtProperty? DefaultProperty(UhtPropertyResolvePhase resolvePhase, UhtPropertySettings propertySettings, IUhtTokenReader tokenReader, UhtToken matchedToken)
 		{
-			UhtProperty? property = null;
+			UhtProperty? property;
 			UhtSession session = propertySettings.Outer.Session;
-			int typeStartPos = tokenReader.PeekToken().InputStartPos;
 
 			bool gotConst = tokenReader.TryOptional("const");
 
diff --git a/Engine/Source/Programs/Shared/EpicGames.UHT/Types/Properties/UhtStructProperty.cs b/Engine/Source/Programs/Shared/EpicGames.UHT/Types/Properties/UhtStructProperty.cs
--- a/Engine/Source/Programs/Shared/EpicGames.UHT/Types/Properties/UhtStructProperty.cs
+++ b/Engine/Source/Programs/Shared/EpicGames.UHT/Types/Properties/UhtStructProperty.cs
@@ -1,6 +1,5 @@
 // Copyright Epic Games, Inc. All Rights Reserved.
 
-using System;
 using System.Collections.Generic;
 using System.Diagnostics.CodeAnalysis;
 using System.Globalization;
diff --git a/Engine/Source/Programs/Shared/EpicGames.UHT/Types/UhtEnum.cs b/Engine/Source/Programs/Shared/EpicGames.UHT/Types/UhtEnum.cs
--- a/Engine/Source/Programs/Shared/EpicGames.UHT/Types/UhtEnum.cs
+++ b/Engine/Source/Programs/Shared/EpicGames.UHT/Types/UhtEnum.cs
@@ -133,11 +133,6 @@ namespace EpicGames.UHT.Types
 		/// </summary>
 		public string CppType { get; set; } = String.Empty;
 
-		/// <summary>
-		/// True if the enum was declared in an WITH_EDITORONLY_DATA block
-		/// </summary>
-		public bool IsEditorOnly { get; set; } = false;
-
 		/// <summary>
 		/// Collection of enumeration values
 		/// </summary>
diff --git a/Engine/Source/Programs/Shared/EpicGames.UHT/Types/UhtProperty.cs b/Engine/Source/Programs/Shared/EpicGames.UHT/Types/UhtProperty.cs
--- a/Engine/Source/Programs/Shared/EpicGames.UHT/Types/UhtProperty.cs
+++ b/Engine/Source/Programs/Shared/EpicGames.UHT/Types/UhtProperty.cs
@@ -690,6 +690,11 @@ namespace EpicGames.UHT.Types
 		/// </summary>
 		public UhtPropertyExportFlags PropertyExportFlags { get; set; }
 
+		/// <summary>
+		/// #define scope of where the property exists
+		/// </summary>
+		public UhtDefineScope DefineScope { get; set; }
+
 		/// <summary>
 		/// Allocator used for containers
 		/// </summary>
@@ -762,6 +767,7 @@ namespace EpicGames.UHT.Types
 			PropertyFlags = parentPropertySettings.PropertyFlags;
 			DisallowPropertyFlags = parentPropertySettings.DisallowPropertyFlags;
 			PropertyExportFlags = UhtPropertyExportFlags.Public;
+			DefineScope = parentPropertySettings.DefineScope;
 			RepNotifyName = null;
 			Allocator = UhtPropertyAllocator.Default;
 			Options = parentPropertySettings.Options;
@@ -790,6 +796,7 @@ namespace EpicGames.UHT.Types
 			PropertyCategory = propertyCategory;
 			PropertyFlags = EPropertyFlags.None;
 			DisallowPropertyFlags = disallowPropertyFlags;
+			DefineScope = UhtDefineScope.None;
 			PropertyExportFlags = UhtPropertyExportFlags.Public;
 			RepNotifyName = null;
 			Allocator = UhtPropertyAllocator.Default;
@@ -822,6 +829,7 @@ namespace EpicGames.UHT.Types
 			PropertyFlags = property.PropertyFlags;
 			DisallowPropertyFlags = property.DisallowPropertyFlags;
 			PropertyExportFlags = property.PropertyExportFlags;
+			DefineScope = property.DefineScope;
 			Allocator = property.Allocator;
 			Options = options;
 			PointerType = property.PointerType;
@@ -1029,7 +1037,7 @@ namespace EpicGames.UHT.Types
 		/// If true, the property is editor only
 		/// </summary>
 		[JsonIgnore]
-		public bool IsEditorOnlyProperty => PropertyFlags.HasAnyFlags(EPropertyFlags.DevelopmentAssets);
+		public bool IsEditorOnlyProperty => PropertyFlags.HasAnyFlags(EPropertyFlags.EditorOnly);
 
 		/// <summary>
 		/// Construct a new property
@@ -1058,6 +1066,7 @@ namespace EpicGames.UHT.Types
 			PropertyFlags = propertySettings.PropertyFlags;
 			DisallowPropertyFlags = propertySettings.DisallowPropertyFlags;
 			PropertyExportFlags = propertySettings.PropertyExportFlags;
+			DefineScope = propertySettings.DefineScope;
 			Allocator = propertySettings.Allocator;
 			PointerType = propertySettings.PointerType;
 			RepNotifyName = propertySettings.RepNotifyName;
@@ -1958,7 +1967,7 @@ namespace EpicGames.UHT.Types
 		/// <returns>True if GC barriers need to run</returns>				
 		public bool NeedsGCBarrierWhenPassedToFunction(UhtFunction function)
 		{
-			if (this.RefQualifier != UhtPropertyRefQualifier.NonConstRef)
+			if (RefQualifier != UhtPropertyRefQualifier.NonConstRef)
 			{
 				return false;
 			}
diff --git a/Engine/Source/Programs/Shared/EpicGames.UHT/Types/UhtType.cs b/Engine/Source/Programs/Shared/EpicGames.UHT/Types/UhtType.cs
--- a/Engine/Source/Programs/Shared/EpicGames.UHT/Types/UhtType.cs
+++ b/Engine/Source/Programs/Shared/EpicGames.UHT/Types/UhtType.cs
@@ -523,6 +523,77 @@ namespace EpicGames.UHT.Types
 		}
 	}
 
+	/// <summary>
+	/// #define scope where the type was defined.  This only includes macros such as WITH_EDITOR and WITH_EDITORONLY_DATA.
+	/// Not all types support the given options.
+	/// </summary>
+	[Flags]
+	public enum UhtDefineScope
+	{
+
+		/// <summary>
+		/// The scope is invalid.  Used during code generation.
+		/// </summary>
+		Invalid = -1,
+
+		/// <summary>
+		/// Exists outside of a define scope
+		/// </summary>
+		None = 0,
+
+		/// <summary>
+		/// Exists inside of a WITH_EDITORONLY_DATA block
+		/// </summary>
+		EditorOnlyData = 1 << 0,
+
+		/// <summary>
+		/// Number of unique scope combinations
+		/// </summary>
+		ScopeCount = 1 << 1,
+	}
+
+	/// <summary>
+	/// Helper methods for testing flags.  These methods perform better than the generic HasFlag which hits
+	/// the GC and stalls.
+	/// </summary>
+	public static class UhtDefineScopeExtensions
+	{
+
+		/// <summary>
+		/// Test to see if any of the specified flags are set
+		/// </summary>
+		/// <param name="inFlags">Current flags</param>
+		/// <param name="testFlags">Flags to test for</param>
+		/// <returns>True if any of the flags are set</returns>
+		public static bool HasAnyFlags(this UhtDefineScope inFlags, UhtDefineScope testFlags)
+		{
+			return (inFlags & testFlags) != 0;
+		}
+
+		/// <summary>
+		/// Test to see if all of the specified flags are set
+		/// </summary>
+		/// <param name="inFlags">Current flags</param>
+		/// <param name="testFlags">Flags to test for</param>
+		/// <returns>True if all the flags are set</returns>
+		public static bool HasAllFlags(this UhtDefineScope inFlags, UhtDefineScope testFlags)
+		{
+			return (inFlags & testFlags) == testFlags;
+		}
+
+		/// <summary>
+		/// Test to see if a specific set of flags have a specific value.
+		/// </summary>
+		/// <param name="inFlags">Current flags</param>
+		/// <param name="testFlags">Flags to test for</param>
+		/// <param name="matchFlags">Expected value of the tested flags</param>
+		/// <returns>True if the given flags have a specific value.</returns>
+		public static bool HasExactFlags(this UhtDefineScope inFlags, UhtDefineScope testFlags, UhtDefineScope matchFlags)
+		{
+			return (inFlags & testFlags) == matchFlags;
+		}
+	}
+
 	/// <summary>
 	/// Access specifier for type
 	/// </summary>
@@ -731,6 +802,12 @@ namespace EpicGames.UHT.Types
 		/// </summary>
 		public int LineNumber { get; set; }
 
+		/// <summary>
+		/// #define scope where the type was defined
+		/// </summary>
+		[JsonConverter(typeof(JsonStringEnumConverter))]
+		public UhtDefineScope DefineScope { get; set; } = UhtDefineScope.None;
+
 		/// <summary>
 		/// Return a combination of the engine type name followed by the path name of the type
 		/// </summary>
@@ -813,7 +890,7 @@ namespace EpicGames.UHT.Types
 		#endregion
 
 		/// <summary>
-		/// Construct the base type information.  This constructor is used exclusively by UHTPackage which is at 
+		/// Construct the base type information.  This constructor is used exclusively by UhtPackage which is at 
 		/// the root of all type hierarchies.
 		/// </summary>
 		/// <param name="session"></param>
