Cleanup of empty #define as part of UHT generated code.
Always trim #defines found inside of WITH_ blocks

diff --git a/Engine/Source/Programs/Shared/EpicGames.UHT/Exporters/CodeGen/UhtHeaderCodeGeneratorHFile.cs b/Engine/Source/Programs/Shared/EpicGames.UHT/Exporters/CodeGen/UhtHeaderCodeGeneratorHFile.cs
--- a/Engine/Source/Programs/Shared/EpicGames.UHT/Exporters/CodeGen/UhtHeaderCodeGeneratorHFile.cs
+++ b/Engine/Source/Programs/Shared/EpicGames.UHT/Exporters/CodeGen/UhtHeaderCodeGeneratorHFile.cs
@@ -536,7 +536,8 @@ namespace EpicGames.UHT.Exporters.CodeGen
 			AppendRpcFunctions(builder, classObj, alternateUsesLegacy, rpcFunctions);
 
 			// Output property accessors
-			AppendPropertyAccessors(builder, classObj);
+			IEnumerable<UhtProperty> getterSetterProperties = classObj.Properties.Where(x => x.PropertyExportFlags.HasAnyFlags(UhtPropertyExportFlags.GetterFound | UhtPropertyExportFlags.SetterFound));
+			AppendPropertyAccessors(builder, classObj, getterSetterProperties);
 
 			// Collect the callback function and sort by name to make the order stable
 			List<UhtFunction> callbackFunctions = new(classObj.Functions.Where(x => x.FunctionFlags.HasAnyFlags(EFunctionFlags.Event) && x.SuperFunction == null));
@@ -606,7 +607,7 @@ namespace EpicGames.UHT.Exporters.CodeGen
 
 				AppendProlog(builder, classObj);
 				AppendGeneratedBodyMacroBlock(builder, classObj, nativeInterface!, alternateUsesLegacy, rpcFunctions, autoGetterSetterProperties, 
-					sparseScriptStructs.Any(), sparseProperties, hasCallbacks, null);
+					sparseScriptStructs.Any(), sparseProperties, getterSetterProperties, hasCallbacks, null);
 			}
 			else
 			{
@@ -631,7 +632,7 @@ namespace EpicGames.UHT.Exporters.CodeGen
 				AppendAutoGettersSetters(builder, classObj, autoGetterSetterProperties);
 				AppendProlog(builder, classObj);
 				AppendGeneratedBodyMacroBlock(builder, classObj, classObj, usesLegacy, rpcFunctions, autoGetterSetterProperties, 
-					sparseScriptStructs.Any(), sparseProperties, hasCallbacks, usesLegacy ? "GENERATED_UCLASS_BODY" : null);
+					sparseScriptStructs.Any(), sparseProperties, getterSetterProperties, hasCallbacks, usesLegacy ? "GENERATED_UCLASS_BODY" : null);
 			}
 
 			// Forward declare the StaticClass specialization in the header
@@ -660,7 +661,7 @@ namespace EpicGames.UHT.Exporters.CodeGen
 			}
 
 			UhtUsedDefineScopes<UhtType> notifyTypes = GetFieldNotifyTypes(classObj);
-			return builder.AppendSingleMacro(notifyTypes, UhtDefineScopeNames.Standard, this, classObj, FieldNotifyMacroSuffix, false,
+			return builder.AppendSingleMacro(notifyTypes, UhtDefineScopeNames.Standard, this, classObj, FieldNotifyMacroSuffix,
 				(builder, filteredTypes) =>
 				{
 					builder.Append("\tUE_FIELD_NOTIFICATION_DECLARE_CLASS_DESCRIPTOR_BEGIN(").Append(PackageApi).Append(") \\\r\n");
@@ -697,7 +698,7 @@ namespace EpicGames.UHT.Exporters.CodeGen
 				return builder;
 			}
 
-			return builder.AppendMultiMacros(autoGetterSetterProperties, UhtDefineScopeNames.Standard, this, classObj, AutoGettersSettersMacroSuffix, true,
+			return builder.AppendMultiMacros(autoGetterSetterProperties, UhtDefineScopeNames.Standard, this, classObj, AutoGettersSettersMacroSuffix,
 				(builder, properties) =>
 				{
 					foreach (UhtProperty property in properties)
@@ -740,7 +741,7 @@ namespace EpicGames.UHT.Exporters.CodeGen
 			AppendSparseStructDeclarations(builder, classObj, sparseScriptStructs);
 
 			builder.AppendMultiMacros(sparseProperties, UhtDefineScopeNames.Standard, this, classObj, SparseDataPropertyAccessorsMacroSuffix, 
-				true, (builder, properties) =>
+				(builder, properties) =>
 				{
 					foreach (UhtProperty property in properties)
 					{
@@ -769,18 +770,15 @@ namespace EpicGames.UHT.Exporters.CodeGen
 		{
 			using (UhtMacroCreator macro = new(builder, this, classObj, SparseDataMacroSuffix))
 			{
-				if (sparseScriptStructs.Any())
-				{
-					string api = classObj.ClassFlags.HasAnyFlags(EClassFlags.MinimalAPI) ? PackageApi : "";
+				string api = classObj.ClassFlags.HasAnyFlags(EClassFlags.MinimalAPI) ? PackageApi : "";
 
-					foreach (UhtScriptStruct sparseScriptStruct in sparseScriptStructs)
-					{
-						string sparseDataType = sparseScriptStruct.EngineName;
+				foreach (UhtScriptStruct sparseScriptStruct in sparseScriptStructs)
+				{
+					string sparseDataType = sparseScriptStruct.EngineName;
 
-						builder.Append(api).Append('F').Append(sparseDataType).Append("* Get").Append(sparseDataType).Append("() const; \\\r\n");
-						builder.Append(api).Append("const F").Append(sparseDataType).Append("* Get").Append(sparseDataType).Append("(EGetSparseClassDataMethod GetMethod) const; \\\r\n");
-						builder.Append(api).Append("static UScriptStruct* StaticGet").Append(sparseDataType).Append("ScriptStruct(); \\\r\n");
-					}
+					builder.Append(api).Append('F').Append(sparseDataType).Append("* Get").Append(sparseDataType).Append("() const; \\\r\n");
+					builder.Append(api).Append("const F").Append(sparseDataType).Append("* Get").Append(sparseDataType).Append("(EGetSparseClassDataMethod GetMethod) const; \\\r\n");
+					builder.Append(api).Append("static UScriptStruct* StaticGet").Append(sparseDataType).Append("ScriptStruct(); \\\r\n");
 				}
 			}
 
@@ -789,7 +787,7 @@ namespace EpicGames.UHT.Exporters.CodeGen
 
 		private StringBuilder AppendRpcFunctions(StringBuilder builder, UhtClass classObj, bool usesLegacy, UhtUsedDefineScopes<UhtFunction> functions)
 		{
-			builder.AppendMultiMacros(functions, UhtDefineScopeNames.WithEditor, this, classObj, usesLegacy ? RpcWrappersMacroSuffix : RpcWrappersNoPureDeclsMacroSuffix, false,
+			builder.AppendMultiMacros(functions, UhtDefineScopeNames.WithEditor, this, classObj, usesLegacy ? RpcWrappersMacroSuffix : RpcWrappersNoPureDeclsMacroSuffix, 
 				(builder, functions) =>
 				{
 					if (usesLegacy)
@@ -809,13 +807,13 @@ namespace EpicGames.UHT.Exporters.CodeGen
 			return builder;
 		}
 
-		private StringBuilder AppendPropertyAccessors(StringBuilder builder, UhtClass classObj)
+		private StringBuilder AppendPropertyAccessors(StringBuilder builder, UhtClass classObj, IEnumerable<UhtProperty> getterSetterProperties)
 		{
-			using (UhtMacroCreator macro = new(builder, this, classObj, AccessorsMacroSuffix))
+			if (getterSetterProperties.Any())
 			{
-				foreach (UhtType type in classObj.Children)
+				using (UhtMacroCreator macro = new(builder, this, classObj, AccessorsMacroSuffix))
 				{
-					if (type is UhtProperty property)
+					foreach (UhtProperty property in getterSetterProperties)
 					{
 						if (property.PropertyExportFlags.HasAnyFlags(UhtPropertyExportFlags.GetterFound))
 						{
@@ -1463,7 +1461,7 @@ namespace EpicGames.UHT.Exporters.CodeGen
 
 		private StringBuilder AppendGeneratedBodyMacroBlock(StringBuilder builder, UhtClass classObj, UhtClass bodyClassObj, bool isLegacy, 
 			UhtUsedDefineScopes<UhtFunction> rpcFunctions, UhtUsedDefineScopes<UhtProperty> autoGetterSetterProperties, bool hasSparseStructs,
-			UhtUsedDefineScopes<UhtProperty> sparseProperties, bool hasCallbacks, string? deprecatedMacroName)
+			UhtUsedDefineScopes<UhtProperty> sparseProperties, IEnumerable<UhtProperty> getterSetterProperties, bool hasCallbacks, string? deprecatedMacroName)
 		{
 			bool isInterface = classObj.ClassFlags.HasAnyFlags(EClassFlags.Interface);
 			using (UhtMacroCreator macro = new(builder, this, bodyClassObj, isLegacy ? GeneratedBodyLegacyMacroSuffix : GeneratedBodyMacroSuffix))
@@ -1480,11 +1478,11 @@ namespace EpicGames.UHT.Exporters.CodeGen
 					builder.AppendMultiMacroRefs(sparseProperties, this, classObj, SparseDataPropertyAccessorsMacroSuffix);
 				}
 				builder.AppendMultiMacroRefs(rpcFunctions, this, classObj, isLegacy ? RpcWrappersMacroSuffix : RpcWrappersNoPureDeclsMacroSuffix);
-				builder.Append('\t').AppendMacroName(this, classObj, AccessorsMacroSuffix).Append(" \\\r\n");
-				if (!autoGetterSetterProperties.IsEmpty)
+				if (getterSetterProperties.Any())
 				{
-					builder.AppendMultiMacroRefs(autoGetterSetterProperties, this, classObj, AutoGettersSettersMacroSuffix);
+					builder.Append('\t').AppendMacroName(this, classObj, AccessorsMacroSuffix).Append(" \\\r\n");
 				}
+				builder.AppendMultiMacroRefs(autoGetterSetterProperties, this, classObj, AutoGettersSettersMacroSuffix);
 				if (hasCallbacks)
 				{
 					builder.Append('\t').AppendMacroName(this, classObj, CallbackWrappersMacroSuffix).Append(" \\\r\n");
diff --git a/Engine/Source/Programs/Shared/EpicGames.UHT/Exporters/CodeGen/UhtUsedDefineScopes.cs b/Engine/Source/Programs/Shared/EpicGames.UHT/Exporters/CodeGen/UhtUsedDefineScopes.cs
--- a/Engine/Source/Programs/Shared/EpicGames.UHT/Exporters/CodeGen/UhtUsedDefineScopes.cs
+++ b/Engine/Source/Programs/Shared/EpicGames.UHT/Exporters/CodeGen/UhtUsedDefineScopes.cs
@@ -125,11 +125,8 @@ namespace EpicGames.UHT.Exporters.CodeGen
 		/// <returns>Enumeration</returns>
 		public IEnumerable<UhtDefineScope> EnumerateDefinedScopes()
 		{
-			// At this time, we alway return None
-			yield return UhtDefineScope.None;
-
-			ulong present = _present >> 1;
-			for (int index = 1; present != 0; ++index, present >>= 1)
+			ulong present = _present;
+			for (int index = 0; present != 0; ++index, present >>= 1)
 			{
 				if ((present & 1) != 0)
 				{
@@ -210,11 +207,10 @@ namespace EpicGames.UHT.Exporters.CodeGen
 		/// <param name="generator">Header code generator</param>
 		/// <param name="outerType">Output type owning the instances</param>
 		/// <param name="macroSuffix">Macro being created</param>
-		/// <param name="trimMacro">If true, trim the end of the prior macro</param>
 		/// <param name="appendAction">Action to invoke to append an instance</param>
 		/// <returns>String builder</returns>
 		public static StringBuilder AppendMultiMacros<T>(this StringBuilder builder, UhtUsedDefineScopes<T> instances, UhtDefineScopeNames defineScopeNames, 
-			UhtHeaderCodeGenerator generator, UhtType outerType, string macroSuffix, bool trimMacro, Action<StringBuilder, IEnumerable<T>> appendAction) where T : UhtType
+			UhtHeaderCodeGenerator generator, UhtType outerType, string macroSuffix, Action<StringBuilder, IEnumerable<T>> appendAction) where T : UhtType
 		{
 			foreach (UhtDefineScope defineScope in instances.EnumerateDefinedScopes())
 			{
@@ -228,17 +224,14 @@ namespace EpicGames.UHT.Exporters.CodeGen
 					// We can skip writing the macros if there are no properties to declare, as the 'if' and 'else' would be the same
 					if (defineScope != UhtDefineScope.None)
 					{
-						if (trimMacro)
+						// Trim the extra newlines added after the macro generator
+						if (builder.Length > 4 &&
+							builder[^4] == '\r' &&
+							builder[^3] == '\n' &&
+							builder[^2] == '\r' &&
+							builder[^1] == '\n')
 						{
-							// Trim the extra newlines added after the macro generator
-							if (builder.Length > 4 &&
-								builder[^4] == '\r' &&
-								builder[^3] == '\n' &&
-								builder[^2] == '\r' &&
-								builder[^1] == '\n')
-							{
-								builder.Length -= 4;
-							}
+							builder.Length -= 4;
 						}
 
 						builder.AppendElsePreprocessor(defineScope, defineScopeNames);
@@ -246,7 +239,7 @@ namespace EpicGames.UHT.Exporters.CodeGen
 					}
 				}
 
-				if (trimMacro && defineScope != UhtDefineScope.None)
+				if (defineScope != UhtDefineScope.None)
 				{
 					builder.Append("\r\n\r\n");
 				}
@@ -283,11 +276,10 @@ namespace EpicGames.UHT.Exporters.CodeGen
 		/// <param name="generator">Header code generator</param>
 		/// <param name="outerType">Output type owning the instances</param>
 		/// <param name="macroSuffix">Macro being created</param>
-		/// <param name="trimMacro">If true, trim the end of the prior macro</param>
 		/// <param name="appendAction">Action to invoke to append an instance</param>
 		/// <returns>String builder</returns>
 		public static StringBuilder AppendSingleMacro<T>(this StringBuilder builder, UhtUsedDefineScopes<T> instances, UhtDefineScopeNames defineScopeNames, 
-			UhtHeaderCodeGenerator generator, UhtType outerType, string macroSuffix, bool trimMacro, Action<StringBuilder, IEnumerable<T>> appendAction) where T : UhtType
+			UhtHeaderCodeGenerator generator, UhtType outerType, string macroSuffix, Action<StringBuilder, IEnumerable<T>> appendAction) where T : UhtType
 		{
 			if (instances.SoleScope == UhtDefineScope.None)
 			{
@@ -316,19 +308,16 @@ namespace EpicGames.UHT.Exporters.CodeGen
 					{
 						appendAction(builder, instances.Instances.Where(x => (x.DefineScope & ~defineScope) == 0));
 					}
-					if (trimMacro)
+					if (builder.Length > 4 &&
+						builder[^4] == '\r' &&
+						builder[^3] == '\n' &&
+						builder[^2] == '\r' &&
+						builder[^1] == '\n')
 					{
-						if (builder.Length > 4 &&
-							builder[^4] == '\r' &&
-							builder[^3] == '\n' &&
-							builder[^2] == '\r' &&
-							builder[^1] == '\n')
-						{
-							builder.Length -= 4;
-						}
+						builder.Length -= 4;
 					}
 				}
-				builder.Append("#endif\r\n");
+				builder.Append("#endif\r\n\r\n\r\n");
 			}
 			return builder;
 		}
