Box: Include the worker id number in any log messages output from Horde. Also add a warning when the initiator has been waiting for an attach message for too long.

diff --git a/Engine/Source/Programs/Horde/Horde.Agent/Commands/Compute/ComputeCommand.cs b/Engine/Source/Programs/Horde/Horde.Agent/Commands/Compute/ComputeCommand.cs
--- a/Engine/Source/Programs/Horde/Horde.Agent/Commands/Compute/ComputeCommand.cs
+++ b/Engine/Source/Programs/Horde/Horde.Agent/Commands/Compute/ComputeCommand.cs
@@ -58,7 +58,7 @@ namespace Horde.Agent.Commands.Compute
 				requirements = new Requirements(Condition.Parse(Requirements));
 			}
 
-			await using IComputeLease? lease = await client.TryAssignWorkerAsync(new ClusterId(ClusterId), requirements, null, CancellationToken.None);
+			await using IComputeLease? lease = await client.TryAssignWorkerAsync(new ClusterId(ClusterId), requirements, null, logger, CancellationToken.None);
 			if (lease == null)
 			{
 				throw new Exception("Unable to create lease");
diff --git a/Engine/Source/Programs/Horde/Horde/Commands/Compute/ComputeRun.cs b/Engine/Source/Programs/Horde/Horde/Commands/Compute/ComputeRun.cs
--- a/Engine/Source/Programs/Horde/Horde/Commands/Compute/ComputeRun.cs
+++ b/Engine/Source/Programs/Horde/Horde/Commands/Compute/ComputeRun.cs
@@ -72,7 +72,7 @@ namespace Horde.Commands.Compute
 				requirements = new Requirements(Condition.Parse(Requirements));
 			}
 
-			await using IComputeLease? lease = await client.TryAssignWorkerAsync(new ClusterId(ClusterId), requirements, null, CancellationToken.None);
+			await using IComputeLease? lease = await client.TryAssignWorkerAsync(new ClusterId(ClusterId), requirements, null, logger, CancellationToken.None);
 			if (lease == null)
 			{
 				throw new Exception("Unable to create lease");
diff --git a/Engine/Source/Programs/Horde/Samples/RemoteClient/RemoteClientApp.cs b/Engine/Source/Programs/Horde/Samples/RemoteClient/RemoteClientApp.cs
--- a/Engine/Source/Programs/Horde/Samples/RemoteClient/RemoteClientApp.cs
+++ b/Engine/Source/Programs/Horde/Samples/RemoteClient/RemoteClientApp.cs
@@ -59,7 +59,7 @@ namespace RemoteClient
 				requirements = new Requirements(Condition.Parse(options.Condition));
 			}
 
-			await using IComputeLease? lease = await client.TryAssignWorkerAsync(new ClusterId("default"), requirements, null);
+			await using IComputeLease? lease = await client.TryAssignWorkerAsync(new ClusterId("default"), requirements, null, logger);
 			if (lease == null)
 			{
 				logger.LogInformation("Unable to connect to remote");
diff --git a/Engine/Source/Programs/Shared/EpicGames.Core/PrefixLogger.cs b/Engine/Source/Programs/Shared/EpicGames.Core/PrefixLogger.cs
new file mode 100644
--- /dev/null
+++ b/Engine/Source/Programs/Shared/EpicGames.Core/PrefixLogger.cs
@@ -0,0 +1,53 @@
+// Copyright Epic Games, Inc. All Rights Reserved.
+
+using System;
+using System.Collections.Generic;
+using Microsoft.Extensions.Logging;
+
+namespace EpicGames.Core
+{
+	/// <summary>
+	/// ILogger implementation that adds a prefix to the start of each line
+	/// </summary>
+	public class PrefixLogger : ILogger
+	{
+		readonly string _prefix;
+		readonly ILogger _inner;
+
+		/// <summary>
+		/// Constructor
+		/// </summary>
+		/// <param name="prefix">Prefix to insert</param>
+		/// <param name="inner">Logger to write to</param>
+		public PrefixLogger(string prefix, ILogger inner)
+		{
+			_prefix = prefix;
+			_inner = inner;
+		}
+
+		/// <inheritdoc/>
+		public IDisposable BeginScope<TState>(TState state) => _inner.BeginScope<TState>(state);
+
+		/// <inheritdoc/>
+		public bool IsEnabled(LogLevel logLevel) => _inner.IsEnabled(logLevel);
+
+		/// <inheritdoc/>
+		public void Log<TState>(LogLevel logLevel, EventId eventId, TState state, Exception? exception, Func<TState, Exception?, string> formatter)
+		{
+			if (state is IEnumerable<KeyValuePair<string, object>> enumerable)
+			{
+				List<KeyValuePair<string, object>> copy = new List<KeyValuePair<string, object>>(enumerable);
+
+				int idx = copy.FindIndex(x => x.Key.Equals("{OriginalFormat}", StringComparison.OrdinalIgnoreCase));
+				if (idx != -1 && copy[idx].Value is string format)
+				{
+					copy[idx] = new KeyValuePair<string, object>(copy[idx].Key, "[{_tag}] " + format);
+					copy.Add(new KeyValuePair<string, object>("_tag", _prefix));
+					_inner.Log(logLevel, eventId, copy, exception, (s, e) => $"[{_prefix}] {formatter(state, exception)}");
+					return;
+				}
+			}
+			_inner.Log(logLevel, eventId, state, exception, formatter);
+		}
+	}
+}
diff --git a/Engine/Source/Programs/Shared/EpicGames.Horde/Compute/Clients/AgentComputeClient.cs b/Engine/Source/Programs/Shared/EpicGames.Horde/Compute/Clients/AgentComputeClient.cs
--- a/Engine/Source/Programs/Shared/EpicGames.Horde/Compute/Clients/AgentComputeClient.cs
+++ b/Engine/Source/Programs/Shared/EpicGames.Horde/Compute/Clients/AgentComputeClient.cs
@@ -64,21 +64,15 @@ namespace EpicGames.Horde.Compute.Clients
 
 		/// <inheritdoc/>
 		public ValueTask DisposeAsync() => new ValueTask();
-
-		/// <inheritdoc/>
-		public Task<IComputeLease?> TryAssignWorkerAsync(ClusterId clusterId, Requirements? requirements, CancellationToken cancellationToken)
-		{
-			return TryAssignWorkerAsync(clusterId, requirements, null, cancellationToken);
-		}
 		
 		/// <inheritdoc/>
-		public async Task<IComputeLease?> TryAssignWorkerAsync(ClusterId clusterId, Requirements? requirements, string? requestId, CancellationToken cancellationToken)
+		public async Task<IComputeLease?> TryAssignWorkerAsync(ClusterId clusterId, Requirements? requirements, string? requestId, ILogger logger, CancellationToken cancellationToken)
 		{
-			_logger.LogInformation("** CLIENT **");
-			_logger.LogInformation("Launching {Path} to handle remote", _hordeAgentAssembly);
+			logger.LogInformation("** CLIENT **");
+			logger.LogInformation("Launching {Path} to handle remote", _hordeAgentAssembly);
 
 			// The connection logic is an async enumerator that returns the socket, then shuts down.
-			IAsyncEnumerator<RemoteComputeSocket> source = ConnectAsync(cancellationToken).GetAsyncEnumerator(cancellationToken);
+			IAsyncEnumerator<RemoteComputeSocket> source = ConnectAsync(logger, cancellationToken).GetAsyncEnumerator(cancellationToken);
 			if (!await source.MoveNextAsync())
 			{
 				await source.DisposeAsync();
@@ -88,13 +82,13 @@ namespace EpicGames.Horde.Compute.Clients
 			return new LeaseImpl(source);
 		}
 
-		async IAsyncEnumerable<RemoteComputeSocket> ConnectAsync([EnumeratorCancellation] CancellationToken cancellationToken)
+		async IAsyncEnumerable<RemoteComputeSocket> ConnectAsync(ILogger logger, [EnumeratorCancellation] CancellationToken cancellationToken)
 		{
 			using Socket listener = new Socket(SocketType.Stream, ProtocolType.IP);
 			listener.Bind(new IPEndPoint(IPAddress.Loopback, _port));
 			listener.Listen();
 
-			using BackgroundTask agentTask = BackgroundTask.StartNew(ctx => RunAgentAsync(_hordeAgentAssembly, _port, _logger, ctx));
+			using BackgroundTask agentTask = BackgroundTask.StartNew(ctx => RunAgentAsync(_hordeAgentAssembly, _port, logger, ctx));
 			using Socket tcpSocket = await listener.AcceptAsync(cancellationToken);
 
 			await using RemoteComputeSocket socket = new RemoteComputeSocket(new TcpTransport(tcpSocket), _logger);
diff --git a/Engine/Source/Programs/Shared/EpicGames.Horde/Compute/Clients/LocalComputeClient.cs b/Engine/Source/Programs/Shared/EpicGames.Horde/Compute/Clients/LocalComputeClient.cs
--- a/Engine/Source/Programs/Shared/EpicGames.Horde/Compute/Clients/LocalComputeClient.cs
+++ b/Engine/Source/Programs/Shared/EpicGames.Horde/Compute/Clients/LocalComputeClient.cs
@@ -35,39 +35,6 @@ namespace EpicGames.Horde.Compute.Clients
 			public ValueTask CloseAsync(CancellationToken cancellationToken) => _socket.CloseAsync(cancellationToken);
 		}
 
-		class PrefixLogger : ILogger
-		{
-			readonly string _prefix;
-			readonly ILogger _inner;
-
-			public PrefixLogger(string prefix, ILogger inner)
-			{
-				_prefix = prefix;
-				_inner = inner;
-			}
-
-			public IDisposable BeginScope<TState>(TState state) => _inner.BeginScope<TState>(state);
-			public bool IsEnabled(LogLevel logLevel) => _inner.IsEnabled(logLevel);
-
-			public void Log<TState>(LogLevel logLevel, EventId eventId, TState state, Exception? exception, Func<TState, Exception?, string> formatter)
-			{
-				if (state is IEnumerable<KeyValuePair<string, object>> enumerable)
-				{
-					List<KeyValuePair<string, object>> copy = new List<KeyValuePair<string, object>>(enumerable);
-
-					int idx = copy.FindIndex(x => x.Key.Equals("{OriginalFormat}", StringComparison.OrdinalIgnoreCase));
-					if (idx != -1 && copy[idx].Value is string format)
-					{
-						copy[idx] = new KeyValuePair<string, object>(copy[idx].Key, "[{_tag}] " + format);
-						copy.Add(new KeyValuePair<string, object>("_tag", _prefix));
-						_inner.Log(logLevel, eventId, copy, exception, (s, e) => $"[{_prefix}] {formatter(state, exception)}");
-						return;
-					}
-				}
-				_inner.Log(logLevel, eventId, state, exception, formatter);
-			}
-		}
-
 		readonly BackgroundTask _listenerTask;
 		readonly Socket _listener;
 		readonly Socket _socket;
@@ -120,18 +87,12 @@ namespace EpicGames.Horde.Compute.Clients
 				await socket.CloseAsync(cancellationToken);
 			}
 		}
-		
-		/// <inheritdoc/>
-		public Task<IComputeLease?> TryAssignWorkerAsync(ClusterId clusterId, Requirements? requirements, CancellationToken cancellationToken)
-		{
-			return TryAssignWorkerAsync(clusterId, requirements, null, cancellationToken);
-		}
 
 		/// <inheritdoc/>
-		public Task<IComputeLease?> TryAssignWorkerAsync(ClusterId clusterId, Requirements? requirements, string? requestId, CancellationToken cancellationToken)
+		public Task<IComputeLease?> TryAssignWorkerAsync(ClusterId clusterId, Requirements? requirements, string? requestId, ILogger logger, CancellationToken cancellationToken)
 		{
 #pragma warning disable CA2000 // Dispose objects before losing scope
-			RemoteComputeSocket socket = new RemoteComputeSocket(new TcpTransport(_socket), new PrefixLogger("CLIENT", _logger));
+			RemoteComputeSocket socket = new RemoteComputeSocket(new TcpTransport(_socket), logger);
 			return Task.FromResult<IComputeLease?>(new LeaseImpl(socket));
 #pragma warning restore CA2000 // Dispose objects before losing scope
 		}
diff --git a/Engine/Source/Programs/Shared/EpicGames.Horde/Compute/Clients/ServerComputeClient.cs b/Engine/Source/Programs/Shared/EpicGames.Horde/Compute/Clients/ServerComputeClient.cs
--- a/Engine/Source/Programs/Shared/EpicGames.Horde/Compute/Clients/ServerComputeClient.cs
+++ b/Engine/Source/Programs/Shared/EpicGames.Horde/Compute/Clients/ServerComputeClient.cs
@@ -140,15 +140,9 @@ namespace EpicGames.Horde.Compute.Clients
 		}
 
 		/// <inheritdoc/>
-		public Task<IComputeLease?> TryAssignWorkerAsync(ClusterId clusterId, Requirements? requirements, CancellationToken cancellationToken)
+		public async Task<IComputeLease?> TryAssignWorkerAsync(ClusterId clusterId, Requirements? requirements, string? requestId, ILogger logger, CancellationToken cancellationToken)
 		{
-			return TryAssignWorkerAsync(clusterId, requirements, null, cancellationToken);
-		}
-		
-		/// <inheritdoc/>
-		public async Task<IComputeLease?> TryAssignWorkerAsync(ClusterId clusterId, Requirements? requirements, string? requestId, CancellationToken cancellationToken)
-		{
-			IAsyncEnumerator<LeaseInfo> source = ConnectAsync(clusterId, requirements, requestId, cancellationToken).GetAsyncEnumerator(cancellationToken);
+			IAsyncEnumerator<LeaseInfo> source = ConnectAsync(clusterId, requirements, requestId, logger, cancellationToken).GetAsyncEnumerator(cancellationToken);
 			if (!await source.MoveNextAsync())
 			{
 				await source.DisposeAsync();
@@ -158,7 +152,7 @@ namespace EpicGames.Horde.Compute.Clients
 		}
 
 		/// <inheritdoc/>
-		async IAsyncEnumerable<LeaseInfo> ConnectAsync(ClusterId clusterId, Requirements? requirements, string? requestId, [EnumeratorCancellation] CancellationToken cancellationToken)
+		async IAsyncEnumerable<LeaseInfo> ConnectAsync(ClusterId clusterId, Requirements? requirements, string? requestId, ILogger workerLogger, [EnumeratorCancellation] CancellationToken cancellationToken)
 		{
 			_logger.LogDebug("Requesting compute resource");
 
@@ -195,7 +189,7 @@ namespace EpicGames.Horde.Compute.Clients
 				}
 			}
 
-			_logger.LogDebug("Connecting to {AgentId} ({Ip}) with nonce {Nonce}...", responseMessage.AgentId, responseMessage.Ip, responseMessage.Nonce);
+			workerLogger.LogDebug("Connecting to {AgentId} ({Ip}) with nonce {Nonce}...", responseMessage.AgentId, responseMessage.Ip, responseMessage.Nonce);
 
 			// Connect to the remote machine
 			using Socket socket = new Socket(SocketType.Stream, ProtocolType.Tcp);
@@ -204,12 +198,12 @@ namespace EpicGames.Horde.Compute.Clients
 			// Send the nonce
 			byte[] nonce = StringUtils.ParseHexString(responseMessage.Nonce);
 			await socket.SendMessageAsync(nonce, SocketFlags.None, cancellationToken);
-			_logger.LogInformation("Connected to {AgentId} ({Ip}) under lease {LeaseId}", responseMessage.AgentId, responseMessage.Ip, responseMessage.LeaseId);
+			workerLogger.LogInformation("Connected to {AgentId} ({Ip}) under lease {LeaseId}", responseMessage.AgentId, responseMessage.Ip, responseMessage.LeaseId);
 
 			// Pass the rest of the call over to the handler
 			byte[] key = StringUtils.ParseHexString(responseMessage.Key);
 
-			await using RemoteComputeSocket computeSocket = new RemoteComputeSocket(new TcpTransport(socket), _logger);
+			await using RemoteComputeSocket computeSocket = new RemoteComputeSocket(new TcpTransport(socket), workerLogger);
 			yield return new LeaseInfo(responseMessage.Properties, responseMessage.AssignedResources, computeSocket);
 		}
 	}
diff --git a/Engine/Source/Programs/Shared/EpicGames.Horde/Compute/IComputeClient.cs b/Engine/Source/Programs/Shared/EpicGames.Horde/Compute/IComputeClient.cs
--- a/Engine/Source/Programs/Shared/EpicGames.Horde/Compute/IComputeClient.cs
+++ b/Engine/Source/Programs/Shared/EpicGames.Horde/Compute/IComputeClient.cs
@@ -3,6 +3,7 @@
 using System;
 using System.Threading;
 using System.Threading.Tasks;
+using Microsoft.Extensions.Logging;
 
 namespace EpicGames.Horde.Compute
 {
@@ -11,21 +12,27 @@ namespace EpicGames.Horde.Compute
 	/// </summary>
 	public interface IComputeClient : IAsyncDisposable
 	{
-		/// <summary>
-		/// Adds a new remote request (deprecated, use method below instead)
-		/// </summary>
-		/// <param name="clusterId">Cluster to execute the request</param>
-		/// <param name="requirements">Requirements for the agent</param>
-		/// <param name="cancellationToken">Cancellation token for the operation</param>
-		public Task<IComputeLease?> TryAssignWorkerAsync(ClusterId clusterId, Requirements? requirements, CancellationToken cancellationToken = default);
-		
 		/// <summary>
 		/// Adds a new remote request
 		/// </summary>
 		/// <param name="clusterId">Cluster to execute the request</param>
 		/// <param name="requirements">Requirements for the agent</param>
 		/// <param name="requestId">Optional ID identifying the request over multiple calls, such as retrying the same request</param>
+		/// <param name="logger">Logger for output from this worker</param>
 		/// <param name="cancellationToken">Cancellation token for the operation</param>
-		public Task<IComputeLease?> TryAssignWorkerAsync(ClusterId clusterId, Requirements? requirements, string? requestId, CancellationToken cancellationToken = default);
+		public Task<IComputeLease?> TryAssignWorkerAsync(ClusterId clusterId, Requirements? requirements, string? requestId, ILogger logger, CancellationToken cancellationToken = default);
+	}
+
+	/// <summary>
+	/// Extension methods for <see cref="IComputeClient"/>
+	/// </summary>
+	public static class ComputeClientExtensions
+	{
+		/// <inheritdoc/>
+		[Obsolete("Prefer the overload taking a requestId parameter")]
+		public static Task<IComputeLease?> TryAssignWorkerAsync(this IComputeClient computeClient, ClusterId clusterId, Requirements? requirements, ILogger logger, CancellationToken cancellationToken)
+		{
+			return computeClient.TryAssignWorkerAsync(clusterId, requirements, null, logger, cancellationToken);
+		}
 	}
 }
