Changed serialization declaration to use new scope system

diff --git a/Engine/Source/Programs/Shared/EpicGames.UHT/Exporters/CodeGen/UhtHeaderCodeGenerator.cs b/Engine/Source/Programs/Shared/EpicGames.UHT/Exporters/CodeGen/UhtHeaderCodeGenerator.cs
--- a/Engine/Source/Programs/Shared/EpicGames.UHT/Exporters/CodeGen/UhtHeaderCodeGenerator.cs
+++ b/Engine/Source/Programs/Shared/EpicGames.UHT/Exporters/CodeGen/UhtHeaderCodeGenerator.cs
@@ -634,34 +634,37 @@ namespace EpicGames.UHT.Exporters.CodeGen
 
 	internal static class UhtHaederCodeGeneratorStringBuilderExtensions
 	{
-		public static StringBuilder AppendMacroName(this StringBuilder builder, string fileId, int lineNumber, string macroSuffix, UhtDefineScope defineScope = UhtDefineScope.None)
+		public static StringBuilder AppendMacroName(this StringBuilder builder, string fileId, int lineNumber, string macroSuffix, UhtDefineScope defineScope = UhtDefineScope.None, bool includeSuffix = true)
 		{
 			builder.Append(fileId).Append('_').Append(lineNumber).Append('_').Append(macroSuffix);
-			if (defineScope.HasAnyFlags(UhtDefineScope.EditorOnlyData))
+			if (includeSuffix)
 			{
-				builder.Append("_EOD");
+				if (defineScope.HasAnyFlags(UhtDefineScope.EditorOnlyData))
+				{
+					builder.Append("_EOD");
+				}
 			}
 			return builder;
 		}
 
-		public static StringBuilder AppendMacroName(this StringBuilder builder, UhtHeaderCodeGenerator generator, int lineNumber, string macroSuffix, UhtDefineScope defineScope = UhtDefineScope.None)
+		public static StringBuilder AppendMacroName(this StringBuilder builder, UhtHeaderCodeGenerator generator, int lineNumber, string macroSuffix, UhtDefineScope defineScope = UhtDefineScope.None, bool includeSuffix = true)
 		{
-			return builder.AppendMacroName(generator.FileId, lineNumber, macroSuffix, defineScope);
+			return builder.AppendMacroName(generator.FileId, lineNumber, macroSuffix, defineScope, includeSuffix);
 		}
 
-		public static StringBuilder AppendMacroName(this StringBuilder builder, UhtHeaderCodeGenerator generator, UhtType type, string macroSuffix, UhtDefineScope defineScope = UhtDefineScope.None)
+		public static StringBuilder AppendMacroName(this StringBuilder builder, UhtHeaderCodeGenerator generator, UhtType type, string macroSuffix, UhtDefineScope defineScope = UhtDefineScope.None, bool includeSuffix = true)
 		{
 			if (type is UhtClass classObj)
 			{
-				return builder.AppendMacroName(generator, classObj.GeneratedBodyLineNumber, macroSuffix, defineScope);
+				return builder.AppendMacroName(generator, classObj.GeneratedBodyLineNumber, macroSuffix, defineScope, includeSuffix);
 			}
 			else if (type is UhtScriptStruct scriptStruct)
 			{
-				return builder.AppendMacroName(generator, scriptStruct.MacroDeclaredLineNumber, macroSuffix, defineScope);
+				return builder.AppendMacroName(generator, scriptStruct.MacroDeclaredLineNumber, macroSuffix, defineScope, includeSuffix);
 			}
 			else if (type is UhtFunction function)
 			{
-				return builder.AppendMacroName(generator, function.MacroLineNumber, macroSuffix, defineScope);
+				return builder.AppendMacroName(generator, function.MacroLineNumber, macroSuffix, defineScope, includeSuffix);
 			}
 			else
 			{
diff --git a/Engine/Source/Programs/Shared/EpicGames.UHT/Exporters/CodeGen/UhtHeaderCodeGeneratorCppFile.cs b/Engine/Source/Programs/Shared/EpicGames.UHT/Exporters/CodeGen/UhtHeaderCodeGeneratorCppFile.cs
--- a/Engine/Source/Programs/Shared/EpicGames.UHT/Exporters/CodeGen/UhtHeaderCodeGeneratorCppFile.cs
+++ b/Engine/Source/Programs/Shared/EpicGames.UHT/Exporters/CodeGen/UhtHeaderCodeGeneratorCppFile.cs
@@ -1388,15 +1388,8 @@ namespace EpicGames.UHT.Exporters.CodeGen
 		{
 			if (!classObj.SerializerArchiveType.HasAnyFlags(serializerType))
 			{
-				if (classObj.EnclosingDefine.Length > 0)
-				{
-					builder.Append("#if ").Append(classObj.EnclosingDefine).Append("\r\n");
-				}
-				builder.Append('\t').Append(macroText).Append('(').Append(classObj.SourceName).Append(")\r\n");
-				if (classObj.EnclosingDefine.Length > 0)
-				{
-					builder.Append("#endif\r\n");
-				}
+				builder.AppendScoped(UhtDefineScopeNames.Standard, classObj.SerializerDefineScope,
+					builder => builder.Append('\t').Append(macroText).Append('(').Append(classObj.SourceName).Append(")\r\n"));
 			}
 			return builder;
 		}
diff --git a/Engine/Source/Programs/Shared/EpicGames.UHT/Exporters/CodeGen/UhtHeaderCodeGeneratorHFile.cs b/Engine/Source/Programs/Shared/EpicGames.UHT/Exporters/CodeGen/UhtHeaderCodeGeneratorHFile.cs
--- a/Engine/Source/Programs/Shared/EpicGames.UHT/Exporters/CodeGen/UhtHeaderCodeGeneratorHFile.cs
+++ b/Engine/Source/Programs/Shared/EpicGames.UHT/Exporters/CodeGen/UhtHeaderCodeGeneratorHFile.cs
@@ -920,28 +920,8 @@ namespace EpicGames.UHT.Exporters.CodeGen
 		{
 			if (!classObj.SerializerArchiveType.HasAnyFlags(type))
 			{
-				if (classObj.EnclosingDefine.Length != 0)
-				{
-					builder.Append("#if ").Append(classObj.EnclosingDefine).Append("\r\n");
-				}
-				AppendSerializerFunction(builder, classObj, api, declare);
-				if (classObj.EnclosingDefine.Length != 0)
-				{
-					builder.Append("#else\r\n");
-					using (UhtMacroCreator macro = new(builder, this, classObj, ArchiveSerializerMacroSuffix))
-					{
-					}
-					builder.Append("#endif\r\n");
-				}
-			}
-			return builder;
-		}
-
-		private StringBuilder AppendSerializerFunction(StringBuilder builder, UhtClass classObj, string api, string declare)
-		{
-			using (UhtMacroCreator macro = new(builder, this, classObj, ArchiveSerializerMacroSuffix))
-			{
-				builder.Append('\t').Append(declare).Append('(').Append(classObj.SourceName).Append(", ").Append(api[0..^1]).Append(") \\\r\n");
+				builder.AppendScopedMacro(UhtDefineScopeNames.Standard, classObj.SerializerDefineScope, this, classObj, ArchiveSerializerMacroSuffix, false,
+					builder => builder.Append('\t').Append(declare).Append('(').Append(classObj.SourceName).Append(", ").Append(api[0..^1]).Append(") \\\r\n"));
 			}
 			return builder;
 		}
diff --git a/Engine/Source/Programs/Shared/EpicGames.UHT/Exporters/CodeGen/UhtMacroCreator.cs b/Engine/Source/Programs/Shared/EpicGames.UHT/Exporters/CodeGen/UhtMacroCreator.cs
--- a/Engine/Source/Programs/Shared/EpicGames.UHT/Exporters/CodeGen/UhtMacroCreator.cs
+++ b/Engine/Source/Programs/Shared/EpicGames.UHT/Exporters/CodeGen/UhtMacroCreator.cs
@@ -12,16 +12,16 @@ namespace EpicGames.UHT.Exporters.CodeGen
 		private readonly StringBuilder _builder;
 		private readonly int _startingLength;
 
-		public UhtMacroCreator(StringBuilder builder, UhtHeaderCodeGenerator generator, int lineNumber, string macroSuffix, UhtDefineScope defineScope = UhtDefineScope.None)
+		public UhtMacroCreator(StringBuilder builder, UhtHeaderCodeGenerator generator, int lineNumber, string macroSuffix, UhtDefineScope defineScope = UhtDefineScope.None, bool includeSuffix = true)
 		{
-			builder.Append("#define ").AppendMacroName(generator, lineNumber, macroSuffix, defineScope).Append(" \\\r\n");
+			builder.Append("#define ").AppendMacroName(generator, lineNumber, macroSuffix, defineScope, includeSuffix).Append(" \\\r\n");
 			_builder = builder;
 			_startingLength = builder.Length;
 		}
 
-		public UhtMacroCreator(StringBuilder builder, UhtHeaderCodeGenerator generator, UhtType type, string macroSuffix, UhtDefineScope defineScope = UhtDefineScope.None)
+		public UhtMacroCreator(StringBuilder builder, UhtHeaderCodeGenerator generator, UhtType type, string macroSuffix, UhtDefineScope defineScope = UhtDefineScope.None, bool includeSuffix = true)
 		{
-			builder.Append("#define ").AppendMacroName(generator, type, macroSuffix, defineScope).Append(" \\\r\n");
+			builder.Append("#define ").AppendMacroName(generator, type, macroSuffix, defineScope, includeSuffix).Append(" \\\r\n");
 			_builder = builder;
 			_startingLength = builder.Length;
 		}
diff --git a/Engine/Source/Programs/Shared/EpicGames.UHT/Exporters/CodeGen/UhtUsedDefineScopes.cs b/Engine/Source/Programs/Shared/EpicGames.UHT/Exporters/CodeGen/UhtUsedDefineScopes.cs
--- a/Engine/Source/Programs/Shared/EpicGames.UHT/Exporters/CodeGen/UhtUsedDefineScopes.cs
+++ b/Engine/Source/Programs/Shared/EpicGames.UHT/Exporters/CodeGen/UhtUsedDefineScopes.cs
@@ -198,52 +198,88 @@ namespace EpicGames.UHT.Exporters.CodeGen
 	/// </summary>
 	internal static class UhtUsedDefineScopesExtensions
 	{
+
 		/// <summary>
-		/// Append multi macros for the given collection of scopes
+		/// Append a macro scoped macro
+		/// </summary>
+		/// <param name="builder">Destination builder</param>
+		/// <param name="defineScopeNames">Which set of scope names will be used</param>
+		/// <param name="defineScope">Specified scope</param>
+		/// <param name="appendAction">Action to invoke to append an instance</param>
+		/// <returns>String builder</returns>
+		public static StringBuilder AppendScoped(this StringBuilder builder, UhtDefineScopeNames defineScopeNames, UhtDefineScope defineScope, Action<StringBuilder> appendAction)
+		{
+			using UhtMacroBlockEmitter blockEmitter = new(builder, defineScopeNames, defineScope);
+			appendAction(builder);
+			return builder;
+		}
+
+		/// <summary>
+		/// Append a macro scoped macro
 		/// </summary>
 		/// <param name="builder">Destination builder</param>
-		/// <param name="instances">Defined scope instances</param>
 		/// <param name="defineScopeNames">Which set of scope names will be used</param>
+		/// <param name="defineScope">Specified scope</param>
 		/// <param name="generator">Header code generator</param>
 		/// <param name="outerType">Output type owning the instances</param>
 		/// <param name="macroSuffix">Macro being created</param>
+		/// <param name="includeSuffix">If true, include such things as _EOD onto the macro name</param>
 		/// <param name="appendAction">Action to invoke to append an instance</param>
 		/// <returns>String builder</returns>
-		public static StringBuilder AppendMultiMacros<T>(this StringBuilder builder, UhtUsedDefineScopes<T> instances, UhtDefineScopeNames defineScopeNames, 
-			UhtHeaderCodeGenerator generator, UhtType outerType, string macroSuffix, Action<StringBuilder, IEnumerable<T>> appendAction) where T : UhtType
+		public static StringBuilder AppendScopedMacro(this StringBuilder builder, UhtDefineScopeNames defineScopeNames, UhtDefineScope defineScope,
+			UhtHeaderCodeGenerator generator, UhtType outerType, string macroSuffix, bool includeSuffix, Action<StringBuilder> appendAction)
 		{
-			foreach (UhtDefineScope defineScope in instances.EnumerateDefinedScopes())
+			using (UhtMacroBlockEmitter blockEmitter = new(builder, defineScopeNames, defineScope))
 			{
-				using (UhtMacroBlockEmitter blockEmitter = new(builder, defineScopeNames, defineScope))
+				using (UhtMacroCreator macro = new(builder, generator, outerType, macroSuffix, defineScope, includeSuffix))
 				{
-					using (UhtMacroCreator macro = new(builder, generator, outerType, macroSuffix, defineScope))
-					{
-						appendAction(builder, instances.Instances.Where(x => x.DefineScope == defineScope));
-					}
-
-					// We can skip writing the macros if there are no properties to declare, as the 'if' and 'else' would be the same
-					if (defineScope != UhtDefineScope.None)
-					{
-						// Trim the extra newlines added after the macro generator
-						if (builder.Length > 4 &&
-							builder[^4] == '\r' &&
-							builder[^3] == '\n' &&
-							builder[^2] == '\r' &&
-							builder[^1] == '\n')
-						{
-							builder.Length -= 4;
-						}
-
-						builder.AppendElsePreprocessor(defineScope, defineScopeNames);
-						using UhtMacroCreator macro = new(builder, generator, outerType, macroSuffix, defineScope); // Empty macro
-					}
+					appendAction(builder);
 				}
 
+				// We can skip writing the macros if there are no properties to declare, as the 'if' and 'else' would be the same
 				if (defineScope != UhtDefineScope.None)
 				{
-					builder.Append("\r\n\r\n");
+					// Trim the extra newlines added after the macro generator
+					if (builder.Length > 4 &&
+						builder[^4] == '\r' &&
+						builder[^3] == '\n' &&
+						builder[^2] == '\r' &&
+						builder[^1] == '\n')
+					{
+						builder.Length -= 4;
+					}
+
+					builder.AppendElsePreprocessor(defineScope, defineScopeNames);
+					using UhtMacroCreator macro = new(builder, generator, outerType, macroSuffix, defineScope, includeSuffix); // Empty macro
 				}
 			}
+
+			if (defineScope != UhtDefineScope.None)
+			{
+				builder.Append("\r\n\r\n");
+			}
+			return builder;
+		}
+
+		/// <summary>
+		/// Append multi macros for the given collection of scopes
+		/// </summary>
+		/// <param name="builder">Destination builder</param>
+		/// <param name="instances">Defined scope instances</param>
+		/// <param name="defineScopeNames">Which set of scope names will be used</param>
+		/// <param name="generator">Header code generator</param>
+		/// <param name="outerType">Output type owning the instances</param>
+		/// <param name="macroSuffix">Macro being created</param>
+		/// <param name="appendAction">Action to invoke to append an instance</param>
+		/// <returns>String builder</returns>
+		public static StringBuilder AppendMultiMacros<T>(this StringBuilder builder, UhtUsedDefineScopes<T> instances, UhtDefineScopeNames defineScopeNames,
+			UhtHeaderCodeGenerator generator, UhtType outerType, string macroSuffix, Action<StringBuilder, IEnumerable<T>> appendAction) where T : UhtType
+		{
+			foreach (UhtDefineScope defineScope in instances.EnumerateDefinedScopes())
+			{
+				AppendScopedMacro(builder, defineScopeNames, defineScope, generator, outerType, macroSuffix, true,
+					builder => appendAction(builder, instances.Instances.Where(x => x.DefineScope == defineScope)));
+			}
 			return builder;
 		}
 
diff --git a/Engine/Source/Programs/Shared/EpicGames.UHT/Parsers/UhtHeaderFileParser.cs b/Engine/Source/Programs/Shared/EpicGames.UHT/Parsers/UhtHeaderFileParser.cs
--- a/Engine/Source/Programs/Shared/EpicGames.UHT/Parsers/UhtHeaderFileParser.cs
+++ b/Engine/Source/Programs/Shared/EpicGames.UHT/Parsers/UhtHeaderFileParser.cs
@@ -1318,7 +1318,7 @@ namespace EpicGames.UHT.Parsers
 				if (declaration.CompilerDirectives == UhtCompilerDirective.None || declaration.CompilerDirectives == UhtCompilerDirective.WithEditorOnlyData)
 				{
 					classObj.SerializerArchiveType |= archiveType;
-					classObj.EnclosingDefine = declaration.CompilerDirectives == UhtCompilerDirective.None ? "" : "WITH_EDITORONLY_DATA";
+					classObj.SerializerDefineScope = declaration.CompilerDirectives == UhtCompilerDirective.None ? UhtDefineScope.None : UhtDefineScope.EditorOnlyData;
 				}
 				else
 				{
diff --git a/Engine/Source/Programs/Shared/EpicGames.UHT/Types/UhtClass.cs b/Engine/Source/Programs/Shared/EpicGames.UHT/Types/UhtClass.cs
--- a/Engine/Source/Programs/Shared/EpicGames.UHT/Types/UhtClass.cs
+++ b/Engine/Source/Programs/Shared/EpicGames.UHT/Types/UhtClass.cs
@@ -298,8 +298,8 @@ namespace EpicGames.UHT.Types
 		/// <summary>
 		/// If needed, the #if block define for the serializer
 		/// </summary>
-		[JsonIgnore]
-		public string EnclosingDefine { get; set; } = String.Empty;
+		[JsonConverter(typeof(JsonStringEnumConverter))]
+		public UhtDefineScope SerializerDefineScope { get; set; } = UhtDefineScope.None;
 
 		/// <summary>
 		/// The class within
