Horde: Use memory mapped files for file storage backend.

diff --git a/Engine/Source/Programs/Shared/EpicGames.Horde.Tests/BackendTests.cs b/Engine/Source/Programs/Shared/EpicGames.Horde.Tests/BackendTests.cs
--- a/Engine/Source/Programs/Shared/EpicGames.Horde.Tests/BackendTests.cs
+++ b/Engine/Source/Programs/Shared/EpicGames.Horde.Tests/BackendTests.cs
@@ -52,6 +52,29 @@ namespace EpicGames.Horde.Tests
 			}
 		}
 
+		[TestMethod]
+		public async Task TestFileBackendFileMapping()
+		{
+			using (TempDir tempDir = new TempDir("Cache"))
+			{
+				using FileStorageBackend backend = new FileStorageBackend(tempDir.Location);
+				string path = await backend.WriteBytesAsync(Encoding.UTF8.GetBytes("hello world"));
+
+				using (IStorageObject handle = backend.Read(path, 0, null))
+				{
+					Assert.AreEqual("hello world", Encoding.UTF8.GetString(handle.Data.Span));
+				}
+				using (IStorageObject handle = backend.Read(path, 0, 5))
+				{
+					Assert.AreEqual("hello", Encoding.UTF8.GetString(handle.Data.Span));
+				}
+				using (IStorageObject handle = backend.Read(path, 4, 3))
+				{
+					Assert.AreEqual("o w", Encoding.UTF8.GetString(handle.Data.Span));
+				}
+			}
+		}
+
 		[TestMethod]
 		public async Task TestCacheBackend()
 		{
diff --git a/Engine/Source/Programs/Shared/EpicGames.Horde/Storage/Backends/FileStorageBackend.cs b/Engine/Source/Programs/Shared/EpicGames.Horde/Storage/Backends/FileStorageBackend.cs
--- a/Engine/Source/Programs/Shared/EpicGames.Horde/Storage/Backends/FileStorageBackend.cs
+++ b/Engine/Source/Programs/Shared/EpicGames.Horde/Storage/Backends/FileStorageBackend.cs
@@ -3,6 +3,7 @@
 using System;
 using System.Collections.Generic;
 using System.IO;
+using System.IO.MemoryMappedFiles;
 using System.Runtime.CompilerServices;
 using System.Threading;
 using System.Threading.Tasks;
@@ -10,11 +11,208 @@ using EpicGames.Core;
 
 namespace EpicGames.Horde.Storage.Backends
 {
+	/// <summary>
+	/// Handle to a raw object read from storgae
+	/// </summary>
+	public interface IStorageObject : IDisposable
+	{
+		/// <summary>
+		/// Data for the item
+		/// </summary>
+		public ReadOnlyMemory<byte> Data { get; }
+	}
+
 	/// <summary>
 	/// Storage backend that utilizes the local filesystem
 	/// </summary>
 	public sealed class FileStorageBackend : IStorageBackend
 	{
+		// Item which has been opened from the cache using a memory mapped file
+		class MappedFile : IDisposable
+		{
+			public string Path { get; }
+			public LinkedListNode<MappedFile> ListNode { get; }
+
+			MemoryMappedFile? _memoryMappedFile;
+			MemoryMappedViewAccessor? _memoryMappedViewAccessor;
+			MemoryMappedView? _memoryMappedView;
+
+			int _refCount = 1;
+			ReadOnlyMemory<byte> _data;
+
+			public int RefCount => _refCount;
+
+			public ulong MappedSize => _memoryMappedViewAccessor?.SafeMemoryMappedViewHandle.ByteLength ?? 0UL;
+
+			public MappedFile(string path, FileInfo fileInfo)
+			{
+				Path = path;
+				ListNode = new LinkedListNode<MappedFile>(this);
+
+				try
+				{
+					_memoryMappedFile = MemoryMappedFile.CreateFromFile(fileInfo.FullName, FileMode.Open, null, 0);
+					_memoryMappedViewAccessor = _memoryMappedFile.CreateViewAccessor(0, 0, MemoryMappedFileAccess.Read);
+					_memoryMappedView = new MemoryMappedView(_memoryMappedViewAccessor);
+
+					_data = _memoryMappedView.GetMemory(0, (int)fileInfo.Length);
+				}
+				catch
+				{
+					Dispose();
+					throw;
+				}
+			}
+
+			public void Dispose()
+			{
+				_data = ReadOnlyMemory<byte>.Empty;
+
+				if (_memoryMappedView != null)
+				{
+					_memoryMappedView.Dispose();
+					_memoryMappedView = null;
+				}
+				if (_memoryMappedViewAccessor != null)
+				{
+					_memoryMappedViewAccessor.Dispose();
+					_memoryMappedViewAccessor = null;
+				}
+				if (_memoryMappedFile != null)
+				{
+					_memoryMappedFile.Dispose();
+					_memoryMappedFile = null;
+				}
+			}
+
+			public ReadOnlyMemory<byte> GetData(int offset, int? length)
+			{
+				if (length == null)
+				{
+					return _data.Slice(offset);
+				}
+				else
+				{
+					return _data.Slice(offset, length.Value);
+				}
+			}
+
+			public void AddRef()
+			{
+				Interlocked.Increment(ref _refCount);
+			}
+
+			public void Release()
+			{
+				if (Interlocked.Decrement(ref _refCount) == 0)
+				{
+					Dispose();
+				}
+			}
+
+			public override string ToString() => Path;
+		}
+
+		// Handle to a file in memory
+		class MappedFileHandle : IStorageObject
+		{
+			MappedFile? _mappedFile;
+			ReadOnlyMemory<byte> _data;
+
+			public ReadOnlyMemory<byte> Data => _data;
+
+			public MappedFileHandle(MappedFile? mappedFile, ReadOnlyMemory<byte> data)
+			{
+				_mappedFile = mappedFile;
+				_mappedFile?.AddRef();
+				_data = data;
+			}
+				
+			public MappedFileHandle Clone()
+			{
+				_mappedFile?.AddRef();
+				return new MappedFileHandle(_mappedFile, _data);
+			}
+
+			public void Dispose()
+			{
+				if (_mappedFile != null)
+				{
+					_mappedFile.Release();
+					_mappedFile = null!;
+				}
+
+				_data = ReadOnlyMemory<byte>.Empty;
+			}
+		}
+
+		// Stream data from a memory mapped item
+		class MappedFileStream : Stream
+		{
+			MappedFile? _mappedFile;
+
+			ReadOnlyMemory<byte> _data;
+			int _offset;
+
+			public override bool CanRead => true;
+			public override bool CanSeek => true;
+			public override bool CanWrite => false;
+
+			public override long Length => _data.Length;
+
+			public override long Position
+			{
+				get => _offset;
+				set => _offset = (int)Math.Clamp(value, 0, _data.Length);
+			}
+
+			public MappedFileStream(MappedFile mappedFile, ReadOnlyMemory<byte> data)
+			{
+				_mappedFile = mappedFile;
+				_mappedFile.AddRef();
+				_data = data;
+			}
+
+			public override void Close()
+			{
+				base.Close();
+
+				if (_mappedFile != null)
+				{
+					_mappedFile.Release();
+					_mappedFile = null;
+				}
+
+				_data = ReadOnlyMemory<byte>.Empty;
+			}
+
+			public override void Flush() { }
+
+			public override int Read(Span<byte> buffer)
+			{
+				int length = (int)Math.Min(Length - _offset, buffer.Length);
+				_data.Span.Slice(_offset, length).CopyTo(buffer);
+				_offset += length;
+				return length;
+			}
+
+			public override int Read(byte[] buffer, int offset, int count) => Read(buffer.AsSpan(offset, count));
+
+			public override long Seek(long offset, SeekOrigin origin)
+			{
+				return _offset = origin switch
+				{
+					SeekOrigin.Begin => (int)Math.Clamp(offset, 0, Length),
+					SeekOrigin.Current => (int)Math.Clamp(_offset + offset, 0, Length),
+					SeekOrigin.End => (int)Math.Clamp(Length + offset, 0, Length),
+					_ => throw new InvalidOperationException()
+				};
+			}
+
+			public override void SetLength(long value) => throw new InvalidOperationException();
+			public override void Write(byte[] buffer, int offset, int count) => throw new InvalidOperationException();
+		}
+
 		/// <summary>
 		/// Base directory for log files
 		/// </summary>
@@ -23,6 +221,15 @@ namespace EpicGames.Horde.Storage.Backends
 		/// <inheritdoc/>
 		public bool SupportsRedirects => false;
 
+		readonly object _lockObject = new object();
+		readonly Dictionary<string, MappedFile> _pathToMappedFile = new Dictionary<string, MappedFile>(StringComparer.Ordinal);
+		readonly LinkedList<MappedFile> _mappedFiles = new LinkedList<MappedFile>();
+
+		long _mappedSize;
+
+		const long MaxMappedSize = 1024L * 1024 * 1024;
+		const int MaxMappedCount = 128;
+
 		/// <summary>
 		/// Constructor
 		/// </summary>
@@ -36,20 +243,86 @@ namespace EpicGames.Horde.Storage.Backends
 		/// <inheritdoc/>
 		public void Dispose()
 		{
+			lock (_lockObject)
+			{
+				UnmapFiles(0, 0);
+			}
 		}
 
 		/// <summary>
 		/// Gets the path for storing a file on disk
 		/// </summary>
 		FileReference GetBlobFile(string path) => FileReference.Combine(_baseDir, $"{path}.blob");
-	
+
+		/// <summary>
+		/// Finds an existing mapped file or adds a new one for the given path
+		/// </summary>
+		MappedFile FindOrAddMappedFile(string path)
+		{
+			MappedFile? mappedFile;
+			if (!_pathToMappedFile.TryGetValue(path, out mappedFile))
+			{
+				FileInfo fileInfo = GetBlobFile(path).ToFileInfo();
+
+				long maxSize = MaxMappedSize - fileInfo.Length;
+				if (_mappedSize > maxSize || _mappedFiles.Count + 1 > MaxMappedCount)
+				{
+					UnmapFiles(maxSize, MaxMappedCount - 1);
+				}
+
+				mappedFile = new MappedFile(path, fileInfo);
+				_pathToMappedFile.Add(path, mappedFile);
+				_mappedFiles.AddFirst(mappedFile.ListNode);
+
+				_mappedSize += (long)mappedFile.MappedSize;
+			}
+			return mappedFile;
+		}
+
+		/// <summary>
+		/// Discard mapped files until only a certain size is mapped in memory
+		/// </summary>
+		/// <param name="maxMappedSize">Maximum mapped size</param>
+		/// <param name="maxMappedCount">Maximum number of mapped files</param>
+		void UnmapFiles(long maxMappedSize, int maxMappedCount)
+		{
+			for (LinkedListNode<MappedFile>? listNode = _mappedFiles.Last; listNode != null && (_mappedSize > maxMappedSize || _mappedFiles.Count > maxMappedCount); )
+			{
+				LinkedListNode<MappedFile>? nextListNode = listNode.Previous;
+				if (listNode.Value.RefCount == 1)
+				{
+					_pathToMappedFile.Remove(listNode.Value.Path);
+					_mappedFiles.Remove(listNode);
+					listNode.Value.Release();
+				}
+				listNode = nextListNode;
+			}
+		}
+
 		/// <inheritdoc/>
 		public Task<Stream> OpenAsync(string path, int offset, int? length, CancellationToken cancellationToken)
 		{
-			FileReference location = GetBlobFile(path);
-			Stream stream = FileReference.Open(location, FileMode.Open, FileAccess.Read, FileShare.Read);
-			stream.Seek(offset, SeekOrigin.Begin);
-			return Task.FromResult(stream);
+			lock (_lockObject)
+			{
+				MappedFile mappedFile = FindOrAddMappedFile(path);
+				return Task.FromResult<Stream>(new MappedFileStream(mappedFile, mappedFile.GetData(offset, length)));
+			}
+		}
+
+		/// <summary>
+		/// Maps a file into memory for reading, and returns a handle to it
+		/// </summary>
+		/// <param name="path">Path to the file</param>
+		/// <param name="offset">Offset of the data to retrieve</param>
+		/// <param name="length">Length of the data</param>
+		/// <returns>Handle to the data. Must be disposed by the caller.</returns>
+		public IStorageObject Read(string path, int offset, int? length)
+		{
+			lock (_lockObject)
+			{
+				MappedFile mappedFile = FindOrAddMappedFile(path);
+				return new MappedFileHandle(mappedFile, mappedFile.GetData(offset, length));
+			}
 		}
 
 		/// <inheritdoc/>
