EpicGames.Core: Remove double brackets around log prefixes.

diff --git a/Engine/Source/Programs/Shared/EpicGames.Core/PrefixLogger.cs b/Engine/Source/Programs/Shared/EpicGames.Core/PrefixLogger.cs
--- a/Engine/Source/Programs/Shared/EpicGames.Core/PrefixLogger.cs
+++ b/Engine/Source/Programs/Shared/EpicGames.Core/PrefixLogger.cs
@@ -41,9 +41,9 @@ namespace EpicGames.Core
 				int idx = copy.FindIndex(x => x.Key.Equals("{OriginalFormat}", StringComparison.OrdinalIgnoreCase));
 				if (idx != -1 && copy[idx].Value is string format)
 				{
-					copy[idx] = new KeyValuePair<string, object>(copy[idx].Key, "[{_tag}] " + format);
+					copy[idx] = new KeyValuePair<string, object>(copy[idx].Key, "{_tag} " + format);
 					copy.Add(new KeyValuePair<string, object>("_tag", _prefix));
-					_inner.Log(logLevel, eventId, copy, exception, (s, e) => $"[{_prefix}] {formatter(state, exception)}");
+					_inner.Log(logLevel, eventId, copy, exception, (s, e) => $"{_prefix} {formatter(state, exception)}");
 					return;
 				}
 			}
