Fix C# static analyzer warning in VS preview.

diff --git a/Engine/Source/Programs/Shared/EpicGames.MsBuild/CsProjBuilder.cs b/Engine/Source/Programs/Shared/EpicGames.MsBuild/CsProjBuilder.cs
--- a/Engine/Source/Programs/Shared/EpicGames.MsBuild/CsProjBuilder.cs
+++ b/Engine/Source/Programs/Shared/EpicGames.MsBuild/CsProjBuilder.cs
@@ -302,7 +302,7 @@ namespace EpicGames.MsBuild
 					// nuget.g.props and nuget.g.targets are generated by Restore, and are frequently re-written;
 					// it should be safe to ignore these files - changes to references from a .csproj file will
 					// show up as that file being out of date.
-					if (ImportPath.IndexOf("nuget.g.", StringComparison.Ordinal) != -1)
+					if (ImportPath.Contains("nuget.g.", StringComparison.Ordinal))
 					{
 						continue;
 					}
