Horde: Add an interface for reading blobs from storage backends that allows using memory mapped files.

diff --git a/Engine/Source/Programs/Horde/Horde.Server/Storage/Backends/AwsStorageBackend.cs b/Engine/Source/Programs/Horde/Horde.Server/Storage/Backends/AwsStorageBackend.cs
--- a/Engine/Source/Programs/Horde/Horde.Server/Storage/Backends/AwsStorageBackend.cs
+++ b/Engine/Source/Programs/Horde/Horde.Server/Storage/Backends/AwsStorageBackend.cs
@@ -17,6 +17,7 @@ using Microsoft.Extensions.Configuration;
 using Microsoft.Extensions.Logging;
 using EpicGames.Horde.Storage;
 using OpenTelemetry.Trace;
+using EpicGames.Core;
 
 namespace Horde.Server.Storage.Backends
 {
@@ -103,6 +104,13 @@ namespace Horde.Server.Storage.Backends
 	/// </summary>
 	public sealed class AwsStorageBackend : IStorageBackend, IDisposable
 	{
+		class StorageObject : IStorageObject
+		{
+			public ReadOnlyMemory<byte> Data { get; }
+			public StorageObject(ReadOnlyMemory<byte> data) => Data = data;
+			public void Dispose() { }
+		}
+
 		/// <summary>
 		/// S3 Client
 		/// </summary>
@@ -277,12 +285,12 @@ namespace Horde.Server.Storage.Backends
 			{
 				range = $"bytes={offset}-{offset + length.Value - 1}";
 			}
-			return ReadAsync(path, new ByteRange(range), cancellationToken);
+			return OpenAsync(path, new ByteRange(range), cancellationToken);
 		}
 
-		async Task<Stream> ReadAsync(string path, ByteRange? byteRange, CancellationToken cancellationToken)
+		async Task<Stream> OpenAsync(string path, ByteRange? byteRange, CancellationToken cancellationToken)
 		{
-			using TelemetrySpan span = OpenTelemetryTracers.Horde.StartActiveSpan($"{nameof(AwsStorageBackend)}.{nameof(ReadAsync)}");
+			using TelemetrySpan span = OpenTelemetryTracers.Horde.StartActiveSpan($"{nameof(AwsStorageBackend)}.{nameof(OpenAsync)}");
 			span.SetAttribute("path", path);
 			
 			string fullPath = GetFullPath(path);
@@ -294,7 +302,7 @@ namespace Horde.Server.Storage.Backends
 			{
 				semaLock = await _semaphore.UseWaitAsync(cancellationToken);
 
-				semaphoreSpan = OpenTelemetryTracers.Horde.StartActiveSpan($"{nameof(AwsStorageBackend)}.{nameof(ReadAsync)}.Semaphore");
+				semaphoreSpan = OpenTelemetryTracers.Horde.StartActiveSpan($"{nameof(AwsStorageBackend)}.{nameof(OpenAsync)}.Semaphore");
 				semaphoreSpan.SetAttribute("path", path);
 
 				GetObjectRequest newGetRequest = new GetObjectRequest();
@@ -339,6 +347,13 @@ namespace Horde.Server.Storage.Backends
 			}
 		}
 
+		/// <inheritdoc/>
+		public async Task<IStorageObject> ReadAsync(string path, int offset, int? length, CancellationToken cancellationToken)
+		{
+			using Stream stream = await OpenAsync(path, offset, length, cancellationToken);
+			return new StorageObject(await stream.ReadAllBytesAsync(cancellationToken));
+		}
+
 		/// <inheritdoc/>
 		public ValueTask<Uri?> TryGetReadRedirectAsync(string path, CancellationToken cancellationToken = default) => new ValueTask<Uri?>(GetPresignedUrl(path, HttpVerb.GET));
 
diff --git a/Engine/Source/Programs/Horde/Horde.Server/Storage/StorageBackendProvider.cs b/Engine/Source/Programs/Horde/Horde.Server/Storage/StorageBackendProvider.cs
--- a/Engine/Source/Programs/Horde/Horde.Server/Storage/StorageBackendProvider.cs
+++ b/Engine/Source/Programs/Horde/Horde.Server/Storage/StorageBackendProvider.cs
@@ -81,11 +81,13 @@ namespace Horde.Server.Storage
 			/// <inheritdoc/>
 			public Task<Stream> OpenAsync(string path, int offset, int? length, CancellationToken cancellationToken = default)
 			{
-				if (_backend == null)
-				{
-					throw new InvalidOperationException("Backend has already been disposed");
-				}
-				return _backend.OpenAsync(path, offset, length, cancellationToken);
+				return _backend?.OpenAsync(path, offset, length, cancellationToken) ?? throw new InvalidOperationException("Backend has already been disposed");
+			}
+
+			/// <inheritdoc/>
+			public Task<IStorageObject> ReadAsync(string path, int offset, int? length, CancellationToken cancellationToken = default)
+			{
+				return _backend?.ReadAsync(path, offset, length, cancellationToken) ?? throw new InvalidOperationException("Backend has already been disposed");
 			}
 
 			/// <inheritdoc/>
diff --git a/Engine/Source/Programs/Horde/Horde.Server/Storage/StorageService.cs b/Engine/Source/Programs/Horde/Horde.Server/Storage/StorageService.cs
--- a/Engine/Source/Programs/Horde/Horde.Server/Storage/StorageService.cs
+++ b/Engine/Source/Programs/Horde/Horde.Server/Storage/StorageService.cs
@@ -113,6 +113,13 @@ namespace Horde.Server.Storage
 				return await _inner.OpenAsync(fullPath, offset, length, cancellationToken);
 			}
 
+			/// <inheritdoc/>
+			public async Task<IStorageObject> ReadAsync(string path, int offset, int? length, CancellationToken cancellationToken = default)
+			{
+				string fullPath = $"{_prefix}{path}";
+				return await _inner.ReadAsync(fullPath, offset, length, cancellationToken);
+			}
+
 			/// <inheritdoc/>
 			public async Task<string> WriteAsync(Stream stream, string? prefix = null, CancellationToken cancellationToken = default)
 			{
diff --git a/Engine/Source/Programs/Horde/Horde.Server/Storage/TypedStorageBackend.cs b/Engine/Source/Programs/Horde/Horde.Server/Storage/TypedStorageBackend.cs
--- a/Engine/Source/Programs/Horde/Horde.Server/Storage/TypedStorageBackend.cs
+++ b/Engine/Source/Programs/Horde/Horde.Server/Storage/TypedStorageBackend.cs
@@ -45,6 +45,9 @@ namespace Horde.Server.Storage
 			/// <inheritdoc/>
 			public Task<Stream> OpenAsync(string path, int offset, int? length, CancellationToken cancellationToken) => _inner.OpenAsync(path, offset, length, cancellationToken);
 
+			/// <inheritdoc/>
+			public Task<IStorageObject> ReadAsync(string path, int offset, int? length, CancellationToken cancellationToken) => _inner.ReadAsync(path, offset, length, cancellationToken);
+
 			/// <inheritdoc/>
 			public Task<string> WriteAsync(Stream stream, string? prefix, CancellationToken cancellationToken) => _inner.WriteAsync(stream, prefix, cancellationToken);
 
diff --git a/Engine/Source/Programs/Horde/Horde/Program.cs b/Engine/Source/Programs/Horde/Horde/Program.cs
--- a/Engine/Source/Programs/Horde/Horde/Program.cs
+++ b/Engine/Source/Programs/Horde/Horde/Program.cs
@@ -64,7 +64,7 @@ namespace Horde
 		{
 			CmdConfig cmdConfig = serviceProvider.GetRequiredService<IOptions<CmdConfig>>().Value;
 			DirectoryReference cacheDir = DirectoryReference.Combine(GetDataDir(), String.IsNullOrEmpty(cmdConfig.Cache.CacheDir)? "Cache" : cmdConfig.Cache.CacheDir);
-			return new StorageBackendCache(cacheDir, cmdConfig.Cache.CacheSize * 1024 * 1024);
+			return new StorageBackendCache(cacheDir, cmdConfig.Cache.CacheSize * 1024 * 1024, serviceProvider.GetRequiredService<ILogger<StorageBackendCache>>());
 		}
 
 		static DirectoryReference GetAppDir()
diff --git a/Engine/Source/Programs/Shared/EpicGames.Horde.Tests/BackendTests.cs b/Engine/Source/Programs/Shared/EpicGames.Horde.Tests/BackendTests.cs
--- a/Engine/Source/Programs/Shared/EpicGames.Horde.Tests/BackendTests.cs
+++ b/Engine/Source/Programs/Shared/EpicGames.Horde.Tests/BackendTests.cs
@@ -8,6 +8,7 @@ using System.Threading.Tasks;
 using EpicGames.Core;
 using EpicGames.Horde.Storage;
 using EpicGames.Horde.Storage.Backends;
+using Microsoft.Extensions.Logging.Abstractions;
 using Microsoft.VisualStudio.TestTools.UnitTesting;
 
 namespace EpicGames.Horde.Tests
@@ -80,10 +81,11 @@ namespace EpicGames.Horde.Tests
 		{
 			using (TempDir tempDir = new TempDir("Cache"))
 			{
-				StorageBackendCache cache = new StorageBackendCache(tempDir.Location, 12);
+				using StorageBackendCache cache = new StorageBackendCache(tempDir.Location, 12, NullLogger.Instance);
 
 				using MemoryStorageBackend memoryBackend = new MemoryStorageBackend();
-				using IStorageBackend cacheBackend = StorageBackendCache.Wrap(memoryBackend, cache);
+				using IStorageBackend cacheBackend = cache.CreateWrapper("", memoryBackend);
+
 				await TestBackendAsync(cacheBackend);
 
 				Assert.AreEqual(1, cache.Items.Count());
diff --git a/Engine/Source/Programs/Shared/EpicGames.Horde/Compute/AgentStorageClient.cs b/Engine/Source/Programs/Shared/EpicGames.Horde/Compute/AgentStorageClient.cs
--- a/Engine/Source/Programs/Shared/EpicGames.Horde/Compute/AgentStorageClient.cs
+++ b/Engine/Source/Programs/Shared/EpicGames.Horde/Compute/AgentStorageClient.cs
@@ -56,6 +56,13 @@ namespace EpicGames.Horde.Compute
 
 	class AgentStorageBackend : IStorageBackend
 	{
+		class StorageObject : IStorageObject
+		{
+			public ReadOnlyMemory<byte> Data { get; }
+			public StorageObject(ReadOnlyMemory<byte> data) => Data = data;
+			public void Dispose() { }
+		}
+
 		readonly AgentMessageChannel _channel;
 		readonly SemaphoreSlim _semaphore;
 
@@ -101,6 +108,16 @@ namespace EpicGames.Horde.Compute
 			}
 		}
 
+		/// <inheritdoc/>
+		public async Task<IStorageObject> ReadAsync(string path, int offset, int? length, CancellationToken cancellationToken = default)
+		{
+			using (Stream stream = await OpenAsync(path, offset, length, cancellationToken))
+			{
+				byte[] data = await stream.ReadAllBytesAsync(cancellationToken);
+				return new StorageObject(data);
+			}
+		}
+
 		/// <inheritdoc/>
 		public Task<string> WriteAsync(Stream stream, string? prefix = null, CancellationToken cancellationToken = default) => throw new NotSupportedException();
 
diff --git a/Engine/Source/Programs/Shared/EpicGames.Horde/Storage/Backends/FileStorageBackend.cs b/Engine/Source/Programs/Shared/EpicGames.Horde/Storage/Backends/FileStorageBackend.cs
--- a/Engine/Source/Programs/Shared/EpicGames.Horde/Storage/Backends/FileStorageBackend.cs
+++ b/Engine/Source/Programs/Shared/EpicGames.Horde/Storage/Backends/FileStorageBackend.cs
@@ -11,17 +11,6 @@ using EpicGames.Core;
 
 namespace EpicGames.Horde.Storage.Backends
 {
-	/// <summary>
-	/// Handle to a raw object read from storgae
-	/// </summary>
-	public interface IStorageObject : IDisposable
-	{
-		/// <summary>
-		/// Data for the item
-		/// </summary>
-		public ReadOnlyMemory<byte> Data { get; }
-	}
-
 	/// <summary>
 	/// Storage backend that utilizes the local filesystem
 	/// </summary>
@@ -33,12 +22,15 @@ namespace EpicGames.Horde.Storage.Backends
 			public string Path { get; }
 			public LinkedListNode<MappedFile> ListNode { get; }
 
+			readonly FileInfo _fileInfo;
+
 			MemoryMappedFile? _memoryMappedFile;
 			MemoryMappedViewAccessor? _memoryMappedViewAccessor;
 			MemoryMappedView? _memoryMappedView;
 
 			int _refCount = 1;
 			ReadOnlyMemory<byte> _data;
+			bool _deleteOnDispose;
 
 			public int RefCount => _refCount;
 
@@ -49,6 +41,8 @@ namespace EpicGames.Horde.Storage.Backends
 				Path = path;
 				ListNode = new LinkedListNode<MappedFile>(this);
 
+				_fileInfo = fileInfo;
+
 				try
 				{
 					_memoryMappedFile = MemoryMappedFile.CreateFromFile(fileInfo.FullName, FileMode.Open, null, 0);
@@ -83,6 +77,15 @@ namespace EpicGames.Horde.Storage.Backends
 					_memoryMappedFile.Dispose();
 					_memoryMappedFile = null;
 				}
+
+				if (_deleteOnDispose)
+				{
+					try
+					{
+						_fileInfo.Delete();
+					}
+					catch { }
+				}
 			}
 
 			public ReadOnlyMemory<byte> GetData(int offset, int? length)
@@ -110,6 +113,11 @@ namespace EpicGames.Horde.Storage.Backends
 				}
 			}
 
+			public void DeleteOnDispose()
+			{
+				_deleteOnDispose = true;
+			}
+
 			public override string ToString() => Path;
 		}
 
@@ -127,7 +135,7 @@ namespace EpicGames.Horde.Storage.Backends
 				_mappedFile?.AddRef();
 				_data = data;
 			}
-				
+
 			public MappedFileHandle Clone()
 			{
 				_mappedFile?.AddRef();
@@ -146,73 +154,6 @@ namespace EpicGames.Horde.Storage.Backends
 			}
 		}
 
-		// Stream data from a memory mapped item
-		class MappedFileStream : Stream
-		{
-			MappedFile? _mappedFile;
-
-			ReadOnlyMemory<byte> _data;
-			int _offset;
-
-			public override bool CanRead => true;
-			public override bool CanSeek => true;
-			public override bool CanWrite => false;
-
-			public override long Length => _data.Length;
-
-			public override long Position
-			{
-				get => _offset;
-				set => _offset = (int)Math.Clamp(value, 0, _data.Length);
-			}
-
-			public MappedFileStream(MappedFile mappedFile, ReadOnlyMemory<byte> data)
-			{
-				_mappedFile = mappedFile;
-				_mappedFile.AddRef();
-				_data = data;
-			}
-
-			public override void Close()
-			{
-				base.Close();
-
-				if (_mappedFile != null)
-				{
-					_mappedFile.Release();
-					_mappedFile = null;
-				}
-
-				_data = ReadOnlyMemory<byte>.Empty;
-			}
-
-			public override void Flush() { }
-
-			public override int Read(Span<byte> buffer)
-			{
-				int length = (int)Math.Min(Length - _offset, buffer.Length);
-				_data.Span.Slice(_offset, length).CopyTo(buffer);
-				_offset += length;
-				return length;
-			}
-
-			public override int Read(byte[] buffer, int offset, int count) => Read(buffer.AsSpan(offset, count));
-
-			public override long Seek(long offset, SeekOrigin origin)
-			{
-				return _offset = origin switch
-				{
-					SeekOrigin.Begin => (int)Math.Clamp(offset, 0, Length),
-					SeekOrigin.Current => (int)Math.Clamp(_offset + offset, 0, Length),
-					SeekOrigin.End => (int)Math.Clamp(Length + offset, 0, Length),
-					_ => throw new InvalidOperationException()
-				};
-			}
-
-			public override void SetLength(long value) => throw new InvalidOperationException();
-			public override void Write(byte[] buffer, int offset, int count) => throw new InvalidOperationException();
-		}
-
 		/// <summary>
 		/// Base directory for log files
 		/// </summary>
@@ -286,7 +227,7 @@ namespace EpicGames.Horde.Storage.Backends
 		/// <param name="maxMappedCount">Maximum number of mapped files</param>
 		void UnmapFiles(long maxMappedSize, int maxMappedCount)
 		{
-			for (LinkedListNode<MappedFile>? listNode = _mappedFiles.Last; listNode != null && (_mappedSize > maxMappedSize || _mappedFiles.Count > maxMappedCount); )
+			for (LinkedListNode<MappedFile>? listNode = _mappedFiles.Last; listNode != null && (_mappedSize > maxMappedSize || _mappedFiles.Count > maxMappedCount);)
 			{
 				LinkedListNode<MappedFile>? nextListNode = listNode.Previous;
 				if (listNode.Value.RefCount == 1)
@@ -302,11 +243,8 @@ namespace EpicGames.Horde.Storage.Backends
 		/// <inheritdoc/>
 		public Task<Stream> OpenAsync(string path, int offset, int? length, CancellationToken cancellationToken)
 		{
-			lock (_lockObject)
-			{
-				MappedFile mappedFile = FindOrAddMappedFile(path);
-				return Task.FromResult<Stream>(new MappedFileStream(mappedFile, mappedFile.GetData(offset, length)));
-			}
+			IStorageObject storageObject = Read(path, offset, length);
+			return Task.FromResult<Stream>(storageObject.CreateStream());
 		}
 
 		/// <summary>
@@ -325,6 +263,12 @@ namespace EpicGames.Horde.Storage.Backends
 			}
 		}
 
+		/// <inheritdoc/>
+		public Task<IStorageObject> ReadAsync(string path, int offset, int? length, CancellationToken cancellationToken = default)
+		{
+			return Task.FromResult(Read(path, offset, length));
+		}
+
 		/// <inheritdoc/>
 		public async Task<string> WriteAsync(Stream stream, string? prefix = null, CancellationToken cancellationToken = default)
 		{
@@ -339,7 +283,7 @@ namespace EpicGames.Horde.Storage.Backends
 			FileReference finalLocation = GetBlobFile(path);
 			DirectoryReference.CreateDirectory(finalLocation.Directory);
 			FileReference tempLocation = new FileReference($"{finalLocation}.tmp");
-		
+
 			using (Stream outputStream = FileReference.Open(tempLocation, FileMode.Create, FileAccess.Write, FileShare.Read))
 			{
 				await stream.CopyToAsync(outputStream, cancellationToken);
@@ -370,11 +314,36 @@ namespace EpicGames.Horde.Storage.Backends
 			return Task.FromResult(FileReference.Exists(location));
 		}
 
+		/// <summary>
+		/// Delete a file from the store
+		/// </summary>
+		/// <param name="path"></param>
+		public void Delete(string path)
+		{
+			lock(_lockObject)
+			{
+				MappedFile? mappedFile;
+				if (_pathToMappedFile.TryGetValue(path, out mappedFile))
+				{
+					mappedFile.DeleteOnDispose();
+
+					_pathToMappedFile.Remove(path);
+					_mappedFiles.Remove(mappedFile.ListNode);
+
+					mappedFile.Release();
+				}
+				else
+				{
+					FileReference location = GetBlobFile(path);
+					FileReference.Delete(location);
+				}
+			}
+		}
+
 		/// <inheritdoc/>
 		public Task DeleteAsync(string path, CancellationToken cancellationToken)
 		{
-			FileReference location = GetBlobFile(path);
-			FileReference.Delete(location);
+			Delete(path);
 			return Task.CompletedTask;
 		}
 
@@ -424,3 +393,4 @@ namespace EpicGames.Horde.Storage.Backends
 		public ValueTask<(string, Uri)?> TryGetWriteRedirectAsync(string? prefix = null, CancellationToken cancellationToken = default) => default;
 	}
 }
+
diff --git a/Engine/Source/Programs/Shared/EpicGames.Horde/Storage/Backends/HttpStorageBackend.cs b/Engine/Source/Programs/Shared/EpicGames.Horde/Storage/Backends/HttpStorageBackend.cs
--- a/Engine/Source/Programs/Shared/EpicGames.Horde/Storage/Backends/HttpStorageBackend.cs
+++ b/Engine/Source/Programs/Shared/EpicGames.Horde/Storage/Backends/HttpStorageBackend.cs
@@ -8,6 +8,7 @@ using System.Net.Http.Headers;
 using System.Net.Http.Json;
 using System.Threading;
 using System.Threading.Tasks;
+using EpicGames.Core;
 using Microsoft.Extensions.Logging;
 
 namespace EpicGames.Horde.Storage.Backends
@@ -17,6 +18,13 @@ namespace EpicGames.Horde.Storage.Backends
 	/// </summary>
 	public sealed class HttpStorageBackend : IStorageBackend
 	{
+		class StorageObject : IStorageObject
+		{
+			public ReadOnlyMemory<byte> Data { get; }
+			public StorageObject(ReadOnlyMemory<byte> data) => Data = data;
+			public void Dispose() { }
+		}
+
 		class WriteBlobResponse
 		{
 			public string Blob { get; set; } = String.Empty;
@@ -86,6 +94,16 @@ namespace EpicGames.Horde.Storage.Backends
 			}
 		}
 
+		/// <inheritdoc/>
+		public async Task<IStorageObject> ReadAsync(string path, int offset, int? length, CancellationToken cancellationToken = default)
+		{
+			using (Stream stream = await OpenAsync(path, offset, length, cancellationToken))
+			{
+				byte[] data = await stream.ReadAllBytesAsync(cancellationToken);
+				return new StorageObject(data);
+			}
+		}
+
 		/// <inheritdoc/>
 		public async Task<string> WriteAsync(Stream stream, string? prefix = null, CancellationToken cancellationToken = default)
 		{
diff --git a/Engine/Source/Programs/Shared/EpicGames.Horde/Storage/Backends/MemoryStorageBackend.cs b/Engine/Source/Programs/Shared/EpicGames.Horde/Storage/Backends/MemoryStorageBackend.cs
--- a/Engine/Source/Programs/Shared/EpicGames.Horde/Storage/Backends/MemoryStorageBackend.cs
+++ b/Engine/Source/Programs/Shared/EpicGames.Horde/Storage/Backends/MemoryStorageBackend.cs
@@ -16,6 +16,13 @@ namespace EpicGames.Horde.Storage.Backends
 	/// </summary>
 	public sealed class MemoryStorageBackend : IStorageBackend
 	{
+		class StorageObject : IStorageObject
+		{
+			public ReadOnlyMemory<byte> Data { get; }
+			public StorageObject(ReadOnlyMemory<byte> data) => Data = data;
+			public void Dispose() { }
+		}
+
 		/// <summary>
 		/// Data storage
 		/// </summary>
@@ -36,13 +43,24 @@ namespace EpicGames.Horde.Storage.Backends
 
 		/// <inheritdoc/>
 		public Task<Stream> OpenAsync(string path, int offset, int? length, CancellationToken cancellationToken)
+		{
+			return Task.FromResult<Stream>(new ReadOnlyMemoryStream(GetData(path, offset, length)));
+		}
+
+		/// <inheritdoc/>
+		public Task<IStorageObject> ReadAsync(string path, int offset, int? length, CancellationToken cancellationToken)
+		{
+			return Task.FromResult<IStorageObject>(new StorageObject(GetData(path, offset, length)));
+		}
+
+		ReadOnlyMemory<byte> GetData(string path, int offset, int? length)
 		{
 			ReadOnlyMemory<byte> data = _pathToData[path].AsMemory(offset);
 			if (length != null && length.Value < data.Length)
 			{
 				data = data.Slice(0, length.Value);
 			}
-			return Task.FromResult<Stream>(new ReadOnlyMemoryStream(data));
+			return data;
 		}
 
 		/// <inheritdoc/>
diff --git a/Engine/Source/Programs/Shared/EpicGames.Horde/Storage/IStorageBackend.cs b/Engine/Source/Programs/Shared/EpicGames.Horde/Storage/IStorageBackend.cs
--- a/Engine/Source/Programs/Shared/EpicGames.Horde/Storage/IStorageBackend.cs
+++ b/Engine/Source/Programs/Shared/EpicGames.Horde/Storage/IStorageBackend.cs
@@ -10,6 +10,17 @@ using EpicGames.Core;
 
 namespace EpicGames.Horde.Storage
 {
+	/// <summary>
+	/// Handle to a raw object read from storage.
+	/// </summary>
+	public interface IStorageObject : IDisposable
+	{
+		/// <summary>
+		/// Data for the item. Consumers of this interface must not cache this value beyond the lifetime of the object.
+		/// </summary>
+		public ReadOnlyMemory<byte> Data { get; }
+	}
+
 	/// <summary>
 	/// Interface for a low-level storage backend.
 	/// </summary>
@@ -30,6 +41,16 @@ namespace EpicGames.Horde.Storage
 		/// <returns></returns>
 		Task<Stream> OpenAsync(string path, int offset, int? length, CancellationToken cancellationToken = default);
 
+		/// <summary>
+		/// Reads an object into memory and returns a handle to it.
+		/// </summary>
+		/// <param name="path">Path to the file</param>
+		/// <param name="offset">Offset of the data to retrieve</param>
+		/// <param name="length">Length of the data</param>
+		/// <param name="cancellationToken">Cancellation token for the operation</param>
+		/// <returns>Handle to the data. Must be disposed by the caller.</returns>
+		Task<IStorageObject> ReadAsync(string path, int offset, int? length, CancellationToken cancellationToken = default);
+
 		/// <summary>
 		/// Writes a stream to the storage backend. If the stream throws an exception during read, the write will be aborted.
 		/// </summary>
@@ -127,6 +148,109 @@ namespace EpicGames.Horde.Storage
 	/// </summary>
 	public static class StorageBackendExtensions
 	{
+		// Wraps an IStorageObject in a stream
+		class StorageObjectStream : Stream
+		{
+			readonly IStorageObject _storageObject;
+			int _offset;
+
+			public override bool CanRead => true;
+			public override bool CanSeek => true;
+			public override bool CanWrite => false;
+
+			public override long Length => _storageObject.Data.Length;
+
+			public override long Position
+			{
+				get => _offset;
+				set => _offset = (int)Math.Clamp(value, 0, _storageObject.Data.Length);
+			}
+
+			public StorageObjectStream(IStorageObject storageObject) => _storageObject = storageObject;
+
+			protected override void Dispose(bool disposing)
+			{
+				base.Dispose(disposing);
+
+				if (disposing)
+				{
+					_storageObject.Dispose();
+				}
+			}
+
+			public override void Flush() { }
+
+			public override int Read(Span<byte> buffer)
+			{
+				int length = (int)Math.Min(Length - _offset, buffer.Length);
+				_storageObject.Data.Span.Slice(_offset, length).CopyTo(buffer);
+				_offset += length;
+				return length;
+			}
+
+			public override int Read(byte[] buffer, int offset, int count) => Read(buffer.AsSpan(offset, count));
+
+			public override long Seek(long offset, SeekOrigin origin)
+			{
+				return _offset = origin switch
+				{
+					SeekOrigin.Begin => (int)Math.Clamp(offset, 0, Length),
+					SeekOrigin.Current => (int)Math.Clamp(_offset + offset, 0, Length),
+					SeekOrigin.End => (int)Math.Clamp(Length + offset, 0, Length),
+					_ => throw new InvalidOperationException()
+				};
+			}
+
+			public override void SetLength(long value) => throw new InvalidOperationException();
+			public override void Write(byte[] buffer, int offset, int count) => throw new InvalidOperationException();
+		}
+
+		/// <summary>
+		/// Create a stream to wrap a storage object. The object will be disposed when the stream is closed.
+		/// </summary>
+		/// <param name="storageObject">Storage object to wrap</param>
+		/// <returns>Stream for the given storage object</returns>
+		public static Stream CreateStream(this IStorageObject storageObject) => new StorageObjectStream(storageObject);
+
+		// Creates a slice of a storage object
+		class StorageObjectSlice : IStorageObject
+		{
+			readonly IStorageObject _inner;
+			readonly int _offset;
+			readonly int _length;
+
+			public ReadOnlyMemory<byte> Data => _inner.Data.Slice(_offset, _length);
+
+			public StorageObjectSlice(IStorageObject inner, int offset, int? length)
+			{
+				_inner = inner;
+				_offset = offset;
+				_length = length ?? (inner.Data.Length - offset);
+			}
+
+			public void Dispose() => _inner.Dispose();
+		}
+
+		/// <summary>
+		/// Create a stream to wrap a storage object. The object will be disposed when the stream is closed.
+		/// </summary>
+		/// <param name="storageObject">Storage object to wrap</param>
+		/// <param name="offset">Offset of the slice</param>
+		/// <param name="length">Length to take for the slice</param>
+		/// <returns>Stream for the given storage object</returns>
+		public static IStorageObject CreateSlice(this IStorageObject storageObject, int offset, int? length)
+		{
+			int actualLength = length ?? (storageObject.Data.Length - offset);
+			if (offset == 0 && actualLength == storageObject.Data.Length)
+			{
+				return storageObject;
+			}
+			else
+			{
+				return new StorageObjectSlice(storageObject, offset, length);
+			}
+		}
+
 		/// <summary>
 		/// Attempts to open a read stream for the given path.
 		/// </summary>
@@ -136,6 +260,15 @@ namespace EpicGames.Horde.Storage
 		/// <returns>Stream for the object</returns>
 		public static Task<Stream> OpenAsync(this IStorageBackend storageBackend, string path, CancellationToken cancellationToken = default) => storageBackend.OpenAsync(path, 0, null, cancellationToken);
 
+		/// <summary>
+		/// Attempts to open a read stream for the given path.
+		/// </summary>
+		/// <param name="storageBackend">Backend to read from</param>
+		/// <param name="path">Object name within the store</param>
+		/// <param name="cancellationToken">Cancellation token for the operation</param>
+		/// <returns>Stream for the object</returns>
+		public static Task<IStorageObject> ReadAsync(this IStorageBackend storageBackend, string path, CancellationToken cancellationToken = default) => storageBackend.ReadAsync(path, 0, null, cancellationToken);
+
 		/// <summary>
 		/// Reads an object as an array of bytes
 		/// </summary>
@@ -145,14 +278,8 @@ namespace EpicGames.Horde.Storage
 		/// <returns>Contents of the object</returns>
 		public static async Task<byte[]> ReadBytesAsync(this IStorageBackend storageBackend, string path, CancellationToken cancellationToken = default)
 		{
-			using (Stream inputStream = await storageBackend.OpenAsync(path, cancellationToken))
-			{
-				using (MemoryStream outputStream = new MemoryStream())
-				{
-					await inputStream.CopyToAsync(outputStream, cancellationToken);
-					return outputStream.ToArray();
-				}
-			}
+			using IStorageObject storageObject = await storageBackend.ReadAsync(path, cancellationToken);
+			return storageObject.Data.ToArray();
 		}
 
 		/// <summary>
diff --git a/Engine/Source/Programs/Shared/EpicGames.Horde/Storage/StorageBackendCache.cs b/Engine/Source/Programs/Shared/EpicGames.Horde/Storage/StorageBackendCache.cs
--- a/Engine/Source/Programs/Shared/EpicGames.Horde/Storage/StorageBackendCache.cs
+++ b/Engine/Source/Programs/Shared/EpicGames.Horde/Storage/StorageBackendCache.cs
@@ -1,6 +1,8 @@
 // Copyright Epic Games, Inc. All Rights Reserved.
 
 using EpicGames.Core;
+using EpicGames.Horde.Storage.Backends;
+using Microsoft.Extensions.Logging;
 using System;
 using System.Collections.Generic;
 using System.IO;
@@ -13,17 +15,19 @@ namespace EpicGames.Horde.Storage
 	/// <summary>
 	/// Implementation of a local disk cache which can be shared by multiple backends
 	/// </summary>
-	public sealed class StorageBackendCache
+	public sealed class StorageBackendCache : IDisposable
 	{
 		class Item
 		{
 			public string Key { get; }
+			public string Path { get; }
 			public long Length { get; }
 			public LinkedListNode<Item> ListNode { get; }
 
-			public Item(string key, long length)
+			public Item(string key, string path, long length)
 			{
 				Key = key;
+				Path = path;
 				Length = length;
 				ListNode = new LinkedListNode<Item>(this);
 			}
@@ -67,12 +71,14 @@ namespace EpicGames.Horde.Storage
 
 			public async Task<Stream> OpenAsync(string path, int offset, int? length, CancellationToken cancellationToken = default)
 			{
-				Stream stream = await _cacheStorage.ReadAsync($"{_keyPrefix}{path}", ctx => _inner.OpenAsync(path, ctx), cancellationToken);
-				if (offset != 0)
-				{
-					stream.Seek(offset, SeekOrigin.Begin);
-				}
-				return stream;
+				IStorageObject storageObject = await ReadAsync(path, offset, length, cancellationToken);
+				return storageObject.CreateStream();
+			}
+
+			public async Task<IStorageObject> ReadAsync(string path, int offset, int? length, CancellationToken cancellationToken = default)
+			{
+				IStorageObject storageObject = await _cacheStorage.ReadAsync($"{_keyPrefix}{path}", ctx => _inner.OpenAsync(path, ctx), cancellationToken);
+				return storageObject.CreateSlice(offset, length);
 			}
 
 			public Task<string> WriteAsync(Stream stream, string? prefix = null, CancellationToken cancellationToken = default) => _inner.WriteAsync(stream, prefix, cancellationToken);
@@ -90,8 +96,9 @@ namespace EpicGames.Horde.Storage
 
 		object LockObject => _items;
 
-		readonly DirectoryReference _cacheDir;
+		readonly FileStorageBackend _backend;
 		readonly long _maxSize;
+		readonly ILogger _logger;
 
 		readonly LinkedList<Item> _items = new LinkedList<Item>();
 		readonly Dictionary<string, Item> _pathToItem = new Dictionary<string, Item>(StringComparer.Ordinal);
@@ -104,31 +111,46 @@ namespace EpicGames.Horde.Storage
 		/// <summary>
 		/// Constructor
 		/// </summary>
-		public StorageBackendCache(DirectoryReference cacheDir, long maxSize)
+		public StorageBackendCache(DirectoryReference cacheDir, long maxSize, ILogger logger)
 		{
-			_cacheDir = cacheDir;
+			_backend = new FileStorageBackend(cacheDir);
 			_maxSize = maxSize;
+			_logger = logger;
+		}
+
+		/// <inheritdoc/>
+		public void Dispose()
+		{
+			foreach (Item item in _items)
+			{
+				_backend.Delete(item.Path);
+			}
+			_backend.Dispose();
+		}
+
+		/// <summary>
+		/// Wraps a storage backend in another backend that routes requests through the cache
+		/// </summary>
+		/// <param name="keyPrefix">Prefix for items in this cache</param>
+		/// <param name="backend">Backend to wrap</param>
+		public IStorageBackend CreateWrapper(string keyPrefix, IStorageBackend backend)
+		{
+			return new BackendWrapper(keyPrefix, this, backend);
 		}
 
 		/// <summary>
 		/// Wraps a storage backend in another backend that routes requests through the cache
 		/// </summary>
+		/// <param name="keyPrefix">Prefix for items in this cache</param>
 		/// <param name="backend">Backend to wrap</param>
 		/// <param name="cache">The cache instance. May be null.</param>
-		public static IStorageBackend Wrap(IStorageBackend backend, StorageBackendCache? cache)
+		public static IStorageBackend CreateWrapper(string keyPrefix, IStorageBackend backend, StorageBackendCache? cache)
 		{
-			if (cache == null)
-			{
-				return backend;
-			}
-			else
-			{
-				return new BackendWrapper("", cache, backend);
-			}
+			return cache?.CreateWrapper(keyPrefix, backend) ?? backend;
 		}
 
 		/// <inheritdoc/>
-		public async Task<Stream> ReadAsync(string key, Func<CancellationToken, Task<Stream>> createStreamAsync, CancellationToken cancellationToken = default)
+		public async Task<IStorageObject> ReadAsync(string key, Func<CancellationToken, Task<Stream>> createStreamAsync, CancellationToken cancellationToken = default)
 		{
 			for (; ; )
 			{
@@ -138,7 +160,9 @@ namespace EpicGames.Horde.Storage
 					Item? item;
 					if (_pathToItem.TryGetValue(key, out item))
 					{
-						return OpenItem(item);
+						_items.Remove(item.ListNode);
+						_items.AddFirst(item.ListNode);
+						return _backend.Read(item.Path, 0, null);
 					}
 
 					if (!_pathToPendingItem.TryGetValue(key, out pendingItem))
@@ -161,15 +185,6 @@ namespace EpicGames.Horde.Storage
 			}
 		}
 
-		Stream OpenItem(Item item)
-		{
-			_items.Remove(item.ListNode);
-			_items.AddFirst(item.ListNode);
-
-			FileReference file = FileReference.Combine(_cacheDir, item.Key);
-			return FileReference.Open(file, FileMode.Open, FileAccess.Read, FileShare.Read | FileShare.Delete);
-		}
-
 		async Task ReadIntoCacheAsync(string key, Func<CancellationToken, Task<Stream>> createStreamAsync, CancellationToken cancellationToken)
 		{
 			using Stream stream = await createStreamAsync(cancellationToken);
@@ -179,30 +194,33 @@ namespace EpicGames.Horde.Storage
 			{
 				_size += totalLength;
 
-				while (_size > _maxSize && _items.Count > 0)
+				LinkedListNode<Item>? lastNode = _items.Last;
+				while (_size > _maxSize && lastNode != null)
 				{
-					LinkedListNode<Item> lastItem = _items.Last!;
-					_pathToItem.Remove(lastItem.Value.Key);
-					_items.RemoveLast();
-
-					_size -= lastItem.Value.Length;
+					LinkedListNode<Item> node = lastNode;
+					lastNode = lastNode.Previous;
 
-					FileReference file = FileReference.Combine(_cacheDir, lastItem.Value.Key);
-					FileReference.Delete(file);
-				}
-			}
+					Item item = node.Value;
+					try
+					{
+						_backend.Delete(item.Path);
 
-			FileReference cacheFile = FileReference.Combine(_cacheDir, key);
-			DirectoryReference.CreateDirectory(cacheFile.Directory);
+						_pathToItem.Remove(item.Key);
+						_items.Remove(node);
 
-			using (FileStream outputStream = FileReference.Open(cacheFile, FileMode.Create, FileAccess.Write, FileShare.None))
-			{
-				await stream.CopyToAsync(outputStream, cancellationToken);
+						_size -= item.Length;
+					}
+					catch (Exception ex)
+					{
+						_logger.LogDebug(ex, "Unable to delete cache item {Path}: {Message}", item.Path, ex.Message);
+					}
+				}
 			}
 
+			string path = await _backend.WriteAsync(stream, cancellationToken: cancellationToken);
 			lock (LockObject)
 			{
-				Item item = new Item(key, totalLength);
+				Item item = new Item(key, path, totalLength);
 				_items.AddFirst(item.ListNode);
 				_pathToItem.Add(key, item);
 
