Horde: Enable more style warnings at build time.

diff --git a/Engine/Source/Programs/Horde/BuildHorde.xml b/Engine/Source/Programs/Horde/BuildHorde.xml
--- a/Engine/Source/Programs/Horde/BuildHorde.xml
+++ b/Engine/Source/Programs/Horde/BuildHorde.xml
@@ -17,7 +17,7 @@
 		<Property Name="StagingDir" Value="$(RootDir)/Engine/Saved/Horde.Server"/>
 
 		<Node Name="Analyze HordeServer">
-			<DotNet Arguments="build &quot;$(RootDir)/Engine/Source/Programs/Horde/Horde.Server/Horde.Server.csproj&quot; /p:Configuration=Analyze /p:TreatWarningsAsErrors=True $(VersionArguments)"/>
+			<DotNet Arguments="build &quot;$(RootDir)/Engine/Source/Programs/Horde/Horde.Server/Horde.Server.csproj&quot; /p:Configuration=Analyze $(VersionArguments)"/>
 		</Node>
 
 		<Node Name="Check HordeServer Licenses" Requires="Analyze HordeServer">
diff --git a/Engine/Source/Programs/Horde/Horde.Agent.TrayApp/Program.cs b/Engine/Source/Programs/Horde/Horde.Agent.TrayApp/Program.cs
--- a/Engine/Source/Programs/Horde/Horde.Agent.TrayApp/Program.cs
+++ b/Engine/Source/Programs/Horde/Horde.Agent.TrayApp/Program.cs
@@ -319,7 +319,7 @@ namespace Horde.Agent.TrayApp
 			public uint dwLowDateTime;
 			public uint dwHighDateTime;
 
-			public ulong Total => dwLowDateTime | ((ulong)dwHighDateTime << 32);
+			public readonly ulong Total => dwLowDateTime | ((ulong)dwHighDateTime << 32);
 		};
 
 		[StructLayout(LayoutKind.Sequential)]
diff --git a/Engine/Source/Programs/Horde/Horde.Server.Tests/JobCollectionTests.cs b/Engine/Source/Programs/Horde/Horde.Server.Tests/JobCollectionTests.cs
--- a/Engine/Source/Programs/Horde/Horde.Server.Tests/JobCollectionTests.cs
+++ b/Engine/Source/Programs/Horde/Horde.Server.Tests/JobCollectionTests.cs
@@ -33,10 +33,7 @@ namespace Horde.Server.Tests
 		static NewNode AddNode(NewGroup group, string name, string[]? inputDependencies, Action<NewNode>? action = null)
 		{
 			NewNode node = new NewNode(name, inputDependencies: inputDependencies?.ToList(), orderDependencies: inputDependencies?.ToList());
-			if (action != null)
-			{
-				action.Invoke(node);
-			}
+			action?.Invoke(node);
 			group.Nodes.Add(node);
 			return node;
 		}
diff --git a/Engine/Source/Programs/Horde/Horde.Server.Tests/JobTaskSourceTests.cs b/Engine/Source/Programs/Horde/Horde.Server.Tests/JobTaskSourceTests.cs
--- a/Engine/Source/Programs/Horde/Horde.Server.Tests/JobTaskSourceTests.cs
+++ b/Engine/Source/Programs/Horde/Horde.Server.Tests/JobTaskSourceTests.cs
@@ -25,19 +25,16 @@ namespace Horde.Server.Tests
 
 		static NewGroup AddGroup(List<NewGroup> groups)
 		{
-			NewGroup @group = new NewGroup("Win64", new List<NewNode>());
-			groups.Add(@group);
-			return @group;
+			NewGroup group = new NewGroup("Win64", new List<NewNode>());
+			groups.Add(group);
+			return group;
 		}
 
-		static NewNode AddNode(NewGroup @group, string name, string[]? inputDependencies, Action<NewNode>? action = null)
+		static NewNode AddNode(NewGroup group, string name, string[]? inputDependencies, Action<NewNode>? action = null)
 		{
 			NewNode node = new NewNode(name, inputDependencies: inputDependencies?.ToList(), orderDependencies: inputDependencies?.ToList());
-			if (action != null)
-			{
-				action.Invoke(node);
-			}
-			@group.Nodes.Add(node);
+			action?.Invoke(node);
+			group.Nodes.Add(node);
 			return node;
 		}
 
diff --git a/Engine/Source/Programs/Shared/.editorconfig b/Engine/Source/Programs/Shared/.editorconfig
--- a/Engine/Source/Programs/Shared/.editorconfig
+++ b/Engine/Source/Programs/Shared/.editorconfig
@@ -115,9 +115,8 @@ dotnet_naming_symbols.all_members.applicable_kinds = *
 
 dotnet_naming_style.pascal_case_style.capitalization = pascal_case
 
-# IDE0073: File header
-dotnet_diagnostic.IDE0073.severity = warning
-file_header_template = Copyright Epic Games, Inc. All Rights Reserved.
+# IDE0005: Using directive is unnecessary
+dotnet_diagnostic.IDE0005.severity = warning
 
 # IDE0035: Remove unreachable code
 dotnet_diagnostic.IDE0035.severity = warning
@@ -152,15 +151,19 @@ dotnet_diagnostic.IDE0059.severity = warning
 # IDE0060: Remove unused parameter
 dotnet_diagnostic.IDE0060.severity = warning
 
+# IDE0073: File header
+dotnet_diagnostic.IDE0073.severity = warning
+file_header_template = Copyright Epic Games, Inc. All Rights Reserved.
+
+# IDE1006: Enforce naming conventions
+dotnet_diagnostic.IDE1006.severity = warning
+
 # CA1012: Abstract types should not have public constructors
 dotnet_diagnostic.CA1012.severity = warning
 
 # CA1822: Make member static
 dotnet_diagnostic.CA1822.severity = warning
 
-# IDE0005: Using directive is unnecessary
-dotnet_diagnostic.IDE0005.severity = warning
-
 # dotnet_style_allow_multiple_blank_lines_experimental
 dotnet_diagnostic.IDE2000.severity = warning
 
diff --git a/Engine/Source/Programs/Shared/EpicGames.Core/ByteArrayBuilder.cs b/Engine/Source/Programs/Shared/EpicGames.Core/ByteArrayBuilder.cs
--- a/Engine/Source/Programs/Shared/EpicGames.Core/ByteArrayBuilder.cs
+++ b/Engine/Source/Programs/Shared/EpicGames.Core/ByteArrayBuilder.cs
@@ -15,7 +15,7 @@ namespace EpicGames.Core
 		{
 			public readonly int RunningIndex;
 			public readonly byte[] Data;
-			public int Length;
+			public int Length { get; set; }
 
 			public Chunk(int runningIndex, int size)
 			{
diff --git a/Engine/Source/Programs/Shared/EpicGames.Core/CsProjBuildHook.cs b/Engine/Source/Programs/Shared/EpicGames.Core/CsProjBuildHook.cs
--- a/Engine/Source/Programs/Shared/EpicGames.Core/CsProjBuildHook.cs
+++ b/Engine/Source/Programs/Shared/EpicGames.Core/CsProjBuildHook.cs
@@ -17,29 +17,28 @@ namespace EpicGames.Core
 	/// </summary>
 	public interface CsProjBuildHook
 	{
-
 		/// <summary>
 		/// Test the cache for a given file to get the last write time of the given file 
 		/// </summary>
-		/// <param name="BasePath">Base path of the file.</param>
-		/// <param name="RelativeFilePath">Relative path of the file</param>
+		/// <param name="basePath">Base path of the file.</param>
+		/// <param name="relativeFilePath">Relative path of the file</param>
 		/// <returns>Last write time of the file.</returns>
-		DateTime GetLastWriteTime(DirectoryReference BasePath, string RelativeFilePath);
+		DateTime GetLastWriteTime(DirectoryReference basePath, string relativeFilePath);
 
 		/// <summary>
 		/// Test the cache for a given file to get the last write time of the given file 
 		/// </summary>
-		/// <param name="BasePath">Base path of the file.</param>
-		/// <param name="RelativeFilePath">Relative path of the file</param>
+		/// <param name="basePath">Base path of the file.</param>
+		/// <param name="relativeFilePath">Relative path of the file</param>
 		/// <returns>Last write time of the file.</returns>
-		DateTime GetLastWriteTime(string BasePath, string RelativeFilePath);
+		DateTime GetLastWriteTime(string basePath, string relativeFilePath);
 
 		/// <summary>
 		/// Return the build record directory for the given base path (i.e. engine dir or project dir)
 		/// </summary>
-		/// <param name="BasePath">The base path for the directory</param>
+		/// <param name="basePath">The base path for the directory</param>
 		/// <returns>Directory for the build records</returns>
-		DirectoryReference GetBuildRecordDirectory(DirectoryReference BasePath);
+		DirectoryReference GetBuildRecordDirectory(DirectoryReference basePath);
 
 		/// <summary>
 		/// Validate the given build records for the project
@@ -51,9 +50,9 @@ namespace EpicGames.Core
 		/// <summary>
 		/// Test to see if the given file spec has any wild cards
 		/// </summary>
-		/// <param name="FileSpec">File spec to test</param>
+		/// <param name="fileSpec">File spec to test</param>
 		/// <returns>True if wildcards are present</returns>
-		bool HasWildcards(string FileSpec);
+		bool HasWildcards(string fileSpec);
 
 		/// <summary>
 		/// Unreal engine directory
diff --git a/Engine/Source/Programs/Shared/EpicGames.Core/CsProjectInfo.cs b/Engine/Source/Programs/Shared/EpicGames.Core/CsProjectInfo.cs
--- a/Engine/Source/Programs/Shared/EpicGames.Core/CsProjectInfo.cs
+++ b/Engine/Source/Programs/Shared/EpicGames.Core/CsProjectInfo.cs
@@ -750,9 +750,9 @@ namespace EpicGames.Core
 				}
 				return bResult;
 			}
-			catch (CsProjectParseException Ex)
+			catch (CsProjectParseException ex)
 			{
-				throw new CsProjectParseException($"{Ex.Message} while parsing {element} in project file {projectInfo.ProjectPath}", Ex);
+				throw new CsProjectParseException($"{ex.Message} while parsing {element} in project file {projectInfo.ProjectPath}", ex);
 			}
 		}
 
diff --git a/Engine/Source/Programs/Shared/EpicGames.Core/Log.cs b/Engine/Source/Programs/Shared/EpicGames.Core/Log.cs
--- a/Engine/Source/Programs/Shared/EpicGames.Core/Log.cs
+++ b/Engine/Source/Programs/Shared/EpicGames.Core/Log.cs
@@ -101,30 +101,30 @@ namespace EpicGames.Core
 		/// <summary>
 		/// Singleton instance of the default output logger
 		/// </summary>
-		private static readonly DefaultLogger DefaultLogger = new DefaultLogger();
+		private static readonly DefaultLogger s_defaultLogger = new DefaultLogger();
 
 		/// <summary>
 		/// Logger instance which parses events and forwards them to the main logger.
 		/// </summary>
-		private static readonly LegacyEventLogger LegacyLogger = new LegacyEventLogger(DefaultLogger);
+		private static readonly LegacyEventLogger s_legacyLogger = new LegacyEventLogger(s_defaultLogger);
 
 		/// <summary>
 		/// Accessor for the global event parser from legacy events
 		/// </summary>
-		public static LogEventParser EventParser => LegacyLogger.Parser;
+		public static LogEventParser EventParser => s_legacyLogger.Parser;
 
 		/// <summary>
 		/// Logger instance
 		/// </summary>
-		public static ILogger Logger => LegacyLogger;
+		public static ILogger Logger => s_legacyLogger;
 
 		/// <summary>
 		/// When true, verbose logging is enabled.
 		/// </summary>
 		public static LogEventType OutputLevel
 		{
-			get => (LogEventType)DefaultLogger.OutputLevel;
-			set => DefaultLogger.OutputLevel = (LogLevel)value;
+			get => (LogEventType)s_defaultLogger.OutputLevel;
+			set => s_defaultLogger.OutputLevel = (LogLevel)value;
 		}
 
 		/// <summary>
@@ -132,8 +132,8 @@ namespace EpicGames.Core
 		/// </summary>
 		public static bool IncludeTimestamps
 		{
-			get => DefaultLogger.IncludeTimestamps;
-			set => DefaultLogger.IncludeTimestamps = value;
+			get => s_defaultLogger.IncludeTimestamps;
+			set => s_defaultLogger.IncludeTimestamps = value;
 		}
 
 		/// <summary>
@@ -151,8 +151,8 @@ namespace EpicGames.Core
 		/// </summary>
 		public static bool ColorConsoleOutput
 		{
-			get => DefaultLogger.ColorConsoleOutput;
-			set => DefaultLogger.ColorConsoleOutput = value;
+			get => s_defaultLogger.ColorConsoleOutput;
+			set => s_defaultLogger.ColorConsoleOutput = value;
 		}
 
 		/// <summary>
@@ -160,8 +160,8 @@ namespace EpicGames.Core
 		/// </summary>
 		public static bool IncludeStartingTimestamp
 		{
-			get => DefaultLogger.IncludeStartingTimestamp;
-			set => DefaultLogger.IncludeStartingTimestamp = value;
+			get => s_defaultLogger.IncludeStartingTimestamp;
+			set => s_defaultLogger.IncludeStartingTimestamp = value;
 		}
 
 		/// <summary>
@@ -180,7 +180,7 @@ namespace EpicGames.Core
 		/// <summary>
 		/// Path to the log file being written to. May be null.
 		/// </summary>
-		public static FileReference? OutputFile => DefaultLogger?.OutputFile;
+		public static FileReference? OutputFile => s_defaultLogger?.OutputFile;
 
 		/// <summary>
 		/// A collection of strings that have been already written once
@@ -193,7 +193,7 @@ namespace EpicGames.Core
 		/// <param name="logger"></param>
 		public static void SetInnerLogger(ILogger logger)
 		{
-			LegacyLogger.SetInnerLogger(logger);
+			s_legacyLogger.SetInnerLogger(logger);
 		}
 
 		/// <summary>
@@ -202,7 +202,7 @@ namespace EpicGames.Core
 		/// <returns></returns>
 		public static async Task FlushAsync()
 		{
-			await DefaultLogger.FlushAsync();
+			await s_defaultLogger.FlushAsync();
 		}
 
 		/// <summary>
@@ -283,7 +283,7 @@ namespace EpicGames.Core
 		/// <returns>The created trace listener</returns>
 		public static void AddFileWriterWithoutBackup(string name, FileReference outputFile)
 		{
-			TextWriterTraceListener firstTextWriter = DefaultLogger.AddFileWriter(name, outputFile);
+			TextWriterTraceListener firstTextWriter = s_defaultLogger.AddFileWriter(name, outputFile);
 			
 			// find the StartupTraceListener in the listeners that was added early on
 			IEnumerable<StartupTraceListener> startupListeners = Trace.Listeners.OfType<StartupTraceListener>();
@@ -301,7 +301,7 @@ namespace EpicGames.Core
 		/// <param name="traceListener">The <see cref="TraceListener"/> to add.</param>
 		public static void AddTraceListener(TraceListener traceListener)
 		{
-			DefaultLogger.AddTraceListener(traceListener);
+			s_defaultLogger.AddTraceListener(traceListener);
 		}
 
 		/// <summary>
@@ -310,7 +310,7 @@ namespace EpicGames.Core
 		/// <param name="traceListener">The <see cref="TraceListener"/> to remove.</param>
 		public static void RemoveTraceListener(TraceListener traceListener)
 		{
-			DefaultLogger.RemoveTraceListener(traceListener);
+			s_defaultLogger.RemoveTraceListener(traceListener);
 		}
 
 		/// <summary>
@@ -760,7 +760,7 @@ namespace EpicGames.Core
 		[Conditional("TRACE")]
 		public static void PushStatus(string message)
 		{
-			DefaultLogger.PushStatus(message);
+			s_defaultLogger.PushStatus(message);
 		}
 
 		/// <summary>
@@ -770,7 +770,7 @@ namespace EpicGames.Core
 		[Conditional("TRACE")]
 		public static void UpdateStatus(string message)
 		{
-			DefaultLogger.UpdateStatus(message);
+			s_defaultLogger.UpdateStatus(message);
 		}
 
 		/// <summary>
@@ -779,7 +779,7 @@ namespace EpicGames.Core
 		[Conditional("TRACE")]
 		public static void PopStatus()
 		{
-			DefaultLogger.PopStatus();
+			s_defaultLogger.PopStatus();
 		}
 	}
 
diff --git a/Engine/Source/Programs/Shared/EpicGames.Core/UHTTypes.cs b/Engine/Source/Programs/Shared/EpicGames.Core/UHTTypes.cs
--- a/Engine/Source/Programs/Shared/EpicGames.Core/UHTTypes.cs
+++ b/Engine/Source/Programs/Shared/EpicGames.Core/UHTTypes.cs
@@ -100,37 +100,37 @@ namespace EpicGames.Core
 				return Name;
 			}
 
-			public bool TryGetDefine(string Name, out string? Value)
+			public bool TryGetDefine(string name, out string? value)
 			{
-				Value = null;
-				int Length = Name.Length;
-				foreach (string Define in PublicDefines)
+				value = null;
+				int length = name.Length;
+				foreach (string define in PublicDefines)
 				{
-					if (!Define.StartsWith(Name, System.StringComparison.Ordinal))
+					if (!define.StartsWith(name, System.StringComparison.Ordinal))
 					{
 						continue;
 					}
-					if (Define.Length > Length)
+					if (define.Length > length)
 					{
-						if (Define[Length] != '=')
+						if (define[length] != '=')
 						{
 							continue;
 						}
-						Value = Define.Substring(Length + 1, Define.Length - Length - 1);
+						value = define.Substring(length + 1, define.Length - length - 1);
 					}
 					return true;
 				}
 				return false;
 			}
 
-			public bool TryGetDefine(string Name, out int Value)
+			public bool TryGetDefine(string name, out int value)
 			{
-				string? String;
-				if (TryGetDefine(Name, out String))
+				string? stringValue;
+				if (TryGetDefine(name, out stringValue))
 				{
-					return int.TryParse(String, out Value);
+					return int.TryParse(stringValue, out value);
 				}
-				Value = 0;
+				value = 0;
 				return false;
 			}
 		}
diff --git a/Engine/Source/Programs/Shared/EpicGames.Core/ZipArchiveExtensions.cs b/Engine/Source/Programs/Shared/EpicGames.Core/ZipArchiveExtensions.cs
--- a/Engine/Source/Programs/Shared/EpicGames.Core/ZipArchiveExtensions.cs
+++ b/Engine/Source/Programs/Shared/EpicGames.Core/ZipArchiveExtensions.cs
@@ -53,11 +53,11 @@ namespace EpicGames.Core
 		{
 			// This seems to be consistently at least 35% faster than ZipFileExtensions.ExtractToFile(Entry, TargetFileName, Overwrite) when the
 			// archive contains many small files.
-			using (FileStream OutputStream = new FileStream(targetFileName, overwrite ? FileMode.Create : FileMode.CreateNew, FileAccess.Write))
+			using (FileStream outputStream = new FileStream(targetFileName, overwrite ? FileMode.Create : FileMode.CreateNew, FileAccess.Write))
 			{
-				using (Stream InStream = entry.Open())
+				using (Stream inputStream = entry.Open())
 				{
-					InStream.CopyTo(OutputStream);
+					inputStream.CopyTo(outputStream);
 				}
 			}
 
diff --git a/Engine/Source/Programs/Shared/EpicGames.Horde.Tests/BundleStoreTests.cs b/Engine/Source/Programs/Shared/EpicGames.Horde.Tests/BundleStoreTests.cs
--- a/Engine/Source/Programs/Shared/EpicGames.Horde.Tests/BundleStoreTests.cs
+++ b/Engine/Source/Programs/Shared/EpicGames.Horde.Tests/BundleStoreTests.cs
@@ -308,9 +308,9 @@ namespace EpicGames.Horde.Tests
 		{
 			using MemoryStorageClient store = new MemoryStorageClient();
 
-			const int length = 4096;
+			const int Length = 4096;
 
-			byte[] chunk = new byte[length];
+			byte[] chunk = new byte[Length];
 			new Random(0).NextBytes(chunk);
 
 			// Generate a tree
@@ -348,14 +348,14 @@ namespace EpicGames.Horde.Tests
 			using IMemoryCache cache = new MemoryCache(new MemoryCacheOptions());
 			using MemoryStorageClient store = new MemoryStorageClient();
 
-			const int length = 1024;
-			const int copies = 4096;
+			const int Length = 1024;
+			const int Copies = 4096;
 
-			byte[] chunk = new byte[length];
+			byte[] chunk = new byte[Length];
 			new Random(0).NextBytes(chunk);
 
-			byte[] data = new byte[chunk.Length * copies];
-			for (int idx = 0; idx < copies; idx++)
+			byte[] data = new byte[chunk.Length * Copies];
+			for (int idx = 0; idx < Copies; idx++)
 			{
 				chunk.CopyTo(data.AsSpan(idx * chunk.Length));
 			}
diff --git a/Engine/Source/Programs/Shared/EpicGames.Horde/Compute/AgentMessage.cs b/Engine/Source/Programs/Shared/EpicGames.Horde/Compute/AgentMessage.cs
--- a/Engine/Source/Programs/Shared/EpicGames.Horde/Compute/AgentMessage.cs
+++ b/Engine/Source/Programs/Shared/EpicGames.Horde/Compute/AgentMessage.cs
@@ -206,9 +206,9 @@ namespace EpicGames.Horde.Compute
 	/// <summary>
 	/// Message requesting that the message loop be forked
 	/// </summary>
-	/// <param name="channelId">New channel to communicate on</param>
-	/// <param name="bufferSize">Size of the buffer</param>
-	public record struct ForkMessage(int channelId, int bufferSize);
+	/// <param name="ChannelId">New channel to communicate on</param>
+	/// <param name="BufferSize">Size of the buffer</param>
+	public record struct ForkMessage(int ChannelId, int BufferSize);
 
 	/// <summary>
 	/// Extract files from a bundle to a path in the remote sandbox
diff --git a/Engine/Source/Programs/Shared/EpicGames.Horde/Compute/AgentMessageHandler.cs b/Engine/Source/Programs/Shared/EpicGames.Horde/Compute/AgentMessageHandler.cs
--- a/Engine/Source/Programs/Shared/EpicGames.Horde/Compute/AgentMessageHandler.cs
+++ b/Engine/Source/Programs/Shared/EpicGames.Horde/Compute/AgentMessageHandler.cs
@@ -81,7 +81,7 @@ namespace EpicGames.Horde.Compute
 						case AgentMessageType.Fork:
 							{
 								ForkMessage fork = message.ParseForkMessage();
-								childTasks.Add(Task.Run(() => RunAsync(socket, fork.channelId, fork.bufferSize, cancellationToken), cancellationToken));
+								childTasks.Add(Task.Run(() => RunAsync(socket, fork.ChannelId, fork.BufferSize, cancellationToken), cancellationToken));
 							}
 							break;
 						case AgentMessageType.WriteFiles:
diff --git a/Engine/Source/Programs/Shared/EpicGames.Horde/Compute/ComputeSocket.cs b/Engine/Source/Programs/Shared/EpicGames.Horde/Compute/ComputeSocket.cs
--- a/Engine/Source/Programs/Shared/EpicGames.Horde/Compute/ComputeSocket.cs
+++ b/Engine/Source/Programs/Shared/EpicGames.Horde/Compute/ComputeSocket.cs
@@ -313,7 +313,7 @@ namespace EpicGames.Horde.Compute
 		class RecvBuffer : IDisposable
 		{
 			public ComputeBufferWriter? _writer;
-			public readonly SemaphoreSlim _semaphore = new SemaphoreSlim(1);
+			public readonly SemaphoreSlim Semaphore = new SemaphoreSlim(1);
 			public int _refCount = 1;
 
 			public RecvBuffer(ComputeBufferWriter writer) => _writer = writer;
@@ -331,22 +331,22 @@ namespace EpicGames.Horde.Compute
 			public void Dispose()
 			{
 				_writer?.Dispose();
-				_semaphore.Dispose();
+				Semaphore.Dispose();
 			}
 		}
 
 		class SendBuffer : IDisposable
 		{
 			public ComputeBufferReader? _reader;
-			public readonly SemaphoreSlim _semaphore = new SemaphoreSlim(1);
+			public readonly SemaphoreSlim Semaphore = new SemaphoreSlim(1);
 			[SuppressMessage("Usage", "CA2213:Disposable fields should be disposed")]
-			public readonly BackgroundTask _task;
+			public readonly BackgroundTask Task;
 			int _refCount = 1;
 
 			public SendBuffer(ComputeBufferReader reader, Func<SendBuffer, CancellationToken, Task> func)
 			{
 				_reader = reader;
-				_task = BackgroundTask.StartNew(ctx => func(this, ctx));
+				Task = BackgroundTask.StartNew(ctx => func(this, ctx));
 			}
 
 			public void AddRef() => Interlocked.Increment(ref _refCount);
@@ -362,7 +362,7 @@ namespace EpicGames.Horde.Compute
 			public void Dispose()
 			{
 				_reader?.Dispose();
-				_semaphore.Dispose();
+				Semaphore.Dispose();
 			}
 		}
 
@@ -540,7 +540,7 @@ namespace EpicGames.Horde.Compute
 			bool result = false;
 			try
 			{
-				await recvBuffer._semaphore.WaitAsync(cancellationToken);
+				await recvBuffer.Semaphore.WaitAsync(cancellationToken);
 				try
 				{
 					ComputeBufferWriter? writer = recvBuffer._writer;
@@ -566,7 +566,7 @@ namespace EpicGames.Horde.Compute
 				}
 				finally
 				{
-					recvBuffer._semaphore.Release();
+					recvBuffer.Semaphore.Release();
 				}
 			}
 			finally
@@ -664,7 +664,7 @@ namespace EpicGames.Horde.Compute
 			}
 
 			// Release the writer
-			await recvBuffer._semaphore.WaitAsync(cancellationToken);
+			await recvBuffer.Semaphore.WaitAsync(cancellationToken);
 			try
 			{
 				recvBuffer._writer?.Dispose();
@@ -672,7 +672,7 @@ namespace EpicGames.Horde.Compute
 			}
 			finally
 			{
-				recvBuffer._semaphore.Release();
+				recvBuffer.Semaphore.Release();
 				recvBuffer.Release(); // Ref added above
 			}
 
@@ -748,7 +748,7 @@ namespace EpicGames.Horde.Compute
 			}
 
 			// Release the reader
-			await sendBuffer._semaphore.WaitAsync(cancellationToken);
+			await sendBuffer.Semaphore.WaitAsync(cancellationToken);
 			try
 			{
 				sendBuffer._reader?.Dispose();
@@ -756,12 +756,12 @@ namespace EpicGames.Horde.Compute
 			}
 			finally
 			{
-				sendBuffer._semaphore.Release();
+				sendBuffer.Semaphore.Release();
 				sendBuffer.Release(); // Added above
 			}
 
 			// Wait for the send task to complete
-			await sendBuffer._task.DisposeAsync();
+			await sendBuffer.Task.DisposeAsync();
 
 			// Remove the buffer from the dictionary
 			lock (_lockObject)
diff --git a/Engine/Source/Programs/Shared/EpicGames.Horde/Storage/Bundles/BundleReader.cs b/Engine/Source/Programs/Shared/EpicGames.Horde/Storage/Bundles/BundleReader.cs
--- a/Engine/Source/Programs/Shared/EpicGames.Horde/Storage/Bundles/BundleReader.cs
+++ b/Engine/Source/Programs/Shared/EpicGames.Horde/Storage/Bundles/BundleReader.cs
@@ -57,7 +57,7 @@ namespace EpicGames.Horde.Storage.Bundles
 		class QueuedBundle
 		{
 			public BundleLocator Locator { get; }
-			public TaskCompletionSource<BundleInfo> BundleInfo = new TaskCompletionSource<BundleInfo>();
+			public TaskCompletionSource<BundleInfo> BundleInfo { get; } = new TaskCompletionSource<BundleInfo>();
 			public List<QueuedPacket> QueuedPackets { get; } = new List<QueuedPacket>(); // Sorted by index
 			public int InfoRefCount { get; set; } // Ref count for reading the header
 			public PendingRead? PendingRead { get; set; }
diff --git a/Engine/Source/Programs/Shared/EpicGames.Horde/Storage/Bundles/BundleWriter.cs b/Engine/Source/Programs/Shared/EpicGames.Horde/Storage/Bundles/BundleWriter.cs
--- a/Engine/Source/Programs/Shared/EpicGames.Horde/Storage/Bundles/BundleWriter.cs
+++ b/Engine/Source/Programs/Shared/EpicGames.Horde/Storage/Bundles/BundleWriter.cs
@@ -294,7 +294,7 @@ namespace EpicGames.Horde.Storage.Bundles
 			long _compressedLength;
 
 			// Map of keys to nodes in the queue
-			public readonly Dictionary<NodeKey, PendingNode> _nodeKeyToInfo = new Dictionary<NodeKey, PendingNode>();
+			readonly Dictionary<NodeKey, PendingNode> _nodeKeyToInfo = new Dictionary<NodeKey, PendingNode>();
 
 			// Queue of nodes for the current bundle
 			readonly List<PendingNode> _queue = new List<PendingNode>();
@@ -756,7 +756,7 @@ namespace EpicGames.Horde.Storage.Bundles
 
 		bool _disposed;
 
-		internal readonly ILogger? _traceLogger;
+		internal ILogger? TraceLogger { get; }
 
 		/// <summary>
 		/// Cache of nodes to deduplicate against
@@ -782,7 +782,7 @@ namespace EpicGames.Horde.Storage.Bundles
 		/// </summary>
 		/// <param name="other"></param>
 		public BundleWriter(BundleWriter other)
-			: this(other._store, other._reader, other._refName, other._options, other._nodeCache, other._writeQueue, other._traceLogger)
+			: this(other._store, other._reader, other._refName, other._options, other._nodeCache, other._writeQueue, other.TraceLogger)
 		{
 			_writeQueue.AddRef();
 		}
@@ -798,7 +798,7 @@ namespace EpicGames.Horde.Storage.Bundles
 			_options = options ?? s_defaultOptions;
 			_nodeCache = nodeCache ?? new NodeCache(_options.NodeCacheSize);
 			_writeQueue = writeQueue ?? new WriteQueue(store, refName.Text, _options.MaxWriteQueueLength, traceLogger);
-			_traceLogger = traceLogger;
+			TraceLogger = traceLogger;
 		}
 
 		/// <inheritdoc/>
@@ -867,14 +867,14 @@ namespace EpicGames.Horde.Storage.Bundles
 			// Check if we have a matching node already in storage
 			if (_nodeCache.TryGetNode(nodeKey, out BlobHandle? handle))
 			{
-				_traceLogger?.LogInformation("Returning cached handle for {NodeKey} -> {Handle}", nodeKey, handle);
+				TraceLogger?.LogInformation("Returning cached handle for {NodeKey} -> {Handle}", nodeKey, handle);
 				return handle;
 			}
 
 			// Append this node data
 			PendingNode pendingNode = currentBundle.WriteNode(nodeKey, size, references);
 			_nodeCache.Add(nodeKey, pendingNode);
-			_traceLogger?.LogInformation("Added new node for {NodeKey} in bundle {BundleId}", nodeKey, currentBundle.BundleId);
+			TraceLogger?.LogInformation("Added new node for {NodeKey} in bundle {BundleId}", nodeKey, currentBundle.BundleId);
 
 			// Add dependencies on all bundles containing a dependent node
 			foreach (BlobHandle reference in references)
@@ -882,7 +882,7 @@ namespace EpicGames.Horde.Storage.Bundles
 				PendingNode? pendingReference = reference as PendingNode;
 				if (pendingReference?.PendingBundle != null)
 				{
-					currentBundle.AddDependencyOn(pendingReference.PendingBundle, _traceLogger);
+					currentBundle.AddDependencyOn(pendingReference.PendingBundle, TraceLogger);
 				}
 			}
 
diff --git a/Engine/Source/Programs/Shared/EpicGames.Horde/Storage/IStorageBackend.cs b/Engine/Source/Programs/Shared/EpicGames.Horde/Storage/IStorageBackend.cs
--- a/Engine/Source/Programs/Shared/EpicGames.Horde/Storage/IStorageBackend.cs
+++ b/Engine/Source/Programs/Shared/EpicGames.Horde/Storage/IStorageBackend.cs
@@ -123,7 +123,7 @@ namespace EpicGames.Horde.Storage
 		/// <summary>
 		/// Incremented value used for each supplied id
 		/// </summary>
-		static int _increment;
+		static int s_increment;
 
 		/// <summary>
 		/// Creates a unique name with a given prefix
@@ -138,7 +138,7 @@ namespace EpicGames.Horde.Storage
 				builder.Append('/');
 			}
 			builder.Append(s_sessionPrefix);
-			builder.Append(Interlocked.Increment(ref _increment));
+			builder.Append(Interlocked.Increment(ref s_increment));
 			return builder.ToString();
 		}
 	}
diff --git a/Engine/Source/Programs/Shared/EpicGames.Horde/Storage/Nodes/DirectoryNode.cs b/Engine/Source/Programs/Shared/EpicGames.Horde/Storage/Nodes/DirectoryNode.cs
--- a/Engine/Source/Programs/Shared/EpicGames.Horde/Storage/Nodes/DirectoryNode.cs
+++ b/Engine/Source/Programs/Shared/EpicGames.Horde/Storage/Nodes/DirectoryNode.cs
@@ -780,8 +780,8 @@ namespace EpicGames.Horde.Storage.Nodes
 		/// <param name="cancellationToken">Cancellation token for the operation</param>
 		public async Task CopyToDirectoryAsync(DirectoryInfo directoryInfo, IProgress<ICopyStats>? progress, ILogger logger, CancellationToken cancellationToken)
 		{
-			int NumTasks = Math.Min(1 + (int)(Length / (10 * 1024 * 1024)), 16);
-			logger.LogInformation("Splitting read into {NumThreads} threads", NumTasks);
+			int numTasks = Math.Min(1 + (int)(Length / (10 * 1024 * 1024)), 16);
+			logger.LogInformation("Splitting read into {NumThreads} threads", numTasks);
 
 			CopyStats? copyStats = null;
 			if (progress != null)
@@ -793,10 +793,10 @@ namespace EpicGames.Horde.Storage.Nodes
 			try
 			{
 				long offset = 0;
-				for (int taskIdx = 0; taskIdx < NumTasks; taskIdx++)
+				for (int taskIdx = 0; taskIdx < numTasks; taskIdx++)
 				{
 					long minOffset = offset;
-					long maxOffset = (Length * (taskIdx + 1)) / NumTasks;
+					long maxOffset = (Length * (taskIdx + 1)) / numTasks;
 					tasks.Add(Task.Run(() => CopyToDirectoryInternalAsync(directoryInfo, minOffset, maxOffset - minOffset, copyStats, logger, cancellationToken), cancellationToken));
 					offset = maxOffset;
 				}
diff --git a/Engine/Source/Programs/UnrealBuildTool/Matchers/.editorconfig b/Engine/Source/Programs/UnrealBuildTool/Matchers/.editorconfig
--- a/Engine/Source/Programs/UnrealBuildTool/Matchers/.editorconfig
+++ b/Engine/Source/Programs/UnrealBuildTool/Matchers/.editorconfig
@@ -7,3 +7,9 @@
 
 # IDE0005: Using directive is unnecessary (Disabled for UBT due to mismatch in projects with implicit usings)
 dotnet_diagnostic.IDE0005.severity = silent 
+
+# IDE0018: Variable declaration can be inlined
+dotnet_diagnostic.IDE0018.severity = silent 
+
+# IDE0090: New expression can be simplified
+dotnet_diagnostic.IDE0090.severity = silent 
diff --git a/Engine/Source/Programs/UnrealBuildTool/Matchers/CompileEventMatcher.cs b/Engine/Source/Programs/UnrealBuildTool/Matchers/CompileEventMatcher.cs
--- a/Engine/Source/Programs/UnrealBuildTool/Matchers/CompileEventMatcher.cs
+++ b/Engine/Source/Programs/UnrealBuildTool/Matchers/CompileEventMatcher.cs
@@ -324,7 +324,7 @@ namespace UnrealBuildTool.Matchers
 			}
 			else
 			{
-				return fileName.Substring(0, index);
+				return fileName[..index];
 			}
 		}
 
