Update Engine/Source/Programs/Shared to use same .editorconfig rules as Horde.

diff --git a/Engine/Source/Programs/Horde/Horde.Server/Logs/LogFileService.cs b/Engine/Source/Programs/Horde/Horde.Server/Logs/LogFileService.cs
--- a/Engine/Source/Programs/Horde/Horde.Server/Logs/LogFileService.cs
+++ b/Engine/Source/Programs/Horde/Horde.Server/Logs/LogFileService.cs
@@ -1898,7 +1898,7 @@ namespace Horde.Server.Logs
 				if(root != null)
 				{
 					LogIndexNode index = await root.IndexRef.ExpandAsync(cancellationToken);
-					await foreach (int lineIdx in index.Search(firstLine, searchText, searchStats, cancellationToken))
+					await foreach (int lineIdx in index.SearchAsync(firstLine, searchText, searchStats, cancellationToken))
 					{
 						yield return lineIdx;
 					}
diff --git a/Engine/Source/Programs/Horde/Horde/Commands/Perforce/PerforceStatus.cs b/Engine/Source/Programs/Horde/Horde/Commands/Perforce/PerforceStatus.cs
--- a/Engine/Source/Programs/Horde/Horde/Commands/Perforce/PerforceStatus.cs
+++ b/Engine/Source/Programs/Horde/Horde/Commands/Perforce/PerforceStatus.cs
@@ -12,7 +12,7 @@ namespace Horde.Commands.Perforce
 	{
 		protected override async Task ExecuteAsync(IPerforceConnection perforce, ManagedWorkspace repo, ILogger logger)
 		{
-			await repo.Status();
+			await repo.StatusAsync();
 		}
 	}
 }
diff --git a/Engine/Source/Programs/Shared/.editorconfig b/Engine/Source/Programs/Shared/.editorconfig
--- a/Engine/Source/Programs/Shared/.editorconfig
+++ b/Engine/Source/Programs/Shared/.editorconfig
@@ -7,8 +7,7 @@
 
 csharp_style_prefer_switch_expression = false
 csharp_style_implicit_object_creation_when_type_is_apparent = false
-dotnet_remove_unnecessary_suppression_exclusions = false
-dotnet_diagnostic.IDE0079.severity = silent
+dotnet_remove_unnecessary_suppression_exclusions = all
 
 # Sort using and Import directives with System.* appearing first
 dotnet_sort_system_directives_first = true
@@ -115,8 +114,9 @@ dotnet_naming_symbols.all_members.applicable_kinds = *
 
 dotnet_naming_style.pascal_case_style.capitalization = pascal_case
 
-# IDE0005: Using directive is unnecessary
-dotnet_diagnostic.IDE0005.severity = warning
+# IDE0073: File header
+dotnet_diagnostic.IDE0073.severity = warning
+file_header_template = Copyright Epic Games, Inc. All Rights Reserved.
 
 # IDE0035: Remove unreachable code
 dotnet_diagnostic.IDE0035.severity = warning
@@ -151,19 +151,24 @@ dotnet_diagnostic.IDE0059.severity = warning
 # IDE0060: Remove unused parameter
 dotnet_diagnostic.IDE0060.severity = warning
 
-# IDE0073: File header
-dotnet_diagnostic.IDE0073.severity = warning
-file_header_template = Copyright Epic Games, Inc. All Rights Reserved.
-
-# IDE1006: Enforce naming conventions
-dotnet_diagnostic.IDE1006.severity = warning
-
 # CA1012: Abstract types should not have public constructors
 dotnet_diagnostic.CA1012.severity = warning
 
 # CA1822: Make member static
 dotnet_diagnostic.CA1822.severity = warning
 
+# IDE0005: Using directive is unnecessary
+dotnet_diagnostic.IDE0005.severity = warning
+
+# IDE0073: Null check can be simplified
+dotnet_diagnostic.IDE0270.severity = silent
+
+# IDE0090: 'new' expression can be simplified
+dotnet_diagnostic.IDE0090.severity = silent
+
+# IDE1006: Enforce naming conventions
+dotnet_diagnostic.IDE1006.severity = warning
+
 # dotnet_style_allow_multiple_blank_lines_experimental
 dotnet_diagnostic.IDE2000.severity = warning
 
@@ -260,6 +265,19 @@ csharp_prefer_braces = true:silent
 csharp_preserve_single_line_blocks = true
 csharp_preserve_single_line_statements = true
 
+# Async methods should have "Async" suffix
+dotnet_naming_rule.async_methods_end_in_async.symbols = any_async_methods
+dotnet_naming_rule.async_methods_end_in_async.style = end_in_async
+dotnet_naming_rule.async_methods_end_in_async.severity = warning
+
+dotnet_naming_symbols.any_async_methods.applicable_kinds = method
+dotnet_naming_symbols.any_async_methods.applicable_accessibilities = *
+dotnet_naming_symbols.any_async_methods.required_modifiers = async
+
+dotnet_naming_style.end_in_async.required_prefix = 
+dotnet_naming_style.end_in_async.required_suffix = Async
+dotnet_naming_style.end_in_async.capitalization = pascal_case
+dotnet_naming_style.end_in_async.word_separator = 
 
 # All Rules Enabled as build warnings
 # Description: All Rules are enabled as build warnings. Rules with IsEnabledByDefault = false are force enabled as build warnings.
@@ -1077,3 +1095,9 @@ dotnet_diagnostic.CA5404.severity = none
 
 # CA5405: Do not always skip token validation in delegates
 dotnet_diagnostic.CA5405.severity = warning
+
+
+[*Test*.cs]
+
+# CA1707: Identifiers should not contain underscores
+dotnet_diagnostic.CA1707.severity = none
\ No newline at end of file
diff --git a/Engine/Source/Programs/Shared/EpicGames.AspNet/ServerTimings.cs b/Engine/Source/Programs/Shared/EpicGames.AspNet/ServerTimings.cs
--- a/Engine/Source/Programs/Shared/EpicGames.AspNet/ServerTimings.cs
+++ b/Engine/Source/Programs/Shared/EpicGames.AspNet/ServerTimings.cs
@@ -112,7 +112,7 @@ namespace EpicGames.AspNet
 			_next = next ?? throw new ArgumentNullException(nameof(next));
 		}
 
-		public async Task Invoke(HttpContext context)
+		public async Task InvokeAsync(HttpContext context)
 		{
 			IServerTiming serverTiming = context.RequestServices.GetRequiredService<IServerTiming>();
 			if (AllowsTrailers(context.Request) && context.Response.SupportsTrailers())
diff --git a/Engine/Source/Programs/Shared/EpicGames.Core.Tests/AsyncThreadPoolWorkQueueTests.cs b/Engine/Source/Programs/Shared/EpicGames.Core.Tests/AsyncThreadPoolWorkQueueTests.cs
--- a/Engine/Source/Programs/Shared/EpicGames.Core.Tests/AsyncThreadPoolWorkQueueTests.cs
+++ b/Engine/Source/Programs/Shared/EpicGames.Core.Tests/AsyncThreadPoolWorkQueueTests.cs
@@ -19,7 +19,7 @@ public sealed class AsyncThreadPoolWorkQueueTests : IDisposable
 	}
 
 	[TestMethod]
-	public async Task ExecuteEnqueuedTasks()
+	public async Task ExecuteEnqueuedTasksAsync()
 	{
 		int c = 0;
 		await _queue.EnqueueAsync((_) => 
@@ -37,7 +37,7 @@ public sealed class AsyncThreadPoolWorkQueueTests : IDisposable
 	}
 	
 	[TestMethod]
-	public async Task ExecuteManyEnqueuedTasks()
+	public async Task ExecuteManyEnqueuedTasksAsync()
 	{
 		int c = 0;
 		int numTasks = 100000;
@@ -56,7 +56,7 @@ public sealed class AsyncThreadPoolWorkQueueTests : IDisposable
 	}
 	
 	[TestMethod]
-	public async Task ExecuteNestedEnqueuedTasks()
+	public async Task ExecuteNestedEnqueuedTasksAsync()
 	{
 		bool task1 = false;
 		bool task2 = false;
@@ -86,7 +86,7 @@ public sealed class AsyncThreadPoolWorkQueueTests : IDisposable
 	
 	[Ignore]
 	[TestMethod]
-	public async Task VerifyMaxConcurrency()
+	public async Task VerifyMaxConcurrencyAsync()
 	{
 		const int NumTasks = 10000;
 		const int NumWorkers = 3;
diff --git a/Engine/Source/Programs/Shared/EpicGames.Core.Tests/JsonTracerTests.cs b/Engine/Source/Programs/Shared/EpicGames.Core.Tests/JsonTracerTests.cs
--- a/Engine/Source/Programs/Shared/EpicGames.Core.Tests/JsonTracerTests.cs
+++ b/Engine/Source/Programs/Shared/EpicGames.Core.Tests/JsonTracerTests.cs
@@ -96,7 +96,7 @@ namespace EpicGames.Core.Tests
 			Assert.AreEqual(tracer.ActiveSpan, null);
 		}
 
-		private static async Task<bool> AsyncSpanTest(JsonTracer tracer, IScope parentScope)
+		private static async Task<bool> AsyncSpanTestAsync(JsonTracer tracer, IScope parentScope)
 		{
 			for (int index = 0; index < 100; index++)
 			{
@@ -126,14 +126,14 @@ namespace EpicGames.Core.Tests
 			using (IScope parentScope = tracer.BuildSpan("Parent").StartActive())
 			{
 				Assert.AreEqual(tracer.ActiveSpan, parentScope.Span);
-				AsyncSpanTest(tracer, parentScope).Wait();
+				AsyncSpanTestAsync(tracer, parentScope).Wait();
 				Assert.AreEqual(tracer.ActiveSpan, parentScope.Span);
 			}
 
 			Assert.AreEqual(tracer.ActiveSpan, null);
 		}
 
-		private static async Task ForEachAsyncSpanTest(JsonTracer tracer, IScope parentScope)
+		private static async Task ForEachAsyncSpanTestAsync(JsonTracer tracer, IScope parentScope)
 		{
 			int[] values = Enumerable.Range(0, 10000).ToArray();
 			await Parallel.ForEachAsync(values, new ParallelOptions { MaxDegreeOfParallelism = 10000 }, async (i, token) =>
@@ -155,7 +155,7 @@ namespace EpicGames.Core.Tests
 			using (IScope parentScope = tracer.BuildSpan("Parent").StartActive())
 			{
 				Assert.AreEqual(tracer.ActiveSpan, parentScope.Span);
-				ForEachAsyncSpanTest(tracer, parentScope).Wait();
+				ForEachAsyncSpanTestAsync(tracer, parentScope).Wait();
 				Assert.AreEqual(tracer.ActiveSpan, parentScope.Span);
 			}
 
diff --git a/Engine/Source/Programs/Shared/EpicGames.Core/AsyncThreadPoolWorkQueue.cs b/Engine/Source/Programs/Shared/EpicGames.Core/AsyncThreadPoolWorkQueue.cs
--- a/Engine/Source/Programs/Shared/EpicGames.Core/AsyncThreadPoolWorkQueue.cs
+++ b/Engine/Source/Programs/Shared/EpicGames.Core/AsyncThreadPoolWorkQueue.cs
@@ -84,7 +84,7 @@ public sealed class AsyncThreadPoolWorkQueue : IDisposable
 		List<Task> tasks = new();
 		for (int i = 0; i < _numWorkers; i++)
 		{
-			tasks.Add(Task.Run(() => Worker(cancellationToken), cancellationToken));
+			tasks.Add(Task.Run(() => WorkerAsync(cancellationToken), cancellationToken));
 		}
 
 		Task queueCompletionTask = Task.Run(async () =>
@@ -101,7 +101,7 @@ public sealed class AsyncThreadPoolWorkQueue : IDisposable
 	/// Worker executing tasks. Any exception encountered will be bubbled up as all calls are awaited.
 	/// </summary>
 	/// <param name="cancellationToken">Cancellation token</param>
-	private async Task Worker(CancellationToken cancellationToken)
+	private async Task WorkerAsync(CancellationToken cancellationToken)
 	{
 		await foreach (Func<CancellationToken, Task> task in _tasks.Reader.ReadAllAsync(cancellationToken))
 		{
diff --git a/Engine/Source/Programs/Shared/EpicGames.Core/AsyncUtils.cs b/Engine/Source/Programs/Shared/EpicGames.Core/AsyncUtils.cs
--- a/Engine/Source/Programs/Shared/EpicGames.Core/AsyncUtils.cs
+++ b/Engine/Source/Programs/Shared/EpicGames.Core/AsyncUtils.cs
@@ -54,7 +54,7 @@ namespace EpicGames.Core
 		/// <param name="task">Task to wait for</param>
 		/// <param name="cancellationToken">Cancellation token for the operation</param>
 		/// <returns>Wrapped task</returns>
-		public static async Task<T> AbandonOnCancel<T>(this Task<T> task, CancellationToken cancellationToken)
+		public static async Task<T> AbandonOnCancelAsync<T>(this Task<T> task, CancellationToken cancellationToken)
 		{
 			if (cancellationToken.CanBeCanceled)
 			{
@@ -164,7 +164,7 @@ namespace EpicGames.Core
 		/// <param name="source">Sequence to enumerate</param>
 		/// <param name="cancellationToken">Cancellation token for the operation</param>
 		/// <returns></returns>
-		public static async IAsyncEnumerable<T> Prefetch<T>(this IAsyncEnumerable<T> source, [EnumeratorCancellation] CancellationToken cancellationToken)
+		public static async IAsyncEnumerable<T> PrefetchAsync<T>(this IAsyncEnumerable<T> source, [EnumeratorCancellation] CancellationToken cancellationToken)
 		{
 			await using IAsyncEnumerator<T> enumerator = source.GetAsyncEnumerator(cancellationToken);
 			if (await enumerator.MoveNextAsync())
diff --git a/Engine/Source/Programs/Shared/EpicGames.Core/HttpClientExtensions.cs b/Engine/Source/Programs/Shared/EpicGames.Core/HttpClientExtensions.cs
--- a/Engine/Source/Programs/Shared/EpicGames.Core/HttpClientExtensions.cs
+++ b/Engine/Source/Programs/Shared/EpicGames.Core/HttpClientExtensions.cs
@@ -30,7 +30,7 @@ namespace EpicGames.Core
 			using (HttpResponseMessage response = await client.GetAsync(url, cancellationToken))
 			{
 				response.EnsureSuccessStatusCode();
-				return await ParseJsonContent<TResponse>(response, cancellationToken);
+				return await ParseJsonContentAsync<TResponse>(response, cancellationToken);
 			}
 		}
 
@@ -40,7 +40,7 @@ namespace EpicGames.Core
 			using (HttpResponseMessage response = await client.GetAsync(url, cancellationToken))
 			{
 				response.EnsureSuccessStatusCode();
-				return await ParseJsonContent<TResponse>(response, cancellationToken);
+				return await ParseJsonContentAsync<TResponse>(response, cancellationToken);
 			}
 		}
 
@@ -81,7 +81,7 @@ namespace EpicGames.Core
 			using (HttpResponseMessage response = await PostAsync(client, url, request, cancellationToken))
 			{
 				response.EnsureSuccessStatusCode();
-				return await ParseJsonContent<TResponse>(response, cancellationToken);
+				return await ParseJsonContentAsync<TResponse>(response, cancellationToken);
 			}
 		}
 
@@ -91,7 +91,7 @@ namespace EpicGames.Core
 			using (HttpResponseMessage response = await PostAsync(client, url, request, cancellationToken))
 			{
 				response.EnsureSuccessStatusCode();
-				return await ParseJsonContent<TResponse>(response, cancellationToken);
+				return await ParseJsonContentAsync<TResponse>(response, cancellationToken);
 			}
 		}
 
@@ -132,7 +132,7 @@ namespace EpicGames.Core
 			using (HttpResponseMessage response = await PutAsync(client, url, request, cancellationToken))
 			{
 				response.EnsureSuccessStatusCode();
-				return await ParseJsonContent<TResponse>(response, cancellationToken);
+				return await ParseJsonContentAsync<TResponse>(response, cancellationToken);
 			}
 		}
 
@@ -142,7 +142,7 @@ namespace EpicGames.Core
 			using (HttpResponseMessage response = await PutAsync(client, url, request, cancellationToken))
 			{
 				response.EnsureSuccessStatusCode();
-				return await ParseJsonContent<TResponse>(response, cancellationToken);
+				return await ParseJsonContentAsync<TResponse>(response, cancellationToken);
 			}
 		}
 
@@ -164,7 +164,7 @@ namespace EpicGames.Core
 		/// <param name="message">The message received</param>
 		/// <param name="cancellationToken">Cancellation token for the operation</param>
 		/// <returns>Parsed object instance</returns>
-		private static async Task<T> ParseJsonContent<T>(HttpResponseMessage message, CancellationToken cancellationToken)
+		private static async Task<T> ParseJsonContentAsync<T>(HttpResponseMessage message, CancellationToken cancellationToken)
 		{
 			byte[] bytes = await message.Content.ReadAsByteArrayAsync(cancellationToken);
 			return JsonSerializer.Deserialize<T>(bytes, new JsonSerializerOptions { PropertyNameCaseInsensitive = true })!;
diff --git a/Engine/Source/Programs/Shared/EpicGames.Horde.Tests/BackendTests.cs b/Engine/Source/Programs/Shared/EpicGames.Horde.Tests/BackendTests.cs
--- a/Engine/Source/Programs/Shared/EpicGames.Horde.Tests/BackendTests.cs
+++ b/Engine/Source/Programs/Shared/EpicGames.Horde.Tests/BackendTests.cs
@@ -37,14 +37,14 @@ namespace EpicGames.Horde.Tests
 		}
 
 		[TestMethod]
-		public async Task TestMemoryBackend()
+		public async Task TestMemoryBackendAsync()
 		{
 			using MemoryStorageBackend backend = new MemoryStorageBackend();
 			await TestBackendAsync(backend);
 		}
 
 		[TestMethod]
-		public async Task TestFileBackend()
+		public async Task TestFileBackendAsync()
 		{
 			using (TempDir tempDir = new TempDir("Cache"))
 			{
@@ -54,7 +54,7 @@ namespace EpicGames.Horde.Tests
 		}
 
 		[TestMethod]
-		public async Task TestFileBackendFileMapping()
+		public async Task TestFileBackendFileMappingAsync()
 		{
 			using (TempDir tempDir = new TempDir("Cache"))
 			{
@@ -77,7 +77,7 @@ namespace EpicGames.Horde.Tests
 		}
 
 		[TestMethod]
-		public async Task TestCacheBackend()
+		public async Task TestCacheBackendAsync()
 		{
 			using (TempDir tempDir = new TempDir("Cache"))
 			{
diff --git a/Engine/Source/Programs/Shared/EpicGames.Horde.Tests/BufferTests.cs b/Engine/Source/Programs/Shared/EpicGames.Horde.Tests/BufferTests.cs
--- a/Engine/Source/Programs/Shared/EpicGames.Horde.Tests/BufferTests.cs
+++ b/Engine/Source/Programs/Shared/EpicGames.Horde.Tests/BufferTests.cs
@@ -21,7 +21,7 @@ namespace EpicGames.Horde.Tests
 		const int ChannelId = 0;
 
 		[TestMethod]
-		public async Task TestSimpleBuffer()
+		public async Task TestSimpleBufferAsync()
 		{
 			using PooledBuffer buffer = new PooledBuffer(2, 1024);
 			using ComputeBufferWriter bufferWriter = buffer.CreateWriter();
@@ -81,13 +81,13 @@ namespace EpicGames.Horde.Tests
 		}
 
 		[TestMethod]
-		public async Task TestPooledBuffer()
+		public async Task TestPooledBufferAsync()
 		{
 			await TestProducerConsumerAsync(length => new PooledBuffer(length), CancellationToken.None);
 		}
 
 		[TestMethod]
-		public async Task TestSharedMemoryBuffer()
+		public async Task TestSharedMemoryBufferAsync()
 		{
 			if (RuntimeInformation.IsOSPlatform(OSPlatform.Windows))
 			{
@@ -161,7 +161,7 @@ namespace EpicGames.Horde.Tests
 		}
 
 		[TestMethod]
-		public async Task TestSendBufferComplete()
+		public async Task TestSendBufferCompleteAsync()
 		{
 			Pipe recvPipe = new Pipe();
 			Pipe sendPipe = new Pipe();
diff --git a/Engine/Source/Programs/Shared/EpicGames.Horde.Tests/BundleStoreTests.cs b/Engine/Source/Programs/Shared/EpicGames.Horde.Tests/BundleStoreTests.cs
--- a/Engine/Source/Programs/Shared/EpicGames.Horde.Tests/BundleStoreTests.cs
+++ b/Engine/Source/Programs/Shared/EpicGames.Horde.Tests/BundleStoreTests.cs
@@ -152,17 +152,17 @@ namespace EpicGames.Horde.Tests
 				await store.WriteRefTargetAsync(new RefName("test"), await writer.WriteNodeAsync(root));
 
 				BundleReader reader = new BundleReader(store, StorageCache.None, NullLogger.Instance);
-				await CheckTree(root);
+				await CheckTreeAsync(root);
 			}
 
 			// Check we can read it back in
 			{
 				SimpleNode root = await store.ReadRefAsync<SimpleNode>(new RefName("test"));
-				await CheckTree(root);
+				await CheckTreeAsync(root);
 			}
 		}
 
-		static async Task CheckTree(SimpleNode root)
+		static async Task CheckTreeAsync(SimpleNode root)
 		{
 			SimpleNode node5 = root;
 			byte[] data5 = node5.Data.ToArray();
@@ -383,7 +383,7 @@ namespace EpicGames.Horde.Tests
 			// Check we can read it back in
 			{
 				DirectoryNode newRoot = await store.ReadRefAsync<DirectoryNode>(new RefName("test"));
-				await CompareTrees(root, newRoot);
+				await CompareTreesAsync(root, newRoot);
 				await CheckLargeFileTreeAsync(root, data);
 
 				NodeRef<ChunkedDataNode> file = root.GetFileEntry("test");
@@ -393,7 +393,7 @@ namespace EpicGames.Horde.Tests
 			}
 		}
 
-		static async Task CompareTrees(DirectoryNode oldNode, DirectoryNode newNode)
+		static async Task CompareTreesAsync(DirectoryNode oldNode, DirectoryNode newNode)
 		{
 			Assert.AreEqual(oldNode.Length, newNode.Length);
 			Assert.AreEqual(oldNode.Files.Count, newNode.Files.Count);
@@ -403,11 +403,11 @@ namespace EpicGames.Horde.Tests
 			{
 				ChunkedDataNode oldFile = await oldFileEntry.ExpandAsync();
 				ChunkedDataNode newFile = await newFileEntry.ExpandAsync();
-				await CompareTrees(oldFile, newFile);
+				await CompareTreesAsync(oldFile, newFile);
 			}
 		}
 
-		static async Task CompareTrees(ChunkedDataNode oldNode, ChunkedDataNode newNode)
+		static async Task CompareTreesAsync(ChunkedDataNode oldNode, ChunkedDataNode newNode)
 		{
 			if (oldNode is InteriorChunkedDataNode oldInteriorNode)
 			{
@@ -419,7 +419,7 @@ namespace EpicGames.Horde.Tests
 				{
 					ChunkedDataNode oldFile = await oldFileRef.ExpandAsync();
 					ChunkedDataNode newFile = await newFileRef.ExpandAsync();
-					await CompareTrees(oldFile, newFile);
+					await CompareTreesAsync(oldFile, newFile);
 					index++;
 				}
 			}
diff --git a/Engine/Source/Programs/Shared/EpicGames.Horde.Tests/BundleTests.cs b/Engine/Source/Programs/Shared/EpicGames.Horde.Tests/BundleTests.cs
--- a/Engine/Source/Programs/Shared/EpicGames.Horde.Tests/BundleTests.cs
+++ b/Engine/Source/Programs/Shared/EpicGames.Horde.Tests/BundleTests.cs
@@ -65,7 +65,7 @@ namespace EpicGames.Horde.Tests
 		}
 
 		[TestMethod]
-		public async Task BasicChunkingTests()
+		public async Task BasicChunkingTestsAsync()
 		{
 			RefName refName = new RefName("test");
 			await using IStorageWriter writer = _storage.CreateWriter(refName);
@@ -186,7 +186,7 @@ namespace EpicGames.Horde.Tests
 		}
 
 		[TestMethod]
-		public async Task BasicTestDirectory()
+		public async Task BasicTestDirectoryAsync()
 		{
 			MemoryStorageClient store = _storage;
 
@@ -238,7 +238,7 @@ namespace EpicGames.Horde.Tests
 		}
 
 		[TestMethod]
-		public async Task DedupTests()
+		public async Task DedupTestsAsync()
 		{
 			Assert.AreEqual(0, _storage.Refs.Count);
 			Assert.AreEqual(0, _storage.Blobs.Count);
@@ -266,7 +266,7 @@ namespace EpicGames.Horde.Tests
 		}
 
 		[TestMethod]
-		public async Task ReloadTests()
+		public async Task ReloadTestsAsync()
 		{
 			BundleOptions options = new BundleOptions();
 			options.MaxBlobSize = 1;
diff --git a/Engine/Source/Programs/Shared/EpicGames.Horde.Tests/ChunkingTests.cs b/Engine/Source/Programs/Shared/EpicGames.Horde.Tests/ChunkingTests.cs
--- a/Engine/Source/Programs/Shared/EpicGames.Horde.Tests/ChunkingTests.cs
+++ b/Engine/Source/Programs/Shared/EpicGames.Horde.Tests/ChunkingTests.cs
@@ -47,7 +47,7 @@ namespace EpicGames.Horde.Tests
 		}
 
 		[TestMethod]
-		public async Task EmptyNodeTest()
+		public async Task EmptyNodeTestAsync()
 		{
 			using MemoryStorageClient store = new MemoryStorageClient();
 
@@ -70,7 +70,7 @@ namespace EpicGames.Horde.Tests
 		}
 
 		[TestMethod]
-		public async Task FixedSizeChunkingTests()
+		public async Task FixedSizeChunkingTestsAsync()
 		{
 			ChunkingOptions options = new ChunkingOptions();
 			options.LeafOptions = new LeafChunkedDataNodeOptions(64, 64, 64);
@@ -80,7 +80,7 @@ namespace EpicGames.Horde.Tests
 		}
 
 		[TestMethod]
-		public async Task VariableSizeChunkingTests()
+		public async Task VariableSizeChunkingTestsAsync()
 		{
 			ChunkingOptions options = new ChunkingOptions();
 			options.LeafOptions = new LeafChunkedDataNodeOptions(32, 64, 96);
@@ -136,10 +136,10 @@ namespace EpicGames.Horde.Tests
 				Assert.IsTrue(spanData.Span.SequenceEqual(data));
 			}
 
-			await CheckSizes(root, options, true);
+			await CheckSizesAsync(root, options, true);
 		}
 
-		static async Task CheckSizes(ChunkedDataNode node, ChunkingOptions options, bool rightmost)
+		static async Task CheckSizesAsync(ChunkedDataNode node, ChunkingOptions options, bool rightmost)
 		{
 			if (node is LeafChunkedDataNode leafNode)
 			{
@@ -157,7 +157,7 @@ namespace EpicGames.Horde.Tests
 				for (int idx = 0; idx < childCount; idx++)
 				{
 					ChunkedDataNode childNode = await interiorNode.Children[idx].ExpandAsync(CancellationToken.None);
-					await CheckSizes(childNode, options, idx == childCount - 1);
+					await CheckSizesAsync(childNode, options, idx == childCount - 1);
 				}
 			}
 		}
diff --git a/Engine/Source/Programs/Shared/EpicGames.Horde.Tests/ComputeSocketTests.cs b/Engine/Source/Programs/Shared/EpicGames.Horde.Tests/ComputeSocketTests.cs
--- a/Engine/Source/Programs/Shared/EpicGames.Horde.Tests/ComputeSocketTests.cs
+++ b/Engine/Source/Programs/Shared/EpicGames.Horde.Tests/ComputeSocketTests.cs
@@ -70,18 +70,18 @@ namespace EpicGames.Horde.Tests
 		}
 
 		[TestMethod]
-		public async Task TestAgentMessageLoopPipe()
+		public async Task TestAgentMessageLoopPipeAsync()
 		{
 			Pipe recvPipe = new Pipe();
 			Pipe sendPipe = new Pipe();
 			await using RemoteComputeSocket localSocket = new RemoteComputeSocket(new PipeTransport(sendPipe.Reader, recvPipe.Writer), new TestLogger());
 			await using RemoteComputeSocket agentSocket = new RemoteComputeSocket(new PipeTransport(recvPipe.Reader, sendPipe.Writer), new TestLogger());
 
-			await RunAgentTests(localSocket, agentSocket);
+			await RunAgentTestsAsync(localSocket, agentSocket);
 		}
 
 		[TestMethod]
-		public async Task TestAgentMessageLoopTcp()
+		public async Task TestAgentMessageLoopTcpAsync()
 		{
 			const int Port = 9990;
 			TcpListener listener = new TcpListener(IPAddress.Loopback, Port);
@@ -96,13 +96,13 @@ namespace EpicGames.Horde.Tests
 			await using RemoteComputeSocket localSocket = new RemoteComputeSocket(new TcpTransport(clientSocket), new TestLogger());
 			await using RemoteComputeSocket agentSocket = new RemoteComputeSocket(new TcpTransport(serverSocket), new TestLogger());
 
-			await RunAgentTests(localSocket, agentSocket);
+			await RunAgentTestsAsync(localSocket, agentSocket);
 		}
 
-		static async Task RunAgentTests(RemoteComputeSocket localSocket, RemoteComputeSocket agentSocket)
+		static async Task RunAgentTestsAsync(RemoteComputeSocket localSocket, RemoteComputeSocket agentSocket)
 		{
 			DirectoryReference tempDir = new DirectoryReference("test-temp");
-			await using (BackgroundTask agentTask = BackgroundTask.StartNew(ctx => RunAgent(agentSocket, tempDir, ctx)))
+			await using (BackgroundTask agentTask = BackgroundTask.StartNew(ctx => RunAgentAsync(agentSocket, tempDir, ctx)))
 			{
 				const int PrimaryChannelId = 0;
 				using (AgentMessageChannel channel = localSocket.CreateAgentMessageChannel(PrimaryChannelId, 4 * 1024 * 1024))
@@ -188,7 +188,7 @@ namespace EpicGames.Horde.Tests
 			await agentSocket.CloseAsync(CancellationToken.None);
 		}
 
-		static async Task RunAgent(ComputeSocket socket, DirectoryReference tempDir, CancellationToken cancellationToken)
+		static async Task RunAgentAsync(ComputeSocket socket, DirectoryReference tempDir, CancellationToken cancellationToken)
 		{
 			using StorageCache storageCache = new StorageCache();
 			AgentMessageHandler handler = new AgentMessageHandler(tempDir, storageCache, null, true, null, NullLogger.Instance);
diff --git a/Engine/Source/Programs/Shared/EpicGames.Horde.Tests/LogTests.cs b/Engine/Source/Programs/Shared/EpicGames.Horde.Tests/LogTests.cs
--- a/Engine/Source/Programs/Shared/EpicGames.Horde.Tests/LogTests.cs
+++ b/Engine/Source/Programs/Shared/EpicGames.Horde.Tests/LogTests.cs
@@ -41,7 +41,7 @@ namespace EpicGames.Horde.Tests
 		}
 
 		[TestMethod]
-		public async Task IndexTests()
+		public async Task IndexTestsAsync()
 		{
 			using MemoryStorageClient store = new MemoryStorageClient();
 
@@ -102,7 +102,7 @@ namespace EpicGames.Horde.Tests
 		}
 
 		[TestMethod]
-		public async Task PartialTokenTests()
+		public async Task PartialTokenTestsAsync()
 		{
 			using MemoryStorageClient store = new MemoryStorageClient();
 
@@ -140,7 +140,7 @@ namespace EpicGames.Horde.Tests
 						string str = lines[lineIdx].Substring(strOfs, strLen);
 
 						SearchStats stats = new SearchStats();
-						List<int> results = await index.Search(0, new SearchTerm(str), stats, CancellationToken.None).ToListAsync();
+						List<int> results = await index.SearchAsync(0, new SearchTerm(str), stats, CancellationToken.None).ToListAsync();
 						Assert.AreEqual(1, results.Count);
 						Assert.AreEqual(lineIdx, results[0]);
 
@@ -208,7 +208,7 @@ namespace EpicGames.Horde.Tests
 		static async Task SearchLogDataTestAsync(LogIndexNode index, string text, int firstLine, int count, int[] expectedLines)
 		{
 			SearchStats stats = new SearchStats();
-			List<int> lines = await index.Search(firstLine, new SearchTerm(text), stats, CancellationToken.None).Take(count).ToListAsync();
+			List<int> lines = await index.SearchAsync(firstLine, new SearchTerm(text), stats, CancellationToken.None).Take(count).ToListAsync();
 			Assert.IsTrue(lines.SequenceEqual(expectedLines));
 		}
 	}
diff --git a/Engine/Source/Programs/Shared/EpicGames.Horde.Tests/StorageCompressorTest.cs b/Engine/Source/Programs/Shared/EpicGames.Horde.Tests/StorageCompressorTest.cs
--- a/Engine/Source/Programs/Shared/EpicGames.Horde.Tests/StorageCompressorTest.cs
+++ b/Engine/Source/Programs/Shared/EpicGames.Horde.Tests/StorageCompressorTest.cs
@@ -41,7 +41,7 @@ namespace EpicGames.Horde.Tests
 	public class FakeStorageCompressorTest
 	{
 		[TestMethod]
-		public async Task CompressAndDecompress()
+		public async Task CompressAndDecompressAsync()
 		{
 			FakeStorageCompressor compressor = new ();
 			byte[] uncompressedData = { 0x41, 0x42, 0x43 };
diff --git a/Engine/Source/Programs/Shared/EpicGames.Horde.Tests/TransportTests.cs b/Engine/Source/Programs/Shared/EpicGames.Horde.Tests/TransportTests.cs
--- a/Engine/Source/Programs/Shared/EpicGames.Horde.Tests/TransportTests.cs
+++ b/Engine/Source/Programs/Shared/EpicGames.Horde.Tests/TransportTests.cs
@@ -16,7 +16,7 @@ namespace EpicGames.Horde.Tests
 	public class TransportTests
 	{
 		[TestMethod]
-		public async Task TestStreamTransport()
+		public async Task TestStreamTransportAsync()
 		{
 			byte[] input = new byte[256 * 1024];
 			new Random(0).NextBytes(input);
@@ -45,7 +45,7 @@ namespace EpicGames.Horde.Tests
 		}
 
 		[TestMethod]
-		public async Task TestAesTransport()
+		public async Task TestAesTransportAsync()
 		{
 			byte[] key = new byte[AesTransport.KeyLength];
 			new Random(1).NextBytes(key);
diff --git a/Engine/Source/Programs/Shared/EpicGames.Horde/Compute/AgentMessage.cs b/Engine/Source/Programs/Shared/EpicGames.Horde/Compute/AgentMessage.cs
--- a/Engine/Source/Programs/Shared/EpicGames.Horde/Compute/AgentMessage.cs
+++ b/Engine/Source/Programs/Shared/EpicGames.Horde/Compute/AgentMessage.cs
@@ -347,7 +347,7 @@ namespace EpicGames.Horde.Compute
 
 		#region Process
 
-		static async Task<AgentMessage> RunStorageServer(this AgentMessageChannel channel, BundleStorageClient storage, CancellationToken cancellationToken = default)
+		static async Task<AgentMessage> RunStorageServerAsync(this AgentMessageChannel channel, BundleStorageClient storage, CancellationToken cancellationToken = default)
 		{
 			for (; ; )
 			{
@@ -386,7 +386,7 @@ namespace EpicGames.Horde.Compute
 				request.Send();
 			}
 
-			using AgentMessage response = await RunStorageServer(channel, storage, cancellationToken);
+			using AgentMessage response = await RunStorageServerAsync(channel, storage, cancellationToken);
 			if (response.Type != AgentMessageType.WriteFilesResponse)
 			{
 				throw new InvalidAgentMessageException(response);
diff --git a/Engine/Source/Programs/Shared/EpicGames.Horde/Compute/AgentMessageHandler.cs b/Engine/Source/Programs/Shared/EpicGames.Horde/Compute/AgentMessageHandler.cs
--- a/Engine/Source/Programs/Shared/EpicGames.Horde/Compute/AgentMessageHandler.cs
+++ b/Engine/Source/Programs/Shared/EpicGames.Horde/Compute/AgentMessageHandler.cs
@@ -111,7 +111,7 @@ namespace EpicGames.Horde.Compute
 						case AgentMessageType.XorRequest:
 							{
 								XorRequestMessage xorRequest = message.AsXorRequest();
-								await RunXor(channel, xorRequest.Data, xorRequest.Value, cancellationToken);
+								await RunXorAsync(channel, xorRequest.Data, xorRequest.Value, cancellationToken);
 							}
 							break;
 						default:
@@ -121,7 +121,7 @@ namespace EpicGames.Horde.Compute
 			}
 		}
 
-		static async ValueTask RunXor(AgentMessageChannel channel, ReadOnlyMemory<byte> source, byte value, CancellationToken cancellationToken)
+		static async ValueTask RunXorAsync(AgentMessageChannel channel, ReadOnlyMemory<byte> source, byte value, CancellationToken cancellationToken)
 		{
 			using IAgentMessageBuilder response = await channel.CreateMessageAsync(AgentMessageType.XorResponse, source.Length, cancellationToken);
 			XorData(source.Span, response.GetSpanAndAdvance(source.Length), value);
diff --git a/Engine/Source/Programs/Shared/EpicGames.Horde/Logs/LogIndexNode.cs b/Engine/Source/Programs/Shared/EpicGames.Horde/Logs/LogIndexNode.cs
--- a/Engine/Source/Programs/Shared/EpicGames.Horde/Logs/LogIndexNode.cs
+++ b/Engine/Source/Programs/Shared/EpicGames.Horde/Logs/LogIndexNode.cs
@@ -180,7 +180,7 @@ namespace EpicGames.Horde.Logs
 		/// <param name="stats">Receives stats for the search</param>
 		/// <param name="cancellationToken">Cancellation token for the operation</param>
 		/// <returns>List of line numbers for the text</returns>
-		public async IAsyncEnumerable<int> Search(int firstLineIndex, SearchTerm text, SearchStats stats, [EnumeratorCancellation] CancellationToken cancellationToken)
+		public async IAsyncEnumerable<int> SearchAsync(int firstLineIndex, SearchTerm text, SearchStats stats, [EnumeratorCancellation] CancellationToken cancellationToken)
 		{
 			int lastBlockCount = 0;
 			foreach (int blockIdx in EnumeratePossibleChunks(text.Bytes, firstLineIndex))
diff --git a/Engine/Source/Programs/Shared/EpicGames.Perforce.Managed.Tests/ManagedWorkspaceTest.cs b/Engine/Source/Programs/Shared/EpicGames.Perforce.Managed.Tests/ManagedWorkspaceTest.cs
--- a/Engine/Source/Programs/Shared/EpicGames.Perforce.Managed.Tests/ManagedWorkspaceTest.cs
+++ b/Engine/Source/Programs/Shared/EpicGames.Perforce.Managed.Tests/ManagedWorkspaceTest.cs
@@ -32,10 +32,10 @@ public class ManagedWorkspaceTest : BasePerforceFixtureTest
 	[TestMethod]
 	[DataRow(true, DisplayName = "With have-table")]
 	[DataRow(false, DisplayName = "Without have-table")]
-	public async Task SyncSingleChangelist(bool useHaveTable)
+	public async Task SyncSingleChangelistAsync(bool useHaveTable)
 	{
-		ManagedWorkspace ws = await CreateManagedWorkspace(useHaveTable);
-		await AssertHaveTableFileCount(0);
+		ManagedWorkspace ws = await CreateManagedWorkspaceAsync(useHaveTable);
+		await AssertHaveTableFileCountAsync(0);
 
 		await SyncAsync(ws, 6);
 		Stream.GetChangelist(6).AssertDepotFiles(SyncDir);
@@ -43,9 +43,9 @@ public class ManagedWorkspaceTest : BasePerforceFixtureTest
 	}
 	
 	[TestMethod]
-	public async Task SyncBackwardsToOlderChangelistRemoveUntracked()
+	public async Task SyncBackwardsToOlderChangelistRemoveUntrackedAsync()
 	{
-		ManagedWorkspace ws = await CreateManagedWorkspace(true);
+		ManagedWorkspace ws = await CreateManagedWorkspaceAsync(true);
 
 		await SyncAsync(ws, 6, removeUntracked: false);
 		Stream.GetChangelist(6).AssertDepotFiles(SyncDir);
@@ -63,9 +63,9 @@ public class ManagedWorkspaceTest : BasePerforceFixtureTest
 	[TestMethod]
 	[DataRow(true, DisplayName = "With have-table")]
 	[DataRow(false, DisplayName = "Without have-table")]
-	public async Task SyncBackwardsToOlderChangelist(bool useHaveTable)
+	public async Task SyncBackwardsToOlderChangelistAsync(bool useHaveTable)
 	{
-		ManagedWorkspace ws = await CreateManagedWorkspace(useHaveTable);
+		ManagedWorkspace ws = await CreateManagedWorkspaceAsync(useHaveTable);
 
 		await SyncAsync(ws, 6);
 		Stream.GetChangelist(6).AssertDepotFiles(SyncDir);
@@ -84,9 +84,9 @@ public class ManagedWorkspaceTest : BasePerforceFixtureTest
 	[TestMethod]
 	[DataRow(true, DisplayName = "With have-table")]
 	[DataRow(false, DisplayName = "Without have-table")]
-	public async Task SyncUsingCacheFiles(bool useHaveTable)
+	public async Task SyncUsingCacheFilesAsync(bool useHaveTable)
 	{
-		ManagedWorkspace ws = await CreateManagedWorkspace(useHaveTable);
+		ManagedWorkspace ws = await CreateManagedWorkspaceAsync(useHaveTable);
 
 		FileReference GetCacheFilePath(int changeNumber)
 		{
@@ -113,9 +113,9 @@ public class ManagedWorkspaceTest : BasePerforceFixtureTest
 	[TestMethod]
 	[DataRow(true, DisplayName = "With have-table")]
 	[DataRow(false, DisplayName = "Without have-table")]
-	public async Task SyncWithViewExclusivePathFirst(bool useHaveTable)
+	public async Task SyncWithViewExclusivePathFirstAsync(bool useHaveTable)
 	{
-		ManagedWorkspace ws = await CreateManagedWorkspace(useHaveTable);
+		ManagedWorkspace ws = await CreateManagedWorkspaceAsync(useHaveTable);
 		ChangelistFixture cl = Stream.GetChangelist(6);
 
 		List<string> view = new() { "-/Data/...", "-/shared.h", };
@@ -133,9 +133,9 @@ public class ManagedWorkspaceTest : BasePerforceFixtureTest
 	[TestMethod]
 	[DataRow(true, DisplayName = "With have-table")]
 	[DataRow(false, DisplayName = "Without have-table")]
-	public async Task SyncWithViewInclusivePathFirst(bool useHaveTable)
+	public async Task SyncWithViewInclusivePathFirstAsync(bool useHaveTable)
 	{
-		ManagedWorkspace ws = await CreateManagedWorkspace(useHaveTable);
+		ManagedWorkspace ws = await CreateManagedWorkspaceAsync(useHaveTable);
 		ChangelistFixture cl = Stream.GetChangelist(6);
 
 		List<string> view = new() { "/...", "-/Data/..." };
@@ -152,9 +152,9 @@ public class ManagedWorkspaceTest : BasePerforceFixtureTest
 	[TestMethod]
 	[DataRow(true, DisplayName = "With have-table")]
 	[DataRow(false, DisplayName = "Without have-table")]
-	public async Task Populate(bool useHaveTable)
+	public async Task PopulateAsync(bool useHaveTable)
 	{
-		ManagedWorkspace ws = await CreateManagedWorkspace(useHaveTable);
+		ManagedWorkspace ws = await CreateManagedWorkspaceAsync(useHaveTable);
 		List<PopulateRequest> populateRequests = new () { new PopulateRequest(PerforceConnection, StreamName, new List<string>()) };
 		await ws.PopulateAsync(populateRequests, false, CancellationToken.None);
 		Stream.LatestChangelist.AssertDepotFiles(SyncDir);
@@ -164,9 +164,9 @@ public class ManagedWorkspaceTest : BasePerforceFixtureTest
 	[TestMethod]
 	[DataRow(true, DisplayName = "With have-table")]
 	[DataRow(false, DisplayName = "Without have-table")]
-	public async Task PopulateWithView(bool useHaveTable)
+	public async Task PopulateWithViewAsync(bool useHaveTable)
 	{
-		ManagedWorkspace ws = await CreateManagedWorkspace(useHaveTable);
+		ManagedWorkspace ws = await CreateManagedWorkspaceAsync(useHaveTable);
 		List<string> view = new () { "-/Data/...", "-/shared.h" };
 		ChangelistFixture cl = Stream.LatestChangelist;
 		List<DepotFileFixture> filtered = cl.StreamFiles
@@ -184,9 +184,9 @@ public class ManagedWorkspaceTest : BasePerforceFixtureTest
 	[TestMethod]
 	[DataRow(true, DisplayName = "With have-table")]
 	[DataRow(false, DisplayName = "Without have-table")]
-	public async Task Unshelve(bool useHaveTable)
+	public async Task UnshelveAsync(bool useHaveTable)
 	{
-		ManagedWorkspace ws = await CreateManagedWorkspace(useHaveTable);
+		ManagedWorkspace ws = await CreateManagedWorkspaceAsync(useHaveTable);
 		await SyncAsync(ws, 7);
 		await ws.UnshelveAsync(PerforceConnection, 8, CancellationToken.None);
 		Stream.GetChangelist(8).AssertDepotFiles(SyncDir);
@@ -198,10 +198,10 @@ public class ManagedWorkspaceTest : BasePerforceFixtureTest
 	[TestMethod]
 	[DataRow(true, DisplayName = "With have-table")]
 	[DataRow(false, DisplayName = "Without have-table")]
-	public async Task Caching_SyncingNewChangelist_UnusedFilesMovedToCache(bool useHaveTable)
+	public async Task Caching_SyncingNewChangelist_UnusedFilesMovedToCacheAsync(bool useHaveTable)
 	{
 		// Arrange
-		ManagedWorkspace ws = await CreateManagedWorkspace(useHaveTable);
+		ManagedWorkspace ws = await CreateManagedWorkspaceAsync(useHaveTable);
 		await SyncAsync(ws, 7);
 		
 		// Act
@@ -214,10 +214,10 @@ public class ManagedWorkspaceTest : BasePerforceFixtureTest
 	[TestMethod]
 	[DataRow(true, DisplayName = "With have-table")]
 	[DataRow(false, DisplayName = "Without have-table")]
-	public async Task Caching_SyncingWithCachedData_FilesPopulatedFromCache(bool useHaveTable)
+	public async Task Caching_SyncingWithCachedData_FilesPopulatedFromCacheAsync(bool useHaveTable)
 	{
 		// Arrange
-		ManagedWorkspace ws = await CreateManagedWorkspace(useHaveTable);
+		ManagedWorkspace ws = await CreateManagedWorkspaceAsync(useHaveTable);
 		await SyncAsync(ws, 7);
 		await SyncAsync(ws, 1);
 		
@@ -229,7 +229,7 @@ public class ManagedWorkspaceTest : BasePerforceFixtureTest
 	}
 
 	[TestMethod]
-	public async Task FixtureMetadataMatchesServerMetadata()
+	public async Task FixtureMetadataMatchesServerMetadataAsync()
 	{
 		FStatOptions options = FStatOptions.IncludeFileSizes;
 
@@ -272,10 +272,10 @@ public class ManagedWorkspaceTest : BasePerforceFixtureTest
 	}
 	
 	[TestMethod]
-	public async Task ReusePerforceClientWithoutHaveTable()
+	public async Task ReusePerforceClientWithoutHaveTableAsync()
 	{
 		// Sync with have-table as normal
-		ManagedWorkspace wsWithHave = await CreateManagedWorkspace(true);
+		ManagedWorkspace wsWithHave = await CreateManagedWorkspaceAsync(true);
 		await wsWithHave.SetupAsync(PerforceConnection, StreamName, CancellationToken.None);
 		ChangelistFixture cl = Stream.GetChangelist(6);
 		await SyncAsync(wsWithHave, cl.Number);
@@ -284,15 +284,15 @@ public class ManagedWorkspaceTest : BasePerforceFixtureTest
 
 		// Create a new ManagedWorkspace without have-table but re-use same Perforce connection
 		// The have-table remnants from above should not interfere with this workspace
-		ManagedWorkspace wsWithoutHave = await CreateManagedWorkspace(false);
+		ManagedWorkspace wsWithoutHave = await CreateManagedWorkspaceAsync(false);
 		await wsWithoutHave.SetupAsync(PerforceConnection, StreamName, CancellationToken.None);
 		await SyncAsync(wsWithoutHave, cl.Number);
-		await AssertHaveTableFileCount(0);
+		await AssertHaveTableFileCountAsync(0);
 		await cl.AssertHaveTableAsync(PerforceConnection, useHaveTable: false);
 		cl.AssertDepotFiles(SyncDir);
 	}
 
-	private async Task<ManagedWorkspace> CreateManagedWorkspace(bool useHaveTable)
+	private async Task<ManagedWorkspace> CreateManagedWorkspaceAsync(bool useHaveTable)
 	{
 		ManagedWorkspaceOptions options = new() { UseHaveTable = useHaveTable };
 		ManagedWorkspace ws = await ManagedWorkspace.CreateAsync(Environment.MachineName, TempDir, options, _mwLogger, CancellationToken.None);
@@ -305,7 +305,7 @@ public class ManagedWorkspaceTest : BasePerforceFixtureTest
 		await managedWorkspace.SyncAsync(PerforceConnection, StreamName, changeNumber, Array.Empty<string>(), removeUntracked, false, cacheFile, CancellationToken.None);
 	}
 	
-	private async Task AssertHaveTableFileCount(int expected)
+	private async Task AssertHaveTableFileCountAsync(int expected)
 	{
 		List<HaveRecord> haveRecords = await PerforceConnection.HaveAsync(new FileSpecList(), CancellationToken.None).ToListAsync();
 
diff --git a/Engine/Source/Programs/Shared/EpicGames.Perforce.Managed/ManagedWorkspace.cs b/Engine/Source/Programs/Shared/EpicGames.Perforce.Managed/ManagedWorkspace.cs
--- a/Engine/Source/Programs/Shared/EpicGames.Perforce.Managed/ManagedWorkspace.cs
+++ b/Engine/Source/Programs/Shared/EpicGames.Perforce.Managed/ManagedWorkspace.cs
@@ -502,7 +502,7 @@ namespace EpicGames.Perforce.Managed
 			{
 				Stopwatch timer = Stopwatch.StartNew();
 
-				(FileInfo[] filesToDelete, DirectoryInfo[] directoriesToDelete) refreshResult = await _workspace.Refresh(removeUntracked, _options.MaxFileConcurrency);
+				(FileInfo[] filesToDelete, DirectoryInfo[] directoriesToDelete) refreshResult = await _workspace.RefreshAsync(removeUntracked, _options.MaxFileConcurrency);
 				filesToDelete = refreshResult.filesToDelete;
 				directoriesToDelete = refreshResult.directoriesToDelete;
 
@@ -805,7 +805,7 @@ namespace EpicGames.Perforce.Managed
 		/// <summary>
 		/// Prints information about the repository state
 		/// </summary>
-		public async Task Status()
+		public async Task StatusAsync()
 		{
 			// Print size stats
 			_logger.LogInformation("Cache contains {NumFiles:n0} files, {TotalSize:n1}mb", _contentIdToTrackedFile.Count, _contentIdToTrackedFile.Values.Sum(x => x.Length) / (1024.0 * 1024.0));
@@ -1733,7 +1733,7 @@ namespace EpicGames.Perforce.Managed
 
 				// Handle the case where two machines may try to write to the cache file at once by writing to a temporary file
 				FileReference tempCacheFile = new FileReference(String.Format("{0}.{1}", cacheFile, Guid.NewGuid()));
-				await contents.Save(tempCacheFile, basePath);
+				await contents.SaveAsync(tempCacheFile, basePath);
 
 				// Try to move it into place
 				try
@@ -2014,7 +2014,7 @@ namespace EpicGames.Perforce.Managed
 							{
 								(int batchBeginIdx, int batchEndIdx) = batches[nextBatchIdx];
 
-								Task task = Task.Run(() => SyncBatch(client, filesToSync, batchBeginIdx, batchEndIdx, fakeSync, cancellationToken), cancellationToken);
+								Task task = Task.Run(() => SyncBatchAsync(client, filesToSync, batchBeginIdx, batchEndIdx, fakeSync, cancellationToken), cancellationToken);
 								tasks[task] = nextBatchIdx++;
 							}
 
@@ -2075,7 +2075,7 @@ namespace EpicGames.Perforce.Managed
 		/// <param name="fakeSync">Whether to fake a sync</param>
 		/// <param name="cancellationToken">Cancellation token for the request</param>
 		/// <returns>Async task</returns>
-		async Task SyncBatch(IPerforceConnection client, WorkspaceFileToSync[] filesToSync, int beginIdx, int endIdx, bool fakeSync, CancellationToken cancellationToken)
+		async Task SyncBatchAsync(IPerforceConnection client, WorkspaceFileToSync[] filesToSync, int beginIdx, int endIdx, bool fakeSync, CancellationToken cancellationToken)
 		{
 			if (fakeSync)
 			{
diff --git a/Engine/Source/Programs/Shared/EpicGames.Perforce.Managed/StreamSnapshotFromMemory.cs b/Engine/Source/Programs/Shared/EpicGames.Perforce.Managed/StreamSnapshotFromMemory.cs
--- a/Engine/Source/Programs/Shared/EpicGames.Perforce.Managed/StreamSnapshotFromMemory.cs
+++ b/Engine/Source/Programs/Shared/EpicGames.Perforce.Managed/StreamSnapshotFromMemory.cs
@@ -114,7 +114,7 @@ namespace EpicGames.Perforce.Managed
 		/// </summary>
 		/// <param name="outputFile">The output file to write to</param>
 		/// <param name="basePath"></param>
-		public async Task Save(FileReference outputFile, Utf8String basePath)
+		public async Task SaveAsync(FileReference outputFile, Utf8String basePath)
 		{
 			CbWriter writer = new CbWriter();
 			writer.BeginObject();
diff --git a/Engine/Source/Programs/Shared/EpicGames.Perforce.Managed/WorkspaceDirectoryInfo.cs b/Engine/Source/Programs/Shared/EpicGames.Perforce.Managed/WorkspaceDirectoryInfo.cs
--- a/Engine/Source/Programs/Shared/EpicGames.Perforce.Managed/WorkspaceDirectoryInfo.cs
+++ b/Engine/Source/Programs/Shared/EpicGames.Perforce.Managed/WorkspaceDirectoryInfo.cs
@@ -127,12 +127,12 @@ namespace EpicGames.Perforce.Managed
 		/// </summary>
 		/// <param name="removeUntracked">Whether to remove files that are not part of the stream</param>
 		/// <param name="numWorkers">Number of concurrent workers to use when refreshing</param>
-		public async Task<(FileInfo[] filesToDelete, DirectoryInfo[] directoriesToDelete)> Refresh(bool removeUntracked, int numWorkers)
+		public async Task<(FileInfo[] filesToDelete, DirectoryInfo[] directoriesToDelete)> RefreshAsync(bool removeUntracked, int numWorkers)
 		{
 			ConcurrentQueue<FileInfo> concurrentFilesToDelete = new ConcurrentQueue<FileInfo>();
 			ConcurrentQueue<DirectoryInfo> concurrentDirectoriesToDelete = new ConcurrentQueue<DirectoryInfo>();
 			using AsyncThreadPoolWorkQueue queue = new (numWorkers);
-			await queue.EnqueueAsync(_ => Refresh(new DirectoryInfo(GetFullName()), removeUntracked, concurrentFilesToDelete, concurrentDirectoriesToDelete, queue));
+			await queue.EnqueueAsync(_ => RefreshAsync(new DirectoryInfo(GetFullName()), removeUntracked, concurrentFilesToDelete, concurrentDirectoriesToDelete, queue));
 			await queue.ExecuteAsync();
 			return (concurrentFilesToDelete.ToArray(), concurrentDirectoriesToDelete.ToArray());
 		}
@@ -145,7 +145,7 @@ namespace EpicGames.Perforce.Managed
 		/// <param name="filesToDelete"></param>
 		/// <param name="directoriesToDelete"></param>
 		/// <param name="queue"></param>
-		async Task Refresh(DirectoryInfo info, bool removeUntracked, ConcurrentQueue<FileInfo> filesToDelete, ConcurrentQueue<DirectoryInfo> directoriesToDelete, AsyncThreadPoolWorkQueue queue)
+		async Task RefreshAsync(DirectoryInfo info, bool removeUntracked, ConcurrentQueue<FileInfo> filesToDelete, ConcurrentQueue<DirectoryInfo> directoriesToDelete, AsyncThreadPoolWorkQueue queue)
 		{
 			// Recurse through subdirectories
 			Dictionary<Utf8String, WorkspaceDirectoryInfo> newNameToSubDirectory = new Dictionary<Utf8String, WorkspaceDirectoryInfo>(NameToSubDirectory.Count, NameToSubDirectory.Comparer);
@@ -155,7 +155,7 @@ namespace EpicGames.Perforce.Managed
 				if (NameToSubDirectory.TryGetValue(subDirectoryInfo.Name, out subDirectory))
 				{
 					newNameToSubDirectory.Add(subDirectory.Name, subDirectory);
-					await queue.EnqueueAsync(_ => subDirectory.Refresh(subDirectoryInfo, removeUntracked, filesToDelete, directoriesToDelete, queue));
+					await queue.EnqueueAsync(_ => subDirectory.RefreshAsync(subDirectoryInfo, removeUntracked, filesToDelete, directoriesToDelete, queue));
 				}
 				else if (removeUntracked)
 				{
diff --git a/Engine/Source/Programs/Shared/EpicGames.Serialization.Tests/CbPackageTests.cs b/Engine/Source/Programs/Shared/EpicGames.Serialization.Tests/CbPackageTests.cs
--- a/Engine/Source/Programs/Shared/EpicGames.Serialization.Tests/CbPackageTests.cs
+++ b/Engine/Source/Programs/Shared/EpicGames.Serialization.Tests/CbPackageTests.cs
@@ -30,7 +30,7 @@ namespace EpicGames.Serialization.Tests
 	public class CbPackageTests
 	{
 		[TestMethod]
-		public async Task RoundTrip()
+		public async Task RoundTripAsync()
 		{
 			TestObject root = new TestObject()
 			{
@@ -56,7 +56,7 @@ namespace EpicGames.Serialization.Tests
 		}
 
 		[TestMethod]
-		public async Task RoundTripComplex()
+		public async Task RoundTripComplexAsync()
 		{
 			byte[] blob = Encoding.UTF8.GetBytes("This is blob contents");
 			IoHash blobHash = IoHash.Compute(blob);
@@ -112,7 +112,7 @@ namespace EpicGames.Serialization.Tests
 		}
 
 		[TestMethod]
-		public async Task ReadNoAttachments()
+		public async Task ReadNoAttachmentsAsync()
 		{
 			using CbPackageBuilder packageBuilder = new CbPackageBuilder();
 			SimpleObject simple = new SimpleObject();
