Horde: Refactoring auth/connection management for Horde API.

* HTTP client instances used to communicate with the Horde server are now provided as typed clients from IHttpClientFactory.
* Logic for getting and refreshing an access token is done via a delegating HttpMessageHandler, which fetches a new token whenever an unauthenticated response is received from the server.
* Moved request/response definitions for /api/v1/server endpoint calls into EpicGames.Horde.

diff --git a/Engine/Source/Programs/Horde/.editorconfig b/Engine/Source/Programs/Horde/.editorconfig
--- a/Engine/Source/Programs/Horde/.editorconfig
+++ b/Engine/Source/Programs/Horde/.editorconfig
@@ -561,8 +561,8 @@ dotnet_diagnostic.CA1815.severity = none
 # CA1816: Dispose methods should call SuppressFinalize
 dotnet_diagnostic.CA1816.severity = warning
 
-# CA1819: Properties should not return arrays
-dotnet_diagnostic.CA1819.severity = warning
+# CA1819: Properties should not return arrays (rationale: System.Text.Json serialization requires array properties)
+dotnet_diagnostic.CA1819.severity = silent
 
 # CA1820: Test for empty strings using string length
 dotnet_diagnostic.CA1820.severity = warning
diff --git a/Engine/Source/Programs/Horde/Horde.Agent/Commands/Compute/ComputeCommand.cs b/Engine/Source/Programs/Horde/Horde.Agent/Commands/Compute/ComputeCommand.cs
deleted file mode 100644
diff --git a/Engine/Source/Programs/Horde/Horde.Agent/Commands/Compute/ComputeWorkerCommand.cs b/Engine/Source/Programs/Horde/Horde.Agent/Commands/Compute/ComputeWorkerCommand.cs
new file mode 100644
--- /dev/null
+++ b/Engine/Source/Programs/Horde/Horde.Agent/Commands/Compute/ComputeWorkerCommand.cs
@@ -0,0 +1,56 @@
+// Copyright Epic Games, Inc. All Rights Reserved.
+
+using System.Net;
+using System.Net.Sockets;
+using EpicGames.Core;
+using EpicGames.Horde.Compute;
+using EpicGames.Horde.Compute.Transports;
+using EpicGames.Horde.Storage;
+using Microsoft.Extensions.Logging;
+
+namespace Horde.Agent.Commands.Compute
+{
+	/// <summary>
+	/// Helper command for hosting a local compute worker in a separate process
+	/// </summary>
+	[Command("computeworker", "Runs the agent as a local compute host, accepting incoming connections on the loopback adapter with a given port")]
+	class ComputeWorkerCommand : Command
+	{
+		readonly StorageCache _storageCache;
+
+		[CommandLine("-Port=")]
+		int Port { get; set; } = 2000;
+
+		public ComputeWorkerCommand(StorageCache storageCache)
+		{
+			_storageCache = storageCache;
+		}
+
+		public override async Task<int> ExecuteAsync(ILogger logger)
+		{
+			logger.LogInformation("** WORKER **");
+
+			using Socket tcpSocket = new Socket(SocketType.Stream, ProtocolType.IP);
+			await tcpSocket.ConnectAsync(IPAddress.Loopback, Port);
+
+			await using (RemoteComputeSocket socket = new RemoteComputeSocket(new TcpTransport(tcpSocket), logger))
+			{
+				logger.LogInformation("Running worker...");
+				await RunWorkerAsync(socket, _storageCache, logger, CancellationToken.None);
+				logger.LogInformation("Worker complete");
+				await socket.CloseAsync(CancellationToken.None);
+			}
+
+			logger.LogInformation("Stopping");
+			return 0;
+		}
+
+		public static async Task RunWorkerAsync(ComputeSocket socket, StorageCache storageCache, ILogger logger, CancellationToken cancellationToken)
+		{
+			DirectoryReference sandboxDir = DirectoryReference.Combine(Program.DataDir, "Sandbox");
+
+			AgentMessageHandler worker = new AgentMessageHandler(sandboxDir, storageCache, null, false, null, logger);
+			await worker.RunAsync(socket, cancellationToken);
+		}
+	}
+}
diff --git a/Engine/Source/Programs/Horde/Horde.Agent/Commands/Compute/CppComputeCommand.cs b/Engine/Source/Programs/Horde/Horde.Agent/Commands/Compute/CppComputeCommand.cs
deleted file mode 100644
diff --git a/Engine/Source/Programs/Horde/Horde.Agent/Commands/Compute/DdcComputeCommand.cs b/Engine/Source/Programs/Horde/Horde.Agent/Commands/Compute/DdcComputeCommand.cs
deleted file mode 100644
diff --git a/Engine/Source/Programs/Horde/Horde.Agent/Commands/Compute/XorComputeCommand.cs b/Engine/Source/Programs/Horde/Horde.Agent/Commands/Compute/XorComputeCommand.cs
deleted file mode 100644
diff --git a/Engine/Source/Programs/Horde/Horde.Server/Server/AccountController.cs b/Engine/Source/Programs/Horde/Horde.Server/Server/AccountController.cs
--- a/Engine/Source/Programs/Horde/Horde.Server/Server/AccountController.cs
+++ b/Engine/Source/Programs/Horde/Horde.Server/Server/AccountController.cs
@@ -5,6 +5,7 @@ using System.Globalization;
 using System.Net;
 using System.Text;
 using System.Threading.Tasks;
+using EpicGames.Horde.Api;
 using Horde.Server.Acls;
 using Horde.Server.Authentication;
 using Microsoft.AspNetCore.Authentication;
diff --git a/Engine/Source/Programs/Horde/Horde.Server/Server/ServerController.cs b/Engine/Source/Programs/Horde/Horde.Server/Server/ServerController.cs
--- a/Engine/Source/Programs/Horde/Horde.Server/Server/ServerController.cs
+++ b/Engine/Source/Programs/Horde/Horde.Server/Server/ServerController.cs
@@ -4,9 +4,11 @@ using System;
 using System.Collections.Generic;
 using System.Diagnostics;
 using System.Reflection;
+using System.Runtime.InteropServices;
 using System.Threading;
 using System.Threading.Tasks;
 using EpicGames.Core;
+using EpicGames.Horde.Api;
 using EpicGames.Perforce;
 using Horde.Server.Agents;
 using Horde.Server.Configuration;
@@ -67,7 +69,11 @@ namespace Horde.Server.Server
 		[ProducesResponseType(typeof(GetServerInfoResponse), 200)]
 		public async Task<ActionResult<GetServerInfoResponse>> GetServerInfoAsync()
 		{
-			string? agentVersion = null;
+			GetServerInfoResponse response = new GetServerInfoResponse();
+
+			FileVersionInfo versionInfo = FileVersionInfo.GetVersionInfo(Assembly.GetExecutingAssembly().Location);
+			response.ServerVersion = versionInfo.ProductVersion ?? String.Empty;
+			response.OsDescription = RuntimeInformation.OSDescription;
 
 			ITool? tool = await _toolCollection.GetAsync(AgentExtensions.DefaultAgentSoftwareToolId, _globalConfig.Value);
 			if (tool != null)
@@ -75,11 +81,11 @@ namespace Horde.Server.Server
 				IToolDeployment? deployment = tool.GetCurrentDeployment(1.0, _clock.UtcNow);
 				if (deployment != null)
 				{
-					agentVersion = deployment.Version;
+					response.AgentVersion = deployment.Version;
 				}
 			}
 
-			return new GetServerInfoResponse(agentVersion);
+			return response;
 		}
 
 		/// <summary>
@@ -121,7 +127,14 @@ namespace Horde.Server.Server
 		[Route("/api/v1/server/auth")]
 		public ActionResult<GetAuthConfigResponse> GetAuthConfig()
 		{
-			return new GetAuthConfigResponse(_globalConfig.Value.ServerSettings);
+			ServerSettings settings = _globalConfig.Value.ServerSettings;
+
+			GetAuthConfigResponse response = new GetAuthConfigResponse();
+			response.Method = settings.AuthMethod;
+			response.ServerUrl = settings.OidcAuthority;
+			response.ClientId = settings.OidcClientId;
+			response.LocalRedirectUrls = settings.OidcLocalRedirectUrls;
+			return response;
 		}
 
 		/// <summary>
diff --git a/Engine/Source/Programs/Horde/Horde.Server/Server/ServerMessages.cs b/Engine/Source/Programs/Horde/Horde.Server/Server/ServerMessages.cs
deleted file mode 100644
diff --git a/Engine/Source/Programs/Horde/Horde.Server/ServerSettings.cs b/Engine/Source/Programs/Horde/Horde.Server/ServerSettings.cs
--- a/Engine/Source/Programs/Horde/Horde.Server/ServerSettings.cs
+++ b/Engine/Source/Programs/Horde/Horde.Server/ServerSettings.cs
@@ -8,6 +8,7 @@ using System.Security.Claims;
 using System.Text.Json.Serialization;
 using EpicGames.Core;
 using EpicGames.Horde;
+using EpicGames.Horde.Api;
 using EpicGames.Horde.Storage;
 using EpicGames.Horde.Storage.Bundles;
 using EpicGames.Horde.Storage.Nodes;
@@ -132,27 +133,6 @@ namespace Horde.Server
 		public ChunkingOptions Chunking { get; set; } = new ChunkingOptions();
 	}
 
-	/// <summary>
-	/// Authentication method used for logging users in
-	/// </summary>
-	public enum AuthMethod
-	{
-		/// <summary>
-		/// No authentication enabled, mainly for demo and testing purposes
-		/// </summary>
-		Anonymous,
-
-		/// <summary>
-		/// OpenID Connect authentication, tailored for Okta
-		/// </summary>
-		Okta,
-		
-		/// <summary>
-		/// Generic OpenID Connect authentication, recommended for most
-		/// </summary>
-		OpenIdConnect,
-	}
-
 	/// <summary>
 	/// Type of run mode this process should use. Each carry different types of workloads. 
 	/// More than one mode can be active. But not all modes are not guaranteed to be compatible with each other and will
diff --git a/Engine/Source/Programs/Horde/Horde/CmdConfig.cs b/Engine/Source/Programs/Horde/Horde/CmdConfig.cs
--- a/Engine/Source/Programs/Horde/Horde/CmdConfig.cs
+++ b/Engine/Source/Programs/Horde/Horde/CmdConfig.cs
@@ -2,10 +2,6 @@
 
 using System.Text.Json;
 using EpicGames.Core;
-using EpicGames.Horde;
-using EpicGames.OIDC;
-using Microsoft.Extensions.Configuration;
-using Microsoft.Extensions.Logging;
 
 namespace Horde
 {
@@ -104,101 +100,4 @@ namespace Horde
 		/// </summary>
 		public long? PacketCacheSize { get; set; }
 	}
-
-	/// <summary>
-	/// Extension methods for <see cref="CmdConfig"/>
-	/// </summary>
-	static class CmdConfigExtensions
-	{
-		class GetAuthConfigResponse
-		{
-			public string? Method { get; set; }
-			public string? ServerUrl { get; set; }
-			public string? ClientId { get; set; }
-			public string[]? LocalRedirectUrls { get; set; }
-		}
-
-		/// <summary>
-		/// Get an access token for the server specified in a config instance
-		/// </summary>
-		public static async Task<string?> GetAccessTokenAsync(this CmdConfig config, ILogger logger, CancellationToken cancellationToken = default)
-		{
-			if (config.Server == null)
-			{
-				throw new InvalidOperationException("No Horde server is configured. Run 'horde login -server=...' to configure.");
-			}
-
-			logger.LogInformation("Getting access token for {Server}", config.Server);
-
-			GetAuthConfigResponse authConfig;
-			using (HttpClient httpClient = new HttpClient())
-			{
-				Uri uri = new Uri(config.Server, "api/v1/server/auth");
-				authConfig = await httpClient.GetAsync<GetAuthConfigResponse>(uri, cancellationToken);
-			}
-
-			if (String.Equals(authConfig.Method, "Anonymous", StringComparison.OrdinalIgnoreCase))
-			{
-				return null;
-			}
-
-			const string OidcProvider = "Horde";
-
-			Dictionary<string, string?> values = new Dictionary<string, string?>();
-			values[$"Providers:{OidcProvider}:DisplayName"] = "Horde";
-			values[$"Providers:{OidcProvider}:ServerUri"] = authConfig.ServerUrl;
-			values[$"Providers:{OidcProvider}:ClientId"] = authConfig.ClientId;
-			values[$"Providers:{OidcProvider}:RedirectUri"] = authConfig.LocalRedirectUrls?.FirstOrDefault();
-
-			ConfigurationBuilder builder = new ConfigurationBuilder();
-			builder.AddInMemoryCollection(values);
-
-			IConfiguration configuration = builder.Build();
-
-			using ITokenStore tokenStore = TokenStoreFactory.CreateTokenStore();
-			OidcTokenManager oidcTokenManager = OidcTokenManager.CreateTokenManager(configuration, tokenStore, new List<string>() { OidcProvider });
-
-			OidcTokenInfo? result = null;
-			try
-			{
-				result = await oidcTokenManager.TryGetAccessToken(OidcProvider, cancellationToken);
-			}
-			catch (Exception ex)
-			{
-				logger.LogTrace(ex, "Unable to get access token; attempting login: {Message}", ex.Message);
-			}
-			if (result == null)
-			{
-				result = await oidcTokenManager.Login(OidcProvider, cancellationToken);
-			}
-
-			if (result.AccessToken == null)
-			{
-				throw new Exception($"Unable to get access token for {config.Server}");
-			}
-
-			logger.LogInformation("Received bearer token for {Server}", config.Server);
-			return result.AccessToken;
-		}
-
-		/// <summary>
-		/// Creates an HTTP client for connecting to the configured Horde server
-		/// </summary>
-		public static async Task<HordeHttpClient> GetHttpClientAsync(this CmdConfig config, ILogger logger, CancellationToken cancellationToken = default)
-		{
-			Uri? server = config.Server;
-			if (server == null)
-			{
-				throw new Exception("No server is configured. Run 'horde login -server=...' to set up.");
-			}
-
-			string? token = await config.GetAccessTokenAsync(logger, cancellationToken);
-			if (token == null)
-			{
-				throw new Exception("Unable to log in to server.");
-			}
-
-			return new HordeHttpClient(server, token, logger);
-		}
-	}
 }
diff --git a/Engine/Source/Programs/Horde/Horde/Commands/Artifacts/ArtifactFind.cs b/Engine/Source/Programs/Horde/Horde/Commands/Artifacts/ArtifactFind.cs
--- a/Engine/Source/Programs/Horde/Horde/Commands/Artifacts/ArtifactFind.cs
+++ b/Engine/Source/Programs/Horde/Horde/Commands/Artifacts/ArtifactFind.cs
@@ -4,7 +4,6 @@ using EpicGames.Core;
 using EpicGames.Horde;
 using EpicGames.Horde.Api;
 using Microsoft.Extensions.Logging;
-using Microsoft.Extensions.Options;
 
 namespace Horde.Commands.Bundles
 {
@@ -14,18 +13,16 @@ namespace Horde.Commands.Bundles
 		[CommandLine("-Key=")]
 		public List<string> Keys { get; } = new List<string>();
 
-		readonly CmdConfig _config;
+		readonly HordeHttpClient _hordeHttpClient;
 
-		public ArtifactFind(IOptions<CmdConfig> config)
+		public ArtifactFind(HordeHttpClient httpClient)
 		{
-			_config = config.Value;
+			_hordeHttpClient = httpClient;
 		}
 
 		public override async Task<int> ExecuteAsync(ILogger logger)
 		{
-			using HordeHttpClient horde = await _config.GetHttpClientAsync(logger);
-
-			List<GetArtifactResponse> artifacts = await horde.FindArtifactsAsync(keys: Keys);
+			List<GetArtifactResponse> artifacts = await _hordeHttpClient.FindArtifactsAsync(keys: Keys);
 			foreach (GetArtifactResponse artifact in artifacts)
 			{
 				logger.LogInformation("");
diff --git a/Engine/Source/Programs/Horde/Horde/Commands/Bundles/BundleCreate.cs b/Engine/Source/Programs/Horde/Horde/Commands/Bundles/BundleCreate.cs
--- a/Engine/Source/Programs/Horde/Horde/Commands/Bundles/BundleCreate.cs
+++ b/Engine/Source/Programs/Horde/Horde/Commands/Bundles/BundleCreate.cs
@@ -2,6 +2,7 @@
 
 using System.Diagnostics;
 using EpicGames.Core;
+using EpicGames.Horde;
 using EpicGames.Horde.Storage;
 using EpicGames.Horde.Storage.Clients;
 using EpicGames.Horde.Storage.Nodes;
@@ -25,8 +26,8 @@ namespace Horde.Commands.Bundles
 		[CommandLine("-Filter=", Description = "Filter for files to include, in P4 syntax (eg. Foo/...).")]
 		public string Filter { get; set; } = "...";
 
-		public BundleCreate(StorageCache storageCache, IOptions<CmdConfig> config)
-			: base(storageCache, config)
+		public BundleCreate(HordeHttpClientFactory httpClientFactory, StorageCache storageCache, IOptions<CmdConfig> config)
+			: base(httpClientFactory, storageCache, config)
 		{
 		}
 
diff --git a/Engine/Source/Programs/Horde/Horde/Commands/Bundles/BundleExtract.cs b/Engine/Source/Programs/Horde/Horde/Commands/Bundles/BundleExtract.cs
--- a/Engine/Source/Programs/Horde/Horde/Commands/Bundles/BundleExtract.cs
+++ b/Engine/Source/Programs/Horde/Horde/Commands/Bundles/BundleExtract.cs
@@ -2,6 +2,7 @@
 
 using System.Diagnostics;
 using EpicGames.Core;
+using EpicGames.Horde;
 using EpicGames.Horde.Storage;
 using EpicGames.Horde.Storage.Bundles;
 using EpicGames.Horde.Storage.Clients;
@@ -29,8 +30,8 @@ namespace Horde.Commands.Bundles
 		[CommandLine("-OutputDir=", Required = true)]
 		public DirectoryReference OutputDir { get; set; } = null!;
 
-		public BundleExtract(StorageCache storageCache, IOptions<CmdConfig> config)
-			: base(storageCache, config)
+		public BundleExtract(HordeHttpClientFactory httpClientFactory, StorageCache storageCache, IOptions<CmdConfig> config)
+			: base(httpClientFactory, storageCache, config)
 		{
 		}
 
@@ -50,7 +51,7 @@ namespace Horde.Commands.Bundles
 			}
 			else if (Node != null)
 			{
-				BundleStorageClient store = (BundleStorageClient)await CreateStorageClientAsync(logger);
+				using BundleStorageClient store = (BundleStorageClient)await CreateStorageClientAsync(logger);
 				BlobHandle handle = store.CreateNodeHandle(BundleNodeLocator.Parse(Node));
 				await ExecuteInternalAsync(store, handle, logger);
 			}
diff --git a/Engine/Source/Programs/Horde/Horde/Commands/Compute/ComputeRun.cs b/Engine/Source/Programs/Horde/Horde/Commands/Compute/ComputeRun.cs
--- a/Engine/Source/Programs/Horde/Horde/Commands/Compute/ComputeRun.cs
+++ b/Engine/Source/Programs/Horde/Horde/Commands/Compute/ComputeRun.cs
@@ -1,9 +1,9 @@
 // Copyright Epic Games, Inc. All Rights Reserved.
 
-using System.Net.Http.Headers;
 using System.Reflection;
 using System.Text.Json;
 using EpicGames.Core;
+using EpicGames.Horde;
 using EpicGames.Horde.Common;
 using EpicGames.Horde.Compute;
 using EpicGames.Horde.Compute.Clients;
@@ -11,7 +11,6 @@ using EpicGames.Horde.Storage.Bundles;
 using EpicGames.Horde.Storage.Clients;
 using EpicGames.Horde.Storage.Nodes;
 using Microsoft.Extensions.Logging;
-using Microsoft.Extensions.Options;
 
 namespace Horde.Commands.Compute
 {
@@ -51,20 +50,20 @@ namespace Horde.Commands.Compute
 		[CommandLine("-Task=", Required = true)]
 		FileReference TaskFile { get; set; } = null!;
 
-		readonly CmdConfig _config;
+		readonly HordeHttpClientFactory _hordeHttpClientFactory;
 
 		/// <summary>
 		/// Constructor
 		/// </summary>
-		public ComputeRun(IOptions<CmdConfig> config)
+		public ComputeRun(HordeHttpClientFactory hordeHttpClientFactory)
 		{
-			_config = config.Value;
+			_hordeHttpClientFactory = hordeHttpClientFactory;
 		}
 
 		/// <inheritdoc/>
 		public override async Task<int> ExecuteAsync(ILogger logger)
 		{
-			await using IComputeClient client = await CreateClientAsync(logger);
+			await using IComputeClient client = CreateComputeClient(logger);
 
 			Requirements? requirements = null;
 			if (Requirements != null)
@@ -82,7 +81,7 @@ namespace Horde.Commands.Compute
 			return result ? 0 : 1;
 		}
 
-		async Task<IComputeClient> CreateClientAsync(ILogger logger)
+		IComputeClient CreateComputeClient(ILogger logger)
 		{
 			if (Local)
 			{
@@ -94,8 +93,7 @@ namespace Horde.Commands.Compute
 			}
 			else
 			{
-				string? accessToken = await _config.GetAccessTokenAsync(logger);
-				return new ServerComputeClient(_config.Server, new AuthenticationHeaderValue("Bearer", accessToken), logger);
+				return new ServerComputeClient(_hordeHttpClientFactory, logger);
 			}
 		}
 
diff --git a/Engine/Source/Programs/Horde/Horde/Commands/LoginCommand.cs b/Engine/Source/Programs/Horde/Horde/Commands/LoginCommand.cs
--- a/Engine/Source/Programs/Horde/Horde/Commands/LoginCommand.cs
+++ b/Engine/Source/Programs/Horde/Horde/Commands/LoginCommand.cs
@@ -1,6 +1,9 @@
 // Copyright Epic Games, Inc. All Rights Reserved.
 
 using EpicGames.Core;
+using EpicGames.Horde;
+using EpicGames.Horde.Api;
+using Microsoft.Extensions.DependencyInjection;
 using Microsoft.Extensions.Logging;
 using Microsoft.Extensions.Options;
 
@@ -21,10 +24,12 @@ namespace Horde.Commands
 		public string? Server { get; set; }
 
 		readonly CmdConfig _config;
+		readonly IServiceProvider _serviceProvider;
 
-		public LoginCommand(IOptions<CmdConfig> config)
+		public LoginCommand(IOptions<CmdConfig> config, IServiceProvider serviceProvider)
 		{
 			_config = config.Value;
+			_serviceProvider = serviceProvider;
 		}
 
 		/// <inheritdoc/>
@@ -36,12 +41,10 @@ namespace Horde.Commands
 				await _config.WriteAsync();
 			}
 
-			if (await _config.GetAccessTokenAsync(logger) == null)
-			{
-				logger.LogError("Unable to log in to server");
-				return 1;
-			}
-
+			HordeHttpClient httpClient = _serviceProvider.GetRequiredService<HordeHttpClient>();
+			GetServerInfoResponse serverInfo = await httpClient.GetAsync<GetServerInfoResponse>("api/v1/server/info");
+			logger.LogInformation("Connected to server version: {Version}", serverInfo.ServerVersion);
+			
 			return 0;
 		}
 	}
diff --git a/Engine/Source/Programs/Horde/Horde/Commands/StorageCommandBase.cs b/Engine/Source/Programs/Horde/Horde/Commands/StorageCommandBase.cs
--- a/Engine/Source/Programs/Horde/Horde/Commands/StorageCommandBase.cs
+++ b/Engine/Source/Programs/Horde/Horde/Commands/StorageCommandBase.cs
@@ -3,6 +3,7 @@
 using EpicGames.Core;
 using EpicGames.Horde;
 using EpicGames.Horde.Storage;
+using EpicGames.Horde.Storage.Clients;
 using Microsoft.Extensions.Logging;
 using Microsoft.Extensions.Options;
 
@@ -11,10 +12,8 @@ namespace Horde.Commands
 	/// <summary>
 	/// Base class for commands that require a configured storage client
 	/// </summary>
-	abstract class StorageCommandBase : Command, IDisposable
+	abstract class StorageCommandBase : Command
 	{
-		HordeHttpClient? _hordeHttpClient;
-
 		/// <summary>
 		/// Namespace to use
 		/// </summary>
@@ -37,37 +36,54 @@ namespace Horde.Commands
 		/// </summary>
 		public CmdConfig Config { get; }
 
+		readonly HordeHttpClientFactory _httpClientFactory;
+
 		/// <summary>
 		/// Constructor
 		/// </summary>
-		public StorageCommandBase(StorageCache storageCache, IOptions<CmdConfig> config)
+		public StorageCommandBase(HordeHttpClientFactory httpClientFactory, StorageCache storageCache, IOptions<CmdConfig> config)
 		{
+			_httpClientFactory = httpClientFactory;
+
 			StorageCache = storageCache;
 			Config = config.Value;
 		}
 
-		/// <inheritdoc/>
-		public void Dispose()
-		{
-			_hordeHttpClient?.Dispose();
-		}
-
 		/// <summary>
 		/// Creates a new client instance
 		/// </summary>
 		/// <param name="logger">Logger for output messages</param>
 		/// <param name="cancellationToken"></param>
-		public async Task<IStorageClient> CreateStorageClientAsync(ILogger logger, CancellationToken cancellationToken = default)
+		public Task<IStorageClient> CreateStorageClientAsync(ILogger logger, CancellationToken cancellationToken = default)
 		{
-			_hordeHttpClient ??= await Config.GetHttpClientAsync(logger, cancellationToken);
-			if (String.IsNullOrEmpty(Path))
+			_ = cancellationToken;
+
+			string? path = Path;
+			if (String.IsNullOrEmpty(path))
 			{
-				return _hordeHttpClient.CreateStorageClient(Namespace, StorageCache);
+				path = $"api/v1/storage/{Namespace}/";
 			}
-			else
+			else if (!path.EndsWith("/", StringComparison.Ordinal))
 			{
-				return _hordeHttpClient.CreateStorageClient(Path, StorageCache);
+				path += "/";
 			}
+
+			HttpClient CreateClient()
+			{
+				HttpClient client = _httpClientFactory.CreateClient().HttpClient;
+				client.BaseAddress = new Uri(client.BaseAddress!, path);
+				return client;
+			}
+
+			HttpClient CreateRedirectClient()
+			{
+				HttpClient client = _httpClientFactory.CreateClient().HttpClient;
+				client.BaseAddress = null;
+				client.DefaultRequestHeaders.Authorization = null;
+				return client;
+			}
+
+			return Task.FromResult<IStorageClient>(new HttpStorageClient(CreateClient, CreateRedirectClient, StorageCache, logger));
 		}
 	}
 }
diff --git a/Engine/Source/Programs/Horde/Horde/Commands/Workspace/WorkspaceSync.cs b/Engine/Source/Programs/Horde/Horde/Commands/Workspace/WorkspaceSync.cs
--- a/Engine/Source/Programs/Horde/Horde/Commands/Workspace/WorkspaceSync.cs
+++ b/Engine/Source/Programs/Horde/Horde/Commands/Workspace/WorkspaceSync.cs
@@ -2,6 +2,7 @@
 
 using System.Diagnostics;
 using EpicGames.Core;
+using EpicGames.Horde;
 using EpicGames.Horde.Storage;
 using EpicGames.Horde.Storage.Bundles;
 using EpicGames.Horde.Storage.Clients;
@@ -34,8 +35,8 @@ namespace Horde.Commands.Workspace
 		[CommandLine("-Stats")]
 		public bool Stats { get; set; }
 
-		public WorkspaceSync(StorageCache storageCache, IOptions<CmdConfig> config)
-			: base(storageCache, config)
+		public WorkspaceSync(HordeHttpClientFactory httpClientFactory, StorageCache storageCache, IOptions<CmdConfig> config)
+			: base(httpClientFactory, storageCache, config)
 		{
 		}
 
diff --git a/Engine/Source/Programs/Horde/Horde/Program.cs b/Engine/Source/Programs/Horde/Horde/Program.cs
--- a/Engine/Source/Programs/Horde/Horde/Program.cs
+++ b/Engine/Source/Programs/Horde/Horde/Program.cs
@@ -11,6 +11,7 @@ using Serilog;
 using Serilog.Formatting.Json;
 using EpicGames.Horde.Storage;
 using Microsoft.Extensions.Options;
+using EpicGames.Horde;
 
 namespace Horde
 {
@@ -34,9 +35,10 @@ namespace Horde
 			services.AddSingleton(loggerFactory);
 			services.AddLogging();
 			services.AddMemoryCache();
+			services.AddSingleton(sp => Options.Create(CmdConfig.Read()));
+			services.AddHordeHttpClient((sp, options) => options.ServerUrl = sp.GetRequiredService<IOptions<CmdConfig>>().Value.Server);
 			services.AddSingleton<StorageCache>(CreateStorageClientCache);
 			services.AddSingleton<StorageBackendCache>(CreateStorageBackendCache);
-			services.AddSingleton(sp => Options.Create(CmdConfig.Read()));
 
 			// Execute all the commands
 			await using ServiceProvider serviceProvider = services.BuildServiceProvider();
diff --git a/Engine/Source/Programs/Horde/Samples/RemoteClient/RemoteClientApp.cs b/Engine/Source/Programs/Horde/Samples/RemoteClient/RemoteClientApp.cs
--- a/Engine/Source/Programs/Horde/Samples/RemoteClient/RemoteClientApp.cs
+++ b/Engine/Source/Programs/Horde/Samples/RemoteClient/RemoteClientApp.cs
@@ -1,17 +1,16 @@
 // Copyright Epic Games, Inc. All Rights Reserved.
 
 using System.Buffers.Binary;
-using System.Net.Http.Headers;
 using System.Reflection;
 using EpicGames.Core;
+using EpicGames.Horde;
 using EpicGames.Horde.Common;
 using EpicGames.Horde.Compute;
 using EpicGames.Horde.Compute.Clients;
 using EpicGames.Horde.Storage.Bundles;
 using EpicGames.Horde.Storage.Clients;
 using EpicGames.Horde.Storage.Nodes;
-using EpicGames.OIDC;
-using Microsoft.Extensions.Configuration;
+using Microsoft.Extensions.DependencyInjection;
 using Microsoft.Extensions.Logging;
 
 namespace RemoteClient
@@ -49,8 +48,22 @@ namespace RemoteClient
 			arguments.ApplyTo(options);
 			arguments.CheckAllArgumentsUsed(logger);
 
+			// Create a DI container that can create and authenticate Horde HTTP clients for us
+			ServiceCollection services = new ServiceCollection();
+			if (options.Server == null)
+			{
+				DirectoryReference sandboxDir = DirectoryReference.Combine(new FileReference(Assembly.GetExecutingAssembly().Location).Directory, "Sandbox");
+				services.AddSingleton<IComputeClient>(sp => new LocalComputeClient(2000, sandboxDir, options.InProc, new PrefixLogger("[REMOTE]", logger)));
+			}
+			else
+			{
+				services.AddHordeHttpClient(x => x.ServerUrl = new Uri(options.Server));
+				services.AddSingleton<IComputeClient, ServerComputeClient>();
+			}
+
 			// Create the client to handle our requests
-			await using IComputeClient client = await CreateClientAsync(options, logger);
+			await using ServiceProvider serviceProvider = services.BuildServiceProvider();
+			IComputeClient client = serviceProvider.GetRequiredService<IComputeClient>();
 
 			// Allocate a worker
 			Requirements? requirements = null;
@@ -59,7 +72,7 @@ namespace RemoteClient
 				requirements = new Requirements(Condition.Parse(options.Condition));
 			}
 
-			await using IComputeLease? lease = await client.TryAssignWorkerAsync(new ClusterId("default"), requirements, null, logger);
+			await using IComputeLease? lease = await client.TryAssignWorkerAsync(new ClusterId("default"), requirements, null, new PrefixLogger("[CLIENT]", logger));
 			if (lease == null)
 			{
 				logger.LogInformation("Unable to connect to remote");
@@ -161,66 +174,5 @@ namespace RemoteClient
 
 			await channel.CloseAsync();
 		}
-
-		static async Task<IComputeClient> CreateClientAsync(ClientAppOptions options, ILogger logger)
-		{
-			if (options.Server == null)
-			{
-				DirectoryReference sandboxDir = DirectoryReference.Combine(new FileReference(Assembly.GetExecutingAssembly().Location).Directory, "Sandbox");
-				return new LocalComputeClient(2000, sandboxDir, options.InProc, logger);
-			}
-			else
-			{
-				AuthenticationHeaderValue? authHeader = await GetAuthHeaderAsync(options, logger);
-				return new ServerComputeClient(new Uri(options.Server), authHeader, logger);
-			}
-		}
-
-		static async Task<AuthenticationHeaderValue?> GetAuthHeaderAsync(ClientAppOptions options, ILogger logger)
-		{
-			if (options.OidcProvider == null)
-			{
-				return null;
-			}
-
-			for (DirectoryReference? currentDir = CurrentAssemblyFile.Directory; currentDir != null; currentDir = currentDir.ParentDirectory)
-			{
-				FileReference buildVersionFile = FileReference.Combine(currentDir, "Build/Build.version");
-				if (FileReference.Exists(buildVersionFile))
-				{
-					string bearerToken = await GetOidcBearerTokenAsync(currentDir, null, options.OidcProvider, logger);
-					return new AuthenticationHeaderValue("Bearer", bearerToken);
-				}
-			}
-
-			throw new Exception($"Unable to find engine directory above {CurrentAssemblyFile}");
-		}
-
-		static async Task<string> GetOidcBearerTokenAsync(DirectoryReference engineDir, DirectoryReference? projectDir, string oidcProvider, ILogger logger)
-		{
-			logger.LogInformation("Performing OIDC token refresh...");
-
-			using ITokenStore tokenStore = TokenStoreFactory.CreateTokenStore();
-			IConfiguration providerConfiguration = ProviderConfigurationFactory.ReadConfiguration(engineDir.ToDirectoryInfo(), projectDir?.ToDirectoryInfo());
-			OidcTokenManager oidcTokenManager = OidcTokenManager.CreateTokenManager(providerConfiguration, tokenStore, new List<string>() { oidcProvider });
-
-			OidcTokenInfo result;
-			try
-			{
-				result = await oidcTokenManager.GetAccessToken(oidcProvider);
-			}
-			catch (NotLoggedInException)
-			{
-				result = await oidcTokenManager.Login(oidcProvider);
-			}
-
-			if (result.AccessToken == null)
-			{
-				throw new Exception($"Unable to get access token for {oidcProvider}");
-			}
-
-			logger.LogInformation("Received bearer token for {OidcProvider}", oidcProvider);
-			return result.AccessToken;
-		}
 	}
 }
\ No newline at end of file
diff --git a/Engine/Source/Programs/Shared/.editorconfig b/Engine/Source/Programs/Shared/.editorconfig
--- a/Engine/Source/Programs/Shared/.editorconfig
+++ b/Engine/Source/Programs/Shared/.editorconfig
@@ -392,7 +392,7 @@ dotnet_diagnostic.CA1054.severity = warning
 dotnet_diagnostic.CA1055.severity = warning
 
 # CA1056: URI-like properties should not be strings
-dotnet_diagnostic.CA1056.severity = warning
+dotnet_diagnostic.CA1056.severity = silent
 
 # CA1058: Types should not extend certain base types
 dotnet_diagnostic.CA1058.severity = warning
@@ -559,8 +559,8 @@ dotnet_diagnostic.CA1815.severity = none
 # CA1816: Dispose methods should call SuppressFinalize
 dotnet_diagnostic.CA1816.severity = warning
 
-# CA1819: Properties should not return arrays
-dotnet_diagnostic.CA1819.severity = warning
+# CA1819: Properties should not return arrays (rationale: System.Text.Json serialization requires array properties)
+dotnet_diagnostic.CA1819.severity = silent
 
 # CA1820: Test for empty strings using string length
 dotnet_diagnostic.CA1820.severity = warning
diff --git a/Engine/Source/Programs/Shared/EpicGames.Horde/.editorconfig b/Engine/Source/Programs/Shared/EpicGames.Horde/.editorconfig
deleted file mode 100644
diff --git a/Engine/Source/Programs/Shared/EpicGames.Horde/Api/ServerMessages.cs b/Engine/Source/Programs/Shared/EpicGames.Horde/Api/ServerMessages.cs
new file mode 100644
--- /dev/null
+++ b/Engine/Source/Programs/Shared/EpicGames.Horde/Api/ServerMessages.cs
@@ -0,0 +1,149 @@
+// Copyright Epic Games, Inc. All Rights Reserved.
+
+using System;
+
+namespace EpicGames.Horde.Api
+{
+	/// <summary>
+	/// Server Info
+	/// </summary>
+	public class GetServerInfoResponse
+	{
+		/// <summary>
+		/// Server version info
+		/// </summary>
+		public string ServerVersion { get; set; } = String.Empty;
+
+		/// <summary>
+		/// The current agent version string
+		/// </summary>
+		public string? AgentVersion { get; set; }
+
+		/// <summary>
+		/// The operating system server is hosted on
+		/// </summary>
+		public string OsDescription { get; set; } = String.Empty;
+	}
+
+	/// <summary>
+	/// Gets connection information to the server
+	/// </summary>
+	public class GetConnectionResponse
+	{
+		/// <summary>
+		/// Public IP address of the remote machine
+		/// </summary>
+		public string? Ip { get; set; }
+
+		/// <summary>
+		/// Public port of the connecting machine
+		/// </summary>
+		public int Port { get; set; }
+	}
+
+	/// <summary>
+	/// Gets ports configured for this server
+	/// </summary>
+	public class GetPortsResponse
+	{
+		/// <summary>
+		/// Port for HTTP communication
+		/// </summary>
+		public int? Http { get; set; }
+
+		/// <summary>
+		/// Port number for HTTPS communication
+		/// </summary>
+		public int? Https { get; set; }
+
+		/// <summary>
+		/// Port number for unencrpyted HTTPS communication
+		/// </summary>
+		public int? UnencryptedHttp2 { get; set; }
+	}
+
+	/// <summary>
+	/// Authentication method used for logging users in
+	/// </summary>
+	public enum AuthMethod
+	{
+		/// <summary>
+		/// No authentication enabled, mainly for demo and testing purposes
+		/// </summary>
+		Anonymous,
+
+		/// <summary>
+		/// OpenID Connect authentication, tailored for Okta
+		/// </summary>
+		Okta,
+
+		/// <summary>
+		/// Generic OpenID Connect authentication, recommended for most
+		/// </summary>
+		OpenIdConnect,
+	}
+
+	/// <summary>
+	/// Describes the auth config for this server
+	/// </summary>
+	public class GetAuthConfigResponse
+	{
+		/// <summary>
+		/// Issuer for tokens from the auth provider
+		/// </summary>
+		public AuthMethod Method { get; set; }
+
+		/// <summary>
+		/// Issuer for tokens from the auth provider
+		/// </summary>
+		public string? ServerUrl { get; set; }
+
+		/// <summary>
+		/// Client id for the OIDC authority
+		/// </summary>
+		public string? ClientId { get; set; }
+
+		/// <summary>
+		/// Optional redirect url provided to OIDC login for external tools (typically to a local server)
+		/// </summary>
+		public string[]? LocalRedirectUrls { get; set; }
+	}
+
+	/// <summary>
+	/// Request to validate server configuration with the given files replacing their checked-in counterparts.
+	/// </summary>
+	public class PreflightConfigRequest
+	{
+		/// <summary>
+		/// Perforce cluster to retrieve from
+		/// </summary>
+		public string? Cluster { get; set; }
+
+		/// <summary>
+		/// Change to test
+		/// </summary>
+		public int ShelvedChange { get; set; }
+	}
+
+	/// <summary>
+	/// Response from validating config files
+	/// </summary>
+	public class PreflightConfigResponse
+	{
+		/// <summary>
+		/// Whether the files were validated successfully
+		/// </summary>
+		public bool Result { get; set; }
+
+		/// <summary>
+		/// Output message from validation
+		/// </summary>
+		public string? Message { get; set; }
+
+		/// <summary>
+		/// Detailed response
+		/// </summary>
+		public string? Detail { get; set; }
+	}
+}
+
diff --git a/Engine/Source/Programs/Shared/EpicGames.Horde/Compute/Clients/LocalComputeClient.cs b/Engine/Source/Programs/Shared/EpicGames.Horde/Compute/Clients/LocalComputeClient.cs
--- a/Engine/Source/Programs/Shared/EpicGames.Horde/Compute/Clients/LocalComputeClient.cs
+++ b/Engine/Source/Programs/Shared/EpicGames.Horde/Compute/Clients/LocalComputeClient.cs
@@ -54,7 +54,7 @@ namespace EpicGames.Horde.Compute.Clients
 			_listener.Bind(new IPEndPoint(IPAddress.Loopback, port));
 			_listener.Listen();
 
-			_listenerTask = BackgroundTask.StartNew(ctx => RunListenerAsync(_listener, sandboxDir, _executeInProcess, new PrefixLogger("REMOTE", logger), ctx));
+			_listenerTask = BackgroundTask.StartNew(ctx => RunListenerAsync(_listener, sandboxDir, _executeInProcess, logger, ctx));
 
 			_socket = new Socket(SocketType.Stream, ProtocolType.IP);
 			_socket.Connect(IPAddress.Loopback, port);
diff --git a/Engine/Source/Programs/Shared/EpicGames.Horde/Compute/Clients/ServerComputeClient.cs b/Engine/Source/Programs/Shared/EpicGames.Horde/Compute/Clients/ServerComputeClient.cs
--- a/Engine/Source/Programs/Shared/EpicGames.Horde/Compute/Clients/ServerComputeClient.cs
+++ b/Engine/Source/Programs/Shared/EpicGames.Horde/Compute/Clients/ServerComputeClient.cs
@@ -8,7 +8,6 @@ using System;
 using System.Collections.Generic;
 using System.Net;
 using System.Net.Http;
-using System.Net.Http.Headers;
 using System.Net.Http.Json;
 using System.Net.Sockets;
 using System.Runtime.CompilerServices;
@@ -81,50 +80,21 @@ namespace EpicGames.Horde.Compute.Clients
 			}
 		}
 
-		readonly HttpClient? _defaultHttpClient;
-		readonly Func<CancellationToken, Task<HttpClient>> _createHttpClientAsync;
+		readonly HordeHttpClientFactory _hordeHttpClientFactory;
 		readonly CancellationTokenSource _cancellationSource = new CancellationTokenSource();
 		readonly ILogger _logger;
 
 		/// <summary>
 		/// Constructor
 		/// </summary>
-		/// <param name="serverUri">Uri of the server to connect to</param>
-		/// <param name="authHeader">Authentication header</param>
+		/// <param name="hordeHttpClientFactory">Factory for constructing Horde http client instances</param>
 		/// <param name="logger">Logger for diagnostic messages</param>
-		public ServerComputeClient(Uri serverUri, AuthenticationHeaderValue? authHeader, ILogger logger)
+		public ServerComputeClient(HordeHttpClientFactory hordeHttpClientFactory, ILogger logger)
 		{
-#pragma warning disable CA2000 // Dispose objects before losing scope
-			// This is disposed via HttpClient
-			SocketsHttpHandler handler = new SocketsHttpHandler();
-			handler.PooledConnectionLifetime = TimeSpan.FromMinutes(2.0);
-
-			_defaultHttpClient = new HttpClient(handler, true);
-			_defaultHttpClient.BaseAddress = serverUri;
-			_defaultHttpClient.DefaultRequestHeaders.Authorization = authHeader;
-#pragma warning restore CA2000 // Dispose objects before losing scope
-
-			_createHttpClientAsync = GetDefaultHttpClientAsync;
+			_hordeHttpClientFactory = hordeHttpClientFactory;
 			_logger = logger;
 		}
 
-		/// <summary>
-		/// Constructor
-		/// </summary>
-		/// <param name="createHttpClientAsync">Creates an HTTP client with the correct base address for the server</param>
-		/// <param name="logger">Logger for diagnostic messages</param>
-		public ServerComputeClient(Func<CancellationToken, Task<HttpClient>> createHttpClientAsync, ILogger logger)
-		{
-			_createHttpClientAsync = createHttpClientAsync;
-			_logger = logger;
-		}
-
-		/// <summary>
-		/// Gets the default http client
-		/// </summary>
-		/// <returns></returns>
-		Task<HttpClient> GetDefaultHttpClientAsync(CancellationToken cancellationToken) => Task.FromResult(_defaultHttpClient!);
-
 		/// <inheritdoc/>
 		public ValueTask DisposeAsync()
 		{
@@ -136,7 +106,6 @@ namespace EpicGames.Horde.Compute.Clients
 		public void Dispose()
 		{
 			_cancellationSource.Dispose();
-			_defaultHttpClient?.Dispose();
 		}
 
 		/// <inheritdoc/>
@@ -160,7 +129,7 @@ namespace EpicGames.Horde.Compute.Clients
 			HordeHttpClient.ConfigureJsonSerializer(jsonSerializerOptions);
 
 			// Assign a compute worker
-			HttpClient client = await _createHttpClientAsync(cancellationToken);
+			HordeHttpClient client = _hordeHttpClientFactory.CreateClient();
 
 			AssignComputeRequest request = new AssignComputeRequest();
 			request.Requirements = requirements;
diff --git a/Engine/Source/Programs/Shared/EpicGames.Horde/EpicGames.Horde.csproj b/Engine/Source/Programs/Shared/EpicGames.Horde/EpicGames.Horde.csproj
--- a/Engine/Source/Programs/Shared/EpicGames.Horde/EpicGames.Horde.csproj
+++ b/Engine/Source/Programs/Shared/EpicGames.Horde/EpicGames.Horde.csproj
@@ -1,4 +1,4 @@
-<Project Sdk="Microsoft.NET.Sdk">
+﻿<Project Sdk="Microsoft.NET.Sdk">
   <PropertyGroup>
     <TargetFramework>net6.0</TargetFramework>
     <Configuration Condition=" '$(Configuration)' == '' ">Development</Configuration>
@@ -50,6 +50,7 @@
   <ItemGroup>
     <ProjectReference Include="..\EpicGames.Core\EpicGames.Core.csproj" />
     <ProjectReference Include="..\EpicGames.IoHash\EpicGames.IoHash.csproj" />
+    <ProjectReference Include="..\EpicGames.OIDC\EpicGames.OIDC.csproj" />
     <ProjectReference Include="..\EpicGames.Oodle\EpicGames.Oodle.csproj" />
     <ProjectReference Include="..\EpicGames.Serialization\EpicGames.Serialization.csproj" />
   </ItemGroup>
diff --git a/Engine/Source/Programs/Shared/EpicGames.Horde/HordeHttpAuthHandler.cs b/Engine/Source/Programs/Shared/EpicGames.Horde/HordeHttpAuthHandler.cs
new file mode 100644
--- /dev/null
+++ b/Engine/Source/Programs/Shared/EpicGames.Horde/HordeHttpAuthHandler.cs
@@ -0,0 +1,123 @@
+// Copyright Epic Games, Inc. All Rights Reserved.
+
+using System;
+using System.Collections.Generic;
+using System.Linq;
+using System.Net;
+using System.Net.Http;
+using System.Net.Http.Headers;
+using System.Threading;
+using System.Threading.Tasks;
+using Microsoft.Extensions.Logging;
+using Microsoft.Extensions.Configuration;
+using EpicGames.OIDC;
+using EpicGames.Horde.Api;
+using System.Net.Http.Json;
+using System.Text.Json;
+using Microsoft.Extensions.Options;
+
+namespace EpicGames.Horde
+{
+	/// <summary>
+	/// HTTP message handler which automatically refreshes access tokens as required
+	/// </summary>
+	public class HordeHttpAuthHandler : DelegatingHandler
+	{
+		const string ClientName = "HordeHttpAuth";
+
+		readonly IHttpClientFactory _httpClientFactory;
+		readonly Uri _serverUrl;
+		readonly ILogger _logger;
+
+		AuthenticationHeaderValue? _authHeader;
+
+		/// <summary>
+		/// Constructor
+		/// </summary>
+		public HordeHttpAuthHandler(IHttpClientFactory httpClientFactory, IOptions<HordeHttpClientOptions> options, ILogger<HordeHttpAuthHandler> logger)
+		{
+			_httpClientFactory = httpClientFactory;
+			_serverUrl = options.Value.ServerUrl;
+			_logger = logger;
+		}
+
+		/// <inheritdoc/>
+		protected override async Task<HttpResponseMessage> SendAsync(HttpRequestMessage request, CancellationToken cancellationToken)
+		{
+			HttpResponseMessage? response = null;
+			if (_authHeader != null)
+			{
+				request.Headers.Authorization = _authHeader;
+				response = await base.SendAsync(request, cancellationToken);
+			}
+			if (response == null || response.StatusCode == HttpStatusCode.Unauthorized)
+			{
+				await RefreshAccessTokenAsync(cancellationToken);
+				request.Headers.Authorization = _authHeader;
+				response = await base.SendAsync(request, cancellationToken);
+			}
+			return response;
+		}
+
+		/// <summary>
+		/// Get an access token for the server specified in a config instance
+		/// </summary>
+		async Task RefreshAccessTokenAsync(CancellationToken cancellationToken)
+		{
+			_logger.LogInformation("Getting access token for {Server}", _serverUrl);
+
+			GetAuthConfigResponse? authConfig;
+			using (HttpClient httpClient = _httpClientFactory.CreateClient(ClientName))
+			{
+				Uri uri = new Uri(_serverUrl, "api/v1/server/auth");
+
+				JsonSerializerOptions jsonOptions = new JsonSerializerOptions();
+				HordeHttpClient.ConfigureJsonSerializer(jsonOptions);
+
+				authConfig = await httpClient.GetFromJsonAsync<GetAuthConfigResponse>(uri, jsonOptions, cancellationToken);
+				if (authConfig == null)
+				{
+					throw new Exception($"Invalid response from {uri}");
+				}
+			}
+
+			const string OidcProvider = "Horde";
+
+			Dictionary<string, string?> values = new Dictionary<string, string?>();
+			values[$"Providers:{OidcProvider}:DisplayName"] = "Horde";
+			values[$"Providers:{OidcProvider}:ServerUri"] = authConfig.ServerUrl;
+			values[$"Providers:{OidcProvider}:ClientId"] = authConfig.ClientId;
+			values[$"Providers:{OidcProvider}:RedirectUri"] = authConfig.LocalRedirectUrls?.FirstOrDefault();
+
+			ConfigurationBuilder builder = new ConfigurationBuilder();
+			builder.AddInMemoryCollection(values);
+
+			IConfiguration configuration = builder.Build();
+
+			using ITokenStore tokenStore = TokenStoreFactory.CreateTokenStore();
+			OidcTokenManager oidcTokenManager = OidcTokenManager.CreateTokenManager(configuration, tokenStore, new List<string>() { OidcProvider });
+
+			OidcTokenInfo? result = null;
+			try
+			{
+				result = await oidcTokenManager.TryGetAccessToken(OidcProvider, cancellationToken);
+			}
+			catch (Exception ex)
+			{
+				_logger.LogTrace(ex, "Unable to get access token; attempting login: {Message}", ex.Message);
+			}
+			if (result == null)
+			{
+				result = await oidcTokenManager.Login(OidcProvider, cancellationToken);
+			}
+
+			if (result.AccessToken == null)
+			{
+				throw new Exception($"Unable to get access token for {_serverUrl}");
+			}
+
+			_logger.LogInformation("Received bearer token for {Server}", _serverUrl);
+			_authHeader = new AuthenticationHeaderValue("Bearer", result.AccessToken);
+		}
+	}
+}
diff --git a/Engine/Source/Programs/Shared/EpicGames.Horde/HordeHttpClient.cs b/Engine/Source/Programs/Shared/EpicGames.Horde/HordeHttpClient.cs
--- a/Engine/Source/Programs/Shared/EpicGames.Horde/HordeHttpClient.cs
+++ b/Engine/Source/Programs/Shared/EpicGames.Horde/HordeHttpClient.cs
@@ -3,88 +3,67 @@
 using System;
 using System.Net.Http;
 using System.Net.Http.Headers;
+using System.Net.Http.Json;
 using System.Text.Json;
 using System.Text.Json.Serialization;
 using System.Threading;
 using System.Threading.Tasks;
-using EpicGames.Core;
-using EpicGames.Horde.Storage;
-using EpicGames.Horde.Storage.Clients;
-using Microsoft.Extensions.Http;
-using Microsoft.Extensions.Logging;
+using Microsoft.Extensions.DependencyInjection;
+using Microsoft.Extensions.Options;
 using Polly;
 using Polly.Extensions.Http;
-using Polly.Retry;
 
 namespace EpicGames.Horde
 {
 	/// <summary>
-	/// Wraps an Http client which communicates with the Horde server
+	/// Options for configuring a Horde HTTP client
 	/// </summary>
-	public sealed class HordeHttpClient : IDisposable
+	public class HordeHttpClientOptions
 	{
-		readonly HttpMessageHandler _httpMessageHandler;
-		readonly Uri _baseUri;
-		readonly AuthenticationHeaderValue? _authHeader;
-		readonly JsonSerializerOptions _jsonSerializerOptions;
-		readonly ILogger _logger;
-
 		/// <summary>
-		/// Constructor
+		/// URL of the Horde server
 		/// </summary>
-		/// <param name="uri">Base URI of the server</param>
-		/// <param name="token">Access token for the connection</param>
-		/// <param name="logger">Logger for output</param>
-		public HordeHttpClient(Uri uri, string token, ILogger logger)
-		{
-			_baseUri = uri;
-			_authHeader = new AuthenticationHeaderValue("Bearer", token);
-
-			_httpMessageHandler = CreateHttpMessageHandler();
+		public Uri ServerUrl { get; set; } = new Uri("http://localhost:5000");
+	}
 
-			_jsonSerializerOptions = new JsonSerializerOptions();
-			ConfigureJsonSerializer(_jsonSerializerOptions);
+	/// <summary>
+	/// Wraps an Http client which communicates with the Horde server
+	/// </summary>
+	public sealed class HordeHttpClient
+	{
+		readonly HttpClient _httpClient;
 
-			_logger = logger;
-		}
+		static readonly JsonSerializerOptions s_jsonSerializerOptions = CreateJsonSerializerOptions();
 
-		/// <inheritdoc/>
-		public void Dispose()
-		{
-			_httpMessageHandler.Dispose();
-		}
+		/// <summary>
+		/// Static instance of the default retry policy
+		/// </summary>
+		public static IAsyncPolicy<HttpResponseMessage> DefaultRetryPolicy { get; } = CreateDefaultRetryPolicy();
 
 		/// <summary>
-		/// Creates an HTTP client using the underlying message handler
+		/// Accessor for the underlying HTTP client
 		/// </summary>
-		public HttpClient CreateHttpClient()
-		{
-			HttpClient httpClient = new HttpClient(_httpMessageHandler, false);
-			httpClient.BaseAddress = _baseUri;
-			httpClient.DefaultRequestHeaders.Authorization = _authHeader;
-			return httpClient;
-		}
+		public HttpClient HttpClient => _httpClient;
 
 		/// <summary>
-		/// Creates an HTTP client using the underlying message handler
+		/// Constructor
 		/// </summary>
-		public HttpClient CreateRedirectHttpClient()
+		/// <param name="httpClient">The inner HTTP client instance</param>
+		/// <param name="options">Options for the client</param>
+		public HordeHttpClient(HttpClient httpClient, IOptions<HordeHttpClientOptions> options)
 		{
-			return new HttpClient(_httpMessageHandler, false);
+			_httpClient = httpClient;
+			_httpClient.BaseAddress = options.Value.ServerUrl;
 		}
 
 		/// <summary>
-		/// Creates a HTTP message handler
+		/// Create the shared instance of JSON options for HordeHttpClient instances
 		/// </summary>
-		/// <returns></returns>
-		public static HttpMessageHandler CreateHttpMessageHandler()
+		static JsonSerializerOptions CreateJsonSerializerOptions()
 		{
-			AsyncRetryPolicy<HttpResponseMessage> retryPolicy = HttpPolicyExtensions
-				.HandleTransientHttpError()
-				.WaitAndRetryAsync(new[] { TimeSpan.FromSeconds(2.0), TimeSpan.FromSeconds(5.0), TimeSpan.FromSeconds(10), TimeSpan.FromSeconds(30) });
-
-			SocketsHttpHandler socketsHandler = new SocketsHttpHandler { PooledConnectionLifetime = TimeSpan.FromMinutes(15) };
-			return new PolicyHttpMessageHandler(retryPolicy) { InnerHandler = socketsHandler };
+			JsonSerializerOptions options = new JsonSerializerOptions();
+			ConfigureJsonSerializer(options);
+			return options;
 		}
 
 		/// <summary>
@@ -103,40 +82,6 @@ namespace EpicGames.Horde
 			options.Converters.Add(new BinaryIdJsonConverterFactory());
 		}
 
-		/// <summary>
-		/// Creates a storage client for this Horde instance
-		/// </summary>
-		/// <param name="namespaceId">Namespace to manipulate</param>
-		/// <param name="storageCache">Cache for storage operations</param>
-		/// <returns>New storage client instance</returns>
-		public HttpStorageClient CreateStorageClient(NamespaceId namespaceId, StorageCache storageCache) => CreateStorageClient($"api/v1/storage/{namespaceId}/", storageCache);
-
-		/// <summary>
-		/// Creates a storage client for this Horde instance
-		/// </summary>
-		/// <param name="path">Base path for storage requests</param>
-		/// <param name="storageCache">Cache for storage operations</param>
-		/// <returns>New storage client instance</returns>
-		public HttpStorageClient CreateStorageClient(string path, StorageCache storageCache)
-		{
-			if (!path.EndsWith("/", StringComparison.Ordinal))
-			{
-				path += "/";
-			}
-
-			Uri baseUri = new Uri(_baseUri, path);
-
-			HttpClient CreateStorageHttpClient()
-			{
-				HttpClient httpClient = new HttpClient(_httpMessageHandler, false);
-				httpClient.BaseAddress = baseUri;
-				httpClient.DefaultRequestHeaders.Authorization = _authHeader;
-				return httpClient;
-			}
-
-			return new HttpStorageClient(CreateStorageHttpClient, CreateRedirectHttpClient, storageCache, _logger);
-		}
-
 		/// <summary>
 		/// Gets a resource from an HTTP endpoint and parses it as a JSON object
 		/// </summary>
@@ -146,8 +91,8 @@ namespace EpicGames.Horde
 		/// <returns>New instance of the object</returns>
 		public async Task<TResponse> GetAsync<TResponse>(string relativePath, CancellationToken cancellationToken = default)
 		{
-			using HttpClient httpClient = CreateHttpClient();
-			return await httpClient.GetAsync<TResponse>(relativePath, cancellationToken);
+			TResponse? response = await _httpClient.GetFromJsonAsync<TResponse>(relativePath, s_jsonSerializerOptions, cancellationToken);
+			return response ?? throw new InvalidCastException($"Expected non-null response from GET to {relativePath}");
 		}
 
 		/// <summary>
@@ -160,8 +105,7 @@ namespace EpicGames.Horde
 		/// <returns>The response parsed into the requested type</returns>
 		public async Task<HttpResponseMessage> PostAsync<TRequest>(string relativePath, TRequest request, CancellationToken cancellationToken = default)
 		{
-			using HttpClient httpClient = CreateHttpClient();
-			return await httpClient.PostAsync<TRequest>(relativePath, request, cancellationToken);
+			return await _httpClient.PostAsJsonAsync<TRequest>(relativePath, request, s_jsonSerializerOptions, cancellationToken);
 		}
 
 		/// <summary>
@@ -175,8 +119,13 @@ namespace EpicGames.Horde
 		/// <returns>The response parsed into the requested type</returns>
 		public async Task<TResponse> PostAsync<TResponse, TRequest>(string relativePath, TRequest request, CancellationToken cancellationToken = default)
 		{
-			using HttpClient httpClient = CreateHttpClient();
-			return await httpClient.PostAsync<TResponse, TRequest>(relativePath, request, cancellationToken);
+			using (HttpResponseMessage response = await PostAsync<TRequest>(relativePath, request, cancellationToken))
+			{
+				response.EnsureSuccessStatusCode();
+
+				TResponse? responseValue = await response.Content.ReadFromJsonAsync<TResponse>(s_jsonSerializerOptions, cancellationToken);
+				return responseValue ?? throw new InvalidCastException($"Expected non-null response from POST to {relativePath}");
+			}
 		}
 
 		/// <summary>
@@ -189,8 +138,7 @@ namespace EpicGames.Horde
 		/// <returns>Response message</returns>
 		public async Task<HttpResponseMessage> PutAsync<TRequest>(string relativePath, TRequest request, CancellationToken cancellationToken)
 		{
-			using HttpClient httpClient = CreateHttpClient();
-			return await httpClient.PutAsync<TRequest>(relativePath, request, cancellationToken);
+			return await _httpClient.PutAsJsonAsync<TRequest>(relativePath, request, s_jsonSerializerOptions, cancellationToken);
 		}
 
 		/// <summary>
@@ -204,8 +152,83 @@ namespace EpicGames.Horde
 		/// <returns>Response message</returns>
 		public async Task<TResponse> PutAsync<TResponse, TRequest>(string relativePath, TRequest request, CancellationToken cancellationToken)
 		{
-			using HttpClient httpClient = CreateHttpClient();
-			return await httpClient.PutAsync<TResponse, TRequest>(relativePath, request, cancellationToken);
+			using (HttpResponseMessage response = await _httpClient.PutAsJsonAsync<TRequest>(relativePath, request, s_jsonSerializerOptions, cancellationToken))
+			{
+				response.EnsureSuccessStatusCode();
+
+				TResponse? responseValue = await response.Content.ReadFromJsonAsync<TResponse>(s_jsonSerializerOptions, cancellationToken);
+				return responseValue ?? throw new InvalidCastException($"Expected non-null response from PUT to {relativePath}");
+			}
+		}
+
+		/// <summary>
+		/// Helper method to construct a default retry policy for Horde requests
+		/// </summary>
+		static IAsyncPolicy<HttpResponseMessage> CreateDefaultRetryPolicy()
+		{
+			return HttpPolicyExtensions
+				.HandleTransientHttpError()
+				.WaitAndRetryAsync(new[] { TimeSpan.FromSeconds(2.0), TimeSpan.FromSeconds(5.0), TimeSpan.FromSeconds(10), TimeSpan.FromSeconds(30) });
+		}
+	}
+
+	/// <summary>
+	/// Factory for creating HordeHttpClient instances
+	/// </summary>
+	public class HordeHttpClientFactory
+	{
+		readonly IServiceProvider _serviceProvider;
+
+		/// <summary>
+		/// Constructor
+		/// </summary>
+		public HordeHttpClientFactory(IServiceProvider serviceProvider) => _serviceProvider = serviceProvider;
+
+		/// <inheritdoc/>
+		public HordeHttpClient CreateClient() => _serviceProvider.GetRequiredService<HordeHttpClient>();
+	}
+
+	/// <summary>
+	/// Extension methods for Horde HTTP clients
+	/// </summary>
+	public static class HordeHttpClientExtensions
+	{
+		/// <summary>
+		/// Registers a Horde HTTP client type, and configures it to use the default OIDC message handler.
+		/// </summary>
+		/// <param name="services">Service collection to add services to</param>
+		public static void AddHordeHttpClient(this IServiceCollection services)
+		{
+			services.AddSingleton<HordeHttpAuthHandler>();
+			services.AddHttpClient<HordeHttpClient>().AddPolicyHandler(HordeHttpClient.DefaultRetryPolicy).AddHttpMessageHandler<HordeHttpAuthHandler>();
+			services.AddTransient<HordeHttpClientFactory>();
+		}
+
+		/// <summary>
+		/// Registers a Horde HTTP client type, and configures it to use the default OIDC message handler.
+		/// </summary>
+		/// <param name="services">Service collection to add services to</param>
+		/// <param name="configureOptions">Callback to modify options for the http client</param>
+		public static void AddHordeHttpClient(this IServiceCollection services, Action<HordeHttpClientOptions> configureOptions)
+		{
+			services.AddHordeHttpClient((sp, options) => configureOptions(options));
+		}
+		/// <summary>
+		/// Registers a Horde HTTP client type, and configures it to use the default OIDC message handler.
+		/// </summary>
+		/// <param name="services">Service collection to add services to</param>
+		/// <param name="configureOptions">Callback to modify options for the http client</param>
+		public static void AddHordeHttpClient(this IServiceCollection services, Action<IServiceProvider, HordeHttpClientOptions> configureOptions)
+		{
+			IOptions<HordeHttpClientOptions> CreateOptions(IServiceProvider serviceProvider)
+			{
+				HordeHttpClientOptions options = new HordeHttpClientOptions();
+				configureOptions(serviceProvider, options);
+				return Options.Create(options);
+			}
+
+			services.AddSingleton<IOptions<HordeHttpClientOptions>>(CreateOptions);
+			services.AddHordeHttpClient();
 		}
 	}
 }
