Horde: Fix error building Horde.Agent in analyze configuration.

diff --git a/Engine/Source/Programs/Shared/EpicGames.Horde/HordeHttpClient.cs b/Engine/Source/Programs/Shared/EpicGames.Horde/HordeHttpClient.cs
--- a/Engine/Source/Programs/Shared/EpicGames.Horde/HordeHttpClient.cs
+++ b/Engine/Source/Programs/Shared/EpicGames.Horde/HordeHttpClient.cs
@@ -2,7 +2,6 @@
 
 using System;
 using System.Net.Http;
-using System.Net.Http.Headers;
 using System.Net.Http.Json;
 using System.Text.Json;
 using System.Text.Json.Serialization;
@@ -184,7 +183,9 @@ namespace EpicGames.Horde
 		/// </summary>
 		public HordeHttpClientFactory(IServiceProvider serviceProvider) => _serviceProvider = serviceProvider;
 
-		/// <inheritdoc/>
+		/// <summary>
+		/// Creates a client for communicating with the Horde server
+		/// </summary>
 		public HordeHttpClient CreateClient() => _serviceProvider.GetRequiredService<HordeHttpClient>();
 	}
 
