Horde: Allow AddHordeHttpClient overloads to modify the HttpClient used to construct HordeHttpClient.

diff --git a/Engine/Source/Programs/Horde/Horde/Program.cs b/Engine/Source/Programs/Horde/Horde/Program.cs
--- a/Engine/Source/Programs/Horde/Horde/Program.cs
+++ b/Engine/Source/Programs/Horde/Horde/Program.cs
@@ -36,7 +36,7 @@ namespace Horde
 			services.AddLogging();
 			services.AddMemoryCache();
 			services.AddSingleton(sp => Options.Create(CmdConfig.Read()));
-			services.AddHordeHttpClient((sp, options) => options.ServerUrl = sp.GetRequiredService<IOptions<CmdConfig>>().Value.Server);
+			services.AddHordeHttpClient((sp, client) => client.BaseAddress = sp.GetRequiredService<IOptions<CmdConfig>>().Value.Server);
 			services.AddSingleton<StorageCache>(CreateStorageClientCache);
 			services.AddSingleton<StorageBackendCache>(CreateStorageBackendCache);
 
diff --git a/Engine/Source/Programs/Horde/Samples/RemoteClient/RemoteClientApp.cs b/Engine/Source/Programs/Horde/Samples/RemoteClient/RemoteClientApp.cs
--- a/Engine/Source/Programs/Horde/Samples/RemoteClient/RemoteClientApp.cs
+++ b/Engine/Source/Programs/Horde/Samples/RemoteClient/RemoteClientApp.cs
@@ -57,7 +57,7 @@ namespace RemoteClient
 			}
 			else
 			{
-				services.AddHordeHttpClient(x => x.ServerUrl = new Uri(options.Server));
+				services.AddHordeHttpClient(x => x.BaseAddress = new Uri(options.Server));
 				services.AddSingleton<IComputeClient, ServerComputeClient>();
 			}
 
diff --git a/Engine/Source/Programs/Shared/EpicGames.Horde/HordeHttpAuthHandler.cs b/Engine/Source/Programs/Shared/EpicGames.Horde/HordeHttpAuthHandler.cs
--- a/Engine/Source/Programs/Shared/EpicGames.Horde/HordeHttpAuthHandler.cs
+++ b/Engine/Source/Programs/Shared/EpicGames.Horde/HordeHttpAuthHandler.cs
@@ -26,7 +26,6 @@ namespace EpicGames.Horde
 		const string ClientName = "HordeHttpAuth";
 
 		readonly IHttpClientFactory _httpClientFactory;
-		readonly Uri _serverUrl;
 		readonly ILogger _logger;
 
 		AuthenticationHeaderValue? _authHeader;
@@ -34,10 +33,9 @@ namespace EpicGames.Horde
 		/// <summary>
 		/// Constructor
 		/// </summary>
-		public HordeHttpAuthHandler(IHttpClientFactory httpClientFactory, IOptions<HordeHttpClientOptions> options, ILogger<HordeHttpAuthHandler> logger)
+		public HordeHttpAuthHandler(IHttpClientFactory httpClientFactory, ILogger<HordeHttpAuthHandler> logger)
 		{
 			_httpClientFactory = httpClientFactory;
-			_serverUrl = options.Value.ServerUrl;
 			_logger = logger;
 		}
 
@@ -52,7 +50,11 @@ namespace EpicGames.Horde
 			}
 			if (response == null || response.StatusCode == HttpStatusCode.Unauthorized)
 			{
-				await RefreshAccessTokenAsync(cancellationToken);
+				if (request.RequestUri != null)
+				{
+					await RefreshAccessTokenAsync(request.RequestUri, cancellationToken);
+				}
+
 				request.Headers.Authorization = _authHeader;
 				response = await base.SendAsync(request, cancellationToken);
 			}
@@ -62,14 +64,15 @@ namespace EpicGames.Horde
 		/// <summary>
 		/// Get an access token for the server specified in a config instance
 		/// </summary>
-		async Task RefreshAccessTokenAsync(CancellationToken cancellationToken)
+		async Task RefreshAccessTokenAsync(Uri requestUrl, CancellationToken cancellationToken)
 		{
-			_logger.LogInformation("Getting access token for {Server}", _serverUrl);
+			Uri serverUrl = new Uri(requestUrl, "/");
+			_logger.LogInformation("Getting access token for {Server}", serverUrl);
 
 			GetAuthConfigResponse? authConfig;
 			using (HttpClient httpClient = _httpClientFactory.CreateClient(ClientName))
 			{
-				Uri uri = new Uri(_serverUrl, "api/v1/server/auth");
+				Uri uri = new Uri(serverUrl, "api/v1/server/auth");
 
 				JsonSerializerOptions jsonOptions = new JsonSerializerOptions();
 				HordeHttpClient.ConfigureJsonSerializer(jsonOptions);
@@ -113,10 +116,10 @@ namespace EpicGames.Horde
 
 			if (result.AccessToken == null)
 			{
-				throw new Exception($"Unable to get access token for {_serverUrl}");
+				throw new Exception($"Unable to get access token for {serverUrl}");
 			}
 
-			_logger.LogInformation("Received bearer token for {Server}", _serverUrl);
+			_logger.LogInformation("Received bearer token for {Server}", serverUrl);
 			_authHeader = new AuthenticationHeaderValue("Bearer", result.AccessToken);
 		}
 	}
diff --git a/Engine/Source/Programs/Shared/EpicGames.Horde/HordeHttpClient.cs b/Engine/Source/Programs/Shared/EpicGames.Horde/HordeHttpClient.cs
--- a/Engine/Source/Programs/Shared/EpicGames.Horde/HordeHttpClient.cs
+++ b/Engine/Source/Programs/Shared/EpicGames.Horde/HordeHttpClient.cs
@@ -14,17 +14,6 @@ using Polly.Extensions.Http;
 
 namespace EpicGames.Horde
 {
-	/// <summary>
-	/// Options for configuring a Horde HTTP client
-	/// </summary>
-	public class HordeHttpClientOptions
-	{
-		/// <summary>
-		/// URL of the Horde server
-		/// </summary>
-		public Uri ServerUrl { get; set; } = new Uri("http://localhost:5000");
-	}
-
 	/// <summary>
 	/// Wraps an Http client which communicates with the Horde server
 	/// </summary>
@@ -48,11 +37,9 @@ namespace EpicGames.Horde
 		/// Constructor
 		/// </summary>
 		/// <param name="httpClient">The inner HTTP client instance</param>
-		/// <param name="options">Options for the client</param>
-		public HordeHttpClient(HttpClient httpClient, IOptions<HordeHttpClientOptions> options)
+		public HordeHttpClient(HttpClient httpClient)
 		{
 			_httpClient = httpClient;
-			_httpClient.BaseAddress = options.Value.ServerUrl;
 		}
 
 		/// <summary>
@@ -200,36 +187,28 @@ namespace EpicGames.Horde
 		/// <param name="services">Service collection to add services to</param>
 		public static void AddHordeHttpClient(this IServiceCollection services)
 		{
-			services.AddSingleton<HordeHttpAuthHandler>();
-			services.AddHttpClient<HordeHttpClient>().AddPolicyHandler(HordeHttpClient.DefaultRetryPolicy).AddHttpMessageHandler<HordeHttpAuthHandler>();
-			services.AddTransient<HordeHttpClientFactory>();
+			services.AddHordeHttpClient((sp, client) => { });
 		}
 
 		/// <summary>
 		/// Registers a Horde HTTP client type, and configures it to use the default OIDC message handler.
 		/// </summary>
 		/// <param name="services">Service collection to add services to</param>
-		/// <param name="configureOptions">Callback to modify options for the http client</param>
-		public static void AddHordeHttpClient(this IServiceCollection services, Action<HordeHttpClientOptions> configureOptions)
+		/// <param name="configureClient">Callback to modify options for the http client</param>
+		public static void AddHordeHttpClient(this IServiceCollection services, Action<HttpClient> configureClient)
 		{
-			services.AddHordeHttpClient((sp, options) => configureOptions(options));
+			services.AddHordeHttpClient((sp, client) => configureClient(client));
 		}
 		/// <summary>
 		/// Registers a Horde HTTP client type, and configures it to use the default OIDC message handler.
 		/// </summary>
 		/// <param name="services">Service collection to add services to</param>
-		/// <param name="configureOptions">Callback to modify options for the http client</param>
-		public static void AddHordeHttpClient(this IServiceCollection services, Action<IServiceProvider, HordeHttpClientOptions> configureOptions)
+		/// <param name="configureClient">Callback to modify options for the http client</param>
+		public static void AddHordeHttpClient(this IServiceCollection services, Action<IServiceProvider, HttpClient> configureClient)
 		{
-			IOptions<HordeHttpClientOptions> CreateOptions(IServiceProvider serviceProvider)
-			{
-				HordeHttpClientOptions options = new HordeHttpClientOptions();
-				configureOptions(serviceProvider, options);
-				return Options.Create(options);
-			}
-
-			services.AddSingleton<IOptions<HordeHttpClientOptions>>(CreateOptions);
-			services.AddHordeHttpClient();
+			services.AddSingleton<HordeHttpAuthHandler>();
+			services.AddHttpClient<HordeHttpClient>(configureClient).AddPolicyHandler(HordeHttpClient.DefaultRetryPolicy).AddHttpMessageHandler<HordeHttpAuthHandler>();
+			services.AddTransient<HordeHttpClientFactory>();
 		}
 	}
 }
