Horde: Remove unnecessary using directives.

diff --git a/Engine/Source/Programs/Shared/EpicGames.Horde/HordeHttpClient.cs b/Engine/Source/Programs/Shared/EpicGames.Horde/HordeHttpClient.cs
--- a/Engine/Source/Programs/Shared/EpicGames.Horde/HordeHttpClient.cs
+++ b/Engine/Source/Programs/Shared/EpicGames.Horde/HordeHttpClient.cs
@@ -8,7 +8,6 @@ using System.Text.Json.Serialization;
 using System.Threading;
 using System.Threading.Tasks;
 using Microsoft.Extensions.DependencyInjection;
-using Microsoft.Extensions.Options;
 using Polly;
 using Polly.Extensions.Http;
 
