Horde: Fix warning

diff --git a/Engine/Source/Programs/Shared/EpicGames.Horde/HordeHttpAuthHandler.cs b/Engine/Source/Programs/Shared/EpicGames.Horde/HordeHttpAuthHandler.cs
--- a/Engine/Source/Programs/Shared/EpicGames.Horde/HordeHttpAuthHandler.cs
+++ b/Engine/Source/Programs/Shared/EpicGames.Horde/HordeHttpAuthHandler.cs
@@ -14,7 +14,6 @@ using EpicGames.OIDC;
 using EpicGames.Horde.Api;
 using System.Net.Http.Json;
 using System.Text.Json;
-using Microsoft.Extensions.Options;
 
 namespace EpicGames.Horde
 {
