Horde: Add helper for allocating request IDs in compute client

diff --git a/Engine/Source/Programs/Shared/EpicGames.Horde.Tests/RequestIdAllocatorTests.cs b/Engine/Source/Programs/Shared/EpicGames.Horde.Tests/RequestIdAllocatorTests.cs
new file mode 100644
--- /dev/null
+++ b/Engine/Source/Programs/Shared/EpicGames.Horde.Tests/RequestIdAllocatorTests.cs
@@ -0,0 +1,37 @@
+// Copyright Epic Games, Inc. All Rights Reserved.
+
+using System;
+using System.Collections.Generic;
+using EpicGames.Horde.Compute;
+using Microsoft.VisualStudio.TestTools.UnitTesting;
+
+namespace EpicGames.Horde.Tests;
+
+[TestClass]
+public class RequestIdAllocatorTests
+{
+	[TestMethod]
+	public void IdsAreReused()
+	{
+		RequestIdAllocator ria = new ();
+		List<string> before = new() { ria.AllocateId(), ria.AllocateId(), ria.AllocateId() };
+		ria.StartBatch();
+		List<string> after = new() { ria.AllocateId(), ria.AllocateId(), ria.AllocateId() };
+		
+		CollectionAssert.AreEquivalent(before, after);
+	}
+	
+	[TestMethod]
+	public void AcceptedIdsAreNotReused()
+	{
+		RequestIdAllocator ria = new ();
+		List<string> before = new() { ria.AllocateId(), ria.AllocateId(), ria.AllocateId() };
+		ria.MarkAccepted(before[1]);
+		
+		ria.StartBatch();
+		List<string> after = new() { ria.AllocateId(), ria.AllocateId(), ria.AllocateId() };
+		Assert.IsTrue(after.Remove(before[0]));
+		Assert.IsFalse(after.Remove(before[1]));
+		Assert.IsTrue(after.Remove(before[2]));
+	}
+}
diff --git a/Engine/Source/Programs/Shared/EpicGames.Horde/Compute/RequestIdAllocator.cs b/Engine/Source/Programs/Shared/EpicGames.Horde/Compute/RequestIdAllocator.cs
new file mode 100644
--- /dev/null
+++ b/Engine/Source/Programs/Shared/EpicGames.Horde/Compute/RequestIdAllocator.cs
@@ -0,0 +1,83 @@
+// Copyright Epic Games, Inc. All Rights Reserved.
+
+using System;
+using System.Collections.Generic;
+using System.Linq;
+
+namespace EpicGames.Horde.Compute;
+
+/// <summary>
+/// Manages and hands out request IDs for compute allocation requests
+/// </summary>
+public class RequestIdAllocator
+{
+	private readonly object _lock = new();
+	private readonly List<string> _pendingIds = new();
+	private readonly List<string> _batchIds = new();
+	private readonly string _allocatorId;
+	private int _requestCounter;
+
+	/// <summary>
+	/// Constructor
+	/// </summary>
+	public RequestIdAllocator()
+	{
+		_allocatorId = GenerateAllocatorId(8);
+	}
+
+	/// <summary>
+	/// Starts a new batch of requests.
+	/// Any requests started during current batch will be reset and marked as unfinished.
+	/// </summary>
+	public void StartBatch()
+	{
+		lock (_lock)
+		{
+			_pendingIds.AddRange(_batchIds);
+			_batchIds.Clear();
+		}
+	}
+		
+	/// <summary>
+	/// Get or create a request ID and mark it as part of current batch
+	/// </summary>
+	/// <returns>A request ID</returns>
+	public string AllocateId()
+	{
+		lock (_lock)
+		{
+			string? reqId;
+			if (_pendingIds.Count > 0)
+			{
+				reqId = _pendingIds[0];
+				_pendingIds.RemoveAt(0);
+			}
+			else
+			{
+				reqId = $"{_allocatorId}-{_requestCounter++}";
+			}
+
+			_batchIds.Add(reqId);
+			return reqId;
+		}
+	}
+
+	/// <summary>
+	/// Mark a request Id as accepted. It won't be re-used for any future requests.
+	/// </summary>
+	/// <param name="reqId">Request ID to mark as finished</param>
+	public void MarkAccepted(string reqId)
+	{
+		lock (_lock)
+		{
+			_batchIds.Remove(reqId);
+		}
+	}
+
+	private static string GenerateAllocatorId(int length)
+	{
+		Random random = new ();
+		string chars = "ABCDEFGHIJKLMNOPQRSTUVWXYZabcdefghijklmnopqrstuvwxyz0123456789";
+		return new string(Enumerable.Repeat(chars, length).Select(s => s[random.Next(s.Length)]).ToArray());
+	}
+}
\ No newline at end of file
