Horde: Add option to disable auth popup when creating a Horde HTTP client.

diff --git a/Engine/Source/Programs/Shared/EpicGames.Horde/HordeHttpClient.cs b/Engine/Source/Programs/Shared/EpicGames.Horde/HordeHttpClient.cs
--- a/Engine/Source/Programs/Shared/EpicGames.Horde/HordeHttpClient.cs
+++ b/Engine/Source/Programs/Shared/EpicGames.Horde/HordeHttpClient.cs
@@ -184,9 +184,10 @@ namespace EpicGames.Horde
 		/// Registers a Horde HTTP client type, and configures it to use the default OIDC message handler.
 		/// </summary>
 		/// <param name="services">Service collection to add services to</param>
-		public static void AddHordeHttpClient(this IServiceCollection services)
+		/// <param name="useAuthChallenge">Whether to prompt the user to authenticate if necessary</param>
+		public static void AddHordeHttpClient(this IServiceCollection services, bool useAuthChallenge = true)
 		{
-			services.AddHordeHttpClient((sp, client) => { });
+			services.AddHordeHttpClient((sp, client) => { }, useAuthChallenge);
 		}
 
 		/// <summary>
@@ -194,19 +195,28 @@ namespace EpicGames.Horde
 		/// </summary>
 		/// <param name="services">Service collection to add services to</param>
 		/// <param name="configureClient">Callback to modify options for the http client</param>
-		public static void AddHordeHttpClient(this IServiceCollection services, Action<HttpClient> configureClient)
+		/// <param name="useAuthChallenge">Whether to prompt the user to authenticate if necessary</param>
+		public static void AddHordeHttpClient(this IServiceCollection services, Action<HttpClient> configureClient, bool useAuthChallenge = true)
 		{
-			services.AddHordeHttpClient((sp, client) => configureClient(client));
+			services.AddHordeHttpClient((sp, client) => configureClient(client), useAuthChallenge);
 		}
+
 		/// <summary>
 		/// Registers a Horde HTTP client type, and configures it to use the default OIDC message handler.
 		/// </summary>
 		/// <param name="services">Service collection to add services to</param>
 		/// <param name="configureClient">Callback to modify options for the http client</param>
-		public static void AddHordeHttpClient(this IServiceCollection services, Action<IServiceProvider, HttpClient> configureClient)
+		/// <param name="useAuthChallenge">Whether to prompt the user to authenticate if necessary</param>
+		public static void AddHordeHttpClient(this IServiceCollection services, Action<IServiceProvider, HttpClient> configureClient, bool useAuthChallenge = true)
 		{
 			services.AddSingleton<HordeHttpAuthHandler>();
-			services.AddHttpClient<HordeHttpClient>(configureClient).AddPolicyHandler(HordeHttpClient.DefaultRetryPolicy).AddHttpMessageHandler<HordeHttpAuthHandler>();
+
+			IHttpClientBuilder builder = services.AddHttpClient<HordeHttpClient>(configureClient).AddPolicyHandler(HordeHttpClient.DefaultRetryPolicy);
+			if (useAuthChallenge)
+			{
+				builder.AddHttpMessageHandler<HordeHttpAuthHandler>();
+			}
+
 			services.AddTransient<HordeHttpClientFactory>();
 		}
 	}
