Horde: Rename StorageCache to BundleReaderCache.

diff --git a/Engine/Source/Programs/Horde/Horde.Agent.Tests/JsonRpcLoggerTests.cs b/Engine/Source/Programs/Horde/Horde.Agent.Tests/JsonRpcLoggerTests.cs
--- a/Engine/Source/Programs/Horde/Horde.Agent.Tests/JsonRpcLoggerTests.cs
+++ b/Engine/Source/Programs/Horde/Horde.Agent.Tests/JsonRpcLoggerTests.cs
@@ -138,7 +138,7 @@ namespace Horde.Agent.Tests
 		[TestMethod]
 		public async Task StorageLoggerTestAsync()
 		{
-			using StorageCache cache = new StorageCache();
+			using BundleReaderCache cache = new BundleReaderCache();
 			using MemoryStorageClient store = new MemoryStorageClient();
 
 			BundleReader reader = new BundleReader(store, cache, NullLogger.Instance);
diff --git a/Engine/Source/Programs/Horde/Horde.Agent/AgentApp.cs b/Engine/Source/Programs/Horde/Horde.Agent/AgentApp.cs
--- a/Engine/Source/Programs/Horde/Horde.Agent/AgentApp.cs
+++ b/Engine/Source/Programs/Horde/Horde.Agent/AgentApp.cs
@@ -233,7 +233,7 @@ namespace Horde.Agent
 			services.AddHostedService(sp => sp.GetRequiredService<WorkerService>());
 
 			services.AddSingleton<IStorageClientFactory, StorageClientFactory>();
-			services.AddSingleton<StorageCache>();
+			services.AddSingleton<BundleReaderCache>();
 
 			services.AddSingleton<ComputeListenerService>();
 			services.AddHostedService(sp => sp.GetRequiredService<ComputeListenerService>());
diff --git a/Engine/Source/Programs/Horde/Horde.Agent/Commands/Compute/ComputeWorkerCommand.cs b/Engine/Source/Programs/Horde/Horde.Agent/Commands/Compute/ComputeWorkerCommand.cs
--- a/Engine/Source/Programs/Horde/Horde.Agent/Commands/Compute/ComputeWorkerCommand.cs
+++ b/Engine/Source/Programs/Horde/Horde.Agent/Commands/Compute/ComputeWorkerCommand.cs
@@ -16,12 +16,12 @@ namespace Horde.Agent.Commands.Compute
 	[Command("computeworker", "Runs the agent as a local compute host, accepting incoming connections on the loopback adapter with a given port")]
 	class ComputeWorkerCommand : Command
 	{
-		readonly StorageCache _storageCache;
+		readonly BundleReaderCache _storageCache;
 
 		[CommandLine("-Port=")]
 		int Port { get; set; } = 2000;
 
-		public ComputeWorkerCommand(StorageCache storageCache)
+		public ComputeWorkerCommand(BundleReaderCache storageCache)
 		{
 			_storageCache = storageCache;
 		}
@@ -45,7 +45,7 @@ namespace Horde.Agent.Commands.Compute
 			return 0;
 		}
 
-		public static async Task RunWorkerAsync(ComputeSocket socket, StorageCache storageCache, ILogger logger, CancellationToken cancellationToken)
+		public static async Task RunWorkerAsync(ComputeSocket socket, BundleReaderCache storageCache, ILogger logger, CancellationToken cancellationToken)
 		{
 			DirectoryReference sandboxDir = DirectoryReference.Combine(AgentApp.DataDir, "Sandbox");
 
diff --git a/Engine/Source/Programs/Horde/Horde.Agent/Leases/Handlers/ComputeHandler.cs b/Engine/Source/Programs/Horde/Horde.Agent/Leases/Handlers/ComputeHandler.cs
--- a/Engine/Source/Programs/Horde/Horde.Agent/Leases/Handlers/ComputeHandler.cs
+++ b/Engine/Source/Programs/Horde/Horde.Agent/Leases/Handlers/ComputeHandler.cs
@@ -76,7 +76,7 @@ namespace Horde.Agent.Leases.Handlers
 		static TimeSpan NoDataTimeout { get; } = TimeSpan.FromSeconds(20);
 
 		readonly ComputeListenerService _listenerService;
-		readonly StorageCache _storageCache;
+		readonly BundleReaderCache _storageCache;
 		readonly IServerLoggerFactory _serverLoggerFactory;
 		readonly AgentSettings _settings;
 		readonly ILogger _logger;
@@ -84,7 +84,7 @@ namespace Horde.Agent.Leases.Handlers
 		/// <summary>
 		/// Constructor
 		/// </summary>
-		public ComputeHandler(ComputeListenerService listenerService, StorageCache storageCache, IServerLoggerFactory serverLoggerFactory, IOptions<AgentSettings> settings, ILogger<ComputeHandler> logger)
+		public ComputeHandler(ComputeListenerService listenerService, BundleReaderCache storageCache, IServerLoggerFactory serverLoggerFactory, IOptions<AgentSettings> settings, ILogger<ComputeHandler> logger)
 		{
 			_listenerService = listenerService;
 			_storageCache = storageCache;
diff --git a/Engine/Source/Programs/Horde/Horde.Agent/Services/ServerStorageFactory.cs b/Engine/Source/Programs/Horde/Horde.Agent/Services/ServerStorageFactory.cs
--- a/Engine/Source/Programs/Horde/Horde.Agent/Services/ServerStorageFactory.cs
+++ b/Engine/Source/Programs/Horde/Horde.Agent/Services/ServerStorageFactory.cs
@@ -72,13 +72,13 @@ namespace Horde.Agent.Services
 	class HttpServerStorageFactory : IServerStorageFactory
 	{
 		readonly IHttpClientFactory _httpClientFactory;
-		readonly StorageCache _storageCache;
+		readonly BundleReaderCache _storageCache;
 		readonly ILogger _logger;
 
 		/// <summary>
 		/// Constructor
 		/// </summary>
-		public HttpServerStorageFactory(IHttpClientFactory httpClientFactory, StorageCache storageCache, ILogger<HttpStorageClient> logger)
+		public HttpServerStorageFactory(IHttpClientFactory httpClientFactory, BundleReaderCache storageCache, ILogger<HttpStorageClient> logger)
 		{
 			_httpClientFactory = httpClientFactory;
 			_storageCache = storageCache;
diff --git a/Engine/Source/Programs/Horde/Horde.Agent/Utility/StorageClientFactory.cs b/Engine/Source/Programs/Horde/Horde.Agent/Utility/StorageClientFactory.cs
--- a/Engine/Source/Programs/Horde/Horde.Agent/Utility/StorageClientFactory.cs
+++ b/Engine/Source/Programs/Horde/Horde.Agent/Utility/StorageClientFactory.cs
@@ -15,13 +15,13 @@ namespace Horde.Agent.Utility
 	class StorageClientFactory : IStorageClientFactory
 	{
 		readonly IOptions<AgentSettings> _settings;
-		readonly StorageCache _memoryCache;
+		readonly BundleReaderCache _memoryCache;
 		readonly ILogger _logger;
 
 		/// <summary>
 		/// Constructor
 		/// </summary>
-		public StorageClientFactory(IOptions<AgentSettings> settings, StorageCache memoryCache, ILogger<IStorageClient> logger)
+		public StorageClientFactory(IOptions<AgentSettings> settings, BundleReaderCache memoryCache, ILogger<IStorageClient> logger)
 		{
 			_settings = settings;
 			_memoryCache = memoryCache;
diff --git a/Engine/Source/Programs/Horde/Horde.Server.Tests/TestSetup.cs b/Engine/Source/Programs/Horde/Horde.Server.Tests/TestSetup.cs
--- a/Engine/Source/Programs/Horde/Horde.Server.Tests/TestSetup.cs
+++ b/Engine/Source/Programs/Horde/Horde.Server.Tests/TestSetup.cs
@@ -274,7 +274,7 @@ namespace Horde.Server.Tests
 			services.AddSingleton<IStorageBackend<ArtifactCollectionV1>>(sp => new MemoryStorageBackend().ForType<ArtifactCollectionV1>());
 
 			services.AddSingleton<StorageService>();
-			services.AddSingleton<StorageCache>();
+			services.AddSingleton<BundleReaderCache>();
 
 			services.AddSingleton<ISingletonDocument<DevicePlatformMapV1>>(new SingletonDocumentStub<DevicePlatformMapV1>());
 		}
diff --git a/Engine/Source/Programs/Horde/Horde.Server/Commands/Bundles/CreateCommand.cs b/Engine/Source/Programs/Horde/Horde.Server/Commands/Bundles/CreateCommand.cs
--- a/Engine/Source/Programs/Horde/Horde.Server/Commands/Bundles/CreateCommand.cs
+++ b/Engine/Source/Programs/Horde/Horde.Server/Commands/Bundles/CreateCommand.cs
@@ -47,7 +47,7 @@ namespace Horde.Server.Commands.Bundles
 			}
 			else
 			{
-				using FileStorageClient storageClient = new FileStorageClient(OutputDir, serviceProvider.GetRequiredService<StorageCache>(), logger);
+				using FileStorageClient storageClient = new FileStorageClient(OutputDir, serviceProvider.GetRequiredService<BundleReaderCache>(), logger);
 				await ExecuteInternalAsync(storageClient, logger);
 			}
 
diff --git a/Engine/Source/Programs/Horde/Horde.Server/Commands/Install/SetupCommand.cs b/Engine/Source/Programs/Horde/Horde.Server/Commands/Install/SetupCommand.cs
--- a/Engine/Source/Programs/Horde/Horde.Server/Commands/Install/SetupCommand.cs
+++ b/Engine/Source/Programs/Horde/Horde.Server/Commands/Install/SetupCommand.cs
@@ -58,7 +58,7 @@ namespace Horde.Server.Commands.Install
 
 			RefName refName = new RefName("latest");
 
-			using StorageCache storageCache = new StorageCache();
+			using BundleReaderCache storageCache = new BundleReaderCache();
 
 			using FileStorageClient client = new FileStorageClient(bundleDir, storageCache, logger);
 			await using (IStorageWriter writer = client.CreateWriter(refName))
diff --git a/Engine/Source/Programs/Horde/Horde.Server/Startup.cs b/Engine/Source/Programs/Horde/Horde.Server/Startup.cs
--- a/Engine/Source/Programs/Horde/Horde.Server/Startup.cs
+++ b/Engine/Source/Programs/Horde/Horde.Server/Startup.cs
@@ -524,7 +524,7 @@ namespace Horde.Server
 			services.AddSingleton<StorageService>();
 			services.AddSingleton<IStorageClientFactory>(sp => sp.GetRequiredService<StorageService>());
 			services.AddSingleton<TestDataService>();
-			services.AddSingleton<StorageCache>();
+			services.AddSingleton<BundleReaderCache>();
 
 			if (settings.JiraUrl != null)
 			{
diff --git a/Engine/Source/Programs/Horde/Horde.Server/Storage/StorageService.cs b/Engine/Source/Programs/Horde/Horde.Server/Storage/StorageService.cs
--- a/Engine/Source/Programs/Horde/Horde.Server/Storage/StorageService.cs
+++ b/Engine/Source/Programs/Horde/Horde.Server/Storage/StorageService.cs
@@ -188,7 +188,7 @@ namespace Horde.Server.Storage
 			public NamespaceId NamespaceId { get; }
 			public bool SupportsRedirects { get; }
 
-			public StorageClientImpl(StorageService outer, NamespaceConfig config, StorageBackendImpl backend, StorageCache storageCache, ILogger logger)
+			public StorageClientImpl(StorageService outer, NamespaceConfig config, StorageBackendImpl backend, BundleReaderCache storageCache, ILogger logger)
 				: base(backend, storageCache, logger)
 			{
 				_outer = outer;
@@ -482,7 +482,7 @@ namespace Horde.Server.Storage
 
 		readonly RedisService _redisService;
 		readonly IClock _clock;
-		readonly StorageCache _storageCache;
+		readonly BundleReaderCache _storageCache;
 		readonly IMemoryCache _memoryCache;
 		readonly IStorageBackendProvider _storageBackendProvider;
 		readonly IOptionsMonitor<GlobalConfig> _globalConfig;
@@ -506,7 +506,7 @@ namespace Horde.Server.Storage
 		/// <summary>
 		/// Constructor
 		/// </summary>
-		public StorageService(MongoService mongoService, RedisService redisService, IClock clock, StorageCache storageCache, IMemoryCache memoryCache, IStorageBackendProvider storageBackendProvider, IOptionsMonitor<GlobalConfig> globalConfig, Tracer tracer, ILogger<StorageService> logger)
+		public StorageService(MongoService mongoService, RedisService redisService, IClock clock, BundleReaderCache storageCache, IMemoryCache memoryCache, IStorageBackendProvider storageBackendProvider, IOptionsMonitor<GlobalConfig> globalConfig, Tracer tracer, ILogger<StorageService> logger)
 		{
 			_redisService = redisService;
 			_clock = clock;
diff --git a/Engine/Source/Programs/Horde/Horde.Server/Tools/ToolCollection.cs b/Engine/Source/Programs/Horde/Horde.Server/Tools/ToolCollection.cs
--- a/Engine/Source/Programs/Horde/Horde.Server/Tools/ToolCollection.cs
+++ b/Engine/Source/Programs/Horde/Horde.Server/Tools/ToolCollection.cs
@@ -158,7 +158,7 @@ namespace Horde.Server.Tools
 		private readonly VersionedCollection<ToolId, Tool> _tools;
 		private readonly StorageService _storageService;
 		private readonly IClock _clock;
-		private readonly StorageCache _cache;
+		private readonly BundleReaderCache _cache;
 		private readonly ILogger _logger;
 
 		private static readonly RedisKey s_baseKey = "tools/v1/";
@@ -168,7 +168,7 @@ namespace Horde.Server.Tools
 		/// <summary>
 		/// Constructor
 		/// </summary>
-		public ToolCollection(MongoService mongoService, RedisService redisService, StorageService storageService, StorageCache cache, IClock clock, ILogger<ToolCollection> logger)
+		public ToolCollection(MongoService mongoService, RedisService redisService, StorageService storageService, BundleReaderCache cache, IClock clock, ILogger<ToolCollection> logger)
 		{
 			_tools = new VersionedCollection<ToolId, Tool>(mongoService, "Tools", redisService, s_baseKey, s_types);
 			_storageService = storageService;
diff --git a/Engine/Source/Programs/Shared/EpicGames.Horde/Compute/Clients/LocalComputeClient.cs b/Engine/Source/Programs/Shared/EpicGames.Horde/Compute/Clients/LocalComputeClient.cs
--- a/Engine/Source/Programs/Shared/EpicGames.Horde/Compute/Clients/LocalComputeClient.cs
+++ b/Engine/Source/Programs/Shared/EpicGames.Horde/Compute/Clients/LocalComputeClient.cs
@@ -75,7 +75,7 @@ namespace EpicGames.Horde.Compute.Clients
 		{
 			using Socket tcpSocket = await listener.AcceptAsync(cancellationToken);
 
-			using StorageCache storageCache = new StorageCache();
+			using BundleReaderCache storageCache = new BundleReaderCache();
 
 			await using (RemoteComputeSocket socket = new RemoteComputeSocket(new TcpTransport(tcpSocket), logger))
 			{
diff --git a/Engine/Source/Programs/Shared/EpicGames.Horde/Storage/Bundles/BundleReaderCache.cs b/Engine/Source/Programs/Shared/EpicGames.Horde/Storage/Bundles/BundleReaderCache.cs
new file mode 100644
--- /dev/null
+++ b/Engine/Source/Programs/Shared/EpicGames.Horde/Storage/Bundles/BundleReaderCache.cs
@@ -0,0 +1,159 @@
+// Copyright Epic Games, Inc. All Rights Reserved.
+
+using System;
+using System.Diagnostics.CodeAnalysis;
+using EpicGames.Horde.Storage.Bundles;
+using Microsoft.Extensions.Caching.Memory;
+
+namespace EpicGames.Horde.Storage
+{
+	/// <summary>
+	/// Options for creating a storage cache
+	/// </summary>
+	public class BundleReaderCacheOptions
+	{
+		/// <summary>
+		/// Size of the header cache
+		/// </summary>
+		public long HeaderCacheSize { get; set; } = 64 * 1024 * 1024;
+
+		/// <summary>
+		/// Size of the packet cache
+		/// </summary>
+		public long PacketCacheSize { get; set; } = 192 * 1024 * 1024;
+	}
+
+	/// <summary>
+	/// Implements caching functionality for reading from a storage client
+	/// </summary>
+	public sealed class BundleReaderCache : IDisposable
+	{
+		readonly MemoryCache? _headerCache;
+		readonly MemoryCache? _packetCache;
+
+		/// <summary>
+		/// Instance of an empty cache
+		/// </summary>
+		public static BundleReaderCache None { get; } = new BundleReaderCache();
+
+		/// <summary>
+		/// Size of the configured header cache
+		/// </summary>
+		public long HeaderCacheSize { get; }
+
+		/// <summary>
+		/// Size of the configured packet cache
+		/// </summary>
+		public long PacketCacheSize { get; }
+
+		/// <summary>
+		/// Whether there is a packet cache present
+		/// </summary>
+		public bool HasPacketCache => _packetCache != null;
+
+		/// <summary>
+		/// Constructor
+		/// </summary>
+		public BundleReaderCache() : this(new BundleReaderCacheOptions()) 
+		{ 
+		}
+
+		/// <summary>
+		/// Constructor
+		/// </summary>
+		public BundleReaderCache(BundleReaderCacheOptions options)
+		{
+			if (options.HeaderCacheSize > 0)
+			{
+				HeaderCacheSize = options.HeaderCacheSize;
+				_headerCache = new MemoryCache(new MemoryCacheOptions { SizeLimit = options.HeaderCacheSize });
+			}
+			if (options.PacketCacheSize > 0)
+			{
+				PacketCacheSize = options.PacketCacheSize;
+				_packetCache = new MemoryCache(new MemoryCacheOptions { SizeLimit = options.PacketCacheSize });
+			}
+		}
+
+		/// <inheritdoc/>
+		public void Dispose()
+		{
+			_headerCache?.Dispose();
+			_packetCache?.Dispose();
+		}
+
+		static void AddCachedValue(IMemoryCache? cache, string cacheKey, object value, int size)
+		{
+			if (cache != null)
+			{
+				using (ICacheEntry entry = cache.CreateEntry(cacheKey))
+				{
+					entry.SetValue(value);
+					entry.SetSize(size);
+				}
+			}
+		}
+
+		static bool TryGetCachedValue<T>(IMemoryCache? cache, string cacheKey, [MaybeNull, NotNullWhen(true)] out T value)
+		{
+			if (cache == null)
+			{
+				value = default;
+				return false;
+			}
+			return cache.TryGetValue(cacheKey, out value);
+		}
+
+		static string GetBundleInfoCacheKey(BundleLocator locator) => $"bundle:{locator}";
+		static string GetEncodedPacketCacheKey(BundleLocator locator, int packetIdx) => $"encoded-packet:{locator}#{packetIdx}";
+		static string GetDecodedPacketCacheKey(BundleLocator locator, int packetIdx) => $"decoded-packet:{locator}#{packetIdx}";
+
+		/// <summary>
+		/// Adds a bundle info object to the cache
+		/// </summary>
+		public void AddCachedHeader(BundleLocator locator, BundleInfo bundleInfo)
+		{
+			AddCachedValue(_headerCache, GetBundleInfoCacheKey(locator), bundleInfo, bundleInfo.HeaderLength);
+		}
+
+		/// <summary>
+		/// Try to read a bundle info object from the cache
+		/// </summary>
+		public bool TryGetCachedHeader(BundleLocator locator, [NotNullWhen(true)] out BundleInfo? bundleInfo)
+		{
+			return TryGetCachedValue(_headerCache, GetBundleInfoCacheKey(locator), out bundleInfo);
+		}
+
+		/// <summary>
+		/// Adds an encoded bundle packet to the cache
+		/// </summary>
+		public void AddCachedEncodedPacket(BundleLocator locator, int packetIdx, ReadOnlyMemory<byte> data)
+		{
+			AddCachedValue(_packetCache, GetEncodedPacketCacheKey(locator, packetIdx), data, data.Length);
+		}
+
+		/// <summary>
+		/// Try to read an encoded bundle packet from the cache
+		/// </summary>
+		public bool TryGetCachedEncodedPacket(BundleLocator locator, int packetIdx, out ReadOnlyMemory<byte> data)
+		{
+			return TryGetCachedValue(_packetCache, GetEncodedPacketCacheKey(locator, packetIdx), out data);
+		}
+
+		/// <summary>
+		/// Adds a decoded bundle packet to the cache
+		/// </summary>
+		public void AddCachedDecodedPacket(BundleLocator locator, int packetIdx, ReadOnlyMemory<byte> data)
+		{
+			AddCachedValue(_packetCache, GetDecodedPacketCacheKey(locator, packetIdx), data, data.Length);
+		}
+
+		/// <summary>
+		/// Try to read a decoded bundle packet from the cache
+		/// </summary>
+		public bool TryGetCachedDecodedPacket(BundleLocator locator, int packetIdx, out ReadOnlyMemory<byte> data)
+		{
+			return TryGetCachedValue(_packetCache, GetDecodedPacketCacheKey(locator, packetIdx), out data);
+		}
+	}
+}
diff --git a/Engine/Source/Programs/Shared/EpicGames.Horde/Storage/Clients/BundleStorageClient.cs b/Engine/Source/Programs/Shared/EpicGames.Horde/Storage/Clients/BundleStorageClient.cs
--- a/Engine/Source/Programs/Shared/EpicGames.Horde/Storage/Clients/BundleStorageClient.cs
+++ b/Engine/Source/Programs/Shared/EpicGames.Horde/Storage/Clients/BundleStorageClient.cs
@@ -161,7 +161,7 @@ namespace EpicGames.Horde.Storage.Clients
 		/// <summary>
 		/// Constructor
 		/// </summary>
-		protected BundleStorageClient(IStorageBackend backend, StorageCache cache, ILogger logger)
+		protected BundleStorageClient(IStorageBackend backend, BundleReaderCache cache, ILogger logger)
 		{
 			_backend = backend;
 			BundleReader = new BundleReader(this, cache, logger);
diff --git a/Engine/Source/Programs/Shared/EpicGames.Horde/Storage/Clients/FileStorageClient.cs b/Engine/Source/Programs/Shared/EpicGames.Horde/Storage/Clients/FileStorageClient.cs
--- a/Engine/Source/Programs/Shared/EpicGames.Horde/Storage/Clients/FileStorageClient.cs
+++ b/Engine/Source/Programs/Shared/EpicGames.Horde/Storage/Clients/FileStorageClient.cs
@@ -26,7 +26,7 @@ namespace EpicGames.Horde.Storage.Clients
 		/// <param name="rootDir">Root directory for storing blobs</param>
 		/// <param name="cache">Memory cache for read data</param>
 		/// <param name="logger">Logger interface</param>
-		public FileStorageClient(DirectoryReference rootDir, StorageCache cache, ILogger logger)
+		public FileStorageClient(DirectoryReference rootDir, BundleReaderCache cache, ILogger logger)
 			: base(new FileStorageBackend(rootDir), cache, logger)
 		{
 			_rootDir = rootDir;
diff --git a/Engine/Source/Programs/Shared/EpicGames.Horde/Storage/Clients/HttpStorageClient.cs b/Engine/Source/Programs/Shared/EpicGames.Horde/Storage/Clients/HttpStorageClient.cs
--- a/Engine/Source/Programs/Shared/EpicGames.Horde/Storage/Clients/HttpStorageClient.cs
+++ b/Engine/Source/Programs/Shared/EpicGames.Horde/Storage/Clients/HttpStorageClient.cs
@@ -65,7 +65,7 @@ namespace EpicGames.Horde.Storage.Clients
 		/// <summary>
 		/// Constructor
 		/// </summary>
-		public HttpStorageClient(Uri baseAddress, string? bearerToken, StorageCache cache, ILogger logger)
+		public HttpStorageClient(Uri baseAddress, string? bearerToken, BundleReaderCache cache, ILogger logger)
 			: this(() => CreateAuthenticatedClient(null, baseAddress, bearerToken), () => CreateClient(null), cache, logger)
 		{
 		}
@@ -73,7 +73,7 @@ namespace EpicGames.Horde.Storage.Clients
 		/// <summary>
 		/// Constructor
 		/// </summary>
-		public HttpStorageClient(Func<HttpClient> createClient, Func<HttpClient> createRedirectClient, StorageCache cache, ILogger logger) 
+		public HttpStorageClient(Func<HttpClient> createClient, Func<HttpClient> createRedirectClient, BundleReaderCache cache, ILogger logger) 
 			: base(new HttpStorageBackend(createClient, createRedirectClient, logger), cache, logger)
 		{
 			_createClient = createClient;
@@ -83,7 +83,7 @@ namespace EpicGames.Horde.Storage.Clients
 		/// <summary>
 		/// Constructor
 		/// </summary>
-		public HttpStorageClient(IHttpClientFactory httpClientFactory, Uri baseAddress, string? bearerToken, StorageCache cache, ILogger logger)
+		public HttpStorageClient(IHttpClientFactory httpClientFactory, Uri baseAddress, string? bearerToken, BundleReaderCache cache, ILogger logger)
 			: this(() => CreateAuthenticatedClient(httpClientFactory, baseAddress, bearerToken), () => CreateClient(httpClientFactory), cache, logger)
 		{
 		}
diff --git a/Engine/Source/Programs/Shared/EpicGames.Horde/Storage/StorageCache.cs b/Engine/Source/Programs/Shared/EpicGames.Horde/Storage/StorageCache.cs
deleted file mode 100644
