Horde: Make BundleReader a private implementation detail of BundleStorageClient.

diff --git a/Engine/Source/Programs/Horde/Horde.Agent/Utility/TempStorage.cs b/Engine/Source/Programs/Horde/Horde.Agent/Utility/TempStorage.cs
--- a/Engine/Source/Programs/Horde/Horde.Agent/Utility/TempStorage.cs
+++ b/Engine/Source/Programs/Horde/Horde.Agent/Utility/TempStorage.cs
@@ -719,18 +719,15 @@ namespace Horde.Storage.Utility
 					throw new TempStorageException($"Missing block \"{blockName}\" from node \"{nodeName}\"");
 				}
 
-				BundleStorageClient bundleStorageClient = (BundleStorageClient)storageClient;
-				long initialNumBytesRead = bundleStorageClient.BundleReader.NumBytesRead;
-				int initialNumHeaderReads = bundleStorageClient.BundleReader.NumHeaderReads;
-				int initialNumPacketReads = bundleStorageClient.BundleReader.NumPacketReads;
+				StorageStats initialStats = storageClient.GetStats();
 				Stopwatch timer = Stopwatch.StartNew();
 
 				// Add all the files and flush the ref
 				DirectoryNode rootDirNode = await rootDirEntry.ExpandAsync(cancellationToken);
 				await rootDirNode.CopyToDirectoryAsync(rootDir.ToDirectoryInfo(), new CopyStatsLogger(logger), logger, cancellationToken);
 
-				timer.Stop();
-				logger.LogInformation("Elapsed: {Elapsed}s, Num bytes: {NumBytes:n0}, Num headers: {NumHeaders:n0}, Num packets: {NumPackets:n0}, Header cache: {HeaderCacheSize:n0}mb, Packet cache {PacketCacheSize:n0}mb", (int)timer.Elapsed.TotalSeconds, bundleStorageClient.BundleReader.NumBytesRead - initialNumBytesRead, bundleStorageClient.BundleReader.NumHeaderReads - initialNumHeaderReads, bundleStorageClient.BundleReader.NumPacketReads - initialNumPacketReads, bundleStorageClient.BundleReader.Cache.HeaderCacheSize / (1024.0 * 1024.0), bundleStorageClient.BundleReader.Cache.PacketCacheSize / (1024.0 * 1024.0));
+				StorageStats deltaStats = StorageStats.GetDelta(initialStats, storageClient.GetStats());
+				logger.LogInformation("{Stats}", $"Elapsed: {(int)timer.Elapsed.TotalSeconds}s, {String.Join(", ", deltaStats.Values.Select(x => $"{x.Item1}: {x.Item2:n0}"))}");
 
 				// Read the manifest in
 				manifest = TempStorageBlockManifest.Load(localManifestFile);
diff --git a/Engine/Source/Programs/Horde/Horde.Server/Commands/Bundles/BundleWriteTest.cs b/Engine/Source/Programs/Horde/Horde.Server/Commands/Bundles/BundleWriteTest.cs
--- a/Engine/Source/Programs/Horde/Horde.Server/Commands/Bundles/BundleWriteTest.cs
+++ b/Engine/Source/Programs/Horde/Horde.Server/Commands/Bundles/BundleWriteTest.cs
@@ -39,6 +39,8 @@ namespace Horde.Server.Commands.Bundles
 			}
 
 			IStorageWriter IStorageClient.CreateWriter(RefName refName) => CreateWriter(refName);
+
+			public void GetStats(StorageStats stats) { }
 		}
 
 		public override async Task<int> ExecuteAsync(ILogger logger)
diff --git a/Engine/Source/Programs/Horde/Horde.Server/Storage/Backends/AwsStorageBackend.cs b/Engine/Source/Programs/Horde/Horde.Server/Storage/Backends/AwsStorageBackend.cs
--- a/Engine/Source/Programs/Horde/Horde.Server/Storage/Backends/AwsStorageBackend.cs
+++ b/Engine/Source/Programs/Horde/Horde.Server/Storage/Backends/AwsStorageBackend.cs
@@ -620,5 +620,8 @@ namespace Horde.Server.Storage.Backends
 				}
 			}
 		}
+
+		/// <inheritdoc/>
+		public void GetStats(StorageStats stats) { }
 	}
 }
diff --git a/Engine/Source/Programs/Horde/Horde.Server/Storage/StorageBackendProvider.cs b/Engine/Source/Programs/Horde/Horde.Server/Storage/StorageBackendProvider.cs
--- a/Engine/Source/Programs/Horde/Horde.Server/Storage/StorageBackendProvider.cs
+++ b/Engine/Source/Programs/Horde/Horde.Server/Storage/StorageBackendProvider.cs
@@ -104,6 +104,9 @@ namespace Horde.Server.Storage
 			/// <inheritdoc/>
 			public ValueTask<(string, Uri)?> TryGetWriteRedirectAsync(string? prefix = null, CancellationToken cancellationToken = default) => _backend.TryGetWriteRedirectAsync(prefix, cancellationToken);
 
+			/// <inheritdoc/>
+			public void GetStats(StorageStats stats) => _backend.GetStats(stats);
+
 			#endregion
 		}
 
diff --git a/Engine/Source/Programs/Horde/Horde.Server/Storage/StorageService.cs b/Engine/Source/Programs/Horde/Horde.Server/Storage/StorageService.cs
--- a/Engine/Source/Programs/Horde/Horde.Server/Storage/StorageService.cs
+++ b/Engine/Source/Programs/Horde/Horde.Server/Storage/StorageService.cs
@@ -177,6 +177,9 @@ namespace Horde.Server.Storage
 
 				return redirect;
 			}
+
+			/// <inheritdoc/>
+			public void GetStats(StorageStats stats) { }
 		}
 
 		sealed class StorageClientImpl : BundleStorageClient
@@ -224,7 +227,7 @@ namespace Horde.Server.Storage
 			{
 				await foreach (BundleNodeLocator locator in _outer.FindAliasesAsync(NamespaceId, alias, cancellationToken))
 				{
-					yield return new FlushedNodeHandle(BundleReader, locator);
+					yield return CreateNodeHandle(locator);
 				}
 			}
 
@@ -240,7 +243,7 @@ namespace Horde.Server.Storage
 				{
 					return null;
 				}
-				return new FlushedNodeHandle(BundleReader, locator.Value);
+				return CreateNodeHandle(locator.Value);
 			}
 
 			/// <inheritdoc/>
@@ -322,6 +325,9 @@ namespace Horde.Server.Storage
 
 			#endregion
 
+			/// <inheritdoc/>
+			public void GetStats(StorageStats stats) => _impl.GetStats(stats);
+
 			public bool Authorize(AclAction action, ClaimsPrincipal user) => _impl.Config.Authorize(action, user);
 		}
 
diff --git a/Engine/Source/Programs/Horde/Horde.Server/Storage/TypedStorageBackend.cs b/Engine/Source/Programs/Horde/Horde.Server/Storage/TypedStorageBackend.cs
--- a/Engine/Source/Programs/Horde/Horde.Server/Storage/TypedStorageBackend.cs
+++ b/Engine/Source/Programs/Horde/Horde.Server/Storage/TypedStorageBackend.cs
@@ -70,6 +70,9 @@ namespace Horde.Server.Storage
 
 			/// <inheritdoc/>
 			public ValueTask<(string, Uri)?> TryGetWriteRedirectAsync(string? prefix = null, CancellationToken cancellationToken = default) => _inner.TryGetWriteRedirectAsync(prefix, cancellationToken);
+
+			/// <inheritdoc/>
+			public void GetStats(StorageStats stats) { }
 		}
 
 		/// <summary>
diff --git a/Engine/Source/Programs/Horde/Horde/Commands/Bundles/BundleExtract.cs b/Engine/Source/Programs/Horde/Horde/Commands/Bundles/BundleExtract.cs
--- a/Engine/Source/Programs/Horde/Horde/Commands/Bundles/BundleExtract.cs
+++ b/Engine/Source/Programs/Horde/Horde/Commands/Bundles/BundleExtract.cs
@@ -74,13 +74,8 @@ namespace Horde.Commands.Bundles
 
 			if (Stats)
 			{
-				BundleStorageClient? bundleStorageClient = store as BundleStorageClient;
-				if (bundleStorageClient != null)
-				{
-					logger.LogInformation("Num bytes read: {NumBytes:n0}", bundleStorageClient.BundleReader.NumBytesRead);
-					logger.LogInformation("Num header reads: {NumReads:n0}", bundleStorageClient.BundleReader.NumHeaderReads);
-					logger.LogInformation("Num packet reads: {NumReads:n0}", bundleStorageClient.BundleReader.NumPacketReads);
-				}
+				StorageStats stats = store.GetStats();
+				stats.Print(logger);
 			}
 		}
 	}
diff --git a/Engine/Source/Programs/Horde/Horde/Commands/Workspace/WorkspaceSync.cs b/Engine/Source/Programs/Horde/Horde/Commands/Workspace/WorkspaceSync.cs
--- a/Engine/Source/Programs/Horde/Horde/Commands/Workspace/WorkspaceSync.cs
+++ b/Engine/Source/Programs/Horde/Horde/Commands/Workspace/WorkspaceSync.cs
@@ -89,13 +89,8 @@ namespace Horde.Commands.Workspace
 
 			if (Stats)
 			{
-				BundleStorageClient? bundleStorageClient = store as BundleStorageClient;
-				if (bundleStorageClient != null)
-				{
-					logger.LogInformation("Num bytes read: {NumBytes:n0}", bundleStorageClient.BundleReader.NumBytesRead);
-					logger.LogInformation("Num header reads: {NumReads:n0}", bundleStorageClient.BundleReader.NumHeaderReads);
-					logger.LogInformation("Num packet reads: {NumReads:n0}", bundleStorageClient.BundleReader.NumPacketReads);
-				}
+				StorageStats stats = store.GetStats();
+				stats.Print(logger);
 			}
 
 			return 0;
diff --git a/Engine/Source/Programs/Shared/EpicGames.Horde/Compute/AgentStorageClient.cs b/Engine/Source/Programs/Shared/EpicGames.Horde/Compute/AgentStorageClient.cs
--- a/Engine/Source/Programs/Shared/EpicGames.Horde/Compute/AgentStorageClient.cs
+++ b/Engine/Source/Programs/Shared/EpicGames.Horde/Compute/AgentStorageClient.cs
@@ -138,5 +138,8 @@ namespace EpicGames.Horde.Compute
 
 		/// <inheritdoc/>
 		public ValueTask<(string, Uri)?> TryGetWriteRedirectAsync(string? prefix = null, CancellationToken cancellationToken = default) => throw new NotSupportedException();
+
+		/// <inheritdoc/>
+		public void GetStats(StorageStats stats) { }
 	}
 }
diff --git a/Engine/Source/Programs/Shared/EpicGames.Horde/Storage/Backends/FileStorageBackend.cs b/Engine/Source/Programs/Shared/EpicGames.Horde/Storage/Backends/FileStorageBackend.cs
--- a/Engine/Source/Programs/Shared/EpicGames.Horde/Storage/Backends/FileStorageBackend.cs
+++ b/Engine/Source/Programs/Shared/EpicGames.Horde/Storage/Backends/FileStorageBackend.cs
@@ -391,6 +391,9 @@ namespace EpicGames.Horde.Storage.Backends
 
 		/// <inheritdoc/>
 		public ValueTask<(string, Uri)?> TryGetWriteRedirectAsync(string? prefix = null, CancellationToken cancellationToken = default) => default;
+
+		/// <inheritdoc/>
+		public void GetStats(StorageStats stats) { }
 	}
 }
 
diff --git a/Engine/Source/Programs/Shared/EpicGames.Horde/Storage/Backends/HttpStorageBackend.cs b/Engine/Source/Programs/Shared/EpicGames.Horde/Storage/Backends/HttpStorageBackend.cs
--- a/Engine/Source/Programs/Shared/EpicGames.Horde/Storage/Backends/HttpStorageBackend.cs
+++ b/Engine/Source/Programs/Shared/EpicGames.Horde/Storage/Backends/HttpStorageBackend.cs
@@ -196,5 +196,8 @@ namespace EpicGames.Horde.Storage.Backends
 		}
 
 		#endregion
+
+		/// <inheritdoc/>
+		public void GetStats(StorageStats stats) { }
 	}
 }
diff --git a/Engine/Source/Programs/Shared/EpicGames.Horde/Storage/Backends/MemoryStorageBackend.cs b/Engine/Source/Programs/Shared/EpicGames.Horde/Storage/Backends/MemoryStorageBackend.cs
--- a/Engine/Source/Programs/Shared/EpicGames.Horde/Storage/Backends/MemoryStorageBackend.cs
+++ b/Engine/Source/Programs/Shared/EpicGames.Horde/Storage/Backends/MemoryStorageBackend.cs
@@ -110,5 +110,8 @@ namespace EpicGames.Horde.Storage.Backends
 
 		/// <inheritdoc/>
 		public ValueTask<(string, Uri)?> TryGetWriteRedirectAsync(string? prefix = null, CancellationToken cancellationToken = default) => default;
+
+		/// <inheritdoc/>
+		public void GetStats(StorageStats stats) { }
 	}
 }
diff --git a/Engine/Source/Programs/Shared/EpicGames.Horde/Storage/Bundles/BundleReader.cs b/Engine/Source/Programs/Shared/EpicGames.Horde/Storage/Bundles/BundleReader.cs
--- a/Engine/Source/Programs/Shared/EpicGames.Horde/Storage/Bundles/BundleReader.cs
+++ b/Engine/Source/Programs/Shared/EpicGames.Horde/Storage/Bundles/BundleReader.cs
@@ -139,21 +139,6 @@ namespace EpicGames.Horde.Storage.Bundles
 		int _numPacketReads;
 		long _numBytesRead;
 
-		/// <summary>
-		/// Total number of header reads
-		/// </summary>
-		public long NumBytesRead => _numBytesRead;
-
-		/// <summary>
-		/// Total number of header reads
-		/// </summary>
-		public int NumHeaderReads => _numHeaderReads;
-
-		/// <summary>
-		/// Total number of packet reads
-		/// </summary>
-		public int NumPacketReads => _numPacketReads;
-
 		/// <summary>
 		/// Accessor for the cache
 		/// </summary>
@@ -602,5 +587,15 @@ namespace EpicGames.Horde.Storage.Bundles
 		/// <param name="cancellationToken">Cancellation token for the operation</param>
 		/// <returns>Node data read from the given bundle</returns>
 		public async ValueTask<TNode> ReadNodeAsync<TNode>(BundleNodeLocator locator, CancellationToken cancellationToken = default) where TNode : Node => (TNode)await ReadNodeAsync(locator, cancellationToken);
+
+		/// <summary>
+		/// Gets stats for the reader
+		/// </summary>
+		public void GetStats(StorageStats stats)
+		{
+			stats.Add("Num bytes read", _numBytesRead);
+			stats.Add("Num header reads", _numHeaderReads);
+			stats.Add("Num packet reads", _numPacketReads);
+		}
 	}
 }
diff --git a/Engine/Source/Programs/Shared/EpicGames.Horde/Storage/Clients/BundleStorageClient.cs b/Engine/Source/Programs/Shared/EpicGames.Horde/Storage/Clients/BundleStorageClient.cs
--- a/Engine/Source/Programs/Shared/EpicGames.Horde/Storage/Clients/BundleStorageClient.cs
+++ b/Engine/Source/Programs/Shared/EpicGames.Horde/Storage/Clients/BundleStorageClient.cs
@@ -147,24 +147,20 @@ namespace EpicGames.Horde.Storage.Clients
 	public abstract class BundleStorageClient : IBundleStorageClient
 	{
 		readonly IStorageBackend _backend;
+		readonly BundleReader _bundleReader;
 
 		/// <summary>
 		/// Backend for this client
 		/// </summary>
 		public IStorageBackend Backend => _backend;
 
-		/// <summary>
-		/// Reader for node data
-		/// </summary>
-		public BundleReader BundleReader { get; }
-
 		/// <summary>
 		/// Constructor
 		/// </summary>
 		protected BundleStorageClient(IStorageBackend backend, BundleReaderCache cache, ILogger logger)
 		{
 			_backend = backend;
-			BundleReader = new BundleReader(this, cache, logger);
+			_bundleReader = new BundleReader(this, cache, logger);
 		}
 
 		/// <inheritdoc/>
@@ -193,10 +189,10 @@ namespace EpicGames.Horde.Storage.Clients
 		#region Bundles
 
 		/// <inheritdoc/>
-		public Task<BundleHeader> ReadHeaderAsync(BundleLocator locator, CancellationToken cancellationToken) => BundleReader.ReadHeaderAsync(locator, cancellationToken);
+		public Task<BundleHeader> ReadHeaderAsync(BundleLocator locator, CancellationToken cancellationToken) => _bundleReader.ReadHeaderAsync(locator, cancellationToken);
 
 		/// <inheritdoc/>
-		public async Task<BlobData> ReadNodeDataAsync(BundleNodeLocator locator, CancellationToken cancellationToken) => await BundleReader.ReadNodeDataAsync(locator, cancellationToken);
+		public async Task<BlobData> ReadNodeDataAsync(BundleNodeLocator locator, CancellationToken cancellationToken) => await _bundleReader.ReadNodeDataAsync(locator, cancellationToken);
 
 		#endregion
 
@@ -205,10 +201,10 @@ namespace EpicGames.Horde.Storage.Clients
 		/// <summary>
 		/// Creates a handle to a node from its locator
 		/// </summary>
-		public BundleNodeHandle CreateNodeHandle(BundleNodeLocator locator) => new FlushedNodeHandle(BundleReader, locator);
+		public BundleNodeHandle CreateNodeHandle(BundleNodeLocator locator) => new FlushedNodeHandle(_bundleReader, locator);
 
 		/// <inheritdoc/>
-		public BundleWriter CreateWriter(RefName refName = default, BundleOptions? options = null) => new BundleWriter(this, BundleReader, refName, options);
+		public BundleWriter CreateWriter(RefName refName = default, BundleOptions? options = null) => new BundleWriter(this, _bundleReader, refName, options);
 
 		/// <inheritdoc/>
 		IStorageWriter IStorageClient.CreateWriter(RefName refName) => CreateWriter(refName);
@@ -259,7 +255,10 @@ namespace EpicGames.Horde.Storage.Clients
 
 		/// <inheritdoc/>
 		public abstract Task WriteRefTargetAsync(RefName name, BundleNodeLocator target, RefOptions? options = null, CancellationToken cancellationToken = default);
-		
+
 		#endregion
+
+		/// <inheritdoc/>
+		public void GetStats(StorageStats stats) => _bundleReader.GetStats(stats);
 	}
 }
diff --git a/Engine/Source/Programs/Shared/EpicGames.Horde/Storage/Clients/FileStorageClient.cs b/Engine/Source/Programs/Shared/EpicGames.Horde/Storage/Clients/FileStorageClient.cs
--- a/Engine/Source/Programs/Shared/EpicGames.Horde/Storage/Clients/FileStorageClient.cs
+++ b/Engine/Source/Programs/Shared/EpicGames.Horde/Storage/Clients/FileStorageClient.cs
@@ -41,7 +41,7 @@ namespace EpicGames.Horde.Storage.Clients
 		public async ValueTask<BlobHandle> ReadRefAsync(FileReference file)
 		{
 			string text = await FileReference.ReadAllTextAsync(file);
-			return new FlushedNodeHandle(BundleReader, BundleNodeLocator.Parse(text));
+			return CreateNodeHandle(BundleNodeLocator.Parse(text));
 		}
 
 		FileReference GetRefFile(RefName name) => FileReference.Combine(_rootDir, name.ToString() + ".ref");
@@ -93,7 +93,7 @@ namespace EpicGames.Horde.Storage.Clients
 
 			_logger.LogInformation("Reading {File}", file);
 			string text = await FileReference.ReadAllTextAsync(file, cancellationToken);
-			return new FlushedNodeHandle(BundleReader, BundleNodeLocator.Parse(text));
+			return CreateNodeHandle(BundleNodeLocator.Parse(text));
 		}
 
 		/// <inheritdoc/>
diff --git a/Engine/Source/Programs/Shared/EpicGames.Horde/Storage/Clients/HttpStorageClient.cs b/Engine/Source/Programs/Shared/EpicGames.Horde/Storage/Clients/HttpStorageClient.cs
--- a/Engine/Source/Programs/Shared/EpicGames.Horde/Storage/Clients/HttpStorageClient.cs
+++ b/Engine/Source/Programs/Shared/EpicGames.Horde/Storage/Clients/HttpStorageClient.cs
@@ -156,7 +156,7 @@ namespace EpicGames.Horde.Storage.Clients
 						FindNodesResponse? message = await response.Content.ReadFromJsonAsync<FindNodesResponse>(cancellationToken: cancellationToken);
 						foreach (FindNodeResponse node in message!.Nodes)
 						{
-							yield return new FlushedNodeHandle(BundleReader, new BundleNodeLocator(node.Hash, node.Blob, node.ExportIdx));
+							yield return CreateNodeHandle(new BundleNodeLocator(node.Hash, node.Blob, node.ExportIdx));
 						}
 					}
 				}
@@ -222,7 +222,7 @@ namespace EpicGames.Horde.Storage.Clients
 							response.EnsureSuccessStatusCode();
 							ReadRefResponse? data = await response.Content.ReadFromJsonAsync<ReadRefResponse>(cancellationToken: cancellationToken);
 							_logger.LogDebug("Read ref {RefName} -> {Blob}#{ExportIdx}", name, data!.Blob, data!.ExportIdx);
-							return new FlushedNodeHandle(BundleReader, new BundleNodeLocator(data.Hash, data!.Blob, data!.ExportIdx));
+							return CreateNodeHandle(new BundleNodeLocator(data.Hash, data!.Blob, data!.ExportIdx));
 						}
 					}
 				}
diff --git a/Engine/Source/Programs/Shared/EpicGames.Horde/Storage/Clients/MemoryStorageClient.cs b/Engine/Source/Programs/Shared/EpicGames.Horde/Storage/Clients/MemoryStorageClient.cs
--- a/Engine/Source/Programs/Shared/EpicGames.Horde/Storage/Clients/MemoryStorageClient.cs
+++ b/Engine/Source/Programs/Shared/EpicGames.Horde/Storage/Clients/MemoryStorageClient.cs
@@ -102,7 +102,7 @@ namespace EpicGames.Horde.Storage.Clients
 			BundleNodeLocator hashedLocator;
 			if (_refs.TryGetValue(name, out hashedLocator))
 			{
-				return Task.FromResult<BundleNodeHandle?>(new FlushedNodeHandle(BundleReader, hashedLocator)); 
+				return Task.FromResult<BundleNodeHandle?>(CreateNodeHandle(hashedLocator)); 
 			}
 			else
 			{
diff --git a/Engine/Source/Programs/Shared/EpicGames.Horde/Storage/IStorageBackend.cs b/Engine/Source/Programs/Shared/EpicGames.Horde/Storage/IStorageBackend.cs
--- a/Engine/Source/Programs/Shared/EpicGames.Horde/Storage/IStorageBackend.cs
+++ b/Engine/Source/Programs/Shared/EpicGames.Horde/Storage/IStorageBackend.cs
@@ -108,6 +108,11 @@ namespace EpicGames.Horde.Storage
 		/// <param name="cancellationToken">Cancellation token for the operation</param>
 		/// <returns>Path for retrieval, and URI to upload the data to</returns>
 		ValueTask<(string, Uri)?> TryGetWriteRedirectAsync(string? prefix = null, CancellationToken cancellationToken = default);
+
+		/// <summary>
+		/// Gets stats for this storage backend
+		/// </summary>
+		void GetStats(StorageStats stats);
 	}
 
 	/// <summary>
diff --git a/Engine/Source/Programs/Shared/EpicGames.Horde/Storage/IStorageClient.cs b/Engine/Source/Programs/Shared/EpicGames.Horde/Storage/IStorageClient.cs
--- a/Engine/Source/Programs/Shared/EpicGames.Horde/Storage/IStorageClient.cs
+++ b/Engine/Source/Programs/Shared/EpicGames.Horde/Storage/IStorageClient.cs
@@ -2,9 +2,11 @@
 
 using System;
 using System.Collections.Generic;
+using System.Linq;
 using System.Threading;
 using System.Threading.Tasks;
 using EpicGames.Core;
+using Microsoft.Extensions.Logging;
 
 namespace EpicGames.Horde.Storage
 {
@@ -141,6 +143,11 @@ namespace EpicGames.Horde.Storage
 		Task<bool> DeleteRefAsync(RefName name, CancellationToken cancellationToken = default);
 
 		#endregion
+
+		/// <summary>
+		/// Gets a snapshot of the stats for the storage client.
+		/// </summary>
+		void GetStats(StorageStats stats);
 	}
 
 	/// <summary>
@@ -241,6 +248,50 @@ namespace EpicGames.Horde.Storage
 		public static implicit operator RefCacheTime(TimeSpan age) => new RefCacheTime(age);
 	}
 
+	/// <summary>
+	/// Stats for the storage system
+	/// </summary>
+	public class StorageStats
+	{
+		/// <summary>
+		/// Stat name to value
+		/// </summary>
+		public List<(string, long)> Values { get; } = new List<(string, long)>();
+
+		/// <summary>
+		/// Add a new stat to the list
+		/// </summary>
+		public void Add(string name, long value) => Values.Add((name, value));
+
+		/// <summary>
+		/// Prints the table of stats to the logger
+		/// </summary>
+		public void Print(ILogger logger)
+		{
+			foreach ((string key, long value) in Values)
+			{
+				logger.LogInformation("{Key}: {Value:n0}", key, value);
+			}
+		}
+
+		/// <summary>
+		/// Subtract a base set of stats from this one
+		/// </summary>
+		public static StorageStats GetDelta(StorageStats initial, StorageStats finish)
+		{
+			StorageStats result = new StorageStats();
+
+			Dictionary<string, long> initialValues = initial.Values.ToDictionary(x => x.Item1, x => x.Item2, StringComparer.Ordinal);
+			foreach ((string name, long value) in finish.Values.ToArray())
+			{
+				initialValues.TryGetValue(name, out long otherValue);
+				result.Add(name, value - otherValue);
+			}
+
+			return result;
+		}
+	}
+
 	/// <summary>
 	/// Extension methods for <see cref="IStorageClient"/>
 	/// </summary>
@@ -277,5 +328,15 @@ namespace EpicGames.Horde.Storage
 		}
 
 		#endregion
+
+		/// <summary>
+		/// Gets a snapshot of the stats for the storage client.
+		/// </summary>
+		public static StorageStats GetStats(this IStorageClient store)
+		{
+			StorageStats stats = new StorageStats();
+			store.GetStats(stats);
+			return stats;
+		}
 	}
 }
diff --git a/Engine/Source/Programs/Shared/EpicGames.Horde/Storage/StorageBackendCache.cs b/Engine/Source/Programs/Shared/EpicGames.Horde/Storage/StorageBackendCache.cs
--- a/Engine/Source/Programs/Shared/EpicGames.Horde/Storage/StorageBackendCache.cs
+++ b/Engine/Source/Programs/Shared/EpicGames.Horde/Storage/StorageBackendCache.cs
@@ -92,6 +92,8 @@ namespace EpicGames.Horde.Storage
 			public Task<bool> ExistsAsync(string path, CancellationToken cancellationToken = default) => _inner.ExistsAsync(path, cancellationToken);
 			public ValueTask<Uri?> TryGetReadRedirectAsync(string path, CancellationToken cancellationToken = default) => _inner.TryGetReadRedirectAsync(path, cancellationToken);
 			public ValueTask<(string, Uri)?> TryGetWriteRedirectAsync(string? prefix = null, CancellationToken cancellationToken = default) => _inner.TryGetWriteRedirectAsync(prefix, cancellationToken);
+
+			public void GetStats(StorageStats stats) => _inner.GetStats(stats);
 		}
 
 		object LockObject => _items;
