Horde: Remove parameter to HttpStorageBackend which allows overriding the construction of an HttpClient for following redirects.

diff --git a/Engine/Source/Programs/Horde/Horde.Agent/Services/ServerStorageFactory.cs b/Engine/Source/Programs/Horde/Horde.Agent/Services/ServerStorageFactory.cs
--- a/Engine/Source/Programs/Horde/Horde.Agent/Services/ServerStorageFactory.cs
+++ b/Engine/Source/Programs/Horde/Horde.Agent/Services/ServerStorageFactory.cs
@@ -88,7 +88,7 @@ namespace Horde.Agent.Services
 		/// <inheritdoc/>
 		public IStorageClient CreateStorageClient(Uri baseAddress, string token)
 		{
-			return new HttpStorageClient(() => CreateHttpClient(baseAddress, token), CreateHttpRedirectClient, _bundleReaderCache, _logger);
+			return new HttpStorageClient(() => CreateHttpClient(baseAddress, token), _bundleReaderCache, _logger);
 		}
 
 		HttpClient CreateHttpClient(Uri baseAddress, string token)
@@ -98,10 +98,5 @@ namespace Horde.Agent.Services
 			httpClient.DefaultRequestHeaders.Authorization = new AuthenticationHeaderValue("Bearer", token);
 			return httpClient;
 		}
-
-		HttpClient CreateHttpRedirectClient()
-		{
-			return _httpClientFactory.CreateClient(HttpStorageClient.HttpClientName);
-		}
 	}
 }
diff --git a/Engine/Source/Programs/Horde/Horde.Agent/Utility/StorageClientFactory.cs b/Engine/Source/Programs/Horde/Horde.Agent/Utility/StorageClientFactory.cs
--- a/Engine/Source/Programs/Horde/Horde.Agent/Utility/StorageClientFactory.cs
+++ b/Engine/Source/Programs/Horde/Horde.Agent/Utility/StorageClientFactory.cs
@@ -38,7 +38,7 @@ namespace Horde.Agent.Utility
 			}
 			else
 			{
-				client = new HttpStorageClient(() => CreateDefaultHttpClient(namespaceId), () => new HttpClient(), _memoryCache, _logger);
+				client = new HttpStorageClient(() => CreateDefaultHttpClient(namespaceId), _memoryCache, _logger);
 			}
 			return Task.FromResult(client);
 		}
diff --git a/Engine/Source/Programs/Horde/Horde/Commands/StorageCommandBase.cs b/Engine/Source/Programs/Horde/Horde/Commands/StorageCommandBase.cs
--- a/Engine/Source/Programs/Horde/Horde/Commands/StorageCommandBase.cs
+++ b/Engine/Source/Programs/Horde/Horde/Commands/StorageCommandBase.cs
@@ -75,15 +75,7 @@ namespace Horde.Commands
 				return client;
 			}
 
-			HttpClient CreateRedirectClient()
-			{
-				HttpClient client = _httpClientFactory.CreateClient().HttpClient;
-				client.BaseAddress = null;
-				client.DefaultRequestHeaders.Authorization = null;
-				return client;
-			}
-
-			return Task.FromResult<IStorageClient>(new HttpStorageClient(CreateClient, CreateRedirectClient, BundleReaderCache, logger));
+			return Task.FromResult<IStorageClient>(new HttpStorageClient(CreateClient, BundleReaderCache, logger));
 		}
 	}
 }
diff --git a/Engine/Source/Programs/Shared/EpicGames.Horde/Storage/Backends/HttpStorageBackend.cs b/Engine/Source/Programs/Shared/EpicGames.Horde/Storage/Backends/HttpStorageBackend.cs
--- a/Engine/Source/Programs/Shared/EpicGames.Horde/Storage/Backends/HttpStorageBackend.cs
+++ b/Engine/Source/Programs/Shared/EpicGames.Horde/Storage/Backends/HttpStorageBackend.cs
@@ -33,7 +33,6 @@ namespace EpicGames.Horde.Storage.Backends
 		}
 
 		readonly Func<HttpClient> _createClient;
-		readonly Func<HttpClient> _createRedirectClient;
 		readonly ILogger _logger;
 		bool _supportsUploadRedirects = true;
 
@@ -43,10 +42,9 @@ namespace EpicGames.Horde.Storage.Backends
 		/// <summary>
 		/// Constructor
 		/// </summary>
-		public HttpStorageBackend(Func<HttpClient> createClient, Func<HttpClient> createRedirectClient, ILogger logger)
+		public HttpStorageBackend(Func<HttpClient> createClient, ILogger logger)
 		{
 			_createClient = createClient;
-			_createRedirectClient = createRedirectClient;
 			_logger = logger;
 		}
 
@@ -111,8 +109,11 @@ namespace EpicGames.Horde.Storage.Backends
 
 			if (_supportsUploadRedirects)
 			{
-				using (HttpClient redirectHttpClient = _createRedirectClient())
+				using (HttpClient redirectHttpClient = _createClient())
 				{
+					// Don't send the auth header if we're following a redirect
+					redirectHttpClient.DefaultRequestHeaders.Authorization = null;
+
 					WriteBlobResponse redirectResponse = await SendWriteRequestAsync(null, prefix, cancellationToken);
 					if (redirectResponse.UploadUrl != null)
 					{
diff --git a/Engine/Source/Programs/Shared/EpicGames.Horde/Storage/Clients/HttpStorageClient.cs b/Engine/Source/Programs/Shared/EpicGames.Horde/Storage/Clients/HttpStorageClient.cs
--- a/Engine/Source/Programs/Shared/EpicGames.Horde/Storage/Clients/HttpStorageClient.cs
+++ b/Engine/Source/Programs/Shared/EpicGames.Horde/Storage/Clients/HttpStorageClient.cs
@@ -66,15 +66,15 @@ namespace EpicGames.Horde.Storage.Clients
 		/// Constructor
 		/// </summary>
 		public HttpStorageClient(Uri baseAddress, string? bearerToken, BundleReaderCache cache, ILogger logger)
-			: this(() => CreateAuthenticatedClient(null, baseAddress, bearerToken), () => CreateClient(null), cache, logger)
+			: this(() => CreateAuthenticatedClient(null, baseAddress, bearerToken), cache, logger)
 		{
 		}
 
 		/// <summary>
 		/// Constructor
 		/// </summary>
-		public HttpStorageClient(Func<HttpClient> createClient, Func<HttpClient> createRedirectClient, BundleReaderCache cache, ILogger logger) 
-			: base(new HttpStorageBackend(createClient, createRedirectClient, logger), cache, logger)
+		public HttpStorageClient(Func<HttpClient> createClient, BundleReaderCache cache, ILogger logger) 
+			: base(new HttpStorageBackend(createClient, logger), cache, logger)
 		{
 			_createClient = createClient;
 			_logger = logger;
@@ -84,7 +84,7 @@ namespace EpicGames.Horde.Storage.Clients
 		/// Constructor
 		/// </summary>
 		public HttpStorageClient(IHttpClientFactory httpClientFactory, Uri baseAddress, string? bearerToken, BundleReaderCache cache, ILogger logger)
-			: this(() => CreateAuthenticatedClient(httpClientFactory, baseAddress, bearerToken), () => CreateClient(httpClientFactory), cache, logger)
+			: this(() => CreateAuthenticatedClient(httpClientFactory, baseAddress, bearerToken), cache, logger)
 		{
 		}
 
