Horde: Use a cache on the local disk to wrap any requests for bundle data.

diff --git a/Engine/Source/Programs/AutomationTool/BuildGraph/Tasks/DeployToolTask.cs b/Engine/Source/Programs/AutomationTool/BuildGraph/Tasks/DeployToolTask.cs
--- a/Engine/Source/Programs/AutomationTool/BuildGraph/Tasks/DeployToolTask.cs
+++ b/Engine/Source/Programs/AutomationTool/BuildGraph/Tasks/DeployToolTask.cs
@@ -142,7 +142,7 @@ namespace AutomationTool.Tasks
 			HttpClient CreateHttpClient()
 			{
 				HttpClient httpClient = new HttpClient();
-				httpClient.BaseAddress = new Uri(serverUri, $"api/v1/tools/{Parameters.Id}/");
+				httpClient.BaseAddress = serverUri;
 				if (settings?.Token != null)
 				{
 					httpClient.DefaultRequestHeaders.Authorization = new AuthenticationHeaderValue("Bearer", settings.Token);
@@ -152,7 +152,10 @@ namespace AutomationTool.Tasks
 
 			BundleNodeHandle handle;
 
-			HttpStorageClient storageClient = new HttpStorageClient(CreateHttpClient, BundleReaderCache.None, Logger);
+			string basePath = $"api/v1/tools/{Parameters.Id}";
+			HttpStorageBackend storageBackend = new HttpStorageBackend(basePath, CreateHttpClient, Logger);
+			HttpStorageClient storageClient = new HttpStorageClient(basePath, CreateHttpClient, storageBackend, BundleReaderCache.None, Logger);
+
 			await using (BundleWriter treeWriter = storageClient.CreateWriter())
 			{
 				DirectoryNode sandbox = new DirectoryNode();
diff --git a/Engine/Source/Programs/Horde/Horde.Agent.Tests/WorkerServiceTest.cs b/Engine/Source/Programs/Horde/Horde.Agent.Tests/WorkerServiceTest.cs
--- a/Engine/Source/Programs/Horde/Horde.Agent.Tests/WorkerServiceTest.cs
+++ b/Engine/Source/Programs/Horde/Horde.Agent.Tests/WorkerServiceTest.cs
@@ -8,7 +8,9 @@ using System.Threading;
 using System.Threading.Channels;
 using System.Threading.Tasks;
 using EpicGames.Core;
+using EpicGames.Horde;
 using EpicGames.Horde.Storage;
+using EpicGames.Horde.Storage.Clients;
 using Google.Protobuf.WellKnownTypes;
 using Grpc.Core;
 using Horde.Agent.Execution;
@@ -54,11 +56,6 @@ namespace Horde.Agent.Tests
 			public IServerLogger CreateLogger(ISession session, string logId, string? jobId, string? batchId, string? stepId, bool? warnings = null, bool? useNewLogger = null, LogLevel outputLevel = LogLevel.Information) => new FakeServerLogger();
 		}
 
-		class FakeSessionStorageFactory : IServerStorageFactory
-		{
-			public IStorageClient CreateStorageClient(Uri baseAddress, string token) => null!;
-		}
-
 		internal static IJobExecutor NullExecutor = new SimpleTestExecutor(async (step, logger, cancellationToken) =>
 		{
 			await Task.Delay(1, cancellationToken);
@@ -69,8 +66,11 @@ namespace Horde.Agent.Tests
 		{
 			_serviceCollection = new ServiceCollection();
 			_serviceCollection.AddLogging();
+			_serviceCollection.AddHordeHttpClient();
 			_serviceCollection.AddSingleton<IServerLoggerFactory, FakeServerLoggerFactory>();
-			_serviceCollection.AddSingleton<IServerStorageFactory, FakeSessionStorageFactory>();
+			_serviceCollection.AddSingleton<BundleReaderCache>();
+			_serviceCollection.AddSingleton<StorageBackendCache>();
+			_serviceCollection.AddSingleton<HttpStorageClientFactory>();
 
 			_serviceCollection.Configure<AgentSettings>(settings =>
 			{
diff --git a/Engine/Source/Programs/Horde/Horde.Agent/AgentApp.cs b/Engine/Source/Programs/Horde/Horde.Agent/AgentApp.cs
--- a/Engine/Source/Programs/Horde/Horde.Agent/AgentApp.cs
+++ b/Engine/Source/Programs/Horde/Horde.Agent/AgentApp.cs
@@ -9,6 +9,7 @@ using Datadog.Trace;
 using Datadog.Trace.Configuration;
 using Datadog.Trace.OpenTracing;
 using EpicGames.Core;
+using EpicGames.Horde;
 using EpicGames.Horde.Storage;
 using EpicGames.Horde.Storage.Clients;
 using Horde.Agent.Execution;
@@ -19,6 +20,7 @@ using Horde.Agent.Utility;
 using Microsoft.Extensions.Configuration;
 using Microsoft.Extensions.DependencyInjection;
 using Microsoft.Extensions.Logging;
+using Microsoft.Extensions.Options;
 using OpenTracing.Util;
 using Polly;
 using Polly.Extensions.Http;
@@ -171,8 +173,9 @@ namespace Horde.Agent
 				return builder.WaitAndRetryAsync(new[] { TimeSpan.FromSeconds(1), TimeSpan.FromSeconds(5), TimeSpan.FromSeconds(10) });
 			});
 
-			services.AddHttpClient(HttpStorageClient.HttpClientName, builder => 
+			services.AddHttpClient(HordeHttpClient.HttpClientName, builder => 
 				{
+					builder.BaseAddress = serverProfile.Url;
 					builder.Timeout = TimeSpan.FromSeconds(240); // Global timeout
 				})
 				.ConfigurePrimaryHttpMessageHandler(() => new SocketsHttpHandler
@@ -182,7 +185,7 @@ namespace Horde.Agent
 					})
 				.AddPolicyHandler((serviceProvider, request) => Policy.TimeoutAsync<HttpResponseMessage>(30,
 					(outcome, timespan, context) => { 
-						serviceProvider.GetRequiredService<ILogger<HttpServerStorageFactory>>().LogWarning("Http request timed out after {Time}s.", (int)timespan.TotalSeconds); 
+						serviceProvider.GetRequiredService<ILogger<HttpStorageClient>>().LogWarning("Http request timed out after {Time}s.", (int)timespan.TotalSeconds); 
 						return Task.CompletedTask; 
 					}))
 				.AddPolicyHandler((serviceProvider, request) => HttpPolicyExtensions.HandleTransientHttpError()
@@ -194,7 +197,7 @@ namespace Horde.Agent
 							TimeSpan.FromSeconds(30),
 							TimeSpan.FromSeconds(30),
 						},
-						(outcome, timespan, retryAttempt, context) => serviceProvider.GetRequiredService<ILogger<HttpServerStorageFactory>>().LogWarning("Http request failed. Delaying for {DelayMs}ms (attempt #{RetryNum}).", timespan.TotalMilliseconds, retryAttempt)
+						(outcome, timespan, retryAttempt, context) => serviceProvider.GetRequiredService<ILogger<HttpStorageClient>>().LogWarning("Http request failed. Delaying for {DelayMs}ms (attempt #{RetryNum}).", timespan.TotalMilliseconds, retryAttempt)
 					));
 			services.AddHttpClient(AwsInstanceLifecycleService.HttpClientName);
 			services.AddSingleton<AwsInstanceLifecycleService>();
@@ -228,12 +231,12 @@ namespace Horde.Agent
 			services.AddSingleton<CapabilitiesService>();
 			services.AddSingleton<ISessionFactory, SessionFactory>();
 			services.AddSingleton<IServerLoggerFactory, ServerLoggerFactory>();
-			services.AddSingleton<IServerStorageFactory, HttpServerStorageFactory>();
 			services.AddSingleton<WorkerService>();
 			services.AddHostedService(sp => sp.GetRequiredService<WorkerService>());
 
-			services.AddSingleton<IStorageClientFactory, StorageClientFactory>();
 			services.AddSingleton<BundleReaderCache>();
+			services.AddSingleton<StorageBackendCache>(CreateStorageBackendCache);
+			services.AddSingleton<HttpStorageClientFactory>();
 
 			services.AddSingleton<ComputeListenerService>();
 			services.AddHostedService(sp => sp.GetRequiredService<ComputeListenerService>());
@@ -248,6 +251,13 @@ namespace Horde.Agent
 			return await CommandHost.RunAsync(arguments, serviceProvider, typeof(Horde.Agent.Modes.Service.RunCommand));
 		}
 
+		static StorageBackendCache CreateStorageBackendCache(IServiceProvider serviceProvider)
+		{
+			AgentSettings settings = serviceProvider.GetRequiredService<IOptions<AgentSettings>>().Value;
+			DirectoryReference cacheDir = DirectoryReference.Combine(GetDataDir(), String.IsNullOrEmpty(settings.BundleCacheDir) ? "Cache" : settings.BundleCacheDir);
+			return new StorageBackendCache(cacheDir, settings.BundleCacheSize * 1024 * 1024, serviceProvider.GetRequiredService<ILogger<StorageBackendCache>>());
+		}
+
 		static void ConfigureTracing(string environment, string version)
 		{
 			TracerSettings settings = TracerSettings.FromDefaultSources();
diff --git a/Engine/Source/Programs/Horde/Horde.Agent/AgentSettings.cs b/Engine/Source/Programs/Horde/Horde.Agent/AgentSettings.cs
--- a/Engine/Source/Programs/Horde/Horde.Agent/AgentSettings.cs
+++ b/Engine/Source/Programs/Horde/Horde.Agent/AgentSettings.cs
@@ -264,6 +264,16 @@ namespace Horde.Agent
 		/// </summary>
 		public int TelemetryReportInterval { get; set; } = 30 * 1000;
 
+		/// <summary>
+		/// Directory to use for caching bundles
+		/// </summary>
+		public string? BundleCacheDir { get; set; }
+
+		/// <summary>
+		/// Maximum size of the bundle cache, in megabytes.
+		/// </summary>
+		public long BundleCacheSize { get; set; } = 1024;
+
 		/// <summary>
 		/// Key/value properties in addition to those set internally by the agent
 		/// </summary>
diff --git a/Engine/Source/Programs/Horde/Horde.Agent/Execution/JobExecutor.cs b/Engine/Source/Programs/Horde/Horde.Agent/Execution/JobExecutor.cs
--- a/Engine/Source/Programs/Horde/Horde.Agent/Execution/JobExecutor.cs
+++ b/Engine/Source/Programs/Horde/Horde.Agent/Execution/JobExecutor.cs
@@ -8,6 +8,7 @@ using System.Text.Json;
 using System.Text.Json.Serialization;
 using EpicGames.Core;
 using EpicGames.Horde.Storage;
+using EpicGames.Horde.Storage.Clients;
 using EpicGames.Horde.Storage.Nodes;
 using Grpc.Core;
 using Horde.Agent.Parser;
@@ -35,7 +36,7 @@ namespace Horde.Agent.Execution
 	class JobExecutorOptions
 	{
 		public ISession Session { get; }
-		public IServerStorageFactory StorageFactory { get; }
+		public HttpStorageClientFactory StorageFactory { get; }
 		public string JobId { get; }
 		public string BatchId { get; }
 		public BeginBatchResponse Batch { get; }
@@ -44,7 +45,7 @@ namespace Horde.Agent.Execution
 		public string Token { get; }
 		public JobOptions JobOptions { get; }
 
-		public JobExecutorOptions(ISession session, IServerStorageFactory storageFactory, string jobId, string batchId, BeginBatchResponse batch, NamespaceId namespaceId, string storagePrefix, string token, JobOptions jobOptions)
+		public JobExecutorOptions(ISession session, HttpStorageClientFactory storageFactory, string jobId, string batchId, BeginBatchResponse batch, NamespaceId namespaceId, string storagePrefix, string token, JobOptions jobOptions)
 		{
 			Session = session;
 			StorageFactory = storageFactory;
@@ -194,7 +195,7 @@ namespace Horde.Agent.Execution
 		protected bool _compileAutomationTool = true;
 
 		protected ISession Session { get; }
-		protected IServerStorageFactory StorageFactory { get; }
+		protected HttpStorageClientFactory StorageFactory { get; }
 		private readonly NamespaceId _namespaceId;
 		private readonly string _storagePrefix;
 		private readonly string _token;
@@ -498,7 +499,7 @@ namespace Horde.Agent.Execution
 
 				using (GlobalTracer.Instance.BuildSpan("TempStorage").WithTag("resource", "Write").StartActive())
 				{
-					using IStorageClient storage = StorageFactory.CreateStorageClient(Session, _namespaceId, _token);
+					using IStorageClient storage = StorageFactory.CreateClient(_namespaceId, _token);
 
 					Stopwatch timer = Stopwatch.StartNew();
 
@@ -766,7 +767,7 @@ namespace Horde.Agent.Execution
 
 				if (JobOptions.UseNewTempStorage ?? false)
 				{
-					using IStorageClient storage = StorageFactory.CreateStorageClient(Session, _namespaceId, _token);
+					using IStorageClient storage = StorageFactory.CreateClient(_namespaceId, _token);
 
 					RefName refName = TempStorage.GetRefNameForNode(_storagePrefix, SetupStepName);
 
@@ -843,7 +844,7 @@ namespace Horde.Agent.Execution
 				CreateJobArtifactResponse artifact = await jobRpc.Client.CreateArtifactAsync(new CreateJobArtifactRequest { JobId = JobId, StepId = stepId, Type = type }, cancellationToken: cancellationToken);
 				Logger.LogInformation("Created artifact {ArtifactId} with ref {RefName} in ns {Namespace}", artifact.Id, artifact.RefName, artifact.NamespaceId);
 
-				using IStorageClient storage = StorageFactory.CreateStorageClient(Session, new NamespaceId(artifact.NamespaceId), artifact.Token);
+				using IStorageClient storage = StorageFactory.CreateClient(new NamespaceId(artifact.NamespaceId), artifact.Token);
 				await using IStorageWriter writer = storage.CreateWriter(new RefName(artifact.RefName));
 
 				DirectoryNode dir = new DirectoryNode();
@@ -861,7 +862,7 @@ namespace Horde.Agent.Execution
 		{
 			DirectoryReference manifestDir = DirectoryReference.Combine(workspaceDir, "Engine", "Saved", "BuildGraph");
 
-			using IStorageClient storage = StorageFactory.CreateStorageClient(Session, _namespaceId, _token);
+			using IStorageClient storage = StorageFactory.CreateClient(_namespaceId, _token);
 
 			// Create the mapping of tag names to file sets
 			Dictionary<string, HashSet<FileReference>> tagNameToFileSet = new Dictionary<string, HashSet<FileReference>>();
diff --git a/Engine/Source/Programs/Horde/Horde.Agent/Leases/Handlers/JobHandler.cs b/Engine/Source/Programs/Horde/Horde.Agent/Leases/Handlers/JobHandler.cs
--- a/Engine/Source/Programs/Horde/Horde.Agent/Leases/Handlers/JobHandler.cs
+++ b/Engine/Source/Programs/Horde/Horde.Agent/Leases/Handlers/JobHandler.cs
@@ -5,6 +5,7 @@ using System.Reflection;
 using System.Runtime.InteropServices;
 using EpicGames.Core;
 using EpicGames.Horde.Storage;
+using EpicGames.Horde.Storage.Clients;
 using Google.Protobuf;
 using Grpc.Core;
 using Horde.Agent.Execution;
@@ -47,14 +48,14 @@ namespace Horde.Agent.Leases.Handlers
 
 		readonly IEnumerable<IJobExecutorFactory> _executorFactories;
 		readonly AgentSettings _settings;
-		readonly IServerStorageFactory _serverStorageFactory;
+		readonly HttpStorageClientFactory _serverStorageFactory;
 		readonly IServerLoggerFactory _serverLoggerFactory;
 		readonly ILogger _defaultLogger;
 
 		/// <summary>
 		/// Constructor
 		/// </summary>
-		public JobHandler(IEnumerable<IJobExecutorFactory> executorFactories, IOptions<AgentSettings> settings, IServerStorageFactory storageClientFactory, IServerLoggerFactory serverLoggerFactory, ILogger<JobHandler> defaultLogger)
+		public JobHandler(IEnumerable<IJobExecutorFactory> executorFactories, IOptions<AgentSettings> settings, HttpStorageClientFactory storageClientFactory, IServerLoggerFactory serverLoggerFactory, ILogger<JobHandler> defaultLogger)
 		{
 			_executorFactories = executorFactories;
 			_settings = settings.Value;
diff --git a/Engine/Source/Programs/Horde/Horde.Agent/Services/ServerLoggerFactory.cs b/Engine/Source/Programs/Horde/Horde.Agent/Services/ServerLoggerFactory.cs
--- a/Engine/Source/Programs/Horde/Horde.Agent/Services/ServerLoggerFactory.cs
+++ b/Engine/Source/Programs/Horde/Horde.Agent/Services/ServerLoggerFactory.cs
@@ -3,6 +3,7 @@
 using Microsoft.Extensions.Logging;
 using Horde.Agent.Utility;
 using EpicGames.Horde.Storage;
+using EpicGames.Horde.Storage.Clients;
 
 namespace Horde.Agent.Services
 {
@@ -52,13 +53,13 @@ namespace Horde.Agent.Services
 	/// </summary>
 	class ServerLoggerFactory : IServerLoggerFactory
 	{
-		readonly IServerStorageFactory _storageClientFactory;
+		readonly HttpStorageClientFactory _storageClientFactory;
 		readonly ILogger _logger;
 
 		/// <summary>
 		/// Constructor
 		/// </summary>
-		public ServerLoggerFactory(IServerStorageFactory storageClientFactory, ILogger<ServerLoggerFactory> logger)
+		public ServerLoggerFactory(HttpStorageClientFactory storageClientFactory, ILogger<ServerLoggerFactory> logger)
 		{
 			_storageClientFactory = storageClientFactory;
 			_logger = logger;
@@ -71,7 +72,7 @@ namespace Horde.Agent.Services
 			IJsonRpcLogSink sink = new JsonRpcLogSink(session.RpcConnection, jobId, batchId, stepId, _logger);
 			if (useNewLogger ?? false)
 			{
-				IStorageClient storageClient = _storageClientFactory.CreateStorageClient(session, $"/api/v1/logs/{logId}/");
+				IStorageClient storageClient = _storageClientFactory.CreateClient($"api/v1/logs/{logId}", session.Token);
 				sink = new JsonRpcAndStorageLogSink(session.RpcConnection, logId, sink, storageClient, _logger);
 			}
 			return new JsonRpcLogger(sink, logId, warnings, outputLevel, _logger);
diff --git a/Engine/Source/Programs/Horde/Horde.Agent/Services/ServerStorageFactory.cs b/Engine/Source/Programs/Horde/Horde.Agent/Services/ServerStorageFactory.cs
deleted file mode 100644
diff --git a/Engine/Source/Programs/Horde/Horde.Agent/Utility/StorageClientFactory.cs b/Engine/Source/Programs/Horde/Horde.Agent/Utility/StorageClientFactory.cs
deleted file mode 100644
diff --git a/Engine/Source/Programs/Horde/Horde.Server/ServerSettings.cs b/Engine/Source/Programs/Horde/Horde.Server/ServerSettings.cs
--- a/Engine/Source/Programs/Horde/Horde.Server/ServerSettings.cs
+++ b/Engine/Source/Programs/Horde/Horde.Server/ServerSettings.cs
@@ -806,6 +806,16 @@ namespace Horde.Server
 		/// </summary>
 		public bool OpenBrowser { get; set; } = false;
 
+		/// <summary>
+		/// Directory to use for cache data
+		/// </summary>
+		public string? BundleCacheDir { get; set; }
+
+		/// <summary>
+		/// Maximum size of the storage cache on disk, in megabytes
+		/// </summary>
+		public long BundleCacheSize { get; set; } = 1024;
+
 		/// <summary>
 		/// Experimental features to enable on the server.
 		/// </summary>
diff --git a/Engine/Source/Programs/Horde/Horde.Server/Startup.cs b/Engine/Source/Programs/Horde/Horde.Server/Startup.cs
--- a/Engine/Source/Programs/Horde/Horde.Server/Startup.cs
+++ b/Engine/Source/Programs/Horde/Horde.Server/Startup.cs
@@ -525,6 +525,7 @@ namespace Horde.Server
 			services.AddSingleton<IStorageClientFactory>(sp => sp.GetRequiredService<StorageService>());
 			services.AddSingleton<TestDataService>();
 			services.AddSingleton<BundleReaderCache>();
+			services.AddSingleton<StorageBackendCache>(CreateStorageBackendCache);
 
 			if (settings.JiraUrl != null)
 			{
@@ -863,6 +864,13 @@ namespace Horde.Server
 			OnAddHealthChecks(services);
 		}
 
+		static StorageBackendCache CreateStorageBackendCache(IServiceProvider serviceProvider)
+		{
+			ServerSettings serverSettings = serviceProvider.GetRequiredService<IOptions<ServerSettings>>().Value;
+			DirectoryReference cacheDir = DirectoryReference.Combine(ServerApp.DataDir, String.IsNullOrEmpty(serverSettings.BundleCacheDir) ? "Cache" : serverSettings.BundleCacheDir);
+			return new StorageBackendCache(cacheDir, serverSettings.BundleCacheSize * 1024 * 1024, serviceProvider.GetRequiredService<ILogger<StorageBackendCache>>());
+		}
+
 		public static void ConfigureFormatters()
 		{
 			LogValueFormatter.RegisterTypeAnnotation<AgentId>("AgentId");
diff --git a/Engine/Source/Programs/Horde/Horde.Server/Storage/StorageBackendProvider.cs b/Engine/Source/Programs/Horde/Horde.Server/Storage/StorageBackendProvider.cs
--- a/Engine/Source/Programs/Horde/Horde.Server/Storage/StorageBackendProvider.cs
+++ b/Engine/Source/Programs/Horde/Horde.Server/Storage/StorageBackendProvider.cs
@@ -111,15 +111,17 @@ namespace Horde.Server.Storage
 		}
 
 		readonly IServiceProvider _serviceProvider;
+		readonly StorageBackendCache _storageBackendCache;
 		readonly object _lockObject = new object();
 		readonly Dictionary<IoHash, RefCountedBackend> _backends = new Dictionary<IoHash, RefCountedBackend>();
 
 		/// <summary>
 		/// Constructor
 		/// </summary>
-		public StorageBackendProvider(IServiceProvider serviceProvider)
+		public StorageBackendProvider(IServiceProvider serviceProvider, StorageBackendCache storageBackendCache)
 		{
 			_serviceProvider = serviceProvider;
+			_storageBackendCache = storageBackendCache;
 		}
 
 		/// <inheritdoc/>
@@ -180,7 +182,12 @@ namespace Horde.Server.Storage
 				case StorageBackendType.FileSystem:
 					return new FileStorageBackend(DirectoryReference.Combine(ServerApp.DataDir, config.BaseDir ?? "Storage"));
 				case StorageBackendType.Aws:
-					return new AwsStorageBackend(_serviceProvider.GetRequiredService<IConfiguration>(), config, _serviceProvider.GetRequiredService<ILogger<AwsStorageBackend>>());
+					{
+#pragma warning disable CA2000 // False positive? (Will be disposed with cache backend wrapper)
+						IStorageBackend backend = new AwsStorageBackend(_serviceProvider.GetRequiredService<IConfiguration>(), config, _serviceProvider.GetRequiredService<ILogger<AwsStorageBackend>>());
+						return _storageBackendCache.CreateWrapper(config.Id.ToString(), backend);
+#pragma warning restore CA2000
+					}
 				case StorageBackendType.Memory:
 					return new MemoryStorageBackend();
 				default:
diff --git a/Engine/Source/Programs/Horde/Horde.Server/Storage/StorageService.cs b/Engine/Source/Programs/Horde/Horde.Server/Storage/StorageService.cs
--- a/Engine/Source/Programs/Horde/Horde.Server/Storage/StorageService.cs
+++ b/Engine/Source/Programs/Horde/Horde.Server/Storage/StorageService.cs
@@ -191,7 +191,7 @@ namespace Horde.Server.Storage
 			public NamespaceId NamespaceId { get; }
 			public bool SupportsRedirects { get; }
 
-			public StorageClientImpl(StorageService outer, NamespaceConfig config, StorageBackendImpl backend, BundleReaderCache bundleReaderCache, ILogger logger)
+			public StorageClientImpl(StorageService outer, NamespaceConfig config, IStorageBackend backend, BundleReaderCache bundleReaderCache, ILogger logger)
 				: base(backend, bundleReaderCache, logger)
 			{
 				_outer = outer;
@@ -606,7 +606,7 @@ namespace Horde.Server.Storage
 		}
 
 		/// <inheritdoc/>
-		Task<IStorageClient> IStorageClientFactory.CreateClientAsync(NamespaceId namespaceId, CancellationToken cancellationToken) => Task.FromResult<IStorageClient>(CreateClient(namespaceId));
+		IStorageClient IStorageClientFactory.CreateClient(NamespaceId namespaceId) => CreateClient(namespaceId);
 
 		#region Config
 
diff --git a/Engine/Source/Programs/Horde/Horde/CmdApp.cs b/Engine/Source/Programs/Horde/Horde/CmdApp.cs
--- a/Engine/Source/Programs/Horde/Horde/CmdApp.cs
+++ b/Engine/Source/Programs/Horde/Horde/CmdApp.cs
@@ -12,6 +12,7 @@ using Serilog.Formatting.Json;
 using EpicGames.Horde.Storage;
 using Microsoft.Extensions.Options;
 using EpicGames.Horde;
+using EpicGames.Horde.Storage.Clients;
 
 namespace Horde
 {
@@ -39,6 +40,7 @@ namespace Horde
 			services.AddHordeHttpClient((sp, client) => client.BaseAddress = sp.GetRequiredService<IOptions<CmdConfig>>().Value.Server);
 			services.AddSingleton<BundleReaderCache>(CreateStorageClientCache);
 			services.AddSingleton<StorageBackendCache>(CreateStorageBackendCache);
+			services.AddSingleton<HttpStorageClientFactory>();
 
 			// Execute all the commands
 			await using ServiceProvider serviceProvider = services.BuildServiceProvider();
diff --git a/Engine/Source/Programs/Horde/Horde/Commands/Bundles/BundleCreate.cs b/Engine/Source/Programs/Horde/Horde/Commands/Bundles/BundleCreate.cs
--- a/Engine/Source/Programs/Horde/Horde/Commands/Bundles/BundleCreate.cs
+++ b/Engine/Source/Programs/Horde/Horde/Commands/Bundles/BundleCreate.cs
@@ -2,7 +2,6 @@
 
 using System.Diagnostics;
 using EpicGames.Core;
-using EpicGames.Horde;
 using EpicGames.Horde.Storage;
 using EpicGames.Horde.Storage.Clients;
 using EpicGames.Horde.Storage.Nodes;
@@ -26,8 +25,8 @@ namespace Horde.Commands.Bundles
 		[CommandLine("-Filter=", Description = "Filter for files to include, in P4 syntax (eg. Foo/...).")]
 		public string Filter { get; set; } = "...";
 
-		public BundleCreate(HordeHttpClientFactory httpClientFactory, BundleReaderCache bundleReaderCache, IOptions<CmdConfig> config)
-			: base(httpClientFactory, bundleReaderCache, config)
+		public BundleCreate(HttpStorageClientFactory storageClientFactory, BundleReaderCache bundleReaderCache, IOptions<CmdConfig> config)
+			: base(storageClientFactory, bundleReaderCache, config)
 		{
 		}
 
@@ -40,7 +39,7 @@ namespace Horde.Commands.Bundles
 			}
 			else if (!String.IsNullOrEmpty(Ref))
 			{
-				using IStorageClient store = await CreateStorageClientAsync(logger);
+				using IStorageClient store = CreateStorageClient();
 				return await ExecuteInternalAsync(store, logger);
 			}
 			else
diff --git a/Engine/Source/Programs/Horde/Horde/Commands/Bundles/BundleExtract.cs b/Engine/Source/Programs/Horde/Horde/Commands/Bundles/BundleExtract.cs
--- a/Engine/Source/Programs/Horde/Horde/Commands/Bundles/BundleExtract.cs
+++ b/Engine/Source/Programs/Horde/Horde/Commands/Bundles/BundleExtract.cs
@@ -2,7 +2,6 @@
 
 using System.Diagnostics;
 using EpicGames.Core;
-using EpicGames.Horde;
 using EpicGames.Horde.Storage;
 using EpicGames.Horde.Storage.Bundles;
 using EpicGames.Horde.Storage.Clients;
@@ -30,8 +29,8 @@ namespace Horde.Commands.Bundles
 		[CommandLine("-OutputDir=", Required = true)]
 		public DirectoryReference OutputDir { get; set; } = null!;
 
-		public BundleExtract(HordeHttpClientFactory httpClientFactory, BundleReaderCache bundleReaderCache, IOptions<CmdConfig> config)
-			: base(httpClientFactory, bundleReaderCache, config)
+		public BundleExtract(HttpStorageClientFactory storageClientFactory, BundleReaderCache bundleReaderCache, IOptions<CmdConfig> config)
+			: base(storageClientFactory, bundleReaderCache, config)
 		{
 		}
 
@@ -45,13 +44,13 @@ namespace Horde.Commands.Bundles
 			}
 			else if (Ref != null)
 			{
-				using IStorageClient store = await CreateStorageClientAsync(logger);
+				using IStorageClient store = CreateStorageClient();
 				BlobHandle handle = await store.ReadRefTargetAsync(new RefName(Ref));
 				await ExecuteInternalAsync(store, handle, logger);
 			}
 			else if (Node != null)
 			{
-				using BundleStorageClient store = (BundleStorageClient)await CreateStorageClientAsync(logger);
+				using BundleStorageClient store = (BundleStorageClient)CreateStorageClient();
 				BlobHandle handle = store.CreateNodeHandle(BundleNodeLocator.Parse(Node));
 				await ExecuteInternalAsync(store, handle, logger);
 			}
diff --git a/Engine/Source/Programs/Horde/Horde/Commands/StorageCommandBase.cs b/Engine/Source/Programs/Horde/Horde/Commands/StorageCommandBase.cs
--- a/Engine/Source/Programs/Horde/Horde/Commands/StorageCommandBase.cs
+++ b/Engine/Source/Programs/Horde/Horde/Commands/StorageCommandBase.cs
@@ -1,10 +1,8 @@
 // Copyright Epic Games, Inc. All Rights Reserved.
 
 using EpicGames.Core;
-using EpicGames.Horde;
 using EpicGames.Horde.Storage;
 using EpicGames.Horde.Storage.Clients;
-using Microsoft.Extensions.Logging;
 using Microsoft.Extensions.Options;
 
 namespace Horde.Commands
@@ -36,14 +34,14 @@ namespace Horde.Commands
 		/// </summary>
 		public CmdConfig Config { get; }
 
-		readonly HordeHttpClientFactory _httpClientFactory;
+		readonly HttpStorageClientFactory _storageClientFactory;
 
 		/// <summary>
 		/// Constructor
 		/// </summary>
-		public StorageCommandBase(HordeHttpClientFactory httpClientFactory, BundleReaderCache bundleReaderCache, IOptions<CmdConfig> config)
+		public StorageCommandBase(HttpStorageClientFactory storageClientFactory, BundleReaderCache bundleReaderCache, IOptions<CmdConfig> config)
 		{
-			_httpClientFactory = httpClientFactory;
+			_storageClientFactory = storageClientFactory;
 
 			BundleReaderCache = bundleReaderCache;
 			Config = config.Value;
@@ -52,30 +50,16 @@ namespace Horde.Commands
 		/// <summary>
 		/// Creates a new client instance
 		/// </summary>
-		/// <param name="logger">Logger for output messages</param>
-		/// <param name="cancellationToken"></param>
-		public Task<IStorageClient> CreateStorageClientAsync(ILogger logger, CancellationToken cancellationToken = default)
+		public IStorageClient CreateStorageClient()
 		{
-			_ = cancellationToken;
-
-			string? path = Path;
-			if (String.IsNullOrEmpty(path))
+			if (String.IsNullOrEmpty(Path))
 			{
-				path = $"api/v1/storage/{Namespace}/";
+				return _storageClientFactory.CreateClient(Namespace);
 			}
-			else if (!path.EndsWith("/", StringComparison.Ordinal))
+			else
 			{
-				path += "/";
+				return _storageClientFactory.CreateClient(Path);
 			}
-
-			HttpClient CreateClient()
-			{
-				HttpClient client = _httpClientFactory.CreateClient().HttpClient;
-				client.BaseAddress = new Uri(client.BaseAddress!, path);
-				return client;
-			}
-
-			return Task.FromResult<IStorageClient>(new HttpStorageClient(CreateClient, BundleReaderCache, logger));
 		}
 	}
 }
diff --git a/Engine/Source/Programs/Horde/Horde/Commands/Vcs/VcsBase.cs b/Engine/Source/Programs/Horde/Horde/Commands/Vcs/VcsBase.cs
--- a/Engine/Source/Programs/Horde/Horde/Commands/Vcs/VcsBase.cs
+++ b/Engine/Source/Programs/Horde/Horde/Commands/Vcs/VcsBase.cs
@@ -157,7 +157,7 @@ namespace Horde.Commands.Vcs
 			_storageClientFactory = storageClientFactory;
 		}
 
-		protected async ValueTask<IStorageClient> CreateStorageClientAsync() => await _storageClientFactory.CreateClientAsync(NamespaceId);
+		protected IStorageClient CreateStorageClient() => _storageClientFactory.CreateClient(NamespaceId);
 
 		protected static async Task<WorkspaceState> ReadStateAsync(DirectoryReference rootDir)
 		{
diff --git a/Engine/Source/Programs/Horde/Horde/Commands/Vcs/VcsBranch.cs b/Engine/Source/Programs/Horde/Horde/Commands/Vcs/VcsBranch.cs
--- a/Engine/Source/Programs/Horde/Horde/Commands/Vcs/VcsBranch.cs
+++ b/Engine/Source/Programs/Horde/Horde/Commands/Vcs/VcsBranch.cs
@@ -23,7 +23,7 @@ namespace Horde.Commands.Vcs
 
 			WorkspaceState workspaceState = await ReadStateAsync(rootDir);
 
-			using IStorageClient store = await CreateStorageClientAsync();
+			using IStorageClient store = CreateStorageClient();
 
 			RefName branchName = new RefName(Name);
 			if (await store.HasRefAsync(branchName))
diff --git a/Engine/Source/Programs/Horde/Horde/Commands/Vcs/VcsCheckout.cs b/Engine/Source/Programs/Horde/Horde/Commands/Vcs/VcsCheckout.cs
--- a/Engine/Source/Programs/Horde/Horde/Commands/Vcs/VcsCheckout.cs
+++ b/Engine/Source/Programs/Horde/Horde/Commands/Vcs/VcsCheckout.cs
@@ -51,7 +51,7 @@ namespace Horde.Commands.Vcs
 
 			RefName branchName = (Branch != null) ? new RefName(Branch) : workspaceState.Branch;
 
-			using IStorageClient store = await CreateStorageClientAsync();
+			using IStorageClient store = CreateStorageClient();
 
 			CommitNode? tip = await GetCommitAsync(store, branchName, Change);
 			if (tip == null)
diff --git a/Engine/Source/Programs/Horde/Horde/Commands/Vcs/VcsClone.cs b/Engine/Source/Programs/Horde/Horde/Commands/Vcs/VcsClone.cs
--- a/Engine/Source/Programs/Horde/Horde/Commands/Vcs/VcsClone.cs
+++ b/Engine/Source/Programs/Horde/Horde/Commands/Vcs/VcsClone.cs
@@ -25,7 +25,7 @@ namespace Horde.Commands.Vcs
 
 			RefName branchName = new RefName(Branch ?? "ue5-main");
 
-			using IStorageClient storageClient = await CreateStorageClientAsync();
+			using IStorageClient storageClient = CreateStorageClient();
 
 			CommitNode? tip = await GetCommitAsync(storageClient, branchName, Change);
 			if (tip == null)
diff --git a/Engine/Source/Programs/Horde/Horde/Commands/Vcs/VcsLog.cs b/Engine/Source/Programs/Horde/Horde/Commands/Vcs/VcsLog.cs
--- a/Engine/Source/Programs/Horde/Horde/Commands/Vcs/VcsLog.cs
+++ b/Engine/Source/Programs/Horde/Horde/Commands/Vcs/VcsLog.cs
@@ -25,7 +25,7 @@ namespace Horde.Commands.Vcs
 
 			WorkspaceState workspaceState = await ReadStateAsync(rootDir);
 
-			using IStorageClient store = await CreateStorageClientAsync();
+			using IStorageClient store = CreateStorageClient();
 
 			using MemoryCache cache = new MemoryCache(new MemoryCacheOptions());
 
diff --git a/Engine/Source/Programs/Horde/Horde/Commands/Workspace/WorkspaceSync.cs b/Engine/Source/Programs/Horde/Horde/Commands/Workspace/WorkspaceSync.cs
--- a/Engine/Source/Programs/Horde/Horde/Commands/Workspace/WorkspaceSync.cs
+++ b/Engine/Source/Programs/Horde/Horde/Commands/Workspace/WorkspaceSync.cs
@@ -2,7 +2,6 @@
 
 using System.Diagnostics;
 using EpicGames.Core;
-using EpicGames.Horde;
 using EpicGames.Horde.Storage;
 using EpicGames.Horde.Storage.Bundles;
 using EpicGames.Horde.Storage.Clients;
@@ -35,8 +34,8 @@ namespace Horde.Commands.Workspace
 		[CommandLine("-Stats")]
 		public bool Stats { get; set; }
 
-		public WorkspaceSync(HordeHttpClientFactory httpClientFactory, BundleReaderCache bundleReaderCache, IOptions<CmdConfig> config)
-			: base(httpClientFactory, bundleReaderCache, config)
+		public WorkspaceSync(HttpStorageClientFactory storageClientFactory, BundleReaderCache bundleReaderCache, IOptions<CmdConfig> config)
+			: base(storageClientFactory, bundleReaderCache, config)
 		{
 		}
 
@@ -50,13 +49,13 @@ namespace Horde.Commands.Workspace
 			}
 			else if (Ref != null)
 			{
-				using IStorageClient store = await CreateStorageClientAsync(logger);
+				using IStorageClient store = CreateStorageClient();
 				BlobHandle handle = await store.ReadRefTargetAsync(new RefName(Ref));
 				return await ExecuteInternalAsync(store, handle, logger);
 			}
 			else if (Node != null)
 			{
-				using IStorageClient store = await CreateStorageClientAsync(logger);
+				using IStorageClient store = CreateStorageClient();
 				BlobHandle handle = ((BundleStorageClient)store).CreateNodeHandle(BundleNodeLocator.Parse(Node));
 				return await ExecuteInternalAsync(store, handle, logger);
 			}
diff --git a/Engine/Source/Programs/Shared/EpicGames.Horde/HordeHttpAuthHandler.cs b/Engine/Source/Programs/Shared/EpicGames.Horde/HordeHttpAuthHandler.cs
--- a/Engine/Source/Programs/Shared/EpicGames.Horde/HordeHttpAuthHandler.cs
+++ b/Engine/Source/Programs/Shared/EpicGames.Horde/HordeHttpAuthHandler.cs
@@ -22,20 +22,16 @@ namespace EpicGames.Horde
 	/// </summary>
 	public class HordeHttpAuthHandler : DelegatingHandler
 	{
-		const string ClientName = "HordeHttpAuth";
-
-		readonly IHttpClientFactory _httpClientFactory;
-		readonly ILogger _logger;
+		readonly HordeHttpAuthHandlerState _authState;
 
 		AuthenticationHeaderValue? _authHeader;
 
 		/// <summary>
 		/// Constructor
 		/// </summary>
-		public HordeHttpAuthHandler(IHttpClientFactory httpClientFactory, ILogger<HordeHttpAuthHandler> logger)
+		public HordeHttpAuthHandler(HordeHttpAuthHandlerState authState)
 		{
-			_httpClientFactory = httpClientFactory;
-			_logger = logger;
+			_authState = authState;
 		}
 
 		/// <inheritdoc/>
@@ -45,41 +41,119 @@ namespace EpicGames.Horde
 			if (_authHeader != null)
 			{
 				request.Headers.Authorization = _authHeader;
+
 				response = await base.SendAsync(request, cancellationToken);
-			}
-			if (response == null || response.StatusCode == HttpStatusCode.Unauthorized)
-			{
-				if (request.RequestUri != null)
+				if (response.StatusCode == HttpStatusCode.Unauthorized)
 				{
-					await RefreshAccessTokenAsync(request.RequestUri, cancellationToken);
+					_authState.Invalidate(_authHeader);
+					response = null;
 				}
-
+			}
+			if (response == null)
+			{
+				_authHeader = await _authState.GetAuthHeaderAsync(cancellationToken);
 				request.Headers.Authorization = _authHeader;
 				response = await base.SendAsync(request, cancellationToken);
 			}
 			return response;
 		}
+	}
+
+	/// <summary>
+	/// Shared object used to track the latest access obtained token
+	/// </summary>
+	public sealed class HordeHttpAuthHandlerState : IAsyncDisposable
+	{
+		/// <summary>
+		/// HTTP client name
+		/// </summary>
+		public const string HttpClientName = "HordeHttpAuthState";
+
+		readonly object _lockObject = new object();
+		readonly CancellationTokenSource _cancellationTokenSource = new CancellationTokenSource();
+		Task<AuthenticationHeaderValue>? _authHeaderTask;
+		readonly IHttpClientFactory _httpClientFactory;
+		readonly ILogger _logger;
+
+		/// <summary>
+		/// Constructor
+		/// </summary>
+		public HordeHttpAuthHandlerState(IHttpClientFactory httpClientFactory, ILogger<HordeHttpAuthHandler> logger)
+		{
+			_httpClientFactory = httpClientFactory;
+			_logger = logger;
+		}
+
+		/// <inheritdoc/>
+		public async ValueTask DisposeAsync()
+		{
+			if (_authHeaderTask != null)
+			{
+				_cancellationTokenSource.Cancel();
+				await _authHeaderTask;
+				_authHeaderTask = null;
+			}
+			_cancellationTokenSource.Dispose();
+		}
+
+		/// <summary>
+		/// Invalidate a cached header value
+		/// </summary>
+		/// <param name="authHeader">The auth header to invalidate</param>
+		public void Invalidate(AuthenticationHeaderValue authHeader)
+		{
+			lock (_lockObject)
+			{
+				if (_authHeaderTask != null && _authHeaderTask.IsCompleted && _authHeaderTask.Result.Equals(authHeader))
+				{
+					_authHeaderTask = null;
+				}
+			}
+		}
+
+		/// <summary>
+		/// Gets a new auth header
+		/// </summary>
+		/// <param name="cancellationToken">Cancellation token for the operation</param>
+		public async ValueTask<AuthenticationHeaderValue> GetAuthHeaderAsync(CancellationToken cancellationToken)
+		{
+			Task<AuthenticationHeaderValue>? authHeaderTask = _authHeaderTask;
+			if (authHeaderTask == null)
+			{
+				lock (_lockObject)
+				{
+					_authHeaderTask ??= Task.Run(() => GetNewAuthHeaderAsync(_cancellationTokenSource.Token), _cancellationTokenSource.Token);
+					authHeaderTask = _authHeaderTask;
+				}
+			}
+			return await authHeaderTask.WaitAsync(cancellationToken);
+		}
 
 		/// <summary>
 		/// Get an access token for the server specified in a config instance
 		/// </summary>
-		async Task RefreshAccessTokenAsync(Uri requestUrl, CancellationToken cancellationToken)
+		async Task<AuthenticationHeaderValue> GetNewAuthHeaderAsync(CancellationToken cancellationToken)
 		{
-			Uri serverUrl = new Uri(requestUrl, "/");
-			_logger.LogInformation("Getting access token for {Server}", serverUrl);
+			Uri serverUrl;
 
 			GetAuthConfigResponse? authConfig;
-			using (HttpClient httpClient = _httpClientFactory.CreateClient(ClientName))
+			using (HttpClient httpClient = _httpClientFactory.CreateClient(HttpClientName))
 			{
-				Uri uri = new Uri(serverUrl, "api/v1/server/auth");
+				if (httpClient.BaseAddress == null)
+				{
+					throw new Exception("No http client is configured for Horde. Call IServiceCollection.AddHordeHttpClient().");
+				}
+
+				serverUrl = httpClient.BaseAddress;
+				_logger.LogInformation("Getting access token for {Server}", serverUrl);
 
 				JsonSerializerOptions jsonOptions = new JsonSerializerOptions();
 				HordeHttpClient.ConfigureJsonSerializer(jsonOptions);
 
-				authConfig = await httpClient.GetFromJsonAsync<GetAuthConfigResponse>(uri, jsonOptions, cancellationToken);
+				authConfig = await httpClient.GetFromJsonAsync<GetAuthConfigResponse>("api/v1/server/auth", jsonOptions, cancellationToken);
 				if (authConfig == null)
 				{
-					throw new Exception($"Invalid response from {uri}");
+					throw new Exception($"Invalid response from server");
 				}
 			}
 
@@ -119,7 +193,7 @@ namespace EpicGames.Horde
 			}
 
 			_logger.LogInformation("Received bearer token for {Server}", serverUrl);
-			_authHeader = new AuthenticationHeaderValue("Bearer", result.AccessToken);
+			return new AuthenticationHeaderValue("Bearer", result.AccessToken);
 		}
 	}
 }
diff --git a/Engine/Source/Programs/Shared/EpicGames.Horde/HordeHttpClient.cs b/Engine/Source/Programs/Shared/EpicGames.Horde/HordeHttpClient.cs
--- a/Engine/Source/Programs/Shared/EpicGames.Horde/HordeHttpClient.cs
+++ b/Engine/Source/Programs/Shared/EpicGames.Horde/HordeHttpClient.cs
@@ -18,6 +18,11 @@ namespace EpicGames.Horde
 	/// </summary>
 	public sealed class HordeHttpClient
 	{
+		/// <summary>
+		/// Name of clients created from the http client factory
+		/// </summary>
+		public const string HttpClientName = "HordeHttpClient";
+
 		readonly HttpClient _httpClient;
 
 		static readonly JsonSerializerOptions s_jsonSerializerOptions = CreateJsonSerializerOptions();
@@ -209,11 +214,13 @@ namespace EpicGames.Horde
 		/// <param name="useAuthChallenge">Whether to prompt the user to authenticate if necessary</param>
 		public static void AddHordeHttpClient(this IServiceCollection services, Action<IServiceProvider, HttpClient> configureClient, bool useAuthChallenge = true)
 		{
-			services.AddSingleton<HordeHttpAuthHandler>();
-
-			IHttpClientBuilder builder = services.AddHttpClient<HordeHttpClient>(configureClient).AddPolicyHandler(HordeHttpClient.DefaultRetryPolicy);
+			IHttpClientBuilder builder = services.AddHttpClient<HordeHttpClient>(HordeHttpClient.HttpClientName, configureClient).AddPolicyHandler(HordeHttpClient.DefaultRetryPolicy);
 			if (useAuthChallenge)
 			{
+				services.AddSingleton<HordeHttpAuthHandlerState>();
+				services.AddTransient<HordeHttpAuthHandler>();
+				services.AddHttpClient(HordeHttpAuthHandlerState.HttpClientName, configureClient);
+
 				builder.AddHttpMessageHandler<HordeHttpAuthHandler>();
 			}
 
diff --git a/Engine/Source/Programs/Shared/EpicGames.Horde/Storage/Backends/HttpStorageBackend.cs b/Engine/Source/Programs/Shared/EpicGames.Horde/Storage/Backends/HttpStorageBackend.cs
--- a/Engine/Source/Programs/Shared/EpicGames.Horde/Storage/Backends/HttpStorageBackend.cs
+++ b/Engine/Source/Programs/Shared/EpicGames.Horde/Storage/Backends/HttpStorageBackend.cs
@@ -32,6 +32,7 @@ namespace EpicGames.Horde.Storage.Backends
 			public bool? SupportsRedirects { get; set; }
 		}
 
+		readonly string _basePath;
 		readonly Func<HttpClient> _createClient;
 		readonly ILogger _logger;
 		bool _supportsUploadRedirects = true;
@@ -42,8 +43,9 @@ namespace EpicGames.Horde.Storage.Backends
 		/// <summary>
 		/// Constructor
 		/// </summary>
-		public HttpStorageBackend(Func<HttpClient> createClient, ILogger logger)
+		public HttpStorageBackend(string basePath, Func<HttpClient> createClient, ILogger logger)
 		{
+			_basePath = basePath.TrimEnd('/');
 			_createClient = createClient;
 			_logger = logger;
 		}
@@ -78,7 +80,7 @@ namespace EpicGames.Horde.Storage.Backends
 
 			using (HttpClient httpClient = _createClient())
 			{
-				using (HttpRequestMessage request = new HttpRequestMessage(HttpMethod.Get, $"blobs/{path}"))
+				using (HttpRequestMessage request = new HttpRequestMessage(HttpMethod.Get, $"{_basePath}/blobs/{path}"))
 				{
 					if (offset != 0 || length != null)
 					{
@@ -144,7 +146,7 @@ namespace EpicGames.Horde.Storage.Backends
 		{
 			using (HttpClient httpClient = _createClient())
 			{
-				using (HttpRequestMessage request = new HttpRequestMessage(HttpMethod.Post, $"blobs"))
+				using (HttpRequestMessage request = new HttpRequestMessage(HttpMethod.Post, $"{_basePath}/blobs"))
 				{
 					using StringContent stringContent = new StringContent(prefix ?? String.Empty);
 
diff --git a/Engine/Source/Programs/Shared/EpicGames.Horde/Storage/Clients/HttpStorageClient.cs b/Engine/Source/Programs/Shared/EpicGames.Horde/Storage/Clients/HttpStorageClient.cs
--- a/Engine/Source/Programs/Shared/EpicGames.Horde/Storage/Clients/HttpStorageClient.cs
+++ b/Engine/Source/Programs/Shared/EpicGames.Horde/Storage/Clients/HttpStorageClient.cs
@@ -13,10 +13,6 @@ using System.Web;
 using EpicGames.Core;
 using Microsoft.Extensions.Logging;
 using EpicGames.Horde.Storage.Bundles;
-using Microsoft.Extensions.Http;
-using Polly.Extensions.Http;
-using Polly;
-using Polly.Retry;
 using EpicGames.Horde.Storage.Backends;
 
 namespace EpicGames.Horde.Storage.Clients
@@ -26,11 +22,6 @@ namespace EpicGames.Horde.Storage.Clients
 	/// </summary>
 	public class HttpStorageClient : BundleStorageClient
 	{
-		/// <summary>
-		/// Name of clients created from the http client factory
-		/// </summary>
-		public const string HttpClientName = "Horde.HttpStorageClient";
-
 		class WriteBlobResponse
 		{
 			public BundleLocator Blob { get; set; }
@@ -57,76 +48,21 @@ namespace EpicGames.Horde.Storage.Clients
 			public int ExportIdx { get; set; }
 		}
 
+		readonly string _basePath;
 		readonly Func<HttpClient> _createClient;
 		readonly ILogger _logger;
 
-		static readonly HttpMessageHandler s_defaultHttpMessageHandler = CreateDefaultHttpMessageHandler();
-
-		/// <summary>
-		/// Constructor
-		/// </summary>
-		public HttpStorageClient(Uri baseAddress, string? bearerToken, BundleReaderCache cache, ILogger logger)
-			: this(() => CreateAuthenticatedClient(null, baseAddress, bearerToken), cache, logger)
-		{
-		}
-
 		/// <summary>
 		/// Constructor
 		/// </summary>
-		public HttpStorageClient(Func<HttpClient> createClient, BundleReaderCache cache, ILogger logger) 
-			: base(new HttpStorageBackend(createClient, logger), cache, logger)
+		public HttpStorageClient(string basePath, Func<HttpClient> createClient, IStorageBackend backend, BundleReaderCache cache, ILogger logger) 
+			: base(backend, cache, logger)
 		{
+			_basePath = basePath.TrimEnd('/');
 			_createClient = createClient;
 			_logger = logger;
 		}
 
-		/// <summary>
-		/// Constructor
-		/// </summary>
-		public HttpStorageClient(IHttpClientFactory httpClientFactory, Uri baseAddress, string? bearerToken, BundleReaderCache cache, ILogger logger)
-			: this(() => CreateAuthenticatedClient(httpClientFactory, baseAddress, bearerToken), cache, logger)
-		{
-		}
-
-		static HttpMessageHandler CreateDefaultHttpMessageHandler()
-		{
-			AsyncRetryPolicy<HttpResponseMessage> retryPolicy = HttpPolicyExtensions
-				.HandleTransientHttpError()
-				.WaitAndRetryAsync(new[] { TimeSpan.FromSeconds(2.0), TimeSpan.FromSeconds(5.0), TimeSpan.FromSeconds(10), TimeSpan.FromSeconds(30) });
-
-			SocketsHttpHandler socketsHandler = new SocketsHttpHandler { PooledConnectionLifetime = TimeSpan.FromMinutes(15) };
-			return new PolicyHttpMessageHandler(retryPolicy) { InnerHandler = socketsHandler };
-		}
-
-		/// <summary>
-		/// Helper method to create an HTTP client from the given factory
-		/// </summary>
-		static HttpClient CreateClient(IHttpClientFactory? httpClientFactory)
-		{
-			if (httpClientFactory == null)
-			{
-				return new HttpClient(s_defaultHttpMessageHandler, disposeHandler: false);
-			}
-			else
-			{
-				return httpClientFactory.CreateClient(HttpClientName);
-			}
-		}
-
-		/// <summary>
-		/// Helper method to add the base address and auth header to an HTTP client
-		/// </summary>
-		static HttpClient CreateAuthenticatedClient(IHttpClientFactory? httpClientFactory, Uri baseAddress, string? bearerToken)
-		{
-			HttpClient httpClient = CreateClient(httpClientFactory);
-			httpClient.BaseAddress = baseAddress;
-			if (bearerToken != null)
-			{
-				httpClient.DefaultRequestHeaders.Authorization = new AuthenticationHeaderValue("Bearer", bearerToken);
-			}
-			return httpClient;
-		}
-
 		#region Nodes
 
 		/// <inheritdoc/>
@@ -147,7 +83,7 @@ namespace EpicGames.Horde.Storage.Clients
 			_logger.LogDebug("Finding nodes with alias {Alias}", alias);
 			using (HttpClient httpClient = _createClient())
 			{
-				using (HttpRequestMessage request = new HttpRequestMessage(HttpMethod.Get, $"nodes?alias={HttpUtility.UrlEncode(alias.ToString())}"))
+				using (HttpRequestMessage request = new HttpRequestMessage(HttpMethod.Get, $"{_basePath}/nodes?alias={HttpUtility.UrlEncode(alias.ToString())}"))
 				{
 					using (HttpResponseMessage response = await httpClient.SendAsync(request, cancellationToken))
 					{
@@ -173,7 +109,7 @@ namespace EpicGames.Horde.Storage.Clients
 			_logger.LogDebug("Deleting ref {RefName}", name);
 			using (HttpClient httpClient = _createClient())
 			{
-				using (HttpRequestMessage request = new HttpRequestMessage(HttpMethod.Delete, $"refs/{name}"))
+				using (HttpRequestMessage request = new HttpRequestMessage(HttpMethod.Delete, $"{_basePath}/refs/{name}"))
 				{
 					using (HttpResponseMessage response = await httpClient.SendAsync(request, cancellationToken))
 					{
@@ -198,7 +134,7 @@ namespace EpicGames.Horde.Storage.Clients
 		{
 			using (HttpClient httpClient = _createClient())
 			{
-				using (HttpRequestMessage request = new HttpRequestMessage(HttpMethod.Get, $"refs/{name}"))
+				using (HttpRequestMessage request = new HttpRequestMessage(HttpMethod.Get, $"{_basePath}/refs/{name}"))
 				{
 					if (cacheTime.IsSet())
 					{
@@ -235,7 +171,7 @@ namespace EpicGames.Horde.Storage.Clients
 			_logger.LogDebug("Writing ref {RefName} -> {RefTarget}", name, locator);
 			using (HttpClient httpClient = _createClient())
 			{
-				using (HttpResponseMessage response = await httpClient.PutAsync($"refs/{name}", new { blob = locator.Blob, exportIdx = locator.ExportIdx, options }, cancellationToken))
+				using (HttpResponseMessage response = await httpClient.PutAsync($"{_basePath}/refs/{name}", new { blob = locator.Blob, exportIdx = locator.ExportIdx, options }, cancellationToken))
 				{
 					response.EnsureSuccessStatusCode();
 				}
@@ -244,4 +180,63 @@ namespace EpicGames.Horde.Storage.Clients
 
 		#endregion
 	}
+
+	/// <summary>
+	/// Factory for constructing HttpStorageClient instances
+	/// </summary>
+	public class HttpStorageClientFactory : IStorageClientFactory
+	{
+		readonly IHttpClientFactory _httpClientFactory;
+		readonly StorageBackendCache _backendCache;
+		readonly BundleReaderCache _readerCache;
+		readonly ILogger<HttpStorageBackend> _backendLogger;
+		readonly ILogger<HttpStorageClient> _clientLogger;
+
+		/// <summary>
+		/// Constructor
+		/// </summary>
+		public HttpStorageClientFactory(IHttpClientFactory httpClientFactory, StorageBackendCache backendCache, BundleReaderCache readerCache, ILogger<HttpStorageBackend> backendLogger, ILogger<HttpStorageClient> clientLogger)
+		{
+			_httpClientFactory = httpClientFactory;
+			_backendCache = backendCache;
+			_readerCache = readerCache;
+			_backendLogger = backendLogger;
+			_clientLogger = clientLogger;
+		}
+
+		/// <summary>
+		/// Creates a new HTTP storage client
+		/// </summary>
+		/// <param name="basePath">Base path for all requests</param>
+		/// <param name="accessToken">Custom access token to use for requests</param>
+		public HttpStorageClient CreateClient(string basePath, string? accessToken = null)
+		{
+			HttpClient CreateClient()
+			{
+				HttpClient httpClient = _httpClientFactory.CreateClient(HordeHttpClient.HttpClientName);
+				if (accessToken != null)
+				{
+					httpClient.DefaultRequestHeaders.Authorization = new AuthenticationHeaderValue("Bearer", accessToken);
+				}
+				return httpClient;
+			}
+
+			IStorageBackend backend = new HttpStorageBackend(basePath, CreateClient, _backendLogger);
+			if (_backendCache != null)
+			{
+				backend = _backendCache.CreateWrapper(basePath, backend);
+			}
+			return new HttpStorageClient(basePath, CreateClient, backend, _readerCache, _clientLogger);
+		}
+
+		/// <summary>
+		/// Creates a new HTTP storage client
+		/// </summary>
+		/// <param name="namespaceId">Namespace to create a client for</param>
+		/// <param name="accessToken">Custom access token to use for requests</param>
+		public HttpStorageClient CreateClient(NamespaceId namespaceId, string? accessToken = null) => CreateClient($"api/v1/storage/{namespaceId}", accessToken);
+
+		/// <inheritdoc/>
+		IStorageClient IStorageClientFactory.CreateClient(NamespaceId namespaceId) => CreateClient(namespaceId);
+	}
 }
diff --git a/Engine/Source/Programs/Shared/EpicGames.Horde/Storage/IStorageBackend.cs b/Engine/Source/Programs/Shared/EpicGames.Horde/Storage/IStorageBackend.cs
--- a/Engine/Source/Programs/Shared/EpicGames.Horde/Storage/IStorageBackend.cs
+++ b/Engine/Source/Programs/Shared/EpicGames.Horde/Storage/IStorageBackend.cs
@@ -230,7 +230,16 @@ namespace EpicGames.Horde.Storage
 			{
 				_inner = inner;
 				_offset = offset;
-				_length = length ?? (inner.Data.Length - offset);
+
+				int maxLength = inner.Data.Length - offset;
+				if (length == null)
+				{
+					_length = maxLength;
+				}
+				else
+				{
+					_length = Math.Min(length.Value, maxLength);
+				}
 			}
 
 			public void Dispose() => _inner.Dispose();
diff --git a/Engine/Source/Programs/Shared/EpicGames.Horde/Storage/IStorageClient.cs b/Engine/Source/Programs/Shared/EpicGames.Horde/Storage/IStorageClient.cs
--- a/Engine/Source/Programs/Shared/EpicGames.Horde/Storage/IStorageClient.cs
+++ b/Engine/Source/Programs/Shared/EpicGames.Horde/Storage/IStorageClient.cs
@@ -203,9 +203,8 @@ namespace EpicGames.Horde.Storage
 		/// Creates a storage client for the given namespace
 		/// </summary>
 		/// <param name="namespaceId">Namespace to manipulate</param>
-		/// <param name="cancellationToken">Cancellation token for the operation</param>
 		/// <returns>Storage client instance. Must be disposed by the caller.</returns>
-		Task<IStorageClient> CreateClientAsync(NamespaceId namespaceId, CancellationToken cancellationToken = default);
+		IStorageClient CreateClient(NamespaceId namespaceId);
 	}
 
 	/// <summary>
diff --git a/Engine/Source/Programs/Shared/EpicGames.Horde/Storage/StorageBackendCache.cs b/Engine/Source/Programs/Shared/EpicGames.Horde/Storage/StorageBackendCache.cs
--- a/Engine/Source/Programs/Shared/EpicGames.Horde/Storage/StorageBackendCache.cs
+++ b/Engine/Source/Programs/Shared/EpicGames.Horde/Storage/StorageBackendCache.cs
@@ -3,6 +3,7 @@
 using EpicGames.Core;
 using EpicGames.Horde.Storage.Backends;
 using Microsoft.Extensions.Logging;
+using Microsoft.Extensions.Logging.Abstractions;
 using System;
 using System.Collections.Generic;
 using System.IO;
@@ -113,10 +114,32 @@ namespace EpicGames.Horde.Storage
 		/// <summary>
 		/// Constructor
 		/// </summary>
-		public StorageBackendCache(DirectoryReference cacheDir, long maxSize, ILogger logger)
+		public StorageBackendCache()
+			: this(null, null)
 		{
+		}
+
+		/// <summary>
+		/// Constructor
+		/// </summary>
+		public StorageBackendCache(DirectoryReference? cacheDir, long? maxSize)
+			: this(cacheDir, maxSize, NullLogger.Instance)
+		{
+		}
+
+		/// <summary>
+		/// Constructor
+		/// </summary>
+		/// <param name="cacheDir">Directory to storage cache files. Will be cleared on startup. Defaults to a randomly generated directory in the users's temp folder.</param>
+		/// <param name="maxSize">Maximum size of the cache. Defaults to 50mb.</param>
+		/// <param name="logger">Logger for error/warning messages</param>
+		public StorageBackendCache(DirectoryReference? cacheDir, long? maxSize, ILogger logger)
+		{
+			cacheDir ??= new DirectoryReference(Path.Combine(Path.GetTempPath(), $"horde-{Guid.NewGuid().ToString("n")}"));
+			FileUtils.ForceDeleteDirectoryContents(cacheDir);
+
 			_backend = new FileStorageBackend(cacheDir);
-			_maxSize = maxSize;
+			_maxSize = maxSize ?? (50 * 1024 * 1024);
 			_logger = logger;
 		}
 
