Horde: Move Horde HTTP request wrappers into HordeHttpClient, and add extension methods for fetching from an IHttpClientFactory.

diff --git a/Engine/Source/Programs/Horde/Horde/Commands/Compute/ComputeRun.cs b/Engine/Source/Programs/Horde/Horde/Commands/Compute/ComputeRun.cs
--- a/Engine/Source/Programs/Horde/Horde/Commands/Compute/ComputeRun.cs
+++ b/Engine/Source/Programs/Horde/Horde/Commands/Compute/ComputeRun.cs
@@ -3,7 +3,6 @@
 using System.Reflection;
 using System.Text.Json;
 using EpicGames.Core;
-using EpicGames.Horde;
 using EpicGames.Horde.Common;
 using EpicGames.Horde.Compute;
 using EpicGames.Horde.Compute.Clients;
@@ -50,14 +49,14 @@ namespace Horde.Commands.Compute
 		[CommandLine("-Task=", Required = true)]
 		FileReference TaskFile { get; set; } = null!;
 
-		readonly HordeHttpClientFactory _hordeHttpClientFactory;
+		readonly IHttpClientFactory _httpClientFactory;
 
 		/// <summary>
 		/// Constructor
 		/// </summary>
-		public ComputeRun(HordeHttpClientFactory hordeHttpClientFactory)
+		public ComputeRun(IHttpClientFactory httpClientFactory)
 		{
-			_hordeHttpClientFactory = hordeHttpClientFactory;
+			_httpClientFactory = httpClientFactory;
 		}
 
 		/// <inheritdoc/>
@@ -93,7 +92,7 @@ namespace Horde.Commands.Compute
 			}
 			else
 			{
-				return new ServerComputeClient(_hordeHttpClientFactory, logger);
+				return new ServerComputeClient(_httpClientFactory, logger);
 			}
 		}
 
diff --git a/Engine/Source/Programs/Horde/Horde/Commands/LoginCommand.cs b/Engine/Source/Programs/Horde/Horde/Commands/LoginCommand.cs
--- a/Engine/Source/Programs/Horde/Horde/Commands/LoginCommand.cs
+++ b/Engine/Source/Programs/Horde/Horde/Commands/LoginCommand.cs
@@ -23,13 +23,13 @@ namespace Horde.Commands
 		[CommandLine("-Server=")]
 		public string? Server { get; set; }
 
+		readonly IHttpClientFactory _httpClientFactory;
 		readonly CmdConfig _config;
-		readonly IServiceProvider _serviceProvider;
 
-		public LoginCommand(IOptions<CmdConfig> config, IServiceProvider serviceProvider)
+		public LoginCommand(IHttpClientFactory httpClientFactory, IOptions<CmdConfig> config)
 		{
+			_httpClientFactory = httpClientFactory;
 			_config = config.Value;
-			_serviceProvider = serviceProvider;
 		}
 
 		/// <inheritdoc/>
@@ -41,8 +41,9 @@ namespace Horde.Commands
 				await _config.WriteAsync();
 			}
 
-			HordeHttpClient httpClient = _serviceProvider.GetRequiredService<HordeHttpClient>();
-			GetServerInfoResponse serverInfo = await httpClient.GetAsync<GetServerInfoResponse>("api/v1/server/info");
+			using HordeHttpClient httpClient = _httpClientFactory.CreateHordeClient();
+
+			GetServerInfoResponse serverInfo = await httpClient.GetServerInfoAsync();
 			logger.LogInformation("Connected to server version: {Version}", serverInfo.ServerVersion);
 			
 			return 0;
diff --git a/Engine/Source/Programs/Shared/EpicGames.Horde/Api/ArtifactMessages.cs b/Engine/Source/Programs/Shared/EpicGames.Horde/Api/ArtifactMessages.cs
--- a/Engine/Source/Programs/Shared/EpicGames.Horde/Api/ArtifactMessages.cs
+++ b/Engine/Source/Programs/Shared/EpicGames.Horde/Api/ArtifactMessages.cs
@@ -139,148 +139,4 @@ namespace EpicGames.Horde.Api
 		/// </summary>
 		public List<string> Filter { get; set; } = new List<string>();
 	}
-
-	/// <summary>
-	/// Extension methods for the artifacts endpoint
-	/// </summary>
-	public static class ArtifactExtensions
-	{
-		/// <summary>
-		/// Gets metadata about an artifact object
-		/// </summary>
-		/// <param name="horde">The Horde client instance</param>
-		/// <param name="id">Identifier for the artifact</param>
-		/// <param name="cancellationToken">Cancellation token for the operation</param>
-		public static Task<GetArtifactResponse> GetArtifactAsync(this HordeHttpClient horde, ArtifactId id, CancellationToken cancellationToken = default)
-		{
-			return horde.GetAsync<GetArtifactResponse>($"api/v2/artifacts/{id}", cancellationToken);
-		}
-
-		/// <summary>
-		/// Finds artifacts with a set of ids or keys
-		/// </summary>
-		/// <param name="horde">The Horde client instance</param>
-		/// <param name="ids">Artifact ids to return</param>
-		/// <param name="keys">Keys to find</param>
-		/// <param name="cancellationToken">Cancellation token for the operation</param>
-		/// <returns>Information about all the artifacts</returns>
-		public static async Task<List<GetArtifactResponse>> FindArtifactsAsync(this HordeHttpClient horde, IEnumerable<ArtifactId>? ids = null, IEnumerable<string>? keys = null, CancellationToken cancellationToken = default)
-		{
-			QueryStringBuilder queryParams = new QueryStringBuilder();
-			if (ids != null)
-			{
-				foreach (ArtifactId id in ids)
-				{
-					queryParams.Add("id", id.ToString());
-				}
-			}
-			if (keys != null)
-			{
-				foreach (string key in keys)
-				{
-					queryParams.Add("key", key);
-				}
-			}
-
-			FindArtifactsResponse response = await horde.GetAsync<FindArtifactsResponse>($"api/v2/artifacts?{queryParams}", cancellationToken);
-			return response.Artifacts;
-		}
-
-#if false
-			/// <summary>
-			/// Gets metadata about an artifact object
-			/// </summary>
-			/// <param name="id">Identifier of the artifact to retrieve</param>
-			/// <param name="filter">Filter for returned properties</param>
-			/// <returns>Information about all the artifacts</returns>
-			[HttpGet]
-			[Route("/api/v2/artifacts/{id}")]
-			[ProducesResponseType(typeof(FindArtifactsResponse), 200)]
-			public async Task<ActionResult<object>> GetArtifactAsync(ArtifactId id, [FromQuery] PropertyFilter? filter = null);
-
-			/// <summary>
-			/// Retrieves bundles for a particular artifact
-			/// </summary>
-			/// <param name="id">Identifier of the artifact to retrieve</param>
-			/// <param name="locator">The blob locator</param>
-			/// <param name="length">Length of data to return</param>
-			/// <param name="offset">Offset of the data to return</param>
-			/// <param name="cancellationToken">Cancellation token for the operation</param>
-			/// <returns>Information about all the artifacts</returns>
-			[HttpGet]
-			[Route("/api/v2/artifacts/{id}/blobs/{*locator}")]
-			[Route("/api/v2/artifacts/{id}/bundles/{*locator}")]
-			public async Task<ActionResult<object>> ReadArtifactBlobAsync(ArtifactId id, BundleLocator locator, [FromQuery] int? offset = null, [FromQuery] int? length = null, CancellationToken cancellationToken = default)
-
-			/// <summary>
-			/// Retrieves the root blob for an artifact
-			/// </summary>
-			/// <param name="id">Identifier of the artifact to retrieve</param>
-			/// <param name="cancellationToken">Cancellation token for the operation</param>
-			/// <returns>Information about all the artifacts</returns>
-			[HttpGet]
-			[Route("/api/v2/artifacts/{id}/ref")]
-			public async Task<ActionResult<object>> ReadArtifactRefAsync(ArtifactId id, CancellationToken cancellationToken = default)
-
-			/// <summary>
-			/// Gets metadata about an artifact object
-			/// </summary>
-			/// <param name="id">Identifier of the artifact to retrieve</param>
-			/// <param name="path">Path to fetch</param>
-			/// <param name="filter">Filter for returned properties</param>
-			/// <param name="cancellationToken">Cancellation token for the operation</param>
-			/// <returns>Information about all the artifacts</returns>
-			[HttpGet]
-			[Route("/api/v2/artifacts/{id}/browse")]
-			[ProducesResponseType(typeof(GetArtifactDirectoryResponse), 200)]
-			public async Task<ActionResult<object>> BrowseArtifactAsync(ArtifactId id, [FromQuery] string? path = null, [FromQuery] PropertyFilter? filter = null, CancellationToken cancellationToken = default)
-
-			/// <summary>
-			/// Downloads an individual file from an artifact
-			/// </summary>
-			/// <param name="id">Identifier of the artifact to retrieve</param>
-			/// <param name="path">Path to fetch</param>
-			/// <param name="inline">Whether to request the file be downloaded vs displayed inline</param>
-			/// <param name="cancellationToken">Cancellation token for the operation</param>
-			/// <returns>Information about all the artifacts</returns>
-			[HttpGet]
-			[Route("/api/v2/artifacts/{id}/file")]
-			public async Task<ActionResult<object>> GetFileAsync(ArtifactId id, [FromQuery] string path, [FromQuery] bool inline = false, CancellationToken cancellationToken = default)
-
-			/// <summary>
-			/// Downloads an individual file from an artifact
-			/// </summary>
-			/// <param name="id">Identifier of the artifact to retrieve</param>
-			/// <param name="filter">Paths to include in the zip file. The post version of this request allows for more parameters than can fit in a request string.</param>
-			/// <param name="cancellationToken">Cancellation token for the operation</param>
-			/// <returns>Information about all the artifacts</returns>
-			[HttpGet]
-			[Route("/api/v2/artifacts/{id}/zip")]
-			public async Task<ActionResult<object>> GetZipAsync(ArtifactId id, [FromQuery(Name = "filter")] string[]? filter, CancellationToken cancellationToken = default)
-
-			/// <summary>
-			/// Downloads an individual file from an artifact
-			/// </summary>
-			/// <param name="id">Identifier of the artifact to retrieve</param>
-			/// <param name="request">Filter for the zip file</param>
-			/// <param name="cancellationToken">Cancellation token for the operation</param>
-			/// <returns>Information about all the artifacts</returns>
-			[HttpPost]
-			[Route("/api/v2/artifacts/{id}/zip")]
-			public async Task<ActionResult<object>> CreateZipFromFilterAsync(ArtifactId id, CreateZipRequest request, CancellationToken cancellationToken = default)
-
-			/// <summary>
-			/// Gets metadata about an artifact object
-			/// </summary>
-			/// <param name="ids">Artifact ids to return</param>
-			/// <param name="keys">Keys to find</param>
-			/// <param name="filter">Filter for returned values</param>
-			/// <returns>Information about all the artifacts</returns>
-			[HttpGet]
-			[Route("/api/v2/artifacts")]
-			[ProducesResponseType(typeof(FindArtifactsResponse), 200)]
-			public async Task<ActionResult<object>> FindArtifactsAsync([FromQuery(Name = "id")] IEnumerable<ArtifactId>? ids = null, [FromQuery(Name = "key")] IEnumerable<string>? keys = null, [FromQuery] PropertyFilter? filter = null)
-		}
-#endif
-	}
 }
diff --git a/Engine/Source/Programs/Shared/EpicGames.Horde/Api/ProjectMessages.cs b/Engine/Source/Programs/Shared/EpicGames.Horde/Api/ProjectMessages.cs
--- a/Engine/Source/Programs/Shared/EpicGames.Horde/Api/ProjectMessages.cs
+++ b/Engine/Source/Programs/Shared/EpicGames.Horde/Api/ProjectMessages.cs
@@ -2,8 +2,6 @@
 
 using System.Collections.Generic;
 using System.Diagnostics;
-using System.Threading;
-using System.Threading.Tasks;
 
 #pragma warning disable CA2227 // Collection properties should be read only
 
@@ -125,35 +123,4 @@ namespace EpicGames.Horde.Api
 			Row = row;
 		}
 	}
-
-	/// <summary>
-	/// Extension methods for the projects endpoint
-	/// </summary>
-	public static class ProjectExtensions
-	{
-		/// <summary>
-		/// Query all the projects
-		/// </summary>
-		/// <param name="horde">The horde client instance</param>
-		/// <param name="includeStreams">Whether to include streams in the response</param>
-		/// <param name="includeCategories">Whether to include categories in the response</param>
-		/// <param name="cancellationToken">Cancellation token for the operation</param>
-		/// <returns>Information about all the projects</returns>
-		public static Task<List<GetProjectResponse>> GetProjectsAsync(this HordeHttpClient horde, bool includeStreams = false, bool includeCategories = false, CancellationToken cancellationToken = default)
-		{
-			return horde.GetAsync<List<GetProjectResponse>>($"api/v1/projects?includeStreams={includeStreams}&includeCategories={includeCategories}", cancellationToken);
-		}
-
-		/// <summary>
-		/// Retrieve information about a specific project
-		/// </summary>
-		/// <param name="horde">The horde client instance</param>
-		/// <param name="projectId">Id of the project to get information about</param>
-		/// <param name="cancellationToken">Cancellation token for the operation</param>
-		/// <returns>Information about the requested project</returns>
-		public static Task<GetProjectResponse> GetProjectAsync(this HordeHttpClient horde, ProjectId projectId, CancellationToken cancellationToken = default)
-		{
-			return horde.GetAsync<GetProjectResponse>($"api/v1/projects/{projectId}", cancellationToken);
-		}
-	}
 }
diff --git a/Engine/Source/Programs/Shared/EpicGames.Horde/Api/SecretMessages.cs b/Engine/Source/Programs/Shared/EpicGames.Horde/Api/SecretMessages.cs
--- a/Engine/Source/Programs/Shared/EpicGames.Horde/Api/SecretMessages.cs
+++ b/Engine/Source/Programs/Shared/EpicGames.Horde/Api/SecretMessages.cs
@@ -50,33 +50,4 @@ namespace EpicGames.Horde.Api
 			Data = data;
 		}
 	}
-
-	/// <summary>
-	/// Extension methods for the secrets endpoint
-	/// </summary>
-	public static class SecretExtensions
-	{
-		/// <summary>
-		/// Query all the secrets available to the current user
-		/// </summary>
-		/// <param name="horde">The horde client instance</param>
-		/// <param name="cancellationToken">Cancellation token for the operation</param>
-		/// <returns>Information about all the projects</returns>
-		public static Task<GetSecretsResponse> GetSecretsAsync(this HordeHttpClient horde, CancellationToken cancellationToken = default)
-		{
-			return horde.GetAsync<GetSecretsResponse>($"api/v1/secrets", cancellationToken);
-		}
-
-		/// <summary>
-		/// Retrieve information about a specific project
-		/// </summary>
-		/// <param name="horde">The horde client instance</param>
-		/// <param name="secretId">Id of the secret to retrieve</param>
-		/// <param name="cancellationToken">Cancellation token for the operation</param>
-		/// <returns>Information about the requested project</returns>
-		public static Task<GetSecretResponse> GetSecretAsync(this HordeHttpClient horde, SecretId secretId, CancellationToken cancellationToken = default)
-		{
-			return horde.GetAsync<GetSecretResponse>($"api/v1/secrets/{secretId}", cancellationToken);
-		}
-	}
 }
diff --git a/Engine/Source/Programs/Shared/EpicGames.Horde/Compute/Clients/ServerComputeClient.cs b/Engine/Source/Programs/Shared/EpicGames.Horde/Compute/Clients/ServerComputeClient.cs
--- a/Engine/Source/Programs/Shared/EpicGames.Horde/Compute/Clients/ServerComputeClient.cs
+++ b/Engine/Source/Programs/Shared/EpicGames.Horde/Compute/Clients/ServerComputeClient.cs
@@ -80,18 +80,18 @@ namespace EpicGames.Horde.Compute.Clients
 			}
 		}
 
-		readonly HordeHttpClientFactory _hordeHttpClientFactory;
+		readonly IHttpClientFactory _httpClientFactory;
 		readonly CancellationTokenSource _cancellationSource = new CancellationTokenSource();
 		readonly ILogger _logger;
 
 		/// <summary>
 		/// Constructor
 		/// </summary>
-		/// <param name="hordeHttpClientFactory">Factory for constructing Horde http client instances</param>
+		/// <param name="httpClientFactory">Factory for constructing http client instances</param>
 		/// <param name="logger">Logger for diagnostic messages</param>
-		public ServerComputeClient(HordeHttpClientFactory hordeHttpClientFactory, ILogger logger)
+		public ServerComputeClient(IHttpClientFactory httpClientFactory, ILogger logger)
 		{
-			_hordeHttpClientFactory = hordeHttpClientFactory;
+			_httpClientFactory = httpClientFactory;
 			_logger = logger;
 		}
 
@@ -125,18 +125,15 @@ namespace EpicGames.Horde.Compute.Clients
 		{
 			_logger.LogDebug("Requesting compute resource");
 
-			JsonSerializerOptions jsonSerializerOptions = new JsonSerializerOptions();
-			HordeHttpClient.ConfigureJsonSerializer(jsonSerializerOptions);
-
 			// Assign a compute worker
-			HordeHttpClient client = _hordeHttpClientFactory.CreateClient();
+			HttpClient client = _httpClientFactory.CreateClient(HordeHttpClient.HttpClientName);
 
 			AssignComputeRequest request = new AssignComputeRequest();
 			request.Requirements = requirements;
 			request.RequestId = requestId;
 
 			AssignComputeResponse? responseMessage;
-			using (HttpResponseMessage response = await client.PostAsync($"api/v2/compute/{clusterId}", request, _cancellationSource.Token))
+			using (HttpResponseMessage response = await HordeHttpClient.PostAsync(client, $"api/v2/compute/{clusterId}", request, _cancellationSource.Token))
 			{
 				if (response.StatusCode == HttpStatusCode.NotFound)
 				{
@@ -151,7 +148,7 @@ namespace EpicGames.Horde.Compute.Clients
 
 				response.EnsureSuccessStatusCode();
 
-				responseMessage = await response.Content.ReadFromJsonAsync<AssignComputeResponse>(jsonSerializerOptions, cancellationToken);
+				responseMessage = await response.Content.ReadFromJsonAsync<AssignComputeResponse>(HordeHttpClient.JsonSerializerOptions, cancellationToken);
 				if (responseMessage == null)
 				{
 					throw new InvalidOperationException();
diff --git a/Engine/Source/Programs/Shared/EpicGames.Horde/HordeHttpClient.cs b/Engine/Source/Programs/Shared/EpicGames.Horde/HordeHttpClient.cs
--- a/Engine/Source/Programs/Shared/EpicGames.Horde/HordeHttpClient.cs
+++ b/Engine/Source/Programs/Shared/EpicGames.Horde/HordeHttpClient.cs
@@ -1,12 +1,15 @@
 // Copyright Epic Games, Inc. All Rights Reserved.
 
 using System;
+using System.Collections.Generic;
 using System.Net.Http;
 using System.Net.Http.Json;
 using System.Text.Json;
 using System.Text.Json.Serialization;
 using System.Threading;
 using System.Threading.Tasks;
+using EpicGames.Core;
+using EpicGames.Horde.Api;
 using Microsoft.Extensions.DependencyInjection;
 using Polly;
 using Polly.Extensions.Http;
@@ -16,7 +19,7 @@ namespace EpicGames.Horde
 	/// <summary>
 	/// Wraps an Http client which communicates with the Horde server
 	/// </summary>
-	public sealed class HordeHttpClient
+	public sealed class HordeHttpClient : IDisposable
 	{
 		/// <summary>
 		/// Name of clients created from the http client factory
@@ -27,10 +30,7 @@ namespace EpicGames.Horde
 
 		static readonly JsonSerializerOptions s_jsonSerializerOptions = CreateJsonSerializerOptions();
 
-		/// <summary>
-		/// Static instance of the default retry policy
-		/// </summary>
-		public static IAsyncPolicy<HttpResponseMessage> DefaultRetryPolicy { get; } = CreateDefaultRetryPolicy();
+		internal static JsonSerializerOptions JsonSerializerOptions => s_jsonSerializerOptions;
 
 		/// <summary>
 		/// Accessor for the underlying HTTP client
@@ -46,6 +46,12 @@ namespace EpicGames.Horde
 			_httpClient = httpClient;
 		}
 
+		/// <inheritdoc/>
+		public void Dispose()
+		{
+			_httpClient.Dispose();
+		}
+
 		/// <summary>
 		/// Create the shared instance of JSON options for HordeHttpClient instances
 		/// </summary>
@@ -72,16 +78,128 @@ namespace EpicGames.Horde
 			options.Converters.Add(new BinaryIdJsonConverterFactory());
 		}
 
+		#region Artifacts
+
+		/// <summary>
+		/// Gets metadata about an artifact object
+		/// </summary>
+		/// <param name="id">Identifier for the artifact</param>
+		/// <param name="cancellationToken">Cancellation token for the operation</param>
+		public Task<GetArtifactResponse> GetArtifactAsync(ArtifactId id, CancellationToken cancellationToken = default)
+		{
+			return GetAsync<GetArtifactResponse>(_httpClient, $"api/v2/artifacts/{id}", cancellationToken);
+		}
+
+		/// <summary>
+		/// Finds artifacts with a set of ids or keys
+		/// </summary>
+		/// <param name="ids">Artifact ids to return</param>
+		/// <param name="keys">Keys to find</param>
+		/// <param name="cancellationToken">Cancellation token for the operation</param>
+		/// <returns>Information about all the artifacts</returns>
+		public async Task<List<GetArtifactResponse>> FindArtifactsAsync(IEnumerable<ArtifactId>? ids = null, IEnumerable<string>? keys = null, CancellationToken cancellationToken = default)
+		{
+			QueryStringBuilder queryParams = new QueryStringBuilder();
+			if (ids != null)
+			{
+				foreach (ArtifactId id in ids)
+				{
+					queryParams.Add("id", id.ToString());
+				}
+			}
+			if (keys != null)
+			{
+				foreach (string key in keys)
+				{
+					queryParams.Add("key", key);
+				}
+			}
+
+			FindArtifactsResponse response = await GetAsync<FindArtifactsResponse>(_httpClient, $"api/v2/artifacts?{queryParams}", cancellationToken);
+			return response.Artifacts;
+		}
+
+		#endregion
+
+		#region Projects
+
+		/// <summary>
+		/// Query all the projects
+		/// </summary>
+		/// <param name="includeStreams">Whether to include streams in the response</param>
+		/// <param name="includeCategories">Whether to include categories in the response</param>
+		/// <param name="cancellationToken">Cancellation token for the operation</param>
+		/// <returns>Information about all the projects</returns>
+		public Task<List<GetProjectResponse>> GetProjectsAsync(bool includeStreams = false, bool includeCategories = false, CancellationToken cancellationToken = default)
+		{
+			return GetAsync<List<GetProjectResponse>>(_httpClient, $"api/v1/projects?includeStreams={includeStreams}&includeCategories={includeCategories}", cancellationToken);
+		}
+
+		/// <summary>
+		/// Retrieve information about a specific project
+		/// </summary>
+		/// <param name="projectId">Id of the project to get information about</param>
+		/// <param name="cancellationToken">Cancellation token for the operation</param>
+		/// <returns>Information about the requested project</returns>
+		public Task<GetProjectResponse> GetProjectAsync(ProjectId projectId, CancellationToken cancellationToken = default)
+		{
+			return GetAsync<GetProjectResponse>(_httpClient, $"api/v1/projects/{projectId}", cancellationToken);
+		}
+
+		#endregion
+
+		#region Secrets
+
+		/// <summary>
+		/// Query all the secrets available to the current user
+		/// </summary>
+		/// <param name="cancellationToken">Cancellation token for the operation</param>
+		/// <returns>Information about all the projects</returns>
+		public Task<GetSecretsResponse> GetSecretsAsync(CancellationToken cancellationToken = default)
+		{
+			return GetAsync<GetSecretsResponse>(_httpClient, $"api/v1/secrets", cancellationToken);
+		}
+
+		/// <summary>
+		/// Retrieve information about a specific project
+		/// </summary>
+		/// <param name="secretId">Id of the secret to retrieve</param>
+		/// <param name="cancellationToken">Cancellation token for the operation</param>
+		/// <returns>Information about the requested project</returns>
+		public Task<GetSecretResponse> GetSecretAsync(SecretId secretId, CancellationToken cancellationToken = default)
+		{
+			return GetAsync<GetSecretResponse>(_httpClient, $"api/v1/secrets/{secretId}", cancellationToken);
+		}
+
+		#endregion
+
+		#region Server
+
+		/// <summary>
+		/// Gets information about the currently deployed server version
+		/// </summary>
+		/// <param name="cancellationToken">Cancellation token for the operation</param>
+		/// <returns>Information about the deployed server instance</returns>
+		public Task<GetServerInfoResponse> GetServerInfoAsync(CancellationToken cancellationToken = default)
+		{
+			return GetAsync<GetServerInfoResponse>(_httpClient, "api/v1/server/info", cancellationToken);
+		}
+
+		#endregion
+
+		#region Utility Methods
+
 		/// <summary>
 		/// Gets a resource from an HTTP endpoint and parses it as a JSON object
 		/// </summary>
 		/// <typeparam name="TResponse">The object type to return</typeparam>
+		/// <param name="httpClient">Http client instance</param>
 		/// <param name="relativePath">The url to retrieve</param>
 		/// <param name="cancellationToken">Cancels the request</param>
 		/// <returns>New instance of the object</returns>
-		public async Task<TResponse> GetAsync<TResponse>(string relativePath, CancellationToken cancellationToken = default)
+		internal static async Task<TResponse> GetAsync<TResponse>(HttpClient httpClient, string relativePath, CancellationToken cancellationToken = default)
 		{
-			TResponse? response = await _httpClient.GetFromJsonAsync<TResponse>(relativePath, s_jsonSerializerOptions, cancellationToken);
+			TResponse? response = await httpClient.GetFromJsonAsync<TResponse>(relativePath, s_jsonSerializerOptions, cancellationToken);
 			return response ?? throw new InvalidCastException($"Expected non-null response from GET to {relativePath}");
 		}
 
@@ -89,13 +207,14 @@ namespace EpicGames.Horde
 		/// Posts an object to an HTTP endpoint as a JSON object, and parses the response object
 		/// </summary>
 		/// <typeparam name="TRequest">The object type to post</typeparam>
+		/// <param name="httpClient">Http client instance</param>
 		/// <param name="relativePath">The url to retrieve</param>
 		/// <param name="request">The object to post</param>
 		/// <param name="cancellationToken">Cancels the request</param>
 		/// <returns>The response parsed into the requested type</returns>
-		public async Task<HttpResponseMessage> PostAsync<TRequest>(string relativePath, TRequest request, CancellationToken cancellationToken = default)
+		internal static async Task<HttpResponseMessage> PostAsync<TRequest>(HttpClient httpClient, string relativePath, TRequest request, CancellationToken cancellationToken = default)
 		{
-			return await _httpClient.PostAsJsonAsync<TRequest>(relativePath, request, s_jsonSerializerOptions, cancellationToken);
+			return await httpClient.PostAsJsonAsync<TRequest>(relativePath, request, s_jsonSerializerOptions, cancellationToken);
 		}
 
 		/// <summary>
@@ -103,13 +222,14 @@ namespace EpicGames.Horde
 		/// </summary>
 		/// <typeparam name="TResponse">The object type to return</typeparam>
 		/// <typeparam name="TRequest">The object type to post</typeparam>
+		/// <param name="httpClient">Http client instance</param>
 		/// <param name="relativePath">The url to retrieve</param>
 		/// <param name="request">The object to post</param>
 		/// <param name="cancellationToken">Cancels the request</param>
 		/// <returns>The response parsed into the requested type</returns>
-		public async Task<TResponse> PostAsync<TResponse, TRequest>(string relativePath, TRequest request, CancellationToken cancellationToken = default)
+		internal static async Task<TResponse> PostAsync<TResponse, TRequest>(HttpClient httpClient, string relativePath, TRequest request, CancellationToken cancellationToken = default)
 		{
-			using (HttpResponseMessage response = await PostAsync<TRequest>(relativePath, request, cancellationToken))
+			using (HttpResponseMessage response = await PostAsync<TRequest>(httpClient, relativePath, request, cancellationToken))
 			{
 				response.EnsureSuccessStatusCode();
 
@@ -122,13 +242,14 @@ namespace EpicGames.Horde
 		/// Puts an object to an HTTP endpoint as a JSON object
 		/// </summary>
 		/// <typeparam name="TRequest">The object type to post</typeparam>
+		/// <param name="httpClient">Http client instance</param>
 		/// <param name="relativePath">The url to write to</param>
 		/// <param name="request">The object to post</param>
 		/// <param name="cancellationToken">Cancels the request</param>
 		/// <returns>Response message</returns>
-		public async Task<HttpResponseMessage> PutAsync<TRequest>(string relativePath, TRequest request, CancellationToken cancellationToken)
+		internal static async Task<HttpResponseMessage> PutAsync<TRequest>(HttpClient httpClient, string relativePath, TRequest request, CancellationToken cancellationToken)
 		{
-			return await _httpClient.PutAsJsonAsync<TRequest>(relativePath, request, s_jsonSerializerOptions, cancellationToken);
+			return await httpClient.PutAsJsonAsync<TRequest>(relativePath, request, s_jsonSerializerOptions, cancellationToken);
 		}
 
 		/// <summary>
@@ -136,13 +257,14 @@ namespace EpicGames.Horde
 		/// </summary>
 		/// <typeparam name="TResponse">The object type to return</typeparam>
 		/// <typeparam name="TRequest">The object type to post</typeparam>
+		/// <param name="httpClient">Http client instance</param>
 		/// <param name="relativePath">The url to write to</param>
 		/// <param name="request">The object to post</param>
 		/// <param name="cancellationToken">Cancels the request</param>
 		/// <returns>Response message</returns>
-		public async Task<TResponse> PutAsync<TResponse, TRequest>(string relativePath, TRequest request, CancellationToken cancellationToken)
+		internal static async Task<TResponse> PutAsync<TResponse, TRequest>(HttpClient httpClient, string relativePath, TRequest request, CancellationToken cancellationToken)
 		{
-			using (HttpResponseMessage response = await _httpClient.PutAsJsonAsync<TRequest>(relativePath, request, s_jsonSerializerOptions, cancellationToken))
+			using (HttpResponseMessage response = await httpClient.PutAsJsonAsync<TRequest>(relativePath, request, s_jsonSerializerOptions, cancellationToken))
 			{
 				response.EnsureSuccessStatusCode();
 
@@ -151,40 +273,22 @@ namespace EpicGames.Horde
 			}
 		}
 
-		/// <summary>
-		/// Helper method to construct a default retry policy for Horde requests
-		/// </summary>
-		static IAsyncPolicy<HttpResponseMessage> CreateDefaultRetryPolicy()
-		{
-			return HttpPolicyExtensions
-				.HandleTransientHttpError()
-				.WaitAndRetryAsync(new[] { TimeSpan.FromSeconds(2.0), TimeSpan.FromSeconds(5.0), TimeSpan.FromSeconds(10), TimeSpan.FromSeconds(30) });
-		}
+		#endregion
 	}
 
 	/// <summary>
-	/// Factory for creating HordeHttpClient instances
+	/// Extension methods for Horde HTTP clients
 	/// </summary>
-	public class HordeHttpClientFactory
+	public static class HordeHttpClientExtensions
 	{
-		readonly IServiceProvider _serviceProvider;
-
 		/// <summary>
-		/// Constructor
+		/// Creates a <see cref="HordeHttpClient"/> instance from an http client factory
 		/// </summary>
-		public HordeHttpClientFactory(IServiceProvider serviceProvider) => _serviceProvider = serviceProvider;
-
-		/// <summary>
-		/// Creates a client for communicating with the Horde server
-		/// </summary>
-		public HordeHttpClient CreateClient() => _serviceProvider.GetRequiredService<HordeHttpClient>();
-	}
+		public static HordeHttpClient CreateHordeClient(this IHttpClientFactory factory)
+		{
+			return new HordeHttpClient(factory.CreateClient(HordeHttpClient.HttpClientName));
+		}
 
-	/// <summary>
-	/// Extension methods for Horde HTTP clients
-	/// </summary>
-	public static class HordeHttpClientExtensions
-	{
 		/// <summary>
 		/// Registers a Horde HTTP client type, and configures it to use the default OIDC message handler.
 		/// </summary>
@@ -214,7 +318,11 @@ namespace EpicGames.Horde
 		/// <param name="useAuthChallenge">Whether to prompt the user to authenticate if necessary</param>
 		public static void AddHordeHttpClient(this IServiceCollection services, Action<IServiceProvider, HttpClient> configureClient, bool useAuthChallenge = true)
 		{
-			IHttpClientBuilder builder = services.AddHttpClient<HordeHttpClient>(HordeHttpClient.HttpClientName, configureClient).AddPolicyHandler(HordeHttpClient.DefaultRetryPolicy);
+			IHttpClientBuilder builder = services.AddHttpClient<HordeHttpClient>(HordeHttpClient.HttpClientName, configureClient)
+				.AddPolicyHandler(HttpPolicyExtensions
+					.HandleTransientHttpError()
+					.WaitAndRetryAsync(new[] { TimeSpan.FromSeconds(2.0), TimeSpan.FromSeconds(5.0), TimeSpan.FromSeconds(10), TimeSpan.FromSeconds(30) }));
+
 			if (useAuthChallenge)
 			{
 				services.AddSingleton<HordeHttpAuthHandlerState>();
@@ -223,8 +331,6 @@ namespace EpicGames.Horde
 
 				builder.AddHttpMessageHandler<HordeHttpAuthHandler>();
 			}
-
-			services.AddTransient<HordeHttpClientFactory>();
 		}
 	}
 }
