Horde: Move telemetry HTTP messages into EpicGames.Horde.

diff --git a/Engine/Source/Programs/Horde/Horde.Server/Agents/Telemetry/TelemetryController.Messages.cs b/Engine/Source/Programs/Horde/Horde.Server/Agents/Telemetry/TelemetryController.Messages.cs
deleted file mode 100644
diff --git a/Engine/Source/Programs/Horde/Horde.Server/Agents/Telemetry/TelemetryController.cs b/Engine/Source/Programs/Horde/Horde.Server/Agents/Telemetry/TelemetryController.cs
--- a/Engine/Source/Programs/Horde/Horde.Server/Agents/Telemetry/TelemetryController.cs
+++ b/Engine/Source/Programs/Horde/Horde.Server/Agents/Telemetry/TelemetryController.cs
@@ -2,7 +2,9 @@
 
 using System;
 using System.Collections.Generic;
+using System.Linq;
 using System.Threading.Tasks;
+using EpicGames.Horde.Api;
 using Microsoft.AspNetCore.Mvc;
 
 namespace Horde.Server.Agents.Telemetry
@@ -44,7 +46,40 @@ namespace Horde.Server.Agents.Telemetry
 
 			List<IUtilizationTelemetry> telemetry = await _telemetryCollection.GetUtilizationTelemetryAsync(startDateOffset.UtcDateTime, endDateOffset.UtcDateTime);
 
-			return telemetry.ConvertAll(telemetry => new UtilizationTelemetryResponse(telemetry));
+			return telemetry.ConvertAll(CreateTelemetryResponse);
+		}
+
+		static UtilizationTelemetryResponse CreateTelemetryResponse(IUtilizationTelemetry telemetry)
+		{
+			UtilizationTelemetryResponse response = new UtilizationTelemetryResponse();
+			response.StartTime = telemetry.StartTime;
+			response.FinishTime = telemetry.FinishTime;
+			response.AdminTime = telemetry.AdminTime;
+			response.HibernatingTime = telemetry.HibernatingTime;
+			response.NumAgents = telemetry.NumAgents;
+
+			response.Pools.AddRange(telemetry.Pools.Select(CreatePoolTelemetryResponse));
+			return response;
+		}
+
+		static UtilizationTelemetryPool CreatePoolTelemetryResponse(IPoolUtilizationTelemetry pool)
+		{
+			UtilizationTelemetryPool response = new UtilizationTelemetryPool();
+			response.PoolId = pool.PoolId.ToString();
+			response.NumAgents = pool.NumAgents;
+			response.AdminTime = pool.AdminTime;
+			response.HibernatingTime = pool.HibernatingTime;
+			response.OtherTime = pool.OtherTime;
+			response.Streams.AddRange(pool.Streams.Select(CreateStreamTelemetryResponse));
+			return response;
+		}
+
+		static UtilizationTelemetryStream CreateStreamTelemetryResponse(IStreamUtilizationTelemetry stream)
+		{
+			UtilizationTelemetryStream response = new UtilizationTelemetryStream();
+			response.StreamId = stream.StreamId.ToString();
+			response.Time = stream.Time;
+			return response;
 		}
 	}
 }
diff --git a/Engine/Source/Programs/Shared/EpicGames.Horde/Api/ArtifactMessages.cs b/Engine/Source/Programs/Shared/EpicGames.Horde/Api/ArtifactMessages.cs
--- a/Engine/Source/Programs/Shared/EpicGames.Horde/Api/ArtifactMessages.cs
+++ b/Engine/Source/Programs/Shared/EpicGames.Horde/Api/ArtifactMessages.cs
@@ -1,8 +1,6 @@
 ﻿// Copyright Epic Games, Inc. All Rights Reserved.
 
 using System.Collections.Generic;
-using System.Threading;
-using System.Threading.Tasks;
 using EpicGames.Core;
 
 #pragma warning disable CA2227
diff --git a/Engine/Source/Programs/Shared/EpicGames.Horde/Api/SecretMessages.cs b/Engine/Source/Programs/Shared/EpicGames.Horde/Api/SecretMessages.cs
--- a/Engine/Source/Programs/Shared/EpicGames.Horde/Api/SecretMessages.cs
+++ b/Engine/Source/Programs/Shared/EpicGames.Horde/Api/SecretMessages.cs
@@ -2,8 +2,6 @@
 
 using System.Collections.Generic;
 using System.Diagnostics;
-using System.Threading;
-using System.Threading.Tasks;
 
 #pragma warning disable CA2227 // Collection properties should be read only
 
diff --git a/Engine/Source/Programs/Shared/EpicGames.Horde/Api/TelemetryMessages.cs b/Engine/Source/Programs/Shared/EpicGames.Horde/Api/TelemetryMessages.cs
new file mode 100644
--- /dev/null
+++ b/Engine/Source/Programs/Shared/EpicGames.Horde/Api/TelemetryMessages.cs
@@ -0,0 +1,97 @@
+// Copyright Epic Games, Inc. All Rights Reserved.
+
+using System;
+using System.Collections.Generic;
+
+#pragma warning disable CA2227 // Collection properties should be read only
+
+namespace EpicGames.Horde.Api
+{
+	/// <summary>
+	/// Represents one stream in one pool in one hour of telemetry
+	/// </summary>
+	public class UtilizationTelemetryStream
+	{
+		/// <summary>
+		/// Stream Id
+		/// </summary>
+		public string StreamId { get; set; } = String.Empty;
+
+		/// <summary>
+		/// Total time
+		/// </summary>
+		public double Time { get; set; }
+	}
+
+	/// <summary>
+	/// Representation of an hour of time
+	/// </summary>
+	public class UtilizationTelemetryPool
+	{
+		/// <summary>
+		/// Pool id
+		/// </summary>
+		public string PoolId { get; set; } = String.Empty;
+
+		/// <summary>
+		/// Number of agents in this pool
+		/// </summary>
+		public int NumAgents { get; set; }
+
+		/// <summary>
+		/// Total time spent doing admin work
+		/// </summary>
+		public double AdminTime { get; set; }
+
+		/// <summary>
+		/// Time spent hibernating
+		/// </summary>
+		public double HibernatingTime { get; set; }
+
+		/// <summary>
+		/// Total time agents in this pool were doing work for other pools
+		/// </summary>
+		public double OtherTime { get; set; }
+
+		/// <summary>
+		/// List of streams
+		/// </summary>
+		public List<UtilizationTelemetryStream> Streams { get; set; } = new List<UtilizationTelemetryStream>();
+	}
+
+	/// <summary>
+	/// Response data for a utilization request
+	/// </summary>
+	public class UtilizationTelemetryResponse
+	{
+		/// <summary>
+		/// Start hour
+		/// </summary>
+		public DateTimeOffset StartTime { get; set; }
+
+		/// <summary>
+		/// End hour
+		/// </summary>
+		public DateTimeOffset FinishTime { get; set; }
+
+		/// <summary>
+		/// List of pools
+		/// </summary>
+		public List<UtilizationTelemetryPool> Pools { get; set; } = new List<UtilizationTelemetryPool>();
+
+		/// <summary>
+		/// Total admin time
+		/// </summary>
+		public double AdminTime { get; set; }
+
+		/// <summary>
+		/// Total hibernating time
+		/// </summary>
+		public double HibernatingTime { get; set; }
+
+		/// <summary>
+		/// Total agents
+		/// </summary>
+		public int NumAgents { get; set; }
+	}
+}
diff --git a/Engine/Source/Programs/Shared/EpicGames.Horde/Compute/Clients/ServerComputeClient.cs b/Engine/Source/Programs/Shared/EpicGames.Horde/Compute/Clients/ServerComputeClient.cs
--- a/Engine/Source/Programs/Shared/EpicGames.Horde/Compute/Clients/ServerComputeClient.cs
+++ b/Engine/Source/Programs/Shared/EpicGames.Horde/Compute/Clients/ServerComputeClient.cs
@@ -11,7 +11,6 @@ using System.Net.Http;
 using System.Net.Http.Json;
 using System.Net.Sockets;
 using System.Runtime.CompilerServices;
-using System.Text.Json;
 using System.Threading;
 using System.Threading.Tasks;
 
diff --git a/Engine/Source/Programs/Shared/EpicGames.Horde/HordeHttpClient.cs b/Engine/Source/Programs/Shared/EpicGames.Horde/HordeHttpClient.cs
--- a/Engine/Source/Programs/Shared/EpicGames.Horde/HordeHttpClient.cs
+++ b/Engine/Source/Programs/Shared/EpicGames.Horde/HordeHttpClient.cs
@@ -29,14 +29,8 @@ namespace EpicGames.Horde
 		readonly HttpClient _httpClient;
 
 		static readonly JsonSerializerOptions s_jsonSerializerOptions = CreateJsonSerializerOptions();
-
 		internal static JsonSerializerOptions JsonSerializerOptions => s_jsonSerializerOptions;
 
-		/// <summary>
-		/// Accessor for the underlying HTTP client
-		/// </summary>
-		public HttpClient HttpClient => _httpClient;
-
 		/// <summary>
 		/// Constructor
 		/// </summary>
@@ -187,6 +181,28 @@ namespace EpicGames.Horde
 
 		#endregion
 
+		#region Telemetry
+
+		/// <summary>
+		/// Gets telemetry for Horde within a given range
+		/// </summary>
+		/// <param name="endDate">End date for the range</param>
+		/// <param name="range">Number of hours to return</param>
+		/// <param name="tzOffset">Timezone offset</param>
+		/// <param name="cancellationToken">Cancellation token for the operation</param>
+		public Task<List<UtilizationTelemetryResponse>> GetTelemetryAsync(DateTime endDate, int range, int? tzOffset = null, CancellationToken cancellationToken = default)
+		{
+			QueryStringBuilder queryParams = new QueryStringBuilder();
+			queryParams.Add("Range", range.ToString());
+			if (tzOffset != null)
+			{
+				queryParams.Add("TzOffset", tzOffset.Value.ToString());
+			}
+			return GetAsync<List<UtilizationTelemetryResponse>>(_httpClient, $"api/v1/reports/utilization/{endDate}?{queryParams}", cancellationToken);
+		}
+
+		#endregion
+
 		#region Utility Methods
 
 		/// <summary>
