Horde: Move dashboard messages into EpicGames.Horde.

diff --git a/Engine/Source/Programs/Horde/Horde.Server.Tests/DashboardTests.cs b/Engine/Source/Programs/Horde/Horde.Server.Tests/DashboardTests.cs
--- a/Engine/Source/Programs/Horde/Horde.Server.Tests/DashboardTests.cs
+++ b/Engine/Source/Programs/Horde/Horde.Server.Tests/DashboardTests.cs
@@ -3,7 +3,7 @@
 using System.Collections.Generic;
 using System.Threading.Tasks;
 using Microsoft.VisualStudio.TestTools.UnitTesting;
-using Horde.Server.Dashboard;
+using EpicGames.Horde.Api;
 using System.Linq;
 
 namespace Horde.Server.Tests
diff --git a/Engine/Source/Programs/Horde/Horde.Server/Dashboard/DashboardController.cs b/Engine/Source/Programs/Horde/Horde.Server/Dashboard/DashboardController.cs
--- a/Engine/Source/Programs/Horde/Horde.Server/Dashboard/DashboardController.cs
+++ b/Engine/Source/Programs/Horde/Horde.Server/Dashboard/DashboardController.cs
@@ -12,6 +12,7 @@ using Horde.Server.Server;
 using System.Collections.Generic;
 using System.Linq;
 using Horde.Server.Acls;
+using EpicGames.Horde.Api;
 
 namespace Horde.Server.Dashboard
 {
@@ -168,10 +169,10 @@ namespace Horde.Server.Dashboard
 					return NotFound(preview.Id);
 				}
 
-				return new GetDashboardPreviewResponse(updated);
+				return CreatePreviewResponse(updated);
 			}
 
-			return new GetDashboardPreviewResponse(preview);
+			return CreatePreviewResponse(preview);
 		}
 
 		/// <summary>
@@ -195,7 +196,7 @@ namespace Horde.Server.Dashboard
 				return NotFound(request.Id);
 			}			
 
-			return new GetDashboardPreviewResponse(preview);
+			return CreatePreviewResponse(preview);
 		}
 
 		/// <summary>
@@ -208,7 +209,21 @@ namespace Horde.Server.Dashboard
 		public async Task<ActionResult<List<GetDashboardPreviewResponse>>> GetDashbordPreviewsAsync([FromQuery] bool open = true)
 		{			
 			List <IDashboardPreview> previews = await _previewCollection.FindPreviewsAsync(open);			
-			return previews.Select(p => new GetDashboardPreviewResponse(p)).ToList();
+			return previews.Select(CreatePreviewResponse).ToList();
+		}
+
+		static GetDashboardPreviewResponse CreatePreviewResponse(IDashboardPreview preview)
+		{
+			GetDashboardPreviewResponse response = new GetDashboardPreviewResponse();
+			response.Id = preview.Id;
+			response.CreatedAt = preview.CreatedAt;
+			response.Summary = preview.Summary;
+			response.DeployedCL = preview.DeployedCL;
+			response.Open = preview.Open;
+			response.ExampleLink = preview.ExampleLink;
+			response.DiscussionLink = preview.DiscussionLink;
+			response.TrackingLink = preview.TrackingLink;
+			return response;
 		}
 	}
 }
diff --git a/Engine/Source/Programs/Horde/Horde.Server/Dashboard/DashboardMessages.cs b/Engine/Source/Programs/Horde/Horde.Server/Dashboard/DashboardMessages.cs
deleted file mode 100644
diff --git a/Engine/Source/Programs/Horde/Horde/Commands/LoginCommand.cs b/Engine/Source/Programs/Horde/Horde/Commands/LoginCommand.cs
--- a/Engine/Source/Programs/Horde/Horde/Commands/LoginCommand.cs
+++ b/Engine/Source/Programs/Horde/Horde/Commands/LoginCommand.cs
@@ -3,7 +3,6 @@
 using EpicGames.Core;
 using EpicGames.Horde;
 using EpicGames.Horde.Api;
-using Microsoft.Extensions.DependencyInjection;
 using Microsoft.Extensions.Logging;
 using Microsoft.Extensions.Options;
 
diff --git a/Engine/Source/Programs/Shared/EpicGames.Horde/Api/DashboardMessages.cs b/Engine/Source/Programs/Shared/EpicGames.Horde/Api/DashboardMessages.cs
new file mode 100644
--- /dev/null
+++ b/Engine/Source/Programs/Shared/EpicGames.Horde/Api/DashboardMessages.cs
@@ -0,0 +1,160 @@
+// Copyright Epic Games, Inc. All Rights Reserved.
+
+using System;
+
+namespace EpicGames.Horde.Api
+{
+	/// <summary>
+	/// Setting information required by dashboard
+	/// </summary>
+	public class GetDashboardConfigResponse
+	{
+		/// <summary>
+		/// The name of the external issue service
+		/// </summary>
+		public string? ExternalIssueServiceName { get; set; }
+
+		/// <summary>
+		/// The url of the external issue service
+		/// </summary>
+		public string? ExternalIssueServiceUrl { get; set; }
+
+		/// <summary>
+		/// The url of the perforce swarm installation
+		/// </summary>
+		public string? PerforceSwarmUrl { get; set; }
+
+		/// <summary>
+		/// Help email address that users can contact with issues
+		/// </summary>
+		public string? HelpEmailAddress { get; set; }
+
+		/// <summary>
+		/// Help slack channel that users can use for issues
+		/// </summary>
+		public string? HelpSlackChannel { get; set; }
+
+		/// <summary>
+		/// Device problem cooldown in minutes
+		/// </summary>
+		public int DeviceProblemCooldownMinutes { get; set; }
+	}
+
+	/// <summary>
+	/// 
+	/// </summary>
+	public class CreateDashboardPreviewRequest
+	{
+		/// <summary>
+		/// A summary of what the preview item changes
+		/// </summary>
+		public string Summary { get; set; } = String.Empty;
+
+		/// <summary>
+		/// The CL the preview was deployed in
+		/// </summary>
+		public int? DeployedCL { get; set; }
+
+		/// <summary>
+		/// An example of the preview site users can view the changes
+		/// </summary>
+		public string? ExampleLink { get; set; }
+
+		/// <summary>
+		/// Optional Link for discussion the preview item
+		/// </summary>
+		public string? DiscussionLink { get; set; }
+
+		/// <summary>
+		/// Optional Link for discussing the preview item
+		/// </summary>
+		public string? TrackingLink { get; set; }
+	}
+
+	/// <summary>
+	/// 
+	/// </summary>
+	public class UpdateDashboardPreviewRequest
+	{
+		/// <summary>
+		/// The preview item to update
+		/// </summary>
+		public int Id { get; set; }
+
+		/// <summary>
+		/// A summary of what the preview item changes
+		/// </summary>
+		public string? Summary { get; set; }
+
+		/// <summary>
+		/// The CL the preview was deployed in
+		/// </summary>
+		public int? DeployedCL { get; set; }
+
+		/// <summary>
+		/// Whather the preview is under consideration, if false the preview item didn't pass muster
+		/// </summary>
+		public bool? Open { get; set; }
+
+		/// <summary>
+		/// An example of the preview site users can view the changes
+		/// </summary>
+		public string? ExampleLink { get; set; }
+
+		/// <summary>
+		/// Optional Link for discussion the preview item
+		/// </summary>
+		public string? DiscussionLink { get; set; }
+
+		/// <summary>
+		/// Optional Link for discussing the preview item
+		/// </summary>
+		public string? TrackingLink { get; set; }
+	}
+
+	/// <summary>
+	/// 
+	/// </summary>
+	public class GetDashboardPreviewResponse
+	{
+		/// <summary>
+		/// The unique ID of the preview item
+		/// </summary>
+		public int Id { get; set; }
+
+		/// <summary>
+		/// When the preview item was created
+		/// </summary>
+		public DateTime CreatedAt { get; set; }
+
+		/// <summary>
+		/// A summary of what the preview item changes
+		/// </summary>
+		public string Summary { get; set; } = String.Empty;
+
+		/// <summary>
+		/// The CL the preview was deployed in
+		/// </summary>
+		public int? DeployedCL { get; set; }
+
+		/// <summary>
+		/// Whather the preview is under consideration, if false the preview item didn't pass muster
+		/// </summary>
+		public bool Open { get; set; }
+
+		/// <summary>
+		/// An example of the preview site users can view the changes
+		/// </summary>
+		public string? ExampleLink { get; set; }
+
+		/// <summary>
+		/// Optional Link for discussion the preview item
+		/// </summary>
+		public string? DiscussionLink { get; set; }
+
+		/// <summary>
+		/// Optional Link for discussing the preview item
+		/// </summary>
+		public string? TrackingLink { get; set; }
+	}
+}
diff --git a/Engine/Source/Programs/Shared/EpicGames.Horde/HordeHttpClient.cs b/Engine/Source/Programs/Shared/EpicGames.Horde/HordeHttpClient.cs
--- a/Engine/Source/Programs/Shared/EpicGames.Horde/HordeHttpClient.cs
+++ b/Engine/Source/Programs/Shared/EpicGames.Horde/HordeHttpClient.cs
@@ -115,6 +115,39 @@ namespace EpicGames.Horde
 
 		#endregion
 
+		#region Dashboard
+
+		/// <summary>
+		/// Create a new dashboard preview item
+		/// </summary>
+		/// <param name="request">Request to create a new preview item</param>
+		/// <param name="cancellationToken">Cancellation token for the operation</param>
+		/// <returns>Config information needed by the dashboard</returns>
+		public Task<GetDashboardPreviewResponse> CreateDashbordPreviewAsync(CreateDashboardPreviewRequest request, CancellationToken cancellationToken = default)
+		{
+			return PostAsync<GetDashboardPreviewResponse, CreateDashboardPreviewRequest>(_httpClient, $"api/v1/dashboard/preview", request, cancellationToken); 
+		}
+
+		/// <summary>
+		/// Update a dashboard preview item
+		/// </summary>
+		/// <returns>Config information needed by the dashboard</returns>
+		public Task<GetDashboardPreviewResponse> UpdateDashbordPreviewAsync(UpdateDashboardPreviewRequest request, CancellationToken cancellationToken = default)
+		{
+			return PutAsync<GetDashboardPreviewResponse, UpdateDashboardPreviewRequest>(_httpClient, $"api/v1/dashboard/preview", request, cancellationToken);
+		}
+
+		/// <summary>
+		/// Query dashboard preview items
+		/// </summary>
+		/// <returns>Config information needed by the dashboard</returns>
+		public Task<List<GetDashboardPreviewResponse>> GetDashbordPreviewsAsync(bool open = true, CancellationToken cancellationToken = default)
+		{
+			return GetAsync<List<GetDashboardPreviewResponse>>(_httpClient, $"api/v1/dashboard/preview?open={open}", cancellationToken);
+		}
+
+		#endregion
+
 		#region Projects
 
 		/// <summary>
