Horde: Move DeviceId types into EpicGames.Horde.

diff --git a/Engine/Source/Programs/Horde/Horde.Server.Tests/DeviceServiceTests.cs b/Engine/Source/Programs/Horde/Horde.Server.Tests/DeviceServiceTests.cs
--- a/Engine/Source/Programs/Horde/Horde.Server.Tests/DeviceServiceTests.cs
+++ b/Engine/Source/Programs/Horde/Horde.Server.Tests/DeviceServiceTests.cs
@@ -13,6 +13,7 @@ using Microsoft.AspNetCore.Http;
 using Horde.Server.Server;
 using System;
 using EpicGames.Horde;
+using EpicGames.Horde.Api;
 
 namespace Horde.Server.Tests
 {
diff --git a/Engine/Source/Programs/Horde/Horde.Server/Devices/DeviceId.cs b/Engine/Source/Programs/Horde/Horde.Server/Devices/DeviceId.cs
deleted file mode 100644
diff --git a/Engine/Source/Programs/Horde/Horde.Server/Devices/DevicePlatformId.cs b/Engine/Source/Programs/Horde/Horde.Server/Devices/DevicePlatformId.cs
deleted file mode 100644
diff --git a/Engine/Source/Programs/Horde/Horde.Server/Devices/DevicePoolId.cs b/Engine/Source/Programs/Horde/Horde.Server/Devices/DevicePoolId.cs
deleted file mode 100644
diff --git a/Engine/Source/Programs/Horde/Horde.Server/Devices/DeviceReportService.cs b/Engine/Source/Programs/Horde/Horde.Server/Devices/DeviceReportService.cs
--- a/Engine/Source/Programs/Horde/Horde.Server/Devices/DeviceReportService.cs
+++ b/Engine/Source/Programs/Horde/Horde.Server/Devices/DeviceReportService.cs
@@ -5,6 +5,7 @@ using System.Collections.Generic;
 using System.Linq;
 using System.Threading;
 using System.Threading.Tasks;
+using EpicGames.Horde.Api;
 using Horde.Server.Jobs;
 using Horde.Server.Notifications;
 using Horde.Server.Server;
diff --git a/Engine/Source/Programs/Shared/EpicGames.Horde/Api/DeviceId.cs b/Engine/Source/Programs/Shared/EpicGames.Horde/Api/DeviceId.cs
new file mode 100644
--- /dev/null
+++ b/Engine/Source/Programs/Shared/EpicGames.Horde/Api/DeviceId.cs
@@ -0,0 +1,47 @@
+// Copyright Epic Games, Inc. All Rights Reserved.
+
+using System.ComponentModel;
+using EpicGames.Core;
+using EpicGames.Serialization;
+
+namespace EpicGames.Horde.Api
+{
+	/// <summary>
+	/// Identifier for a pool
+	/// </summary>
+	/// <param name="Id">Id to construct from</param>
+	[JsonSchemaString]
+	[TypeConverter(typeof(StringIdTypeConverter<DeviceId, DeviceIdConverter>))]
+	[StringIdConverter(typeof(DeviceIdConverter))]
+	[CbConverter(typeof(StringIdCbConverter<DeviceId, DeviceIdConverter>))]
+	public record struct DeviceId(StringId Id)
+	{
+		/// <summary>
+		/// Constructor
+		/// </summary>
+		public DeviceId(string id) : this(new StringId(id))
+		{
+		}
+
+		/// <inheritdoc cref="StringId.IsEmpty"/>
+		public bool IsEmpty => Id.IsEmpty;
+
+		/// <inheritdoc cref="StringId.Sanitize(System.String)"/>
+		public static DeviceId Sanitize(string text) => new DeviceId(StringId.Sanitize(text));
+
+		/// <inheritdoc/>
+		public override string ToString() => Id.ToString();
+	}
+
+	/// <summary>
+	/// Converter to and from <see cref="StringId"/> instances.
+	/// </summary>
+	class DeviceIdConverter : StringIdConverter<DeviceId>
+	{
+		/// <inheritdoc/>
+		public override DeviceId FromStringId(StringId id) => new DeviceId(id);
+
+		/// <inheritdoc/>
+		public override StringId ToStringId(DeviceId value) => value.Id;
+	}
+}
diff --git a/Engine/Source/Programs/Shared/EpicGames.Horde/Api/DevicePlatformId.cs b/Engine/Source/Programs/Shared/EpicGames.Horde/Api/DevicePlatformId.cs
new file mode 100644
--- /dev/null
+++ b/Engine/Source/Programs/Shared/EpicGames.Horde/Api/DevicePlatformId.cs
@@ -0,0 +1,47 @@
+// Copyright Epic Games, Inc. All Rights Reserved.
+
+using System.ComponentModel;
+using EpicGames.Core;
+using EpicGames.Serialization;
+
+namespace EpicGames.Horde.Api
+{
+	/// <summary>
+	/// Identifier for a device platform
+	/// </summary>
+	/// <param name="Id">Id to construct from</param>
+	[JsonSchemaString]
+	[TypeConverter(typeof(StringIdTypeConverter<DevicePlatformId, DevicePlatformIdConverter>))]
+	[StringIdConverter(typeof(DevicePlatformIdConverter))]
+	[CbConverter(typeof(StringIdCbConverter<DevicePlatformId, DevicePlatformIdConverter>))]
+	public record struct DevicePlatformId(StringId Id)
+	{
+		/// <summary>
+		/// Constructor
+		/// </summary>
+		public DevicePlatformId(string id) : this(new StringId(id))
+		{
+		}
+
+		/// <inheritdoc cref="StringId.IsEmpty"/>
+		public bool IsEmpty => Id.IsEmpty;
+
+		/// <inheritdoc cref="StringId.Sanitize(System.String)"/>
+		public static DevicePlatformId Sanitize(string text) => new DevicePlatformId(StringId.Sanitize(text));
+
+		/// <inheritdoc/>
+		public override string ToString() => Id.ToString();
+	}
+
+	/// <summary>
+	/// Converter to and from <see cref="StringId"/> instances.
+	/// </summary>
+	class DevicePlatformIdConverter : StringIdConverter<DevicePlatformId>
+	{
+		/// <inheritdoc/>
+		public override DevicePlatformId FromStringId(StringId id) => new DevicePlatformId(id);
+
+		/// <inheritdoc/>
+		public override StringId ToStringId(DevicePlatformId value) => value.Id;
+	}
+}
diff --git a/Engine/Source/Programs/Shared/EpicGames.Horde/Api/DevicePoolId.cs b/Engine/Source/Programs/Shared/EpicGames.Horde/Api/DevicePoolId.cs
new file mode 100644
--- /dev/null
+++ b/Engine/Source/Programs/Shared/EpicGames.Horde/Api/DevicePoolId.cs
@@ -0,0 +1,47 @@
+// Copyright Epic Games, Inc. All Rights Reserved.
+
+using System.ComponentModel;
+using EpicGames.Core;
+using EpicGames.Serialization;
+
+namespace EpicGames.Horde.Api
+{
+	/// <summary>
+	/// Identifier for a device pool
+	/// </summary>
+	/// <param name="Id">Id to construct from</param>
+	[JsonSchemaString]
+	[TypeConverter(typeof(StringIdTypeConverter<DevicePoolId, DevicePoolIdConverter>))]
+	[StringIdConverter(typeof(DevicePoolIdConverter))]
+	[CbConverter(typeof(StringIdCbConverter<DevicePoolId, DevicePoolIdConverter>))]
+	public record struct DevicePoolId(StringId Id)
+	{
+		/// <summary>
+		/// Constructor
+		/// </summary>
+		public DevicePoolId(string id) : this(new StringId(id))
+		{
+		}
+
+		/// <inheritdoc cref="StringId.IsEmpty"/>
+		public bool IsEmpty => Id.IsEmpty;
+
+		/// <inheritdoc cref="StringId.Sanitize(System.String)"/>
+		public static DevicePoolId Sanitize(string text) => new DevicePoolId(StringId.Sanitize(text));
+
+		/// <inheritdoc/>
+		public override string ToString() => Id.ToString();
+	}
+
+	/// <summary>
+	/// Converter to and from <see cref="StringId"/> instances.
+	/// </summary>
+	class DevicePoolIdConverter : StringIdConverter<DevicePoolId>
+	{
+		/// <inheritdoc/>
+		public override DevicePoolId FromStringId(StringId id) => new DevicePoolId(id);
+
+		/// <inheritdoc/>
+		public override StringId ToStringId(DevicePoolId value) => value.Id;
+	}
+}
