Horde: Check environment variables for presence of a configured Horde server URL and access token. These are set during CI jobs, and avoid the need to initiate any interactive login flow.

diff --git a/Engine/Source/Programs/Horde/Horde.Agent/Execution/JobExecutor.cs b/Engine/Source/Programs/Horde/Horde.Agent/Execution/JobExecutor.cs
--- a/Engine/Source/Programs/Horde/Horde.Agent/Execution/JobExecutor.cs
+++ b/Engine/Source/Programs/Horde/Horde.Agent/Execution/JobExecutor.cs
@@ -7,6 +7,7 @@ using System.Text;
 using System.Text.Json;
 using System.Text.Json.Serialization;
 using EpicGames.Core;
+using EpicGames.Horde;
 using EpicGames.Horde.Storage;
 using EpicGames.Horde.Storage.Clients;
 using EpicGames.Horde.Storage.Nodes;
@@ -222,8 +223,8 @@ namespace Horde.Agent.Execution
 			_token = options.Token;
 			JobOptions = options.JobOptions;
 
-			_envVars["UE_HORDE_URL"] = options.Session.ServerUrl.ToString();
-			_envVars["UE_HORDE_TOKEN"] = options.Token;
+			_envVars[HordeHttpClient.HordeUrlEnvVarName] = options.Session.ServerUrl.ToString();
+			_envVars[HordeHttpClient.HordeTokenEnvVarName] = options.Token;
 
 			Logger = logger;
 		}
diff --git a/Engine/Source/Programs/Shared/EpicGames.Horde/HordeHttpAuthHandler.cs b/Engine/Source/Programs/Shared/EpicGames.Horde/HordeHttpAuthHandler.cs
--- a/Engine/Source/Programs/Shared/EpicGames.Horde/HordeHttpAuthHandler.cs
+++ b/Engine/Source/Programs/Shared/EpicGames.Horde/HordeHttpAuthHandler.cs
@@ -37,25 +37,30 @@ namespace EpicGames.Horde
 		/// <inheritdoc/>
 		protected override async Task<HttpResponseMessage> SendAsync(HttpRequestMessage request, CancellationToken cancellationToken)
 		{
-			HttpResponseMessage? response = null;
+			// Do not try to override the auth header if the user has specified it explicitly
+			if (request.Headers.Authorization != null)
+			{
+				return await base.SendAsync(request, cancellationToken);
+			}
+
+			// Try to use the cached auth header
 			if (_authHeader != null)
 			{
 				request.Headers.Authorization = _authHeader;
 
-				response = await base.SendAsync(request, cancellationToken);
-				if (response.StatusCode == HttpStatusCode.Unauthorized)
+				HttpResponseMessage response = await base.SendAsync(request, cancellationToken);
+				if (response.StatusCode != HttpStatusCode.Unauthorized)
 				{
-					_authState.Invalidate(_authHeader);
-					response = null;
+					return response;
 				}
+
+				_authState.Invalidate(_authHeader);
 			}
-			if (response == null)
-			{
-				_authHeader = await _authState.GetAuthHeaderAsync(cancellationToken);
-				request.Headers.Authorization = _authHeader;
-				response = await base.SendAsync(request, cancellationToken);
-			}
-			return response;
+
+			// Otherwise update the auth header and try again
+			_authHeader = await _authState.GetAuthHeaderAsync(cancellationToken);
+			request.Headers.Authorization = _authHeader;
+			return await base.SendAsync(request, cancellationToken);
 		}
 	}
 
@@ -145,7 +150,7 @@ namespace EpicGames.Horde
 				}
 
 				serverUrl = httpClient.BaseAddress;
-				_logger.LogInformation("Getting access token for {Server}", serverUrl);
+				_logger.LogInformation("Retrieving auth configuration for {Server}", serverUrl);
 
 				JsonSerializerOptions jsonOptions = new JsonSerializerOptions();
 				HordeHttpClient.ConfigureJsonSerializer(jsonOptions);
@@ -174,16 +179,20 @@ namespace EpicGames.Horde
 			OidcTokenManager oidcTokenManager = OidcTokenManager.CreateTokenManager(configuration, tokenStore, new List<string>() { OidcProvider });
 
 			OidcTokenInfo? result = null;
-			try
-			{
-				result = await oidcTokenManager.TryGetAccessToken(OidcProvider, cancellationToken);
-			}
-			catch (Exception ex)
+			if (oidcTokenManager.GetStatusForProvider(OidcProvider) != OidcStatus.NotLoggedIn)
 			{
-				_logger.LogTrace(ex, "Unable to get access token; attempting login: {Message}", ex.Message);
+				try
+				{
+					result = await oidcTokenManager.TryGetAccessToken(OidcProvider, cancellationToken);
+				}
+				catch (Exception ex)
+				{
+					_logger.LogTrace(ex, "Unable to get access token; attempting login: {Message}", ex.Message);
+				}
 			}
 			if (result == null)
 			{
+				_logger.LogInformation("Logging in to {Server}...", serverUrl);
 				result = await oidcTokenManager.Login(OidcProvider, cancellationToken);
 			}
 
@@ -192,7 +201,7 @@ namespace EpicGames.Horde
 				throw new Exception($"Unable to get access token for {serverUrl}");
 			}
 
-			_logger.LogInformation("Received bearer token for {Server}", serverUrl);
+			_logger.LogInformation("Received access token for {Server}", serverUrl);
 			return new AuthenticationHeaderValue("Bearer", result.AccessToken);
 		}
 	}
diff --git a/Engine/Source/Programs/Shared/EpicGames.Horde/HordeHttpClient.cs b/Engine/Source/Programs/Shared/EpicGames.Horde/HordeHttpClient.cs
--- a/Engine/Source/Programs/Shared/EpicGames.Horde/HordeHttpClient.cs
+++ b/Engine/Source/Programs/Shared/EpicGames.Horde/HordeHttpClient.cs
@@ -3,6 +3,7 @@
 using System;
 using System.Collections.Generic;
 using System.Net.Http;
+using System.Net.Http.Headers;
 using System.Net.Http.Json;
 using System.Text.Json;
 using System.Text.Json.Serialization;
@@ -21,6 +22,16 @@ namespace EpicGames.Horde
 	/// </summary>
 	public sealed class HordeHttpClient : IDisposable
 	{
+		/// <summary>
+		/// Name of an environment variable containing the Horde server URL
+		/// </summary>
+		public const string HordeUrlEnvVarName = "UE_HORDE_URL";
+
+		/// <summary>
+		/// Name of an environment variable containing a token for connecting to the Horde server
+		/// </summary>
+		public const string HordeTokenEnvVarName = "UE_HORDE_TOKEN";
+
 		/// <summary>
 		/// Name of clients created from the http client factory
 		/// </summary>
@@ -367,7 +378,33 @@ namespace EpicGames.Horde
 		/// <param name="useAuthChallenge">Whether to prompt the user to authenticate if necessary</param>
 		public static void AddHordeHttpClient(this IServiceCollection services, Action<IServiceProvider, HttpClient> configureClient, bool useAuthChallenge = true)
 		{
-			IHttpClientBuilder builder = services.AddHttpClient<HordeHttpClient>(HordeHttpClient.HttpClientName, configureClient)
+			// Sets defaults from the environment before calling the user provided configuration method
+			void ConfigureClientFromEnvironment(IServiceProvider serviceProvider, HttpClient httpClient)
+			{
+				configureClient(serviceProvider, httpClient);
+
+				// Only use the token from the environment if the configured base address is missing or matches the one configured in the environment
+				string? hordeUrlEnvVar = Environment.GetEnvironmentVariable(HordeHttpClient.HordeUrlEnvVarName);
+				if (!String.IsNullOrEmpty(hordeUrlEnvVar))
+				{
+					Uri hordeUrl = new Uri(hordeUrlEnvVar);
+					if (httpClient.BaseAddress == null || String.Equals(httpClient.BaseAddress.Host, hordeUrl.Host, StringComparison.OrdinalIgnoreCase))
+					{
+						httpClient.BaseAddress ??= hordeUrl;
+
+						if (httpClient.DefaultRequestHeaders.Authorization == null)
+						{
+							string? hordeToken = Environment.GetEnvironmentVariable(HordeHttpClient.HordeTokenEnvVarName);
+							if (!String.IsNullOrEmpty(hordeToken))
+							{
+								httpClient.DefaultRequestHeaders.Authorization = new AuthenticationHeaderValue("Bearer", hordeToken);
+							}
+						}
+					}
+				}
+			}
+
+			IHttpClientBuilder builder = services.AddHttpClient<HordeHttpClient>(HordeHttpClient.HttpClientName, ConfigureClientFromEnvironment)
 				.AddPolicyHandler(HttpPolicyExtensions
 					.HandleTransientHttpError()
 					.WaitAndRetryAsync(new[] { TimeSpan.FromSeconds(2.0), TimeSpan.FromSeconds(5.0), TimeSpan.FromSeconds(10), TimeSpan.FromSeconds(30) }));
