Horde: Increase max number of bundle decode tasks to 128.

diff --git a/Engine/Source/Programs/Shared/EpicGames.Horde/Storage/Nodes/DirectoryNode.cs b/Engine/Source/Programs/Shared/EpicGames.Horde/Storage/Nodes/DirectoryNode.cs
--- a/Engine/Source/Programs/Shared/EpicGames.Horde/Storage/Nodes/DirectoryNode.cs
+++ b/Engine/Source/Programs/Shared/EpicGames.Horde/Storage/Nodes/DirectoryNode.cs
@@ -780,7 +780,7 @@ namespace EpicGames.Horde.Storage.Nodes
 		/// <param name="cancellationToken">Cancellation token for the operation</param>
 		public async Task CopyToDirectoryAsync(DirectoryInfo directoryInfo, IProgress<ICopyStats>? progress, ILogger logger, CancellationToken cancellationToken)
 		{
-			int numTasks = Math.Min(1 + (int)(Length / (10 * 1024 * 1024)), 16);
+			int numTasks = Math.Min(1 + (int)(Length / (16 * 1024 * 1024)), 128);
 			logger.LogInformation("Splitting read into {NumThreads} threads", numTasks);
 
 			CopyStats? copyStats = null;
