Cleaned up generated CPP file whitespace.

diff --git a/Engine/Source/Programs/Shared/EpicGames.UHT/Exporters/CodeGen/UhtHeaderCodeGenerator.cs b/Engine/Source/Programs/Shared/EpicGames.UHT/Exporters/CodeGen/UhtHeaderCodeGenerator.cs
--- a/Engine/Source/Programs/Shared/EpicGames.UHT/Exporters/CodeGen/UhtHeaderCodeGenerator.cs
+++ b/Engine/Source/Programs/Shared/EpicGames.UHT/Exporters/CodeGen/UhtHeaderCodeGenerator.cs
@@ -146,7 +146,7 @@ namespace EpicGames.UHT.Exporters.CodeGen
 
 		#region Function helper methods
 		protected StringBuilder AppendNativeFunctionHeader(StringBuilder builder, UhtFunction function, UhtPropertyTextType textType, bool isDeclaration,
-			string? alternateFunctionName, string? extraParam, UhtFunctionExportFlags extraExportFlags, string endl)
+			string? alternateFunctionName, string? extraParam, UhtFunctionExportFlags extraExportFlags, int tabs, string endl)
 		{
 			UhtClass? outerClass = function.Outer as UhtClass;
 			UhtFunctionExportFlags exportFlags = function.FunctionExportFlags | extraExportFlags;
@@ -154,10 +154,7 @@ namespace EpicGames.UHT.Exporters.CodeGen
 			bool isInterface = !isDelegate && (outerClass != null && outerClass.ClassFlags.HasAnyFlags(EClassFlags.Interface));
 			bool isK2Override = function.FunctionFlags.HasAnyFlags(EFunctionFlags.BlueprintEvent);
 
-			if (!isDelegate)
-			{
-				builder.Append('\t');
-			}
+			builder.AppendTabs(tabs);
 
 			if (isDeclaration)
 			{
diff --git a/Engine/Source/Programs/Shared/EpicGames.UHT/Exporters/CodeGen/UhtHeaderCodeGeneratorCppFile.cs b/Engine/Source/Programs/Shared/EpicGames.UHT/Exporters/CodeGen/UhtHeaderCodeGeneratorCppFile.cs
--- a/Engine/Source/Programs/Shared/EpicGames.UHT/Exporters/CodeGen/UhtHeaderCodeGeneratorCppFile.cs
+++ b/Engine/Source/Programs/Shared/EpicGames.UHT/Exporters/CodeGen/UhtHeaderCodeGeneratorCppFile.cs
@@ -140,7 +140,7 @@ namespace EpicGames.UHT.Exporters.CodeGen
 					builder.Append("// Cross Module References\r\n");
 					foreach (string crossReference in sorted.Span)
 					{
-						builder.Append(crossReference);
+						builder.Append(crossReference.AsSpan().TrimStart());
 					}
 					builder.Append("// End Cross Module References\r\n");
 				}
@@ -184,24 +184,24 @@ namespace EpicGames.UHT.Exporters.CodeGen
 					string name = $"Z_CompiledInDeferFile_{headerInfo.FileId}";
 					string staticsName = $"{name}_Statics";
 
-					builder.Append("\tstruct ").Append(staticsName).Append("\r\n");
-					builder.Append("\t{\r\n");
+					builder.Append("struct ").Append(staticsName).Append("\r\n");
+					builder.Append("{\r\n");
 
-					builder.AppendIfInstances(enums, UhtDefineScopeNames.Standard, builder => builder.Append("\t\tstatic const FEnumRegisterCompiledInInfo EnumInfo[];\r\n"));
-					builder.AppendIfInstances(scriptStructs, UhtDefineScopeNames.Standard, builder => builder.Append("\t\tstatic const FStructRegisterCompiledInInfo ScriptStructInfo[];\r\n"));
-					builder.AppendIfInstances(classes, UhtDefineScopeNames.Standard, builder => builder.Append("\t\tstatic const FClassRegisterCompiledInInfo ClassInfo[];\r\n"));
+					builder.AppendIfInstances(enums, UhtDefineScopeNames.Standard, builder => builder.Append("\tstatic const FEnumRegisterCompiledInInfo EnumInfo[];\r\n"));
+					builder.AppendIfInstances(scriptStructs, UhtDefineScopeNames.Standard, builder => builder.Append("\tstatic const FStructRegisterCompiledInInfo ScriptStructInfo[];\r\n"));
+					builder.AppendIfInstances(classes, UhtDefineScopeNames.Standard, builder => builder.Append("\tstatic const FClassRegisterCompiledInInfo ClassInfo[];\r\n"));
 
-					builder.Append("\t};\r\n");
+					builder.Append("};\r\n");
 
 					uint combinedHash = UInt32.MaxValue;
 
 					builder.AppendInstances(enums, UhtDefineScopeNames.Standard,
-						builder => builder.Append("\tconst FEnumRegisterCompiledInInfo ").Append(staticsName).Append("::EnumInfo[] = {\r\n"),
+						builder => builder.Append("const FEnumRegisterCompiledInInfo ").Append(staticsName).Append("::EnumInfo[] = {\r\n"),
 						(builder, enumObj) =>
 						{
 							uint hash = ObjectInfos[enumObj.ObjectTypeIndex].Hash;
 							builder
-								.Append("\t\t{ ")
+								.Append("\t{ ")
 								.Append(enumObj.SourceName)
 								.Append("_StaticEnum, TEXT(\"")
 								.Append(enumObj.EngineName)
@@ -210,15 +210,15 @@ namespace EpicGames.UHT.Exporters.CodeGen
 								.Append($", CONSTRUCT_RELOAD_VERSION_INFO(FEnumReloadVersionInfo, {hash}U) }},\r\n");
 							combinedHash = HashCombine(combinedHash, hash);
 						},
-						builder => builder.Append("\t};\r\n"));
+						builder => builder.Append("};\r\n"));
 
 					builder.AppendInstances(scriptStructs, UhtDefineScopeNames.Standard,
-						builder => builder.Append("\tconst FStructRegisterCompiledInInfo ").Append(staticsName).Append("::ScriptStructInfo[] = {\r\n"),
+						builder => builder.Append("const FStructRegisterCompiledInInfo ").Append(staticsName).Append("::ScriptStructInfo[] = {\r\n"),
 						(builder, scriptStruct) =>
 						{
 							uint hash = ObjectInfos[scriptStruct.ObjectTypeIndex].Hash;
 							builder
-								.Append("\t\t{ ")
+								.Append("\t{ ")
 								.Append(scriptStruct.SourceName)
 								.Append("::StaticStruct, Z_Construct_UScriptStruct_")
 								.Append(scriptStruct.SourceName)
@@ -231,15 +231,15 @@ namespace EpicGames.UHT.Exporters.CodeGen
 								.Append($"), {hash}U) }},\r\n");
 							combinedHash = HashCombine(combinedHash, hash);
 						},
-						builder => builder.Append("\t};\r\n"));
+						builder => builder.Append("};\r\n"));
 
 					builder.AppendInstances(classes, UhtDefineScopeNames.Standard,
-						builder => builder.Append("\tconst FClassRegisterCompiledInInfo ").Append(staticsName).Append("::ClassInfo[] = {\r\n"),
+						builder => builder.Append("const FClassRegisterCompiledInInfo ").Append(staticsName).Append("::ClassInfo[] = {\r\n"),
 						(builder, classObj) =>
 						{
 							uint hash = ObjectInfos[classObj.ObjectTypeIndex].Hash;
 							builder
-								.Append("\t\t{ Z_Construct_UClass_")
+								.Append("\t{ Z_Construct_UClass_")
 								.Append(classObj.SourceName)
 								.Append(", ")
 								.Append(classObj.SourceName)
@@ -252,18 +252,18 @@ namespace EpicGames.UHT.Exporters.CodeGen
 								.Append($"), {hash}U) }},\r\n");
 							combinedHash = HashCombine(combinedHash, hash);
 						},
-						builder => builder.Append("\t};\r\n"));
+						builder => builder.Append("};\r\n"));
 
 					builder
-						.Append("\tstatic FRegisterCompiledInInfo ")
+						.Append("static FRegisterCompiledInInfo ")
 						.Append(name)
 						.Append($"_{combinedHash}(TEXT(\"")
 						.Append(Package.EngineName)
 						.Append("\"),\r\n");
 
-					builder.AppendArrayPtrAndCountLine(classes, UhtDefineScopeNames.Standard, staticsName, "ClassInfo", 2, ",\r\n");
-					builder.AppendArrayPtrAndCountLine(scriptStructs, UhtDefineScopeNames.Standard, staticsName, "ScriptStructInfo", 2, ",\r\n");
-					builder.AppendArrayPtrAndCountLine(enums, UhtDefineScopeNames.Standard, staticsName, "EnumInfo", 2, ");\r\n");
+					builder.AppendArrayPtrAndCountLine(classes, UhtDefineScopeNames.Standard, staticsName, "ClassInfo", 1, ",\r\n");
+					builder.AppendArrayPtrAndCountLine(scriptStructs, UhtDefineScopeNames.Standard, staticsName, "ScriptStructInfo", 1, ",\r\n");
+					builder.AppendArrayPtrAndCountLine(enums, UhtDefineScopeNames.Standard, staticsName, "EnumInfo", 1, ");\r\n");
 				}
 
 				if (Session.IncludeDebugOutput)
@@ -355,7 +355,7 @@ namespace EpicGames.UHT.Exporters.CodeGen
 					bool hasUnparsedValue = enumObj.EnumValues.Exists(x => x.Value == -1);
 					if (hasUnparsedValue)
 					{
-						builder.Append("\tstatic_assert(");
+						builder.Append("static_assert(");
 						bool doneFirst = false;
 						foreach (UhtEnumValue value in enumObj.EnumValues)
 						{
@@ -373,37 +373,37 @@ namespace EpicGames.UHT.Exporters.CodeGen
 					}
 				}
 
-				builder.Append("\tstatic FEnumRegistrationInfo ").Append(registrationName).Append(";\r\n");
-				builder.Append("\tstatic UEnum* ").Append(enumObj.SourceName).Append("_StaticEnum()\r\n");
-				builder.Append("\t{\r\n");
+				builder.Append("static FEnumRegistrationInfo ").Append(registrationName).Append(";\r\n");
+				builder.Append("static UEnum* ").Append(enumObj.SourceName).Append("_StaticEnum()\r\n");
+				builder.Append("{\r\n");
 
-				builder.Append("\t\tif (!").Append(registrationName).Append(".OuterSingleton)\r\n");
-				builder.Append("\t\t{\r\n");
-				builder.Append("\t\t\t").Append(registrationName).Append(".OuterSingleton = GetStaticEnum(").Append(singletonName).Append(", (UObject*)")
+				builder.Append("\tif (!").Append(registrationName).Append(".OuterSingleton)\r\n");
+				builder.Append("\t{\r\n");
+				builder.Append("\t\t").Append(registrationName).Append(".OuterSingleton = GetStaticEnum(").Append(singletonName).Append(", (UObject*)")
 					.Append(PackageSingletonName).Append("(), TEXT(\"").Append(enumObj.SourceName).Append("\"));\r\n");
-				builder.Append("\t\t}\r\n");
-				builder.Append("\t\treturn ").Append(registrationName).Append(".OuterSingleton;\r\n");
-
 				builder.Append("\t}\r\n");
+				builder.Append("\treturn ").Append(registrationName).Append(".OuterSingleton;\r\n");
 
-				builder.Append("\ttemplate<> ").Append(PackageApi).Append("UEnum* StaticEnum<").Append(enumObj.CppType).Append(">()\r\n");
-				builder.Append("\t{\r\n");
-				builder.Append("\t\treturn ").Append(enumObj.SourceName).Append("_StaticEnum();\r\n");
-				builder.Append("\t}\r\n");
+				builder.Append("}\r\n");
+
+				builder.Append("template<> ").Append(PackageApi).Append("UEnum* StaticEnum<").Append(enumObj.CppType).Append(">()\r\n");
+				builder.Append("{\r\n");
+				builder.Append("\treturn ").Append(enumObj.SourceName).Append("_StaticEnum();\r\n");
+				builder.Append("}\r\n");
 
 				// Everything from this point on will be part of the definition hash
 				int hashCodeBlockStart = builder.Length;
 
 				// Statics declaration
-				builder.Append("\tstruct ").Append(staticsName).Append("\r\n");
-				builder.Append("\t{\r\n");
-				builder.Append("\t\tstatic const UECodeGen_Private::FEnumeratorParam Enumerators[];\r\n");
+				builder.Append("struct ").Append(staticsName).Append("\r\n");
+				builder.Append("{\r\n");
+				builder.Append("\tstatic const UECodeGen_Private::FEnumeratorParam Enumerators[];\r\n");
 				builder.AppendMetaDataDecl(enumObj, MetaDataParamsName, 2);
-				builder.Append("\t\tstatic const UECodeGen_Private::FEnumParams EnumParams;\r\n");
-				builder.Append("\t};\r\n");
+				builder.Append("\tstatic const UECodeGen_Private::FEnumParams EnumParams;\r\n");
+				builder.Append("};\r\n");
 
 				// Enumerators
-				builder.Append("\tconst UECodeGen_Private::FEnumeratorParam ").Append(staticsName).Append("::Enumerators[] = {\r\n");
+				builder.Append("const UECodeGen_Private::FEnumeratorParam ").Append(staticsName).Append("::Enumerators[] = {\r\n");
 				int enumIndex = 0;
 				foreach (UhtEnumValue value in enumObj.EnumValues)
 				{
@@ -411,37 +411,37 @@ namespace EpicGames.UHT.Exporters.CodeGen
 					{
 						keyName = value.Name.ToString();
 					}
-					builder.Append("\t\t{ ").AppendUTF8LiteralString(keyName).Append(", (int64)").Append(value.Name).Append(" },\r\n");
+					builder.Append("\t{ ").AppendUTF8LiteralString(keyName).Append(", (int64)").Append(value.Name).Append(" },\r\n");
 					++enumIndex;
 				}
-				builder.Append("\t};\r\n");
+				builder.Append("};\r\n");
 
 				// Meta data
-				builder.AppendMetaDataDef(enumObj, staticsName, MetaDataParamsName, 1);
+				builder.AppendMetaDataDef(enumObj, staticsName, MetaDataParamsName, 0);
 
 				// Singleton parameters
-				builder.Append("\tconst UECodeGen_Private::FEnumParams ").Append(staticsName).Append("::EnumParams = {\r\n");
-				builder.Append("\t\t(UObject*(*)())").Append(PackageSingletonName).Append(",\r\n");
-				builder.Append("\t\t").Append(enumDisplayNameFn).Append(",\r\n");
-				builder.Append("\t\t").AppendUTF8LiteralString(enumObj.SourceName).Append(",\r\n");
-				builder.Append("\t\t").AppendUTF8LiteralString(enumObj.CppType).Append(",\r\n");
-				builder.Append("\t\t").Append(staticsName).Append("::Enumerators,\r\n");
-				builder.Append("\t\t").Append(ObjectFlags).Append(",\r\n");
-				builder.Append("\t\tUE_ARRAY_COUNT(").Append(staticsName).Append("::Enumerators),\r\n");
-				builder.Append("\t\t").Append(enumObj.EnumFlags.HasAnyFlags(EEnumFlags.Flags) ? "EEnumFlags::Flags" : "EEnumFlags::None").Append(",\r\n");
-				builder.Append("\t\t(uint8)UEnum::ECppForm::").Append(enumObj.CppForm.ToString()).Append(",\r\n");
-				builder.Append("\t\t").AppendMetaDataParams(enumObj, staticsName, MetaDataParamsName).Append("\r\n");
-				builder.Append("\t};\r\n");
+				builder.Append("const UECodeGen_Private::FEnumParams ").Append(staticsName).Append("::EnumParams = {\r\n");
+				builder.Append("\t(UObject*(*)())").Append(PackageSingletonName).Append(",\r\n");
+				builder.Append('\t').Append(enumDisplayNameFn).Append(",\r\n");
+				builder.Append('\t').AppendUTF8LiteralString(enumObj.SourceName).Append(",\r\n");
+				builder.Append('\t').AppendUTF8LiteralString(enumObj.CppType).Append(",\r\n");
+				builder.Append('\t').Append(staticsName).Append("::Enumerators,\r\n");
+				builder.Append('\t').Append(ObjectFlags).Append(",\r\n");
+				builder.Append("\tUE_ARRAY_COUNT(").Append(staticsName).Append("::Enumerators),\r\n");
+				builder.Append('\t').Append(enumObj.EnumFlags.HasAnyFlags(EEnumFlags.Flags) ? "EEnumFlags::Flags" : "EEnumFlags::None").Append(",\r\n");
+				builder.Append("\t(uint8)UEnum::ECppForm::").Append(enumObj.CppForm.ToString()).Append(",\r\n");
+				builder.Append('\t').AppendMetaDataParams(enumObj, staticsName, MetaDataParamsName).Append("\r\n");
+				builder.Append("};\r\n");
 
 				// Registration singleton
-				builder.Append("\tUEnum* ").Append(singletonName).Append("()\r\n");
+				builder.Append("UEnum* ").Append(singletonName).Append("()\r\n");
+				builder.Append("{\r\n");
+				builder.Append("\tif (!").Append(registrationName).Append(".InnerSingleton)\r\n");
 				builder.Append("\t{\r\n");
-				builder.Append("\t\tif (!").Append(registrationName).Append(".InnerSingleton)\r\n");
-				builder.Append("\t\t{\r\n");
-				builder.Append("\t\t\tUECodeGen_Private::ConstructUEnum(").Append(registrationName).Append(".InnerSingleton, ").Append(staticsName).Append("::EnumParams);\r\n");
-				builder.Append("\t\t}\r\n");
-				builder.Append("\t\treturn ").Append(registrationName).Append(".InnerSingleton;\r\n");
+				builder.Append("\t\tUECodeGen_Private::ConstructUEnum(").Append(registrationName).Append(".InnerSingleton, ").Append(staticsName).Append("::EnumParams);\r\n");
 				builder.Append("\t}\r\n");
+				builder.Append("\treturn ").Append(registrationName).Append(".InnerSingleton;\r\n");
+				builder.Append("}\r\n");
 
 				{
 					using UhtRentedPoolBuffer<char> borrowBuffer = builder.RentPoolBuffer(hashCodeBlockStart, builder.Length - hashCodeBlockStart);
@@ -464,7 +464,6 @@ namespace EpicGames.UHT.Exporters.CodeGen
 				// Inject static assert to verify that we do not add vtable
 				if (scriptStruct.SuperScriptStruct != null)
 				{
-					builder.Append("\r\n");
 					builder.Append("static_assert(std::is_polymorphic<")
 						.Append(scriptStruct.SourceName)
 						.Append(">() == std::is_polymorphic<")
@@ -474,11 +473,10 @@ namespace EpicGames.UHT.Exporters.CodeGen
 						.Append(" cannot be polymorphic unless super ")
 						.Append(scriptStruct.SuperScriptStruct.SourceName)
 						.Append(" is polymorphic\");\r\n");
-					builder.Append("\r\n");
 				}
 
 				// Outer singleton
-				builder.Append("\tstatic FStructRegistrationInfo ").Append(registrationName).Append(";\r\n");
+				builder.Append("static FStructRegistrationInfo ").Append(registrationName).Append(";\r\n");
 				builder.Append("class UScriptStruct* ").Append(scriptStruct.SourceName).Append("::StaticStruct()\r\n");
 				builder.Append("{\r\n");
 				builder.Append("\tif (!").Append(registrationName).Append(".OuterSingleton)\r\n");
@@ -579,9 +577,9 @@ namespace EpicGames.UHT.Exporters.CodeGen
 					// The preprocessor conditional is written here instead of in FastArraySerializerImplementation.h
 					// since it may evaluate differently in different modules, triggering warnings in IncludeTool.
 					builder.Append("#if defined(UE_NET_HAS_IRIS_FASTARRAY_BINDING) && UE_NET_HAS_IRIS_FASTARRAY_BINDING\r\n");
-					builder.Append("\tUE_NET_IMPLEMENT_FASTARRAY(").Append(scriptStruct.SourceName).Append(");\r\n");
+					builder.Append("UE_NET_IMPLEMENT_FASTARRAY(").Append(scriptStruct.SourceName).Append(");\r\n");
 					builder.Append("#else\r\n");
-					builder.Append("\tUE_NET_IMPLEMENT_FASTARRAY_STUB(").Append(scriptStruct.SourceName).Append(");\r\n");
+					builder.Append("UE_NET_IMPLEMENT_FASTARRAY_STUB(").Append(scriptStruct.SourceName).Append(");\r\n");
 					builder.Append("#endif\r\n");
 				}
 			}
@@ -594,73 +592,73 @@ namespace EpicGames.UHT.Exporters.CodeGen
 
 			// Declare the statics structure
 			{
-				builder.Append("\tstruct ").Append(staticsName).Append("\r\n");
-				builder.Append("\t{\r\n");
+				builder.Append("struct ").Append(staticsName).Append("\r\n");
+				builder.Append("{\r\n");
 
 				foreach (UhtScriptStruct noExportStruct in noExportStructs)
 				{
-					AppendMirrorsForNoexportStruct(builder, noExportStruct, 2);
-					builder.Append("\t\tstatic_assert(sizeof(").Append(noExportStruct.SourceName).Append(") < MAX_uint16);\r\n");
-					builder.Append("\t\tstatic_assert(alignof(").Append(noExportStruct.SourceName).Append(") < MAX_uint8);\r\n");
+					AppendMirrorsForNoexportStruct(builder, noExportStruct, 1);
+					builder.Append("\tstatic_assert(sizeof(").Append(noExportStruct.SourceName).Append(") < MAX_uint16);\r\n");
+					builder.Append("\tstatic_assert(alignof(").Append(noExportStruct.SourceName).Append(") < MAX_uint8);\r\n");
 				}
 
 				// Meta data
-				builder.AppendMetaDataDecl(scriptStruct, MetaDataParamsName, 2);
+				builder.AppendMetaDataDecl(scriptStruct, MetaDataParamsName, 1);
 
 				// New struct ops
 				if (scriptStruct.ScriptStructFlags.HasAnyFlags(EStructFlags.Native))
 				{
-					builder.Append("\t\tstatic void* NewStructOps();\r\n");
+					builder.Append("\tstatic void* NewStructOps();\r\n");
 				}
 
-				AppendPropertiesDecl(builder, scriptStruct, scriptStruct.SourceName, properties, staticsName, 2);
+				AppendPropertiesDecl(builder, scriptStruct, scriptStruct.SourceName, properties, staticsName, 1);
 
-				builder.Append("\t\tstatic const UECodeGen_Private::FStructParams ReturnStructParams;\r\n");
-				builder.Append("\t};\r\n");
+				builder.Append("\tstatic const UECodeGen_Private::FStructParams ReturnStructParams;\r\n");
+				builder.Append("};\r\n");
 			}
 
 			// Populate the elements of the static structure
 			{
-				builder.AppendMetaDataDef(scriptStruct, staticsName, MetaDataParamsName, 1);
+				builder.AppendMetaDataDef(scriptStruct, staticsName, MetaDataParamsName, 0);
 
 				if (scriptStruct.ScriptStructFlags.HasAnyFlags(EStructFlags.Native))
 				{
-					builder.Append("\tvoid* ").Append(staticsName).Append("::NewStructOps()\r\n");
-					builder.Append("\t{\r\n");
-					builder.Append("\t\treturn (UScriptStruct::ICppStructOps*)new UScriptStruct::TCppStructOps<").Append(scriptStruct.SourceName).Append(">();\r\n");
-					builder.Append("\t}\r\n");
+					builder.Append("void* ").Append(staticsName).Append("::NewStructOps()\r\n");
+					builder.Append("{\r\n");
+					builder.Append("\treturn (UScriptStruct::ICppStructOps*)new UScriptStruct::TCppStructOps<").Append(scriptStruct.SourceName).Append(">();\r\n");
+					builder.Append("}\r\n");
 				}
 
-				AppendPropertiesDefs(builder, scriptStruct, scriptStruct.SourceName, properties, staticsName, 1);
+				AppendPropertiesDefs(builder, scriptStruct, scriptStruct.SourceName, properties, staticsName, 0);
 
-				builder.Append("\tconst UECodeGen_Private::FStructParams ").Append(staticsName).Append("::ReturnStructParams = {\r\n");
-				builder.Append("\t\t(UObject* (*)())").Append(PackageSingletonName).Append(",\r\n");
-				builder.Append("\t\t").Append(GetSingletonName(scriptStruct.SuperScriptStruct, true)).Append(",\r\n");
+				builder.Append("const UECodeGen_Private::FStructParams ").Append(staticsName).Append("::ReturnStructParams = {\r\n");
+				builder.Append("\t(UObject* (*)())").Append(PackageSingletonName).Append(",\r\n");
+				builder.Append('\t').Append(GetSingletonName(scriptStruct.SuperScriptStruct, true)).Append(",\r\n");
 				if (scriptStruct.ScriptStructFlags.HasAnyFlags(EStructFlags.Native))
 				{
-					builder.Append("\t\t&NewStructOps,\r\n");
+					builder.Append("\t&NewStructOps,\r\n");
 				}
 				else
 				{
-					builder.Append("\t\tnullptr,\r\n");
+					builder.Append("\tnullptr,\r\n");
 				}
-				builder.Append("\t\t").AppendUTF8LiteralString(scriptStruct.EngineName).Append(",\r\n");
-				builder.AppendArrayPtrLine(properties, UhtDefineScopeNames.Standard, staticsName, "PropPointers", 2, ",\r\n");
-				builder.AppendArrayCountLine(properties, UhtDefineScopeNames.Standard, staticsName, "PropPointers", 2, ",\r\n");
-				builder.Append("\t\tsizeof(").Append(scriptStruct.SourceName).Append("),\r\n");
-				builder.Append("\t\talignof(").Append(scriptStruct.SourceName).Append("),\r\n");
-				builder.Append("\t\tRF_Public|RF_Transient|RF_MarkAsNative,\r\n");
-				builder.Append($"\t\tEStructFlags(0x{(uint)(scriptStruct.ScriptStructFlags & ~EStructFlags.ComputedFlags):X8}),\r\n");
-				builder.Append("\t\t").AppendMetaDataParams(scriptStruct, staticsName, MetaDataParamsName).Append("\r\n");
-				builder.Append("\t};\r\n");
+				builder.Append('\t').AppendUTF8LiteralString(scriptStruct.EngineName).Append(",\r\n");
+				builder.AppendArrayPtrLine(properties, UhtDefineScopeNames.Standard, staticsName, "PropPointers", 1, ",\r\n");
+				builder.AppendArrayCountLine(properties, UhtDefineScopeNames.Standard, staticsName, "PropPointers", 1, ",\r\n");
+				builder.Append("\tsizeof(").Append(scriptStruct.SourceName).Append("),\r\n");
+				builder.Append("\talignof(").Append(scriptStruct.SourceName).Append("),\r\n");
+				builder.Append("\tRF_Public|RF_Transient|RF_MarkAsNative,\r\n");
+				builder.Append($"\tEStructFlags(0x{(uint)(scriptStruct.ScriptStructFlags & ~EStructFlags.ComputedFlags):X8}),\r\n");
+				builder.Append('\t').AppendMetaDataParams(scriptStruct, staticsName, MetaDataParamsName).Append("\r\n");
+				builder.Append("};\r\n");
 
 				builder.AppendPropertiesParamsCountStaticAssert(properties, staticsName);
 			}
 
 			// Generate the registration function
 			{
-				builder.Append("\tUScriptStruct* ").Append(singletonName).Append("()\r\n");
-				builder.Append("\t{\r\n");
+				builder.Append("UScriptStruct* ").Append(singletonName).Append("()\r\n");
+				builder.Append("{\r\n");
 				string innerSingletonName;
 				if (scriptStruct.ScriptStructFlags.HasAnyFlags(EStructFlags.Native))
 				{
@@ -668,15 +666,15 @@ namespace EpicGames.UHT.Exporters.CodeGen
 				}
 				else
 				{
-					builder.Append("\t\tstatic UScriptStruct* ReturnStruct = nullptr;\r\n");
+					builder.Append("\tstatic UScriptStruct* ReturnStruct = nullptr;\r\n");
 					innerSingletonName = "ReturnStruct";
 				}
-				builder.Append("\t\tif (!").Append(innerSingletonName).Append(")\r\n");
-				builder.Append("\t\t{\r\n");
-				builder.Append("\t\t\tUECodeGen_Private::ConstructUScriptStruct(").Append(innerSingletonName).Append(", ").Append(staticsName).Append("::ReturnStructParams);\r\n");
-				builder.Append("\t\t}\r\n");
-				builder.Append("\t\treturn ").Append(innerSingletonName).Append(";\r\n");
+				builder.Append("\tif (!").Append(innerSingletonName).Append(")\r\n");
+				builder.Append("\t{\r\n");
+				builder.Append("\t\tUECodeGen_Private::ConstructUScriptStruct(").Append(innerSingletonName).Append(", ").Append(staticsName).Append("::ReturnStructParams);\r\n");
 				builder.Append("\t}\r\n");
+				builder.Append("\treturn ").Append(innerSingletonName).Append(";\r\n");
+				builder.Append("}\r\n");
 			}
 
 			using (UhtRentedPoolBuffer<char> borrowBuffer = builder.RentPoolBuffer(hashCodeBlockStart, builder.Length - hashCodeBlockStart))
@@ -702,8 +700,6 @@ namespace EpicGames.UHT.Exporters.CodeGen
 						continue;
 					}
 					
-					builder.Append("\r\n");
-
 					builder
 						.Append(methodInfo.ReturnType)
 						.Append(' ')
@@ -730,7 +726,7 @@ namespace EpicGames.UHT.Exporters.CodeGen
 					builder
 						.Append(methodInfo.Name)
 						.Append("(TemporaryExecuteContext);\r\n")
-						.Append("}\r\n\r\n");
+						.Append("}\r\n");
 
 					builder
 						.Append(methodInfo.ReturnType)
@@ -744,20 +740,14 @@ namespace EpicGames.UHT.Exporters.CodeGen
 						.Append("& InExecuteContext)\r\n");
 					builder.Append("{\r\n");
 
-					bool wroteLine = false;
 					foreach (UhtRigVMParameter parameter in scriptStruct.RigVMStructInfo.Members)
 					{
 						if (!parameter.RequiresCast())
 						{
 							continue;
 						}
-						wroteLine = true;
 						builder.Append('\t').Append(parameter.CastType).Append(' ').Append(parameter.CastName).Append('(').Append(parameter.Name).Append(");\r\n");
 					}
-					if (wroteLine)
-					{
-						builder.Append("\r\n");
-					}
 					
 					foreach (UhtRigVMMethodInfo predicateInfo in scriptStruct.RigVMStructInfo.Methods)
 					{
@@ -783,7 +773,6 @@ namespace EpicGames.UHT.Exporters.CodeGen
 					builder.Append("\t);\r\n");
 					builder.Append("}\r\n");
 				}
-				builder.Append("\r\n");
 			}
 			return builder;
 		}
@@ -797,10 +786,9 @@ namespace EpicGames.UHT.Exporters.CodeGen
 			string exportFunctionName = GetDelegateFunctionExportName(function);
 			string extraParameter = GetDelegateFunctionExtraParameter(function);
 
-			AppendNativeFunctionHeader(builder, function, UhtPropertyTextType.EventFunctionArgOrRetVal, false, exportFunctionName, extraParameter, UhtFunctionExportFlags.None, "\r\n");
+			AppendNativeFunctionHeader(builder, function, UhtPropertyTextType.EventFunctionArgOrRetVal, false, exportFunctionName, extraParameter, UhtFunctionExportFlags.None, 0, "\r\n");
 			AppendEventFunctionPrologue(builder, function, strippedFunctionName, tabs, "\r\n", true);
 			builder
-				.Append('\t')
 				.Append(strippedFunctionName)
 				.Append('.')
 				.Append(function.FunctionFlags.HasAnyFlags(EFunctionFlags.MulticastDelegate) ? "ProcessMulticastDelegate" : "ProcessDelegate")
@@ -835,36 +823,36 @@ namespace EpicGames.UHT.Exporters.CodeGen
 
 				// Statics declaration
 				{
-					builder.Append("\tstruct ").Append(staticsName).Append("\r\n");
-					builder.Append("\t{\r\n");
+					builder.Append("struct ").Append(staticsName).Append("\r\n");
+					builder.Append("{\r\n");
 
 					if (paramsInStatic)
 					{
 						List<UhtScriptStruct> noExportStructs = FindNoExportStructs(function);
 						foreach (UhtScriptStruct noExportStruct in noExportStructs)
 						{
-							AppendMirrorsForNoexportStruct(builder, noExportStruct, 2);
+							AppendMirrorsForNoexportStruct(builder, noExportStruct, 1);
 						}
-						AppendEventParameter(builder, function, strippedFunctionName, UhtPropertyTextType.EventParameterFunctionMember, false, 2, "\r\n");
+						AppendEventParameter(builder, function, strippedFunctionName, UhtPropertyTextType.EventParameterFunctionMember, false, 1, "\r\n");
 					}
 
-					AppendPropertiesDecl(builder, function, eventParameters, properties, staticsName, 2);
+					AppendPropertiesDecl(builder, function, eventParameters, properties, staticsName, 1);
 
-					builder.AppendMetaDataDecl(function, MetaDataParamsName, 2);
+					builder.AppendMetaDataDecl(function, MetaDataParamsName, 1);
 
-					builder.Append("\t\tstatic const UECodeGen_Private::FFunctionParams FuncParams;\r\n");
+					builder.Append("\tstatic const UECodeGen_Private::FFunctionParams FuncParams;\r\n");
 
-					builder.Append("\t};\r\n");
+					builder.Append("};\r\n");
 				}
 
 				// Statics definition
 				{
-					AppendPropertiesDefs(builder, function, eventParameters, properties, staticsName, 1);
+					AppendPropertiesDefs(builder, function, eventParameters, properties, staticsName, 0);
 
-					builder.AppendMetaDataDef(function, staticsName, MetaDataParamsName, 1);
+					builder.AppendMetaDataDef(function, staticsName, MetaDataParamsName, 0);
 
 					builder
-						.Append("\tconst UECodeGen_Private::FFunctionParams ")
+						.Append("const UECodeGen_Private::FFunctionParams ")
 						.Append(staticsName).Append("::FuncParams = { ")
 						.Append("(UObject*(*)())").Append(GetFunctionOuterFunc(function)).Append(", ")
 						.Append(GetSingletonName(function.SuperFunction, true)).Append(", ")
@@ -913,21 +901,21 @@ namespace EpicGames.UHT.Exporters.CodeGen
 
 					if (!String.IsNullOrEmpty(sizeOfStatic))
 					{
-						builder.Append("\tstatic_assert(").Append(sizeOfStatic).Append(" < MAX_uint16);\r\n");
+						builder.Append("static_assert(").Append(sizeOfStatic).Append(" < MAX_uint16);\r\n");
 					}
 				}
 
 				// Registration function
 				{
-					builder.Append("\tUFunction* ").Append(singletonName).Append("()\r\n");
+					builder.Append("UFunction* ").Append(singletonName).Append("()\r\n");
+					builder.Append("{\r\n");
+					builder.Append("\tstatic UFunction* ReturnFunction = nullptr;\r\n");
+					builder.Append("\tif (!ReturnFunction)\r\n");
 					builder.Append("\t{\r\n");
-					builder.Append("\t\tstatic UFunction* ReturnFunction = nullptr;\r\n");
-					builder.Append("\t\tif (!ReturnFunction)\r\n");
-					builder.Append("\t\t{\r\n");
-					builder.Append("\t\t\tUECodeGen_Private::ConstructUFunction(&ReturnFunction, ").Append(staticsName).Append("::FuncParams);\r\n");
-					builder.Append("\t\t}\r\n");
-					builder.Append("\t\treturn ReturnFunction;\r\n");
+					builder.Append("\t\tUECodeGen_Private::ConstructUFunction(&ReturnFunction, ").Append(staticsName).Append("::FuncParams);\r\n");
 					builder.Append("\t}\r\n");
+					builder.Append("\treturn ReturnFunction;\r\n");
+					builder.Append("}\r\n");
 				}
 			}
 
@@ -1061,11 +1049,11 @@ namespace EpicGames.UHT.Exporters.CodeGen
 					switch (GetConstructorType(classObj))
 					{
 						case ConstructorType.ObjectInitializer:
-							builder.Append('\t').Append(classObj.SourceName).Append("::").Append(classObj.SourceName).Append("(const FObjectInitializer& ObjectInitializer) : Super(ObjectInitializer) {}\r\n");
+							builder.Append(classObj.SourceName).Append("::").Append(classObj.SourceName).Append("(const FObjectInitializer& ObjectInitializer) : Super(ObjectInitializer) {}\r\n");
 							break;
 
 						case ConstructorType.Default:
-							builder.Append('\t').Append(classObj.SourceName).Append("::").Append(classObj.SourceName).Append("() {}\r\n");
+							builder.Append(classObj.SourceName).Append("::").Append(classObj.SourceName).Append("() {}\r\n");
 							break;
 					}
 				}
@@ -1073,12 +1061,12 @@ namespace EpicGames.UHT.Exporters.CodeGen
 
 			if (!classObj.ClassExportFlags.HasAnyFlags(UhtClassExportFlags.HasCustomVTableHelperConstructor))
 			{
-				builder.Append("\tDEFINE_VTABLE_PTR_HELPER_CTOR(").Append(classObj.SourceName).Append(");\r\n");
+				builder.Append("DEFINE_VTABLE_PTR_HELPER_CTOR(").Append(classObj.SourceName).Append(");\r\n");
 			}
 
 			if (!classObj.ClassExportFlags.HasAnyFlags(UhtClassExportFlags.HasDestructor))
 			{
-				builder.Append('\t').Append(classObj.SourceName).Append("::~").Append(classObj.SourceName).Append("() {}\r\n");
+				builder.Append(classObj.SourceName).Append("::~").Append(classObj.SourceName).Append("() {}\r\n");
 			}
 
 			AppendFieldNotify(builder, classObj);
@@ -1139,18 +1127,18 @@ namespace EpicGames.UHT.Exporters.CodeGen
 					if (property.PropertyExportFlags.HasAnyFlags(UhtPropertyExportFlags.GetterSpecifiedAuto))
 					{
 						string getterCallText = property.Getter ?? "Get" + property.SourceName;
-						builder.Append('\t').AppendPropertyText(property, UhtPropertyTextType.GetterRetVal).Append(classObj.SourceName).Append("::").Append(getterCallText).Append("() const\r\n");
-						builder.Append("\t{\r\n");
-						builder.Append('\t', 2).Append("return ").Append(property.SourceName).Append(";\r\n");
-						builder.Append("\t}\r\n");
+						builder.AppendPropertyText(property, UhtPropertyTextType.GetterRetVal).Append(classObj.SourceName).Append("::").Append(getterCallText).Append("() const\r\n");
+						builder.Append("{\r\n");
+						builder.Append('\t').Append("return ").Append(property.SourceName).Append(";\r\n");
+						builder.Append("}\r\n");
 					}
 					if (property.PropertyExportFlags.HasAnyFlags(UhtPropertyExportFlags.SetterSpecifiedAuto))
 					{
 						string setterCallText = property.Setter ?? "Set" + property.SourceName;
-						builder.Append("\tvoid ").Append(classObj.SourceName).Append("::").Append(setterCallText).Append('(').AppendPropertyText(property, UhtPropertyTextType.SetterParameterArgType).Append("InValue").Append(")\r\n");
-						builder.Append("\t{\r\n");
+						builder.Append("void ").Append(classObj.SourceName).Append("::").Append(setterCallText).Append('(').AppendPropertyText(property, UhtPropertyTextType.SetterParameterArgType).Append("InValue").Append(")\r\n");
+						builder.Append("{\r\n");
 						// @todo: setter defn
-						builder.Append("\t}\r\n");
+						builder.Append("}\r\n");
 					}
 				});
 		}
@@ -1163,15 +1151,15 @@ namespace EpicGames.UHT.Exporters.CodeGen
 				{
 					if (property.PropertyExportFlags.HasAnyFlags(UhtPropertyExportFlags.GetterFound))
 					{
-						builder.Append("\tvoid ").Append(classObj.SourceName).Append("::").AppendPropertyGetterWrapperName(property).Append("(const void* Object, void* OutValue)\r\n");
-						builder.Append("\t{\r\n");
+						builder.Append("void ").Append(classObj.SourceName).Append("::").AppendPropertyGetterWrapperName(property).Append("(const void* Object, void* OutValue)\r\n");
+						builder.Append("{\r\n");
 						using (UhtMacroBlockEmitter blockEmitter = new(builder, UhtDefineScopeNames.Standard, property.DefineScope))
 						{
-							builder.Append("\t\tconst ").Append(classObj.SourceName).Append("* Obj = (const ").Append(classObj.SourceName).Append("*)Object;\r\n");
+							builder.Append("\tconst ").Append(classObj.SourceName).Append("* Obj = (const ").Append(classObj.SourceName).Append("*)Object;\r\n");
 							if (property.IsStaticArray)
 							{
 								builder
-									.Append("\t\t")
+									.Append('\t')
 									.AppendPropertyText(property, UhtPropertyTextType.GetterSetterArg)
 									.Append("* Source = (")
 									.AppendPropertyText(property, UhtPropertyTextType.GetterSetterArg)
@@ -1179,13 +1167,13 @@ namespace EpicGames.UHT.Exporters.CodeGen
 									.Append(property.Getter!)
 									.Append("();\r\n");
 								builder
-									.Append("\t\t")
+									.Append('\t')
 									.AppendPropertyText(property, UhtPropertyTextType.GetterSetterArg)
 									.Append("* Result = (")
 									.AppendPropertyText(property, UhtPropertyTextType.GetterSetterArg)
 									.Append("*)OutValue;\r\n");
 								builder
-									.Append("\t\tCopyAssignItems(Result, Source, ")
+									.Append("\tCopyAssignItems(Result, Source, ")
 									.Append(property.ArrayDimensions)
 									.Append(");\r\n");
 							}
@@ -1195,13 +1183,13 @@ namespace EpicGames.UHT.Exporters.CodeGen
 								// Since this is known to be a TEnumAsByte, we just fetch the first byte.  *HOWEVER* on MSB machines where 
 								// the actual enum value is passed in this will fail and return zero if the native size of the enum > 1 byte.
 								builder
-									.Append("\t\t")
+									.Append('\t')
 									.Append("uint8")
 									.Append("& Result = *(")
 									.Append("uint8")
 									.Append("*)OutValue;\r\n");
 								builder
-									.Append("\t\tResult = (")
+									.Append("\tResult = (")
 									.Append("uint8")
 									.Append(")Obj->")
 									.Append(property.Getter!)
@@ -1210,32 +1198,32 @@ namespace EpicGames.UHT.Exporters.CodeGen
 							else
 							{
 								builder
-									.Append("\t\t")
+									.Append('\t')
 									.AppendPropertyText(property, UhtPropertyTextType.GetterSetterArg)
 									.Append("& Result = *(")
 									.AppendPropertyText(property, UhtPropertyTextType.GetterSetterArg)
 									.Append("*)OutValue;\r\n");
 								builder
-									.Append("\t\tResult = (")
+									.Append("\tResult = (")
 									.AppendPropertyText(property, UhtPropertyTextType.GetterSetterArg)
 									.Append(")Obj->")
 									.Append(property.Getter!)
 									.Append("();\r\n");
 							}
 						}
-						builder.Append("\t}\r\n");
+						builder.Append("}\r\n");
 					}
 					if (property.PropertyExportFlags.HasAnyFlags(UhtPropertyExportFlags.SetterFound))
 					{
-						builder.Append("\tvoid ").Append(classObj.SourceName).Append("::").AppendPropertySetterWrapperName(property).Append("(void* Object, const void* InValue)\r\n");
-						builder.Append("\t{\r\n");
+						builder.Append("void ").Append(classObj.SourceName).Append("::").AppendPropertySetterWrapperName(property).Append("(void* Object, const void* InValue)\r\n");
+						builder.Append("{\r\n");
 						using (UhtMacroBlockEmitter blockEmitter = new(builder, UhtDefineScopeNames.Standard, property.DefineScope))
 						{
-							builder.Append("\t\t").Append(classObj.SourceName).Append("* Obj = (").Append(classObj.SourceName).Append("*)Object;\r\n");
+							builder.Append('\t').Append(classObj.SourceName).Append("* Obj = (").Append(classObj.SourceName).Append("*)Object;\r\n");
 							if (property.IsStaticArray)
 							{
 								builder
-									.Append("\t\t")
+									.Append('\t')
 									.AppendPropertyText(property, UhtPropertyTextType.GetterSetterArg)
 									.Append("* Value = (")
 									.AppendPropertyText(property, UhtPropertyTextType.GetterSetterArg)
@@ -1247,7 +1235,7 @@ namespace EpicGames.UHT.Exporters.CodeGen
 								// Since this is known to be a TEnumAsByte, we just fetch the first byte.  *HOWEVER* on MSB machines where 
 								// the actual enum value is passed in this will fail and return zero if the native size of the enum > 1 byte.
 								builder
-									.Append("\t\t")
+									.Append('\t')
 									.AppendPropertyText(property, UhtPropertyTextType.GetterSetterArg)
 									.Append(" Value = (")
 									.AppendPropertyText(property, UhtPropertyTextType.GetterSetterArg)
@@ -1256,18 +1244,18 @@ namespace EpicGames.UHT.Exporters.CodeGen
 							else
 							{
 								builder
-									.Append("\t\t")
+									.Append('\t')
 									.AppendPropertyText(property, UhtPropertyTextType.GetterSetterArg)
 									.Append("& Value = *(")
 									.AppendPropertyText(property, UhtPropertyTextType.GetterSetterArg)
 									.Append("*)InValue;\r\n");
 							}
 							builder
-								.Append("\t\tObj->")
+								.Append("\tObj->")
 								.Append(property.Setter!)
 								.Append("(Value);\r\n");
 						}
-						builder.Append("\t}\r\n");
+						builder.Append("}\r\n");
 					}
 				}
 			}
@@ -1302,28 +1290,28 @@ namespace EpicGames.UHT.Exporters.CodeGen
 		{
 			foreach (UhtFunction function in callbackFunctions)
 			{
-				builder.Append("\tstatic FName NAME_").Append(function.Outer?.SourceName).Append('_').Append(function.SourceName).Append(" = FName(TEXT(\"").Append(function.EngineName).Append("\"));\r\n");
+				builder.Append("static FName NAME_").Append(function.Outer?.SourceName).Append('_').Append(function.SourceName).Append(" = FName(TEXT(\"").Append(function.EngineName).Append("\"));\r\n");
 				string extraParameter = function.FunctionFlags.HasAnyFlags(EFunctionFlags.Const) ? "const UObject* O" : "UObject* O";
-				AppendNativeFunctionHeader(builder, function, UhtPropertyTextType.InterfaceFunctionArgOrRetVal, false, null, extraParameter, 0, "\r\n");
-				builder.Append("\t{\r\n");
-				builder.Append("\t\tcheck(O != NULL);\r\n");
-				builder.Append("\t\tcheck(O->GetClass()->ImplementsInterface(").Append(classObj.SourceName).Append("::StaticClass()));\r\n");
+				AppendNativeFunctionHeader(builder, function, UhtPropertyTextType.InterfaceFunctionArgOrRetVal, false, null, extraParameter, UhtFunctionExportFlags.None, 0, "\r\n");
+				builder.Append("{\r\n");
+				builder.Append("\tcheck(O != NULL);\r\n");
+				builder.Append("\tcheck(O->GetClass()->ImplementsInterface(").Append(classObj.SourceName).Append("::StaticClass()));\r\n");
 				if (function.Children.Count > 0)
 				{
-					builder.Append("\t\t").Append(GetEventStructParametersName(classObj, function.StrippedFunctionName)).Append(" Parms;\r\n");
+					builder.Append('\t').Append(GetEventStructParametersName(classObj, function.StrippedFunctionName)).Append(" Parms;\r\n");
 				}
-				builder.Append("\t\tUFunction* const Func = O->FindFunction(NAME_").Append(function.Outer?.SourceName).Append('_').Append(function.SourceName).Append(");\r\n");
-				builder.Append("\t\tif (Func)\r\n");
-				builder.Append("\t\t{\r\n");
+				builder.Append("\tUFunction* const Func = O->FindFunction(NAME_").Append(function.Outer?.SourceName).Append('_').Append(function.SourceName).Append(");\r\n");
+				builder.Append("\tif (Func)\r\n");
+				builder.Append("\t{\r\n");
 				foreach (UhtType parameter in function.ParameterProperties.Span)
 				{
 					if (parameter is UhtProperty property)
 					{
-						builder.Append("\t\t\tParms.").Append(property.SourceName).Append('=').Append(property.SourceName).Append(";\r\n");
+						builder.Append("\t\tParms.").Append(property.SourceName).Append('=').Append(property.SourceName).Append(";\r\n");
 					}
 				}
 				builder
-					.Append("\t\t\t")
+					.Append("\t\t")
 					.Append(function.FunctionFlags.HasAnyFlags(EFunctionFlags.Const) ? "const_cast<UObject*>(O)" : "O")
 					.Append("->ProcessEvent(Func, ")
 					.Append(function.Children.Count > 0 ? "&Parms" : "NULL")
@@ -1334,24 +1322,24 @@ namespace EpicGames.UHT.Exporters.CodeGen
 					{
 						if (property.PropertyFlags.HasExactFlags(EPropertyFlags.OutParm | EPropertyFlags.ConstParm, EPropertyFlags.OutParm))
 						{
-							builder.Append("\t\t\t").Append(property.SourceName).Append("=Parms.").Append(property.SourceName).Append(";\r\n");
+							builder.Append("\t\t").Append(property.SourceName).Append("=Parms.").Append(property.SourceName).Append(";\r\n");
 						}
 					}
 				}
-				builder.Append("\t\t}\r\n");
+				builder.Append("\t}\r\n");
 
 				// else clause to call back into native if it's a BlueprintNativeEvent
 				if (function.FunctionFlags.HasAnyFlags(EFunctionFlags.Native))
 				{
 					builder
-						.Append("\t\telse if (auto I = (")
+						.Append("\telse if (auto I = (")
 						.Append(function.FunctionFlags.HasAnyFlags(EFunctionFlags.Const) ? "const I" : "I")
 						.Append(classObj.EngineName)
 						.Append("*)(O->GetNativeInterfaceAddress(U")
 						.Append(classObj.EngineName)
 						.Append("::StaticClass())))\r\n");
-					builder.Append("\t\t{\r\n");
-					builder.Append("\t\t\t");
+					builder.Append("\t{\r\n");
+					builder.Append("\t\t");
 					if (function.HasReturnProperty)
 					{
 						builder.Append("Parms.ReturnValue = ");
@@ -1372,14 +1360,14 @@ namespace EpicGames.UHT.Exporters.CodeGen
 						}
 					}
 					builder.Append(");\r\n");
-					builder.Append("\t\t}\r\n");
+					builder.Append("\t}\r\n");
 				}
 
 				if (function.HasReturnProperty)
 				{
-					builder.Append("\t\treturn Parms.ReturnValue;\r\n");
+					builder.Append("\treturn Parms.ReturnValue;\r\n");
 				}
-				builder.Append("\t}\r\n");
+				builder.Append("}\r\n");
 			}
 			return builder;
 		}
@@ -1389,7 +1377,7 @@ namespace EpicGames.UHT.Exporters.CodeGen
 			if (!classObj.SerializerArchiveType.HasAnyFlags(serializerType))
 			{
 				builder.AppendScoped(UhtDefineScopeNames.Standard, classObj.SerializerDefineScope,
-					builder => builder.Append('\t').Append(macroText).Append('(').Append(classObj.SourceName).Append(")\r\n"));
+					builder => builder.Append(macroText).Append('(').Append(classObj.SourceName).Append(")\r\n"));
 			}
 			return builder;
 		}
@@ -1419,17 +1407,17 @@ namespace EpicGames.UHT.Exporters.CodeGen
 				}
 			}
 
-			builder.Append("\tIMPLEMENT_CLASS_NO_AUTO_REGISTRATION(").Append(classObj.SourceName).Append(");\r\n");
+			builder.Append("IMPLEMENT_CLASS_NO_AUTO_REGISTRATION(").Append(classObj.SourceName).Append(");\r\n");
 
 			// Everything from this point on will be part of the definition hash
 			int hashCodeBlockStart = builder.Length;
 
 			// simple ::StaticClass wrapper to avoid header, link and DLL hell
 			{
-				builder.Append("\tUClass* ").Append(GetSingletonName(classObj, false)).Append("()\r\n");
-				builder.Append("\t{\r\n");
-				builder.Append("\t\treturn ").Append(classObj.SourceName).Append("::StaticClass();\r\n");
-				builder.Append("\t}\r\n");
+				builder.Append("UClass* ").Append(GetSingletonName(classObj, false)).Append("()\r\n");
+				builder.Append("{\r\n");
+				builder.Append("\treturn ").Append(classObj.SourceName).Append("::StaticClass();\r\n");
+				builder.Append("}\r\n");
 			}
 
 			// Collect the properties
@@ -1437,49 +1425,49 @@ namespace EpicGames.UHT.Exporters.CodeGen
 
 			// Declare the statics object
 			{
-				builder.Append("\tstruct ").Append(staticsName).Append("\r\n");
-				builder.Append("\t{\r\n");
-				builder.Append("\t\tstatic UObject* (*const DependentSingletons[])();\r\n");
+				builder.Append("struct ").Append(staticsName).Append("\r\n");
+				builder.Append("{\r\n");
+				builder.Append("\tstatic UObject* (*const DependentSingletons[])();\r\n");
 
-				builder.AppendIfInstances(functions, UhtDefineScopeNames.WithEditor, builder => builder.Append("\t\tstatic const FClassFunctionLinkInfo FuncInfo[];\r\n"));
+				builder.AppendIfInstances(functions, UhtDefineScopeNames.WithEditor, builder => builder.Append("\tstatic const FClassFunctionLinkInfo FuncInfo[];\r\n"));
 
-				builder.AppendMetaDataDecl(classObj, MetaDataParamsName, 2);
+				builder.AppendMetaDataDecl(classObj, MetaDataParamsName, 1);
 
-				AppendPropertiesDecl(builder, classObj, classObj.SourceName, properties, staticsName, 2);
+				AppendPropertiesDecl(builder, classObj, classObj.SourceName, properties, staticsName, 1);
 
 				if (hasInterfaces)
 				{
-					builder.Append("\t\tstatic const UECodeGen_Private::FImplementedInterfaceParams InterfaceParams[];\r\n");
+					builder.Append("\tstatic const UECodeGen_Private::FImplementedInterfaceParams InterfaceParams[];\r\n");
 				}
 
-				builder.Append("\t\tstatic const FCppClassTypeInfoStatic StaticCppClassTypeInfo;\r\n");
+				builder.Append("\tstatic const FCppClassTypeInfoStatic StaticCppClassTypeInfo;\r\n");
 
-				builder.Append("\t\tstatic const UECodeGen_Private::FClassParams ClassParams;\r\n");
+				builder.Append("\tstatic const UECodeGen_Private::FClassParams ClassParams;\r\n");
 
-				builder.Append("\t};\r\n");
+				builder.Append("};\r\n");
 			}
 
 			// Define the statics object
 			{
-				builder.Append("\tUObject* (*const ").Append(staticsName).Append("::DependentSingletons[])() = {\r\n");
+				builder.Append("UObject* (*const ").Append(staticsName).Append("::DependentSingletons[])() = {\r\n");
 				if (classObj.SuperClass != null && classObj.SuperClass != classObj)
 				{
-					builder.Append("\t\t(UObject* (*)())").Append(GetSingletonName(classObj.SuperClass, true)).Append(",\r\n");
+					builder.Append("\t(UObject* (*)())").Append(GetSingletonName(classObj.SuperClass, true)).Append(",\r\n");
 				}
-				builder.Append("\t\t(UObject* (*)())").Append(GetSingletonName(Package, true)).Append(",\r\n");
-				builder.Append("\t};\r\n");
+				builder.Append("\t(UObject* (*)())").Append(GetSingletonName(Package, true)).Append(",\r\n");
+				builder.Append("};\r\n");
 
-				builder.Append("\tstatic_assert(UE_ARRAY_COUNT(").Append(staticsName).Append("::DependentSingletons) < 16);\r\n");
+				builder.Append("static_assert(UE_ARRAY_COUNT(").Append(staticsName).Append("::DependentSingletons) < 16);\r\n");
 
 				builder.AppendInstances(functions, UhtDefineScopeNames.WithEditor,
 					builder =>
 					{
-						builder.Append("\tconst FClassFunctionLinkInfo ").Append(staticsName).Append("::FuncInfo[] = {\r\n");
+						builder.Append("const FClassFunctionLinkInfo ").Append(staticsName).Append("::FuncInfo[] = {\r\n");
 					},
 					(builder, function) =>
 					{
 						builder
-							.Append("\t\t{ &")
+							.Append("\t{ &")
 							.Append(GetSingletonName(function, true))
 							.Append(", ")
 							.AppendUTF8LiteralString(function.EngineName)
@@ -1489,17 +1477,17 @@ namespace EpicGames.UHT.Exporters.CodeGen
 					},
 					builder =>
 					{
-						builder.Append("\t};\r\n");
-						builder.Append("\tstatic_assert(UE_ARRAY_COUNT(").Append(staticsName).Append("::FuncInfo) < 2048);\r\n");
+						builder.Append("};\r\n");
+						builder.Append("static_assert(UE_ARRAY_COUNT(").Append(staticsName).Append("::FuncInfo) < 2048);\r\n");
 					});
 
-				builder.AppendMetaDataDef(classObj, staticsName, MetaDataParamsName, 1);
+				builder.AppendMetaDataDef(classObj, staticsName, MetaDataParamsName, 0);
 
-				AppendPropertiesDefs(builder, classObj, classObj.SourceName, properties, staticsName, 1);
+				AppendPropertiesDefs(builder, classObj, classObj.SourceName, properties, staticsName, 0);
 
 				if (hasInterfaces)
 				{
-					builder.Append("\t\tconst UECodeGen_Private::FImplementedInterfaceParams ").Append(staticsName).Append("::InterfaceParams[] = {\r\n");
+					builder.Append("\tconst UECodeGen_Private::FImplementedInterfaceParams ").Append(staticsName).Append("::InterfaceParams[] = {\r\n");
 
 					foreach (UhtStruct structObj in classObj.Bases)
 					{
@@ -1508,7 +1496,7 @@ namespace EpicGames.UHT.Exporters.CodeGen
 							if (baseClass.ClassFlags.HasAnyFlags(EClassFlags.Interface))
 							{
 								builder
-									.Append("\t\t\t{ ")
+									.Append("\t\t{ ")
 									.Append(GetSingletonName(baseClass.AlternateObject, false))
 									.Append(", (int32)VTABLE_OFFSET(")
 									.Append(classObj.SourceName)
@@ -1521,67 +1509,67 @@ namespace EpicGames.UHT.Exporters.CodeGen
 						}
 					}
 
-					builder.Append("\t\t};\r\n");
+					builder.Append("\t};\r\n");
 
-					builder.Append("\tstatic_assert(UE_ARRAY_COUNT(").Append(staticsName).Append("::InterfaceParams) < 64);\r\n");
+					builder.Append("static_assert(UE_ARRAY_COUNT(").Append(staticsName).Append("::InterfaceParams) < 64);\r\n");
 				}
 
-				builder.Append("\tconst FCppClassTypeInfoStatic ").Append(staticsName).Append("::StaticCppClassTypeInfo = {\r\n");
+				builder.Append("const FCppClassTypeInfoStatic ").Append(staticsName).Append("::StaticCppClassTypeInfo = {\r\n");
 				if (classObj.ClassFlags.HasAnyFlags(EClassFlags.Interface))
 				{
-					builder.Append("\t\tTCppClassTypeTraits<I").Append(classObj.EngineName).Append(">::IsAbstract,\r\n");
+					builder.Append("\tTCppClassTypeTraits<I").Append(classObj.EngineName).Append(">::IsAbstract,\r\n");
 				}
 				else
 				{
-					builder.Append("\t\tTCppClassTypeTraits<").Append(classObj.SourceName).Append(">::IsAbstract,\r\n");
+					builder.Append("\tTCppClassTypeTraits<").Append(classObj.SourceName).Append(">::IsAbstract,\r\n");
 				}
-				builder.Append("\t};\r\n");
+				builder.Append("};\r\n");
 
 				EClassFlags classFlags = classObj.ClassFlags & EClassFlags.SaveInCompiledInClasses;
 
-				builder.Append("\tconst UECodeGen_Private::FClassParams ").Append(staticsName).Append("::ClassParams = {\r\n");
-				builder.Append("\t\t&").Append(classObj.SourceName).Append("::StaticClass,\r\n");
+				builder.Append("const UECodeGen_Private::FClassParams ").Append(staticsName).Append("::ClassParams = {\r\n");
+				builder.Append("\t&").Append(classObj.SourceName).Append("::StaticClass,\r\n");
 				if (classObj.Config.Length > 0)
 				{
 					builder.Append("\t\t").AppendUTF8LiteralString(classObj.Config).Append(",\r\n");
 				}
 				else
 				{
-					builder.Append("\t\tnullptr,\r\n");
+					builder.Append("\tnullptr,\r\n");
 				}
-				builder.Append("\t\t&StaticCppClassTypeInfo,\r\n");
-				builder.Append("\t\tDependentSingletons,\r\n");
+				builder.Append("\t&StaticCppClassTypeInfo,\r\n");
+				builder.Append("\tDependentSingletons,\r\n");
 				builder.AppendArrayPtrLine(functions, UhtDefineScopeNames.WithEditor, null, "FuncInfo", 2, ",\r\n");
 				builder.AppendArrayPtrLine(properties, UhtDefineScopeNames.Standard, staticsName, "PropPointers", 2, ",\r\n");
-				builder.Append("\t\t").Append(hasInterfaces ? "InterfaceParams" : "nullptr").Append(",\r\n");
-				builder.Append("\t\tUE_ARRAY_COUNT(DependentSingletons),\r\n");
+				builder.Append('\t').Append(hasInterfaces ? "InterfaceParams" : "nullptr").Append(",\r\n");
+				builder.Append("\tUE_ARRAY_COUNT(DependentSingletons),\r\n");
 				builder.AppendArrayCountLine(functions, UhtDefineScopeNames.WithEditor, null, "FuncInfo", 2, ",\r\n");
 				builder.AppendArrayCountLine(properties, UhtDefineScopeNames.Standard, staticsName, "PropPointers", 2, ",\r\n");
-				builder.Append("\t\t").Append(hasInterfaces ? "UE_ARRAY_COUNT(InterfaceParams)" : "0").Append(",\r\n");
+				builder.Append('\t').Append(hasInterfaces ? "UE_ARRAY_COUNT(InterfaceParams)" : "0").Append(",\r\n");
 				builder.Append($"\t\t0x{(uint)classFlags:X8}u,\r\n");
-				builder.Append("\t\t").AppendMetaDataParams(classObj, staticsName, MetaDataParamsName).Append("\r\n");
-				builder.Append("\t};\r\n");
+				builder.Append('\t').AppendMetaDataParams(classObj, staticsName, MetaDataParamsName).Append("\r\n");
+				builder.Append("};\r\n");
 
 				builder.AppendPropertiesParamsCountStaticAssert(properties, staticsName);
 			}
 
 			// Class registration
 			{
-				builder.Append("\tUClass* ").Append(singletonName).Append("()\r\n");
+				builder.Append("UClass* ").Append(singletonName).Append("()\r\n");
+				builder.Append("{\r\n");
+				builder.Append("\tif (!").Append(registrationName).Append(".OuterSingleton)\r\n");
 				builder.Append("\t{\r\n");
-				builder.Append("\t\tif (!").Append(registrationName).Append(".OuterSingleton)\r\n");
-				builder.Append("\t\t{\r\n");
-				builder.Append("\t\t\tUECodeGen_Private::ConstructUClass(").Append(registrationName).Append(".OuterSingleton, ").Append(staticsName).Append("::ClassParams);\r\n");
+				builder.Append("\t\tUECodeGen_Private::ConstructUClass(").Append(registrationName).Append(".OuterSingleton, ").Append(staticsName).Append("::ClassParams);\r\n");
 				if (sparseDataTypes != null)
 				{
 					foreach (string sparseClass in sparseDataTypes)
 					{
-						builder.Append("\t\t\t").Append(registrationName).Append(".OuterSingleton->SetSparseClassDataStruct(").Append(classObj.SourceName).Append("::StaticGet").Append(sparseClass).Append("ScriptStruct());\r\n");
+						builder.Append("\t\t").Append(registrationName).Append(".OuterSingleton->SetSparseClassDataStruct(").Append(classObj.SourceName).Append("::StaticGet").Append(sparseClass).Append("ScriptStruct());\r\n");
 					}
 				}
-				builder.Append("\t\t}\r\n");
-				builder.Append("\t\treturn ").Append(registrationName).Append(".OuterSingleton;\r\n");
 				builder.Append("\t}\r\n");
+				builder.Append("\treturn ").Append(registrationName).Append(".OuterSingleton;\r\n");
+				builder.Append("}\r\n");
 			}
 
 			// At this point, we can compute the hash... HOWEVER, in the old UHT extra data is appended to the hash block that isn't emitted to the actual output
@@ -1623,9 +1611,9 @@ namespace EpicGames.UHT.Exporters.CodeGen
 
 				if (classObj.ClassExportFlags.HasAnyFlags(UhtClassExportFlags.NoExport))
 				{
-					builder.Append("\t/* friend declarations for pasting into noexport class ").Append(classObj.SourceName).Append("\r\n");
-					builder.Append("\tfriend struct ").Append(staticsName).Append(";\r\n");
-					builder.Append("\t*/\r\n");
+					builder.Append("/* friend declarations for pasting into noexport class ").Append(classObj.SourceName).Append("\r\n");
+					builder.Append("friend struct ").Append(staticsName).Append(";\r\n");
+					builder.Append("*/\r\n");
 				}
 
 				if (Session.IncludeDebugOutput)
@@ -1643,52 +1631,49 @@ namespace EpicGames.UHT.Exporters.CodeGen
 				}
 			}
 
-			builder.Append("\ttemplate<> ").Append(PackageApi).Append("UClass* StaticClass<").Append(classObj.SourceName).Append(">()\r\n");
-			builder.Append("\t{\r\n");
-			builder.Append("\t\treturn ").Append(classObj.SourceName).Append("::StaticClass();\r\n");
-			builder.Append("\t}\r\n");
+			builder.Append("template<> ").Append(PackageApi).Append("UClass* StaticClass<").Append(classObj.SourceName).Append(">()\r\n");
+			builder.Append("{\r\n");
+			builder.Append("\treturn ").Append(classObj.SourceName).Append("::StaticClass();\r\n");
+			builder.Append("}\r\n");
 
 			if (classObj.ClassExportFlags.HasAnyFlags(UhtClassExportFlags.SelfHasReplicatedProperties))
 			{
-				builder.Append("\r\n");
-				builder.Append("\tvoid ").Append(classObj.SourceName).Append("::ValidateGeneratedRepEnums(const TArray<struct FRepRecord>& ClassReps) const\r\n");
-				builder.Append("\t{\r\n");
+				builder.Append("void ").Append(classObj.SourceName).Append("::ValidateGeneratedRepEnums(const TArray<struct FRepRecord>& ClassReps) const\r\n");
+				builder.Append("{\r\n");
 
 				foreach (UhtProperty property in classObj.Properties)
 				{
 					if (property.PropertyFlags.HasAnyFlags(EPropertyFlags.Net))
 					{
-						builder.Append("\t\tstatic const FName Name_").Append(property.SourceName).Append("(TEXT(\"").Append(property.SourceName).Append("\"));\r\n");
+						builder.Append("\tstatic const FName Name_").Append(property.SourceName).Append("(TEXT(\"").Append(property.SourceName).Append("\"));\r\n");
 					}
 				}
-				builder.Append("\r\n");
-				builder.Append("\t\tconst bool bIsValid = true");
+				builder.Append("\tconst bool bIsValid = true");
 				foreach (UhtProperty property in classObj.Properties)
 				{
 					if (property.PropertyFlags.HasAnyFlags(EPropertyFlags.Net))
 					{
 						if (!property.IsStaticArray)
 						{
-							builder.Append("\r\n\t\t\t&& Name_").Append(property.SourceName).Append(" == ClassReps[(int32)ENetFields_Private::").Append(property.SourceName).Append("].Property->GetFName()");
+							builder.Append("\r\n\t\t&& Name_").Append(property.SourceName).Append(" == ClassReps[(int32)ENetFields_Private::").Append(property.SourceName).Append("].Property->GetFName()");
 						}
 						else
 						{
-							builder.Append("\r\n\t\t\t&& Name_").Append(property.SourceName).Append(" == ClassReps[(int32)ENetFields_Private::").Append(property.SourceName).Append("_STATIC_ARRAY].Property->GetFName()");
+							builder.Append("\r\n\t\t&& Name_").Append(property.SourceName).Append(" == ClassReps[(int32)ENetFields_Private::").Append(property.SourceName).Append("_STATIC_ARRAY].Property->GetFName()");
 						}
 					}
 				}
 				builder.Append(";\r\n");
-				builder.Append("\r\n");
-				builder.Append("\t\tcheckf(bIsValid, TEXT(\"UHT Generated Rep Indices do not match runtime populated Rep Indices for properties in ").Append(classObj.SourceName).Append("\"));\r\n");
-				builder.Append("\t}\r\n");
+				builder.Append("\tcheckf(bIsValid, TEXT(\"UHT Generated Rep Indices do not match runtime populated Rep Indices for properties in ").Append(classObj.SourceName).Append("\"));\r\n");
+				builder.Append("}\r\n");
 			}
 			return builder;
 		}
 
 		private static StringBuilder AppendNatives(StringBuilder builder, UhtClass classObj)
 		{
-			builder.Append("\tvoid ").Append(classObj.SourceName).Append("::StaticRegisterNatives").Append(classObj.SourceName).Append("()\r\n");
-			builder.Append("\t{\r\n");
+			builder.Append("void ").Append(classObj.SourceName).Append("::StaticRegisterNatives").Append(classObj.SourceName).Append("()\r\n");
+			builder.Append("{\r\n");
 
 			UhtUsedDefineScopes<UhtFunction> functions = new(classObj.Functions.Where(x => x.FunctionFlags.HasExactFlags(EFunctionFlags.Native | EFunctionFlags.NetRequest, EFunctionFlags.Native)));
 			functions.Instances.Sort((x, y) => StringComparerUE.OrdinalIgnoreCase.Compare(x.EngineName, y.EngineName));
@@ -1697,14 +1682,14 @@ namespace EpicGames.UHT.Exporters.CodeGen
 			{
 				{
 					using UhtMacroBlockEmitter blockEmitter = new(builder, UhtDefineScopeNames.WithEditor, functions.SoleScope);
-					builder.Append("\t\tUClass* Class = ").Append(classObj.SourceName).Append("::StaticClass();\r\n");
-					builder.Append("\t\tstatic const FNameNativePtrPair Funcs[] = {\r\n");
+					builder.Append("\tUClass* Class = ").Append(classObj.SourceName).Append("::StaticClass();\r\n");
+					builder.Append("\tstatic const FNameNativePtrPair Funcs[] = {\r\n");
 
 					foreach (UhtFunction function in functions.Instances)
 					{
 						blockEmitter.Set(function.DefineScope);
 						builder
-							.Append("\t\t\t{ ")
+							.Append("\t\t{ ")
 							.AppendUTF8LiteralString(function.EngineName)
 							.Append(", &")
 							.AppendClassSourceNameOrInterfaceName(classObj)
@@ -1716,12 +1701,12 @@ namespace EpicGames.UHT.Exporters.CodeGen
 					// This will close the block if we have one that isn't editor only
 					blockEmitter.Set(functions.SoleScope);
 
-					builder.Append("\t\t};\r\n");
-					builder.Append("\t\tFNativeFunctionRegistrar::RegisterFunctions(Class, Funcs, UE_ARRAY_COUNT(Funcs));\r\n");
+					builder.Append("\t};\r\n");
+					builder.Append("\tFNativeFunctionRegistrar::RegisterFunctions(Class, Funcs, UE_ARRAY_COUNT(Funcs));\r\n");
 				}
 			}
 
-			builder.Append("\t}\r\n");
+			builder.Append("}\r\n");
 			return builder;
 		}
 
@@ -1753,18 +1738,18 @@ namespace EpicGames.UHT.Exporters.CodeGen
 
 						if (!isInterfaceClass)
 						{
-							builder.Append("\tstatic FName NAME_").Append(classObj.SourceName).Append('_').Append(function.EngineName).Append(" = FName(TEXT(\"").Append(function.EngineName).Append("\"));\r\n");
+							builder.Append("static FName NAME_").Append(classObj.SourceName).Append('_').Append(function.EngineName).Append(" = FName(TEXT(\"").Append(function.EngineName).Append("\"));\r\n");
 						}
 
-						AppendNativeFunctionHeader(builder, function, UhtPropertyTextType.EventFunctionArgOrRetVal, false, null, null, UhtFunctionExportFlags.None, "\r\n");
+						AppendNativeFunctionHeader(builder, function, UhtPropertyTextType.EventFunctionArgOrRetVal, false, null, null, UhtFunctionExportFlags.None, 0, "\r\n");
 
 						if (isInterfaceClass)
 						{
-							builder.Append("\t{\r\n");
+							builder.Append("{\r\n");
 
 							// assert if this is ever called directly
 							builder
-								.Append("\t\tcheck(0 && \"Do not directly call Event functions in Interfaces. Call Execute_")
+								.Append("\tcheck(0 && \"Do not directly call Event functions in Interfaces. Call Execute_")
 								.Append(function.EngineName)
 								.Append(" instead.\");\r\n");
 
@@ -1772,17 +1757,17 @@ namespace EpicGames.UHT.Exporters.CodeGen
 							if (function.ReturnProperty != null)
 							{
 								string eventParmStructName = GetEventStructParametersName(classObj, function.EngineName);
-								builder.Append("\t\t").Append(eventParmStructName).Append(" Parms;\r\n");
-								builder.Append("\t\treturn Parms.ReturnValue;\r\n");
+								builder.Append('\t').Append(eventParmStructName).Append(" Parms;\r\n");
+								builder.Append("\treturn Parms.ReturnValue;\r\n");
 							}
-							builder.Append("\t}\r\n");
+							builder.Append("}\r\n");
 						}
 						else
 						{
 							AppendEventFunctionPrologue(builder, function, function.EngineName, 1, "\r\n", false);
 
 							// Cast away const just in case, because ProcessEvent isn't const
-							builder.Append("\t\t");
+							builder.Append('\t');
 							if (function.FunctionFlags.HasAnyFlags(EFunctionFlags.Const))
 							{
 								builder.Append("const_cast<").Append(classObj.SourceName).Append("*>(this)->");
@@ -1815,10 +1800,10 @@ namespace EpicGames.UHT.Exporters.CodeGen
 			// Output the RPC methods
 			builder.AppendInstances(functions, UhtDefineScopeNames.WithEditor, (builder, function) => 
 			{
-				builder.Append("\tDEFINE_FUNCTION(").AppendClassSourceNameOrInterfaceName(classObj).Append("::").Append(function.UnMarshalAndCallName).Append(")\r\n");
-				builder.Append("\t{\r\n");
+				builder.Append("DEFINE_FUNCTION(").AppendClassSourceNameOrInterfaceName(classObj).Append("::").Append(function.UnMarshalAndCallName).Append(")\r\n");
+				builder.Append("{\r\n");
 				AppendFunctionThunk(builder, function);
-				builder.Append("\t}\r\n");
+				builder.Append("}\r\n");
 			});
 			return builder;
 		}
@@ -1830,25 +1815,25 @@ namespace EpicGames.UHT.Exporters.CodeGen
 			{
 				if (parameter is UhtProperty property)
 				{
-					builder.Append("\t\t").AppendFunctionThunkParameterGet(property).Append(";\r\n");
+					builder.Append('\t').AppendFunctionThunkParameterGet(property).Append(";\r\n");
 				}
 			}
 
-			builder.Append("\t\tP_FINISH;\r\n");
-			builder.Append("\t\tP_NATIVE_BEGIN;\r\n");
+			builder.Append("\tP_FINISH;\r\n");
+			builder.Append("\tP_NATIVE_BEGIN;\r\n");
 
 			// Call the validate function if there is one
 			if (!function.FunctionExportFlags.HasAnyFlags(UhtFunctionExportFlags.CppStatic) && function.FunctionFlags.HasAnyFlags(EFunctionFlags.NetValidate))
 			{
-				builder.Append("\t\tif (!P_THIS->").Append(function.CppValidationImplName).Append('(').AppendFunctionThunkParameterNames(function).Append("))\r\n");
-				builder.Append("\t\t{\r\n");
-				builder.Append("\t\t\tRPC_ValidateFailed(TEXT(\"").Append(function.CppValidationImplName).Append("\"));\r\n");
-				builder.Append("\t\t\treturn;\r\n");   // If we got here, the validation function check failed
-				builder.Append("\t\t}\r\n");
+				builder.Append("\tif (!P_THIS->").Append(function.CppValidationImplName).Append('(').AppendFunctionThunkParameterNames(function).Append("))\r\n");
+				builder.Append("\t{\r\n");
+				builder.Append("\t\tRPC_ValidateFailed(TEXT(\"").Append(function.CppValidationImplName).Append("\"));\r\n");
+				builder.Append("\t\treturn;\r\n");   // If we got here, the validation function check failed
+				builder.Append("\t}\r\n");
 			}
 
 			// Write out the return value
-			builder.Append("\t\t");
+			builder.Append('\t');
 			UhtProperty? returnProperty = function.ReturnProperty;
 			if (returnProperty != null)
 			{
@@ -1864,7 +1849,7 @@ namespace EpicGames.UHT.Exporters.CodeGen
 			{
 				builder.Append("P_THIS->").Append(function.CppImplName).Append('(').AppendFunctionThunkParameterNames(function).Append(");\r\n");
 			}
-			builder.Append("\t\tP_NATIVE_END;\r\n");
+			builder.Append("\tP_NATIVE_END;\r\n");
 			return builder;
 		}
 
@@ -1999,7 +1984,7 @@ namespace EpicGames.UHT.Exporters.CodeGen
 		{
 			builder.AppendIfInstances(properties, UhtDefineScopeNames.Standard, builder =>
 			{
-				builder.Append("\tstatic_assert(UE_ARRAY_COUNT(").Append(staticsName).Append("::PropPointers) < 2048);\r\n");
+				builder.Append("static_assert(UE_ARRAY_COUNT(").Append(staticsName).Append("::PropPointers) < 2048);\r\n");
 			});
 			return builder;
 		}
diff --git a/Engine/Source/Programs/Shared/EpicGames.UHT/Exporters/CodeGen/UhtHeaderCodeGeneratorHFile.cs b/Engine/Source/Programs/Shared/EpicGames.UHT/Exporters/CodeGen/UhtHeaderCodeGeneratorHFile.cs
--- a/Engine/Source/Programs/Shared/EpicGames.UHT/Exporters/CodeGen/UhtHeaderCodeGeneratorHFile.cs
+++ b/Engine/Source/Programs/Shared/EpicGames.UHT/Exporters/CodeGen/UhtHeaderCodeGeneratorHFile.cs
@@ -503,7 +503,7 @@ namespace EpicGames.UHT.Exporters.CodeGen
 					builder.Append(PackageApi);
 				}
 
-				AppendNativeFunctionHeader(builder, function, UhtPropertyTextType.EventFunctionArgOrRetVal, true, exportFunctionName, extraParameter, UhtFunctionExportFlags.None, "; \\\r\n");
+				AppendNativeFunctionHeader(builder, function, UhtPropertyTextType.EventFunctionArgOrRetVal, true, exportFunctionName, extraParameter, UhtFunctionExportFlags.None, 0, "; \\\r\n");
 			}
 			return builder;
 		}
@@ -838,7 +838,7 @@ namespace EpicGames.UHT.Exporters.CodeGen
 					continue;
 				}
 				AppendNetValidateDeclaration(builder, classObj, function);
-				AppendNativeFunctionHeader(builder, function, UhtPropertyTextType.ClassFunctionArgOrRetVal, true, null, null, UhtFunctionExportFlags.None, "; \\\r\n");
+				AppendNativeFunctionHeader(builder, function, UhtPropertyTextType.ClassFunctionArgOrRetVal, true, null, null, UhtFunctionExportFlags.None, 1, "; \\\r\n");
 			}
 			return builder;
 		}
@@ -857,7 +857,7 @@ namespace EpicGames.UHT.Exporters.CodeGen
 				}
 				if (!function.FunctionExportFlags.HasAnyFlags(UhtFunctionExportFlags.ImplFound))
 				{
-					AppendNativeFunctionHeader(builder, function, UhtPropertyTextType.ClassFunctionArgOrRetVal, true, null, null, UhtFunctionExportFlags.None, "; \\\r\n");
+					AppendNativeFunctionHeader(builder, function, UhtPropertyTextType.ClassFunctionArgOrRetVal, true, null, null, UhtFunctionExportFlags.None, 1, "; \\\r\n");
 				}
 			}
 			return builder;
@@ -908,7 +908,7 @@ namespace EpicGames.UHT.Exporters.CodeGen
 					if (!function.FunctionFlags.HasAnyFlags(EFunctionFlags.NetResponse) &&
 						function.EngineName != function.MarshalAndCallName)
 					{
-						AppendNativeFunctionHeader(builder, function, UhtPropertyTextType.EventFunctionArgOrRetVal, true, null, null, UhtFunctionExportFlags.None, " \\\r\n");
+						AppendNativeFunctionHeader(builder, function, UhtPropertyTextType.EventFunctionArgOrRetVal, true, null, null, UhtFunctionExportFlags.None, 1, " \\\r\n");
 						builder.Append(" \\\r\n");
 					}
 				}
@@ -1434,7 +1434,7 @@ namespace EpicGames.UHT.Exporters.CodeGen
 			{
 				AppendNativeFunctionHeader(builder, function, UhtPropertyTextType.InterfaceFunctionArgOrRetVal, true, null,
 					function.FunctionFlags.HasAnyFlags(EFunctionFlags.Const) ? ConstExtraArg : ExtraArg, UhtFunctionExportFlags.None,
-					"; \\\r\n");
+					1, "; \\\r\n");
 			}
 			return builder;
 		}
