Added block comments to generated file to make it easier to find specific code

diff --git a/Engine/Source/Programs/Shared/EpicGames.UHT/Exporters/CodeGen/UhtHeaderCodeGeneratorCppFile.cs b/Engine/Source/Programs/Shared/EpicGames.UHT/Exporters/CodeGen/UhtHeaderCodeGeneratorCppFile.cs
--- a/Engine/Source/Programs/Shared/EpicGames.UHT/Exporters/CodeGen/UhtHeaderCodeGeneratorCppFile.cs
+++ b/Engine/Source/Programs/Shared/EpicGames.UHT/Exporters/CodeGen/UhtHeaderCodeGeneratorCppFile.cs
@@ -10,6 +10,56 @@ using EpicGames.UHT.Utils;
 
 namespace EpicGames.UHT.Exporters.CodeGen
 {
+	internal struct UhtCodeBlockComment : IDisposable
+	{
+		private readonly StringBuilder _builder;
+		private readonly UhtType? _primaryType = null;
+		private readonly UhtType? _secondaryType = null;
+		private readonly string? _text = null;
+
+		public UhtCodeBlockComment(StringBuilder builder, string text)
+		{
+			_builder = builder;
+			_text = text;
+			builder.Append("\r\n// Begin");
+			AppendText();
+			_builder.Append("\r\n");
+		}
+
+		public UhtCodeBlockComment(StringBuilder builder, UhtType primaryType, UhtType? secondaryType = null)
+		{
+			_builder = builder;
+			_primaryType = primaryType;
+			_secondaryType = secondaryType;
+			builder.Append("\r\n// Begin");
+			AppendText();
+			_builder.Append("\r\n");
+		}
+
+		public void Dispose()
+		{
+			_builder.Append("// End");
+			AppendText();
+			_builder.Append("\r\n");
+		}
+
+		private void AppendText()
+		{
+			if (!String.IsNullOrEmpty(_text))
+			{
+				_builder.Append(' ').Append(_text);
+			}
+			if (_primaryType != null)
+			{
+				_builder.Append(' ').Append(_primaryType.EngineType).Append(' ').Append(_primaryType.SourceName);
+			}
+			if (_secondaryType != null)
+			{
+				_builder.Append(' ').Append(_secondaryType.EngineType).Append(' ').Append(_secondaryType.SourceName);
+			}
+		}
+	}
+
 	internal class UhtHeaderCodeGeneratorCppFile : UhtHeaderCodeGenerator
 	{
 
@@ -137,12 +187,11 @@ namespace EpicGames.UHT.Exporters.CodeGen
 				{
 					ReadOnlyMemory<string> sorted = HeaderFile.References.CrossModule.GetSortedReferences(
 						(int objectIndex, bool registered) => GetExternalDecl(objectIndex, registered));
-					builder.Append("// Cross Module References\r\n");
+					using UhtCodeBlockComment blockComment = new(builder, "Cross Module References");
 					foreach (string crossReference in sorted.Span)
 					{
 						builder.Append(crossReference.AsSpan().TrimStart());
 					}
-					builder.Append("// End Cross Module References\r\n");
 				}
 
 				int generatedBodyStart = builder.Length;
@@ -154,11 +203,13 @@ namespace EpicGames.UHT.Exporters.CodeGen
 				{
 					if (field is UhtEnum enumObj)
 					{
+						using UhtCodeBlockComment blockComment = new(builder, field);
 						AppendEnum(builder, enumObj);
 						enums.Add(enumObj);
 					}
 					else if (field is UhtScriptStruct scriptStruct)
 					{
+						using UhtCodeBlockComment blockComment = new(builder, field);
 						AppendScriptStruct(builder, scriptStruct);
 						if (scriptStruct.ScriptStructFlags.HasAnyFlags(EStructFlags.Native))
 						{
@@ -167,12 +218,14 @@ namespace EpicGames.UHT.Exporters.CodeGen
 					}
 					else if (field is UhtFunction function)
 					{
+						using UhtCodeBlockComment blockComment = new(builder, field);
 						AppendDelegate(builder, function);
 					}
 					else if (field is UhtClass classObj)
 					{
 						if (!classObj.ClassFlags.HasAnyFlags(EClassFlags.Intrinsic))
 						{
+							using UhtCodeBlockComment blockComment = new(builder, field);
 							AppendClass(builder, classObj);
 							classes.Add(classObj);
 						}
@@ -184,6 +237,7 @@ namespace EpicGames.UHT.Exporters.CodeGen
 					string name = $"Z_CompiledInDeferFile_{headerInfo.FileId}";
 					string staticsName = $"{name}_Statics";
 
+					using UhtCodeBlockComment blockComment = new(builder, "Registration");
 					builder.Append("struct ").Append(staticsName).Append("\r\n");
 					builder.Append("{\r\n");
 
@@ -1403,6 +1457,7 @@ namespace EpicGames.UHT.Exporters.CodeGen
 			{
 				if (!function.FunctionFlags.HasAnyFlags(EFunctionFlags.Delegate))
 				{
+					using UhtCodeBlockComment blockCommand = new(builder, classObj, function);
 					AppendFunction(builder, function, classObj.ClassExportFlags.HasAnyFlags(UhtClassExportFlags.NoExport));
 				}
 			}
