Add support for TOptional/FOptionalProperty to UHT.

diff --git a/Engine/Source/Programs/Shared/EpicGames.Core/UnrealEngineTypes.cs b/Engine/Source/Programs/Shared/EpicGames.Core/UnrealEngineTypes.cs
--- a/Engine/Source/Programs/Shared/EpicGames.Core/UnrealEngineTypes.cs
+++ b/Engine/Source/Programs/Shared/EpicGames.Core/UnrealEngineTypes.cs
@@ -1200,6 +1200,11 @@ namespace EpicGames.Core
 		/// </summary>
 		PropagateToArrayInner =	ExportObject | PersistentInstance | InstancedReference | ContainsInstancedReference | Config | EditConst | Deprecated | EditorOnly | AutoWeak | UObjectWrapper,
 
+		/// <summary>
+		/// Flags that are propagated to properties inside optional container
+		/// </summary>
+		PropagateToOptionalInner = ExportObject | PersistentInstance | InstancedReference | ContainsInstancedReference | Config | EditConst | Deprecated | EditorOnly | AutoWeak | UObjectWrapper | Edit,
+
 		/// <summary>
 		/// Flags that are propagated to value properties inside map container
 		/// </summary>
diff --git a/Engine/Source/Programs/Shared/EpicGames.UHT/Types/Properties/UhtOptionalProperty.cs b/Engine/Source/Programs/Shared/EpicGames.UHT/Types/Properties/UhtOptionalProperty.cs
new file mode 100644
--- /dev/null
+++ b/Engine/Source/Programs/Shared/EpicGames.UHT/Types/Properties/UhtOptionalProperty.cs
@@ -0,0 +1,257 @@
+// Copyright Epic Games, Inc. All Rights Reserved.
+
+using System.Collections.Generic;
+using System.Diagnostics.CodeAnalysis;
+using System.Text;
+using EpicGames.Core;
+using EpicGames.UHT.Parsers;
+using EpicGames.UHT.Tables;
+using EpicGames.UHT.Tokenizer;
+using EpicGames.UHT.Utils;
+
+namespace EpicGames.UHT.Types
+{
+
+	/// <summary>
+	/// Represents the FOptionalProperty engine type
+	/// </summary>
+	[UnrealHeaderTool]
+	[UhtEngineClass(Name = "OptionalProperty", IsProperty = true)]
+	public class UhtOptionalProperty : UhtContainerBaseProperty
+	{
+		/// <inheritdoc/>
+		public override string EngineClassName => "OptionalProperty";
+
+		/// <inheritdoc/>
+		protected override string CppTypeText => "TOptional";
+
+		/// <inheritdoc/>
+		protected override string PGetMacroText => "TOptional";
+
+		/// <inheritdoc/>
+		protected override UhtPGetArgumentType PGetTypeArgument => UhtPGetArgumentType.TypeText;
+
+		/// <summary>
+		/// Construct a new Optional property
+		/// </summary>
+		/// <param name="propertySettings">Property settings</param>
+		/// <param name="value">Inner property value</param>
+		public UhtOptionalProperty(UhtPropertySettings propertySettings, UhtProperty value) : base(propertySettings, value)
+		{
+			// If the creation of the value property set more flags, then copy those flags to ourselves
+			PropertyFlags |= ValueProperty.PropertyFlags & EPropertyFlags.UObjectWrapper;
+
+			if (ValueProperty.MetaData.ContainsKey(UhtNames.NativeConst))
+			{
+				MetaData.Add(UhtNames.NativeConstTemplateArg, "");
+				ValueProperty.MetaData.Remove(UhtNames.NativeConst);
+			}
+
+			PropertyCaps |= UhtPropertyCaps.PassCppArgsByRef;
+			PropertyCaps &= ~(UhtPropertyCaps.CanBeContainerValue | UhtPropertyCaps.CanBeContainerKey);
+			UpdateCaps();
+
+			ValueProperty.SourceName = SourceName;
+			ValueProperty.EngineName = EngineName;
+			ValueProperty.PropertyFlags = PropertyFlags & EPropertyFlags.PropagateToOptionalInner;
+			ValueProperty.Outer = this;
+			ValueProperty.MetaData.Clear();
+		}
+
+		/// <inheritdoc/>
+		protected override bool ResolveSelf(UhtResolvePhase phase)
+		{
+			bool results = base.ResolveSelf(phase);
+			switch (phase)
+			{
+				case UhtResolvePhase.Final:
+					PropertyFlags |= ResolveAndReturnNewFlags(ValueProperty, phase);
+					MetaData.Add(ValueProperty.MetaData);
+					ValueProperty.PropertyFlags = PropertyFlags & EPropertyFlags.PropagateToOptionalInner;
+					ValueProperty.MetaData.Clear();
+					PropagateFlagsFromInnerAndHandlePersistentInstanceMetadata(this, MetaData, ValueProperty);
+					UpdateCaps();
+					break;
+			}
+			return results;
+		}
+
+		private void UpdateCaps()
+		{
+			PropertyCaps &= ~(UhtPropertyCaps.IsParameterSupportedByBlueprint | UhtPropertyCaps.IsMemberSupportedByBlueprint | UhtPropertyCaps.CanExposeOnSpawn);
+			PropertyCaps |= ValueProperty.PropertyCaps & UhtPropertyCaps.CanExposeOnSpawn;
+			if (ValueProperty.PropertyCaps.HasAnyFlags(UhtPropertyCaps.IsParameterSupportedByBlueprint))
+			{
+				PropertyCaps |= UhtPropertyCaps.IsParameterSupportedByBlueprint | UhtPropertyCaps.IsMemberSupportedByBlueprint;
+			}
+		}
+
+		/// <inheritdoc/>
+		public override void CollectReferencesInternal(IUhtReferenceCollector collector, bool templateProperty)
+		{
+			ValueProperty.CollectReferencesInternal(collector, true);
+		}
+
+		/// <inheritdoc/>
+		public override string? GetForwardDeclarations()
+		{
+			return ValueProperty.GetForwardDeclarations();
+		}
+
+		/// <inheritdoc/>
+		public override IEnumerable<UhtType> EnumerateReferencedTypes()
+		{
+			foreach (UhtType type in ValueProperty.EnumerateReferencedTypes())
+			{
+				yield return type;
+			}
+		}
+
+		/// <inheritdoc/>
+		public override StringBuilder AppendText(StringBuilder builder, UhtPropertyTextType textType, bool isTemplateArgument)
+		{
+			switch (textType)
+			{
+				case UhtPropertyTextType.SparseShort:
+					builder.Append("TOptional");
+					break;
+
+				case UhtPropertyTextType.FunctionThunkParameterArgType:
+					builder.AppendFunctionThunkParameterArrayType(ValueProperty);
+					break;
+
+				default:
+					builder.Append("TOptional<").AppendPropertyText(ValueProperty, textType, true);
+					if (builder[^1] == '>')
+					{
+						// if our internal property type is a template class, add a space between the closing brackets b/c VS.NET cannot parse this correctly
+						builder.Append(' ');
+					}
+					builder.Append('>');
+					break;
+			}
+			return builder;
+		}
+
+		/// <inheritdoc/>
+		public override StringBuilder AppendMemberDecl(StringBuilder builder, IUhtPropertyMemberContext context, string name, string nameSuffix, int tabs)
+		{
+			builder.AppendMemberDecl(ValueProperty, context, name, GetNameSuffix(nameSuffix, "_Inner"), tabs);
+			return AppendMemberDecl(builder, context, name, nameSuffix, tabs, "FGenericPropertyParams");
+		}
+
+		/// <inheritdoc/>
+		public override StringBuilder AppendMemberDef(StringBuilder builder, IUhtPropertyMemberContext context, string name, string nameSuffix, string? offset, int tabs)
+		{
+			builder.AppendMemberDef(ValueProperty, context, name, GetNameSuffix(nameSuffix, "_Inner"), "0", tabs);
+			AppendMemberDefStart(builder, context, name, nameSuffix, offset, tabs, "FGenericPropertyParams", "UECodeGen_Private::EPropertyGenFlags::Optional");
+			AppendMemberDefEnd(builder, context, name, nameSuffix);
+			return builder;
+		}
+
+		/// <inheritdoc/>
+		public override StringBuilder AppendMemberPtr(StringBuilder builder, IUhtPropertyMemberContext context, string name, string nameSuffix, int tabs)
+		{
+			builder.AppendMemberPtr(ValueProperty, context, name, GetNameSuffix(nameSuffix, "_Inner"), tabs);
+			base.AppendMemberPtr(builder, context, name, nameSuffix, tabs);
+			return builder;
+		}
+
+		/// <inheritdoc/>
+		public override void AppendObjectHashes(StringBuilder builder, int startingLength, IUhtPropertyMemberContext context)
+		{
+			ValueProperty.AppendObjectHashes(builder, startingLength, context);
+		}
+
+		/// <inheritdoc/>
+		public override StringBuilder AppendNullConstructorArg(StringBuilder builder, bool isInitializer)
+		{
+			builder.AppendPropertyText(this, UhtPropertyTextType.Construction).Append("()");
+			return builder;
+		}
+
+		/// <inheritdoc/>
+		public override bool SanitizeDefaultValue(IUhtTokenReader defaultValueReader, StringBuilder innerDefaultValue)
+		{
+			return false;
+		}
+
+		/// <inheritdoc/>
+		protected override void ValidateFunctionArgument(UhtFunction function, UhtValidationOptions options)
+		{
+			base.ValidateFunctionArgument(function, options);
+
+			if (function.FunctionFlags.HasAnyFlags(EFunctionFlags.Net))
+			{
+				if (!function.FunctionFlags.HasAnyFlags(EFunctionFlags.NetRequest))
+				{
+					if (RefQualifier != UhtPropertyRefQualifier.ConstRef && !IsStaticArray)
+					{
+						this.LogError("Replicated TOptional parameters must be passed by const reference");
+					}
+				}
+			}
+		}
+
+		/// <inheritdoc/>
+		public override void Validate(UhtStruct outerStruct, UhtProperty outermostProperty, UhtValidationOptions options)
+		{
+			base.Validate(outerStruct, outermostProperty, options);
+
+			if (PropertyFlags.HasAnyFlags(EPropertyFlags.Net))
+			{
+				this.LogError("Replicated Optionals with MemoryImageAllocators are not yet supported");
+			}
+
+			if (ValueProperty is UhtStructProperty structProperty)
+			{
+				if (structProperty.ScriptStruct == outerStruct)
+				{
+					this.LogError($"'Struct' recursion via optionals is unsupported for properties.");
+				}
+			}
+		}
+
+		/// <inheritdoc/>
+		public override bool IsSameType(UhtProperty other)
+		{
+			if (other is UhtOptionalProperty otherOptional)
+			{
+				return ValueProperty.IsSameType(otherOptional.ValueProperty);
+			}
+			return false;
+		}
+
+		#region Keyword
+		[UhtPropertyType(Keyword = "TOptional")]
+		[SuppressMessage("CodeQuality", "IDE0051:Remove unused private members", Justification = "Attribute accessed method")]
+		private static UhtProperty? OptionalProperty(UhtPropertyResolvePhase resolvePhase, UhtPropertySettings propertySettings, IUhtTokenReader tokenReader, UhtToken matchedToken)
+		{
+			using UhtMessageContext tokenContext = new("TOptional");
+			if (!tokenReader.SkipExpectedType(matchedToken.Value, propertySettings.PropertyCategory == UhtPropertyCategory.Member))
+			{
+				return null;
+			}
+			tokenReader.Require('<');
+
+			// Parse the value type
+			UhtProperty? value = UhtPropertyParser.ParseTemplateParam(resolvePhase, propertySettings, propertySettings.SourceName, tokenReader);
+			if (value == null)
+			{
+				return null;
+			}
+
+			if (!value.PropertyCaps.HasAnyFlags(UhtPropertyCaps.CanBeContainerValue))
+			{
+				tokenReader.LogError($"The type \'{value.GetUserFacingDecl()}\' can not be used as a value in a TOptional");
+			}
+
+			tokenReader.Require('>');
+
+			//@TODO: Prevent sparse delegate types from being used in a container
+
+			return new UhtOptionalProperty(propertySettings, value);
+		}
+		#endregion
+	}
+}
diff --git a/Engine/Source/Runtime/CoreUObject/Private/UObject/PropertyOptional.cpp b/Engine/Source/Runtime/CoreUObject/Private/UObject/PropertyOptional.cpp
--- a/Engine/Source/Runtime/CoreUObject/Private/UObject/PropertyOptional.cpp
+++ b/Engine/Source/Runtime/CoreUObject/Private/UObject/PropertyOptional.cpp
@@ -39,6 +39,11 @@ FOptionalProperty::FOptionalProperty(FFieldVariant InOwner, const FName& InName,
 {
 }
 
+FOptionalProperty::FOptionalProperty(FFieldVariant InOwner, const UECodeGen_Private::FGenericPropertyParams& Prop)
+	: Super(InOwner, (const UECodeGen_Private::FPropertyParamsBaseWithOffset&)Prop)
+{
+}
+
 FOptionalProperty::~FOptionalProperty()
 {
 	if (ValueProperty)
@@ -102,6 +107,11 @@ void FOptionalProperty::GetInnerFields(TArray<FField*>& OutFields)
 	ValueProperty->GetInnerFields(OutFields);
 }
 
+void FOptionalProperty::AddCppProperty(FProperty* Property)
+{
+	SetValueProperty(Property);
+}
+
 FString FOptionalProperty::GetCPPType(FString* ExtendedTypeText, uint32 CPPExportFlags) const
 {
 	checkSlow(ValueProperty);
diff --git a/Engine/Source/Runtime/CoreUObject/Private/UObject/UObjectGlobals.cpp b/Engine/Source/Runtime/CoreUObject/Private/UObject/UObjectGlobals.cpp
--- a/Engine/Source/Runtime/CoreUObject/Private/UObject/UObjectGlobals.cpp
+++ b/Engine/Source/Runtime/CoreUObject/Private/UObject/UObjectGlobals.cpp
@@ -5966,6 +5966,15 @@ namespace UECodeGen_Private
 				NewProp = NewFProperty<FFieldPathProperty, FFieldPathPropertyParams>(Outer, *PropBase);
 			}
 			break;
+
+			case EPropertyGenFlags::Optional:
+			{
+				NewProp = NewFProperty<FOptionalProperty, FGenericPropertyParams>(Outer, *PropBase);
+
+				// Next property is the optional inner
+				ReadMore = 1;
+			}
+			break;
 		}
 
 		NewProp->ArrayDim = PropBase->ArrayDim;
diff --git a/Engine/Source/Runtime/CoreUObject/Public/UObject/ObjectMacros.h b/Engine/Source/Runtime/CoreUObject/Public/UObject/ObjectMacros.h
--- a/Engine/Source/Runtime/CoreUObject/Public/UObject/ObjectMacros.h
+++ b/Engine/Source/Runtime/CoreUObject/Public/UObject/ObjectMacros.h
@@ -461,10 +461,11 @@ enum EPropertyFlags : uint64
 #define CPF_ParmFlags				(CPF_Parm | CPF_OutParm | CPF_ReturnParm | CPF_RequiredParm | CPF_ReferenceParm | CPF_ConstParm )
 
 /** Flags that are propagated to properties inside containers */
-#define CPF_PropagateToArrayInner	(CPF_ExportObject | CPF_PersistentInstance | CPF_InstancedReference | CPF_ContainsInstancedReference | CPF_Config | CPF_EditConst | CPF_Deprecated | CPF_EditorOnly | CPF_AutoWeak | CPF_UObjectWrapper )
-#define CPF_PropagateToMapValue		(CPF_ExportObject | CPF_PersistentInstance | CPF_InstancedReference | CPF_ContainsInstancedReference | CPF_Config | CPF_EditConst | CPF_Deprecated | CPF_EditorOnly | CPF_AutoWeak | CPF_UObjectWrapper | CPF_Edit )
-#define CPF_PropagateToMapKey		(CPF_ExportObject | CPF_PersistentInstance | CPF_InstancedReference | CPF_ContainsInstancedReference | CPF_Config | CPF_EditConst | CPF_Deprecated | CPF_EditorOnly | CPF_AutoWeak | CPF_UObjectWrapper | CPF_Edit )
-#define CPF_PropagateToSetElement	(CPF_ExportObject | CPF_PersistentInstance | CPF_InstancedReference | CPF_ContainsInstancedReference | CPF_Config | CPF_EditConst | CPF_Deprecated | CPF_EditorOnly | CPF_AutoWeak | CPF_UObjectWrapper | CPF_Edit )
+#define CPF_PropagateToArrayInner    (CPF_ExportObject | CPF_PersistentInstance | CPF_InstancedReference | CPF_ContainsInstancedReference | CPF_Config | CPF_EditConst | CPF_Deprecated | CPF_EditorOnly | CPF_AutoWeak | CPF_UObjectWrapper )
+#define CPF_PropagateToOptionalInner (CPF_ExportObject | CPF_PersistentInstance | CPF_InstancedReference | CPF_ContainsInstancedReference | CPF_Config | CPF_EditConst | CPF_Deprecated | CPF_EditorOnly | CPF_AutoWeak | CPF_UObjectWrapper | CPF_Edit )
+#define CPF_PropagateToMapValue      (CPF_ExportObject | CPF_PersistentInstance | CPF_InstancedReference | CPF_ContainsInstancedReference | CPF_Config | CPF_EditConst | CPF_Deprecated | CPF_EditorOnly | CPF_AutoWeak | CPF_UObjectWrapper | CPF_Edit )
+#define CPF_PropagateToMapKey        (CPF_ExportObject | CPF_PersistentInstance | CPF_InstancedReference | CPF_ContainsInstancedReference | CPF_Config | CPF_EditConst | CPF_Deprecated | CPF_EditorOnly | CPF_AutoWeak | CPF_UObjectWrapper | CPF_Edit )
+#define CPF_PropagateToSetElement    (CPF_ExportObject | CPF_PersistentInstance | CPF_InstancedReference | CPF_ContainsInstancedReference | CPF_Config | CPF_EditConst | CPF_Deprecated | CPF_EditorOnly | CPF_AutoWeak | CPF_UObjectWrapper | CPF_Edit )
 
 /** The flags that should never be set on interface properties */
 #define CPF_InterfaceClearMask		(CPF_ExportObject|CPF_InstancedReference|CPF_ContainsInstancedReference)
diff --git a/Engine/Source/Runtime/CoreUObject/Public/UObject/PropertyOptional.h b/Engine/Source/Runtime/CoreUObject/Public/UObject/PropertyOptional.h
--- a/Engine/Source/Runtime/CoreUObject/Public/UObject/PropertyOptional.h
+++ b/Engine/Source/Runtime/CoreUObject/Public/UObject/PropertyOptional.h
@@ -3,6 +3,7 @@
 #pragma once
 
 #include "UObject/UnrealType.h"
+#include "UObject/UObjectGlobals.h"
 #include "Misc/Optional.h"
 
 // 
@@ -184,6 +185,7 @@ class COREUOBJECT_API FOptionalProperty : public FProperty, public FOptionalProp
 public:
 
 	FOptionalProperty(FFieldVariant InOwner, const FName& InName, EObjectFlags InObjectFlags);
+	FOptionalProperty(FFieldVariant InOwner, const UECodeGen_Private::FGenericPropertyParams& Prop);
 	virtual ~FOptionalProperty();
 
 	// Sets the optional property's value property.
@@ -198,6 +200,7 @@ public:
 	virtual void PostDuplicate(const FField& InField) override;
 	virtual FField* GetInnerFieldByName(const FName& InName) override;
 	virtual void GetInnerFields(TArray<FField*>& OutFields) override;
+	virtual void AddCppProperty(FProperty* Property) override;
 	// End of Field interface
 
 	// UHT interface
diff --git a/Engine/Source/Runtime/CoreUObject/Public/UObject/UObjectGlobals.h b/Engine/Source/Runtime/CoreUObject/Public/UObject/UObjectGlobals.h
--- a/Engine/Source/Runtime/CoreUObject/Public/UObject/UObjectGlobals.h
+++ b/Engine/Source/Runtime/CoreUObject/Public/UObject/UObjectGlobals.h
@@ -3382,6 +3382,7 @@ namespace UECodeGen_Private
 		Enum              = 0x1E,
 		FieldPath         = 0x1F,
 		LargeWorldCoordinatesReal = 0x20,
+		Optional          = 0x21,
 
 		// Property-specific flags
 		NativeBool        = 0x40,
