Registration declarations now uses contexpr to initialize the data.

diff --git a/Engine/Source/Programs/Shared/EpicGames.UHT/Exporters/CodeGen/UhtHeaderCodeGeneratorCppFile.cs b/Engine/Source/Programs/Shared/EpicGames.UHT/Exporters/CodeGen/UhtHeaderCodeGeneratorCppFile.cs
--- a/Engine/Source/Programs/Shared/EpicGames.UHT/Exporters/CodeGen/UhtHeaderCodeGeneratorCppFile.cs
+++ b/Engine/Source/Programs/Shared/EpicGames.UHT/Exporters/CodeGen/UhtHeaderCodeGeneratorCppFile.cs
@@ -237,25 +237,19 @@ namespace EpicGames.UHT.Exporters.CodeGen
 					string name = $"Z_CompiledInDeferFile_{headerInfo.FileId}";
 					string staticsName = $"{name}_Statics";
 
+					uint combinedHash = UInt32.MaxValue;
+
 					using UhtCodeBlockComment blockComment = new(builder, "Registration");
 					builder.Append("struct ").Append(staticsName).Append("\r\n");
 					builder.Append("{\r\n");
 
-					builder.AppendIfInstances(enums, UhtDefineScopeNames.Standard, builder => builder.Append("\tstatic const FEnumRegisterCompiledInInfo EnumInfo[];\r\n"));
-					builder.AppendIfInstances(scriptStructs, UhtDefineScopeNames.Standard, builder => builder.Append("\tstatic const FStructRegisterCompiledInInfo ScriptStructInfo[];\r\n"));
-					builder.AppendIfInstances(classes, UhtDefineScopeNames.Standard, builder => builder.Append("\tstatic const FClassRegisterCompiledInInfo ClassInfo[];\r\n"));
-
-					builder.Append("};\r\n");
-
-					uint combinedHash = UInt32.MaxValue;
-
-					builder.AppendInstances(enums, UhtDefineScopeNames.Standard,
-						builder => builder.Append("const FEnumRegisterCompiledInInfo ").Append(staticsName).Append("::EnumInfo[] = {\r\n"),
+					builder.AppendInstances(enums, UhtDefineScopeNames.Standard, 
+						builder => builder.Append("\tstatic constexpr FEnumRegisterCompiledInInfo EnumInfo[] = {\r\n"),
 						(builder, enumObj) =>
 						{
 							uint hash = ObjectInfos[enumObj.ObjectTypeIndex].Hash;
 							builder
-								.Append("\t{ ")
+								.Append("\t\t{ ")
 								.Append(enumObj.SourceName)
 								.Append("_StaticEnum, TEXT(\"")
 								.Append(enumObj.EngineName)
@@ -264,15 +258,15 @@ namespace EpicGames.UHT.Exporters.CodeGen
 								.Append($", CONSTRUCT_RELOAD_VERSION_INFO(FEnumReloadVersionInfo, {hash}U) }},\r\n");
 							combinedHash = HashCombine(combinedHash, hash);
 						},
-						builder => builder.Append("};\r\n"));
+						builder => builder.Append("\t};\r\n"));
 
 					builder.AppendInstances(scriptStructs, UhtDefineScopeNames.Standard,
-						builder => builder.Append("const FStructRegisterCompiledInInfo ").Append(staticsName).Append("::ScriptStructInfo[] = {\r\n"),
+						builder => builder.Append("\tstatic constexpr FStructRegisterCompiledInInfo ScriptStructInfo[] = {\r\n"),
 						(builder, scriptStruct) =>
 						{
 							uint hash = ObjectInfos[scriptStruct.ObjectTypeIndex].Hash;
 							builder
-								.Append("\t{ ")
+								.Append("\t\t{ ")
 								.Append(scriptStruct.SourceName)
 								.Append("::StaticStruct, Z_Construct_UScriptStruct_")
 								.Append(scriptStruct.SourceName)
@@ -285,15 +279,15 @@ namespace EpicGames.UHT.Exporters.CodeGen
 								.Append($"), {hash}U) }},\r\n");
 							combinedHash = HashCombine(combinedHash, hash);
 						},
-						builder => builder.Append("};\r\n"));
+						builder => builder.Append("\t};\r\n"));
 
 					builder.AppendInstances(classes, UhtDefineScopeNames.Standard,
-						builder => builder.Append("const FClassRegisterCompiledInInfo ").Append(staticsName).Append("::ClassInfo[] = {\r\n"),
+						builder => builder.Append("\tstatic constexpr FClassRegisterCompiledInInfo ClassInfo[] = {\r\n"),
 						(builder, classObj) =>
 						{
 							uint hash = ObjectInfos[classObj.ObjectTypeIndex].Hash;
 							builder
-								.Append("\t{ Z_Construct_UClass_")
+								.Append("\t\t{ Z_Construct_UClass_")
 								.Append(classObj.SourceName)
 								.Append(", ")
 								.Append(classObj.SourceName)
@@ -306,7 +300,9 @@ namespace EpicGames.UHT.Exporters.CodeGen
 								.Append($"), {hash}U) }},\r\n");
 							combinedHash = HashCombine(combinedHash, hash);
 						},
-						builder => builder.Append("};\r\n"));
+						builder => builder.Append("\t};\r\n"));
+
+					builder.Append("};\r\n");
 
 					builder
 						.Append("static FRegisterCompiledInInfo ")
