EpicGames.Core: Swallow DirectoryNotFound exceptions when trying to enumerate/delete the contents of a directory. If some other process is also deleting directories, we can clash when trying to enumerate a directory's contents.

diff --git a/Engine/Source/Programs/Shared/EpicGames.Core/FileUtils.cs b/Engine/Source/Programs/Shared/EpicGames.Core/FileUtils.cs
--- a/Engine/Source/Programs/Shared/EpicGames.Core/FileUtils.cs
+++ b/Engine/Source/Programs/Shared/EpicGames.Core/FileUtils.cs
@@ -278,6 +278,10 @@ namespace EpicGames.Core
 			{
 				directory.Delete(true);
 			}
+			catch (DirectoryNotFoundException)
+			{
+				// Race condition with something else deleting the same directory.
+			}
 			catch (Exception ex)
 			{
 				throw new WrappedFileOrDirectoryException(ex, String.Format("Unable to delete '{0}': {1}", directory.FullName, ex.Message.TrimEnd()));
@@ -307,21 +311,28 @@ namespace EpicGames.Core
 			{
 				if (directory.Exists)
 				{
-					foreach (FileInfo file in directory.EnumerateFiles())
-					{
-						ForceDeleteFile(file);
-					}
-					foreach (DirectoryInfo subDirectory in directory.EnumerateDirectories())
+					try
 					{
-						if (subDirectory.Attributes.HasFlag(FileAttributes.ReparsePoint))
+						foreach (FileInfo file in directory.EnumerateFiles())
 						{
-							ForceDeleteDirectoryInternal(subDirectory);
+							ForceDeleteFile(file);
 						}
-						else
+						foreach (DirectoryInfo subDirectory in directory.EnumerateDirectories())
 						{
-							ForceDeleteDirectory(subDirectory);
+							if (subDirectory.Attributes.HasFlag(FileAttributes.ReparsePoint))
+							{
+								ForceDeleteDirectoryInternal(subDirectory);
+							}
+							else
+							{
+								ForceDeleteDirectory(subDirectory);
+							}
 						}
 					}
+					catch (DirectoryNotFoundException)
+					{
+						// Race condition with something else deleting the same directory.
+					}
 				}
 			}
 		}
