Meta data code generation uses constexpr initialization.

diff --git a/Engine/Source/Programs/Shared/EpicGames.UHT/Exporters/CodeGen/UhtHeaderCodeGeneratorCppFile.cs b/Engine/Source/Programs/Shared/EpicGames.UHT/Exporters/CodeGen/UhtHeaderCodeGeneratorCppFile.cs
--- a/Engine/Source/Programs/Shared/EpicGames.UHT/Exporters/CodeGen/UhtHeaderCodeGeneratorCppFile.cs
+++ b/Engine/Source/Programs/Shared/EpicGames.UHT/Exporters/CodeGen/UhtHeaderCodeGeneratorCppFile.cs
@@ -448,7 +448,7 @@ namespace EpicGames.UHT.Exporters.CodeGen
 				builder.Append("struct ").Append(staticsName).Append("\r\n");
 				builder.Append("{\r\n");
 				builder.Append("\tstatic const UECodeGen_Private::FEnumeratorParam Enumerators[];\r\n");
-				builder.AppendMetaDataDecl(enumObj, MetaDataParamsName, 2);
+				builder.AppendMetaDataDecl(enumObj, null, null, MetaDataParamsName, 1);
 				builder.Append("\tstatic const UECodeGen_Private::FEnumParams EnumParams;\r\n");
 				builder.Append("};\r\n");
 
@@ -466,9 +466,6 @@ namespace EpicGames.UHT.Exporters.CodeGen
 				}
 				builder.Append("};\r\n");
 
-				// Meta data
-				builder.AppendMetaDataDef(enumObj, staticsName, MetaDataParamsName, 0);
-
 				// Singleton parameters
 				builder.Append("const UECodeGen_Private::FEnumParams ").Append(staticsName).Append("::EnumParams = {\r\n");
 				builder.Append("\t(UObject*(*)())").Append(PackageSingletonName).Append(",\r\n");
@@ -638,6 +635,7 @@ namespace EpicGames.UHT.Exporters.CodeGen
 			int hashCodeBlockStart = builder.Length;
 
 			// Collect the properties
+			PropertyMemberContextImpl context = new(CodeGenerator, scriptStruct, scriptStruct.SourceName, staticsName);
 			UhtUsedDefineScopes<UhtProperty> properties = new(scriptStruct.Properties);
 
 			// Declare the statics structure
@@ -653,7 +651,7 @@ namespace EpicGames.UHT.Exporters.CodeGen
 				}
 
 				// Meta data
-				builder.AppendMetaDataDecl(scriptStruct, MetaDataParamsName, 1);
+				builder.AppendMetaDataDecl(scriptStruct, context, properties, MetaDataParamsName, 1);
 
 				// New struct ops
 				if (scriptStruct.ScriptStructFlags.HasAnyFlags(EStructFlags.Native))
@@ -661,16 +659,14 @@ namespace EpicGames.UHT.Exporters.CodeGen
 					builder.Append("\tstatic void* NewStructOps();\r\n");
 				}
 
-				AppendPropertiesDecl(builder, scriptStruct, scriptStruct.SourceName, properties, staticsName, 1);
+				AppendPropertiesDecl(builder, context, properties, 1);
 
-				builder.Append("\tstatic const UECodeGen_Private::FStructParams ReturnStructParams;\r\n");
+				builder.Append("\tstatic const UECodeGen_Private::FStructParams StructParams;\r\n");
 				builder.Append("};\r\n");
 			}
 
 			// Populate the elements of the static structure
 			{
-				builder.AppendMetaDataDef(scriptStruct, staticsName, MetaDataParamsName, 0);
-
 				if (scriptStruct.ScriptStructFlags.HasAnyFlags(EStructFlags.Native))
 				{
 					builder.Append("void* ").Append(staticsName).Append("::NewStructOps()\r\n");
@@ -679,9 +675,9 @@ namespace EpicGames.UHT.Exporters.CodeGen
 					builder.Append("}\r\n");
 				}
 
-				AppendPropertiesDefs(builder, scriptStruct, scriptStruct.SourceName, properties, staticsName, 0);
+				AppendPropertiesDefs(builder, context, properties, 0);
 
-				builder.Append("const UECodeGen_Private::FStructParams ").Append(staticsName).Append("::ReturnStructParams = {\r\n");
+				builder.Append("const UECodeGen_Private::FStructParams ").Append(staticsName).Append("::StructParams = {\r\n");
 				builder.Append("\t(UObject* (*)())").Append(PackageSingletonName).Append(",\r\n");
 				builder.Append('\t').Append(GetSingletonName(scriptStruct.SuperScriptStruct, true)).Append(",\r\n");
 				if (scriptStruct.ScriptStructFlags.HasAnyFlags(EStructFlags.Native))
@@ -701,8 +697,6 @@ namespace EpicGames.UHT.Exporters.CodeGen
 				builder.Append($"\tEStructFlags(0x{(uint)(scriptStruct.ScriptStructFlags & ~EStructFlags.ComputedFlags):X8}),\r\n");
 				builder.Append('\t').AppendMetaDataParams(scriptStruct, staticsName, MetaDataParamsName).Append("\r\n");
 				builder.Append("};\r\n");
-
-				builder.AppendPropertiesParamsCountStaticAssert(properties, staticsName);
 			}
 
 			// Generate the registration function
@@ -721,7 +715,7 @@ namespace EpicGames.UHT.Exporters.CodeGen
 				}
 				builder.Append("\tif (!").Append(innerSingletonName).Append(")\r\n");
 				builder.Append("\t{\r\n");
-				builder.Append("\t\tUECodeGen_Private::ConstructUScriptStruct(").Append(innerSingletonName).Append(", ").Append(staticsName).Append("::ReturnStructParams);\r\n");
+				builder.Append("\t\tUECodeGen_Private::ConstructUScriptStruct(").Append(innerSingletonName).Append(", ").Append(staticsName).Append("::StructParams);\r\n");
 				builder.Append("\t}\r\n");
 				builder.Append("\treturn ").Append(innerSingletonName).Append(";\r\n");
 				builder.Append("}\r\n");
@@ -869,6 +863,7 @@ namespace EpicGames.UHT.Exporters.CodeGen
 			{
 
 				// Collect the properties
+				PropertyMemberContextImpl context = new(CodeGenerator, function, eventParameters, staticsName);
 				UhtUsedDefineScopes<UhtProperty> properties = new(function.Properties);
 
 				// Statics declaration
@@ -886,9 +881,9 @@ namespace EpicGames.UHT.Exporters.CodeGen
 						AppendEventParameter(builder, function, strippedFunctionName, UhtPropertyTextType.EventParameterFunctionMember, false, 1, "\r\n");
 					}
 
-					AppendPropertiesDecl(builder, function, eventParameters, properties, staticsName, 1);
+					AppendPropertiesDecl(builder, context, properties, 1);
 
-					builder.AppendMetaDataDecl(function, MetaDataParamsName, 1);
+					builder.AppendMetaDataDecl(function, context, properties, MetaDataParamsName, 1);
 
 					builder.Append("\tstatic const UECodeGen_Private::FFunctionParams FuncParams;\r\n");
 
@@ -897,9 +892,7 @@ namespace EpicGames.UHT.Exporters.CodeGen
 
 				// Statics definition
 				{
-					AppendPropertiesDefs(builder, function, eventParameters, properties, staticsName, 0);
-
-					builder.AppendMetaDataDef(function, staticsName, MetaDataParamsName, 0);
+					AppendPropertiesDefs(builder, context, properties, 0);
 
 					builder
 						.Append("const UECodeGen_Private::FFunctionParams ")
@@ -947,8 +940,6 @@ namespace EpicGames.UHT.Exporters.CodeGen
 						.AppendMetaDataParams(function, staticsName, MetaDataParamsName)
 						.Append(" };\r\n");
 
-					builder.AppendPropertiesParamsCountStaticAssert(properties, staticsName);
-
 					if (!String.IsNullOrEmpty(sizeOfStatic))
 					{
 						builder.Append("static_assert(").Append(sizeOfStatic).Append(" < MAX_uint16);\r\n");
@@ -1023,15 +1014,12 @@ namespace EpicGames.UHT.Exporters.CodeGen
 			return builder;
 		}
 
-		private StringBuilder AppendPropertiesDecl(StringBuilder builder, UhtStruct structObj, string structSourceName, UhtUsedDefineScopes<UhtProperty> properties, string staticsName, int tabs)
+		private static StringBuilder AppendPropertiesDecl(StringBuilder builder, IUhtPropertyMemberContext context, UhtUsedDefineScopes<UhtProperty> properties, int tabs)
 		{
 			if (properties.IsEmpty)
 			{
 				return builder;
 			}
-
-			PropertyMemberContextImpl context = new(CodeGenerator, structObj, structSourceName, staticsName);
-
 			builder.AppendInstances(properties, UhtDefineScopeNames.Standard,
 				builder => { },
 				(builder, property) => builder.AppendMemberDecl(property, context, property.EngineName, "", tabs),
@@ -1039,22 +1027,23 @@ namespace EpicGames.UHT.Exporters.CodeGen
 			return builder;
 		}
 
-		private StringBuilder AppendPropertiesDefs(StringBuilder builder, UhtStruct structObj, string structSourceName, UhtUsedDefineScopes<UhtProperty> properties, string staticsName, int tabs)
+		private static StringBuilder AppendPropertiesDefs(StringBuilder builder, IUhtPropertyMemberContext context, UhtUsedDefineScopes<UhtProperty> properties, int tabs)
 		{
 			if (properties.IsEmpty)
 			{
 				return builder;
 			}
-
-			PropertyMemberContextImpl context = new(CodeGenerator, structObj, structSourceName, staticsName);
-
 			builder.AppendInstances(properties, UhtDefineScopeNames.Standard,
 				(builder, property) => builder.AppendMemberDef(property, context, property.EngineName, "", null, tabs));
 
 			builder.AppendInstances(properties, UhtDefineScopeNames.Standard,
-				builder => builder.AppendTabs(tabs).Append("const UECodeGen_Private::FPropertyParamsBase* const ").Append(staticsName).Append("::PropPointers[] = {\r\n"),
+				builder => builder.AppendTabs(tabs).Append("const UECodeGen_Private::FPropertyParamsBase* const ").Append(context.StaticsName).Append("::PropPointers[] = {\r\n"),
 				(builder, property) => builder.AppendMemberPtr(property, context, property.EngineName, "", tabs + 1),
-				builder => builder.AppendTabs(tabs).Append("};\r\n"));
+				builder =>
+				{
+					builder.AppendTabs(tabs).Append("};\r\n");
+					builder.AppendTabs(tabs).Append("static_assert(UE_ARRAY_COUNT(").Append(context.StaticsName).Append("::PropPointers) < 2048);\r\n");
+				});
 			return builder;
 		}
 
@@ -1179,7 +1168,7 @@ namespace EpicGames.UHT.Exporters.CodeGen
 						string getterCallText = property.Getter ?? "Get" + property.SourceName;
 						builder.AppendPropertyText(property, UhtPropertyTextType.GetterRetVal).Append(classObj.SourceName).Append("::").Append(getterCallText).Append("() const\r\n");
 						builder.Append("{\r\n");
-						builder.Append('\t').Append("return ").Append(property.SourceName).Append(";\r\n");
+						builder.Append("\treturn ").Append(property.SourceName).Append(";\r\n");
 						builder.Append("}\r\n");
 					}
 					if (property.PropertyExportFlags.HasAnyFlags(UhtPropertyExportFlags.SetterSpecifiedAuto))
@@ -1482,9 +1471,9 @@ namespace EpicGames.UHT.Exporters.CodeGen
 
 				builder.AppendIfInstances(functions, UhtDefineScopeNames.WithEditor, builder => builder.Append("\tstatic const FClassFunctionLinkInfo FuncInfo[];\r\n"));
 
-				builder.AppendMetaDataDecl(classObj, MetaDataParamsName, 1);
+				builder.AppendMetaDataDecl(classObj, context, properties, MetaDataParamsName, 1);
 
-				AppendPropertiesDecl(builder, classObj, classObj.SourceName, properties, staticsName, 1);
+				AppendPropertiesDecl(builder, context, properties, 1);
 
 				if (hasInterfaces)
 				{
@@ -1532,9 +1521,7 @@ namespace EpicGames.UHT.Exporters.CodeGen
 						builder.Append("static_assert(UE_ARRAY_COUNT(").Append(staticsName).Append("::FuncInfo) < 2048);\r\n");
 					});
 
-				builder.AppendMetaDataDef(classObj, staticsName, MetaDataParamsName, 0);
-
-				AppendPropertiesDefs(builder, classObj, classObj.SourceName, properties, staticsName, 0);
+				AppendPropertiesDefs(builder, context, properties, 0);
 
 				if (hasInterfaces)
 				{
@@ -1600,8 +1587,6 @@ namespace EpicGames.UHT.Exporters.CodeGen
 				builder.Append($"\t\t0x{(uint)classFlags:X8}u,\r\n");
 				builder.Append('\t').AppendMetaDataParams(classObj, staticsName, MetaDataParamsName).Append("\r\n");
 				builder.Append("};\r\n");
-
-				builder.AppendPropertiesParamsCountStaticAssert(properties, staticsName);
 			}
 
 			// Class registration
@@ -2023,21 +2008,5 @@ namespace EpicGames.UHT.Exporters.CodeGen
 			}
 			return builder;
 		}
-
-		/// <summary>
-		/// Append static_assert to check that properties params count fit within struct
-		/// </summary>
-		/// <param name="builder">Destination builder</param>
-		/// <param name="properties">Collection of properties</param>
-		/// <param name="staticsName">Name of the statics section for this structure</param>
-		/// <returns>Destination builder</returns>
-		public static StringBuilder AppendPropertiesParamsCountStaticAssert(this StringBuilder builder, UhtUsedDefineScopes<UhtProperty> properties, string staticsName)
-		{
-			builder.AppendIfInstances(properties, UhtDefineScopeNames.Standard, builder =>
-			{
-				builder.Append("static_assert(UE_ARRAY_COUNT(").Append(staticsName).Append("::PropPointers) < 2048);\r\n");
-			});
-			return builder;
-		}
 	}
 }
diff --git a/Engine/Source/Programs/Shared/EpicGames.UHT/Exporters/CodeGen/UhtPackageCodeGenerator.cs b/Engine/Source/Programs/Shared/EpicGames.UHT/Exporters/CodeGen/UhtPackageCodeGenerator.cs
--- a/Engine/Source/Programs/Shared/EpicGames.UHT/Exporters/CodeGen/UhtPackageCodeGenerator.cs
+++ b/Engine/Source/Programs/Shared/EpicGames.UHT/Exporters/CodeGen/UhtPackageCodeGenerator.cs
@@ -1,6 +1,9 @@
 // Copyright Epic Games, Inc. All Rights Reserved.
 
+using System.Linq;
+using System.Text;
 using EpicGames.UHT.Types;
+using EpicGames.UHT.Utils;
 
 namespace EpicGames.UHT.Exporters.CodeGen
 {
@@ -118,4 +121,40 @@ namespace EpicGames.UHT.Exporters.CodeGen
 		}
 		#endregion
 	}
+
+	/// <summary>
+	/// Helper formatting methods
+	/// </summary>
+	public static class UhtPackageCodeGeneratorExtensions
+	{
+		/// <summary>
+		/// Append the meta data declaration
+		/// </summary>
+		/// <param name="builder">Destination builder</param>
+		/// <param name="type">Source type containing the meta data</param>
+		/// <param name="propertyContext">Context for formatting properties</param>
+		/// <param name="properties">Optional collection of properties to output</param>
+		/// <param name="name">Name</param>
+		/// <param name="tabs">Number of tabs to indent</param>
+		/// <returns>Destination builder</returns>
+		public static StringBuilder AppendMetaDataDecl(this StringBuilder builder, UhtType type, IUhtPropertyMemberContext? propertyContext, UhtUsedDefineScopes<UhtProperty>? properties, string name, int tabs)
+		{
+			UhtStruct? structObj = type as UhtStruct;
+			if (!type.MetaData.IsEmpty() || (properties != null && properties.Instances.Any(x => !x.MetaData.IsEmpty())))
+			{
+				builder.Append("#if WITH_METADATA\r\n");
+				builder.AppendMetaDataDecl(type, null, name, null, null, tabs);
+				if (propertyContext != null && properties != null)
+				{
+					builder.AppendInstances(properties, UhtDefineScopeNames.Standard,
+						(builder, property) =>
+						{
+							builder.AppendMetaDataDecl(property, propertyContext, property.EngineName, "", tabs);
+						});
+				}
+				builder.Append("#endif // WITH_METADATA\r\n");
+			}
+			return builder;
+		}
+	}
 }
diff --git a/Engine/Source/Programs/Shared/EpicGames.UHT/Exporters/CodeGen/UhtPackageCodeGeneratorCppFile.cs b/Engine/Source/Programs/Shared/EpicGames.UHT/Exporters/CodeGen/UhtPackageCodeGeneratorCppFile.cs
--- a/Engine/Source/Programs/Shared/EpicGames.UHT/Exporters/CodeGen/UhtPackageCodeGeneratorCppFile.cs
+++ b/Engine/Source/Programs/Shared/EpicGames.UHT/Exporters/CodeGen/UhtPackageCodeGeneratorCppFile.cs
@@ -110,8 +110,6 @@ namespace EpicGames.UHT.Exporters.CodeGen
 					builder.Append(info.RegsiteredExternalDecl);
 				}
 
-				builder.AppendMetaDataDef(Package, null, MetaDataParamsName, 3);
-
 				builder.Append("\tstatic FPackageRegistrationInfo Z_Registration_Info_UPackage_").Append(strippedName).Append(";\r\n");
 				builder.Append("\tFORCENOINLINE UPackage* ").Append(singletonName).Append("()\r\n");
 				builder.Append("\t{\r\n");
diff --git a/Engine/Source/Programs/Shared/EpicGames.UHT/Types/Properties/UhtArrayProperty.cs b/Engine/Source/Programs/Shared/EpicGames.UHT/Types/Properties/UhtArrayProperty.cs
--- a/Engine/Source/Programs/Shared/EpicGames.UHT/Types/Properties/UhtArrayProperty.cs
+++ b/Engine/Source/Programs/Shared/EpicGames.UHT/Types/Properties/UhtArrayProperty.cs
@@ -133,6 +133,13 @@ namespace EpicGames.UHT.Types
 			return builder;
 		}
 
+		/// <inheritdoc/>
+		public override StringBuilder AppendMetaDataDecl(StringBuilder builder, IUhtPropertyMemberContext context, string name, string nameSuffix, int tabs)
+		{
+			builder.AppendMetaDataDecl(ValueProperty, context, name, GetNameSuffix(nameSuffix, "_Inner"), tabs);
+			return base.AppendMetaDataDecl(builder, context, name, nameSuffix, tabs);
+		}
+
 		/// <inheritdoc/>
 		public override StringBuilder AppendMemberDecl(StringBuilder builder, IUhtPropertyMemberContext context, string name, string nameSuffix, int tabs)
 		{
diff --git a/Engine/Source/Programs/Shared/EpicGames.UHT/Types/Properties/UhtBoolProperty.cs b/Engine/Source/Programs/Shared/EpicGames.UHT/Types/Properties/UhtBoolProperty.cs
--- a/Engine/Source/Programs/Shared/EpicGames.UHT/Types/Properties/UhtBoolProperty.cs
+++ b/Engine/Source/Programs/Shared/EpicGames.UHT/Types/Properties/UhtBoolProperty.cs
@@ -159,7 +159,6 @@ namespace EpicGames.UHT.Types
 		/// <inheritdoc/>
 		public override StringBuilder AppendMemberDecl(StringBuilder builder, IUhtPropertyMemberContext context, string name, string nameSuffix, int tabs)
 		{
-			builder.AppendMetaDataDecl(this, context, name, nameSuffix, tabs);
 			if (Outer is not UhtProperty)
 			{
 				builder.AppendTabs(tabs).Append("static void ").AppendNameDecl(context, name, nameSuffix).Append("_SetBit(void* Obj);\r\n");
@@ -171,7 +170,6 @@ namespace EpicGames.UHT.Types
 		/// <inheritdoc/>
 		public override StringBuilder AppendMemberDef(StringBuilder builder, IUhtPropertyMemberContext context, string name, string nameSuffix, string? offset, int tabs)
 		{
-			builder.AppendMetaDataDef(this, context, name, nameSuffix, tabs);
 			if (Outer == context.OuterStruct)
 			{
 				builder.AppendTabs(tabs).Append("void ").AppendNameDef(context, name, nameSuffix).Append("_SetBit(void* Obj)\r\n");
@@ -184,7 +182,7 @@ namespace EpicGames.UHT.Types
 				IsNativeBool ?
 				"UECodeGen_Private::EPropertyGenFlags::Bool | UECodeGen_Private::EPropertyGenFlags::NativeBool" :
 				"UECodeGen_Private::EPropertyGenFlags::Bool ",
-				false, false);
+				false);
 
 			builder.Append("sizeof(").Append(CppTypeText).Append("), ");
 
diff --git a/Engine/Source/Programs/Shared/EpicGames.UHT/Types/Properties/UhtClassProperty.cs b/Engine/Source/Programs/Shared/EpicGames.UHT/Types/Properties/UhtClassProperty.cs
--- a/Engine/Source/Programs/Shared/EpicGames.UHT/Types/Properties/UhtClassProperty.cs
+++ b/Engine/Source/Programs/Shared/EpicGames.UHT/Types/Properties/UhtClassProperty.cs
@@ -132,7 +132,7 @@ namespace EpicGames.UHT.Types
 		/// <inheritdoc/>
 		public override StringBuilder AppendMemberDecl(StringBuilder builder, IUhtPropertyMemberContext context, string name, string nameSuffix, int tabs)
 		{
-			return base.AppendMemberDecl(builder, context, name, nameSuffix, tabs, "FClassPropertyParams");
+			return AppendMemberDecl(builder, context, name, nameSuffix, tabs, "FClassPropertyParams");
 		}
 
 		/// <inheritdoc/>
diff --git a/Engine/Source/Programs/Shared/EpicGames.UHT/Types/Properties/UhtEnumProperty.cs b/Engine/Source/Programs/Shared/EpicGames.UHT/Types/Properties/UhtEnumProperty.cs
--- a/Engine/Source/Programs/Shared/EpicGames.UHT/Types/Properties/UhtEnumProperty.cs
+++ b/Engine/Source/Programs/Shared/EpicGames.UHT/Types/Properties/UhtEnumProperty.cs
@@ -212,6 +212,16 @@ namespace EpicGames.UHT.Types
 			return AppendEnumText(builder, this, Enum, textType, isTemplateArgument);
 		}
 
+		/// <inheritdoc/>
+		public override StringBuilder AppendMetaDataDecl(StringBuilder builder, IUhtPropertyMemberContext context, string name, string nameSuffix, int tabs)
+		{
+			if (UnderlyingProperty != null)
+			{
+				builder.AppendMetaDataDecl(UnderlyingProperty, context, name, GetNameSuffix(nameSuffix, "_Underlying"), tabs);
+			}
+			return base.AppendMetaDataDecl(builder, context, name, nameSuffix, tabs);
+		}
+
 		/// <inheritdoc/>
 		public override StringBuilder AppendMemberDecl(StringBuilder builder, IUhtPropertyMemberContext context, string name, string nameSuffix, int tabs)
 		{
@@ -238,7 +248,7 @@ namespace EpicGames.UHT.Types
 			}
 			else
 			{
-				AppendMemberDefStart(builder, context, name, nameSuffix, offset, tabs, "FBytePropertyParams", "UECodeGen_Private::EPropertyGenFlags::Byte", true, true);
+				AppendMemberDefStart(builder, context, name, nameSuffix, offset, tabs, "FBytePropertyParams", "UECodeGen_Private::EPropertyGenFlags::Byte");
 				AppendMemberDefRef(builder, context, Enum, true);
 				AppendMemberDefEnd(builder, context, name, nameSuffix);
 			}
diff --git a/Engine/Source/Programs/Shared/EpicGames.UHT/Types/Properties/UhtMapProperty.cs b/Engine/Source/Programs/Shared/EpicGames.UHT/Types/Properties/UhtMapProperty.cs
--- a/Engine/Source/Programs/Shared/EpicGames.UHT/Types/Properties/UhtMapProperty.cs
+++ b/Engine/Source/Programs/Shared/EpicGames.UHT/Types/Properties/UhtMapProperty.cs
@@ -168,6 +168,14 @@ namespace EpicGames.UHT.Types
 			return builder;
 		}
 
+		/// <inheritdoc/>
+		public override StringBuilder AppendMetaDataDecl(StringBuilder builder, IUhtPropertyMemberContext context, string name, string nameSuffix, int tabs)
+		{
+			builder.AppendMetaDataDecl(ValueProperty, context, name, GetNameSuffix(nameSuffix, "_ValueProp"), tabs);
+			builder.AppendMetaDataDecl(KeyProperty, context, name, GetNameSuffix(nameSuffix, "_Key_KeyProp"), tabs);
+			return base.AppendMetaDataDecl(builder, context, name, nameSuffix, tabs);
+		}
+
 		/// <inheritdoc/>
 		public override StringBuilder AppendMemberDecl(StringBuilder builder, IUhtPropertyMemberContext context, string name, string nameSuffix, int tabs)
 		{
diff --git a/Engine/Source/Programs/Shared/EpicGames.UHT/Types/Properties/UhtOptionalProperty.cs b/Engine/Source/Programs/Shared/EpicGames.UHT/Types/Properties/UhtOptionalProperty.cs
--- a/Engine/Source/Programs/Shared/EpicGames.UHT/Types/Properties/UhtOptionalProperty.cs
+++ b/Engine/Source/Programs/Shared/EpicGames.UHT/Types/Properties/UhtOptionalProperty.cs
@@ -133,6 +133,13 @@ namespace EpicGames.UHT.Types
 			return builder;
 		}
 
+		/// <inheritdoc/>
+		public override StringBuilder AppendMetaDataDecl(StringBuilder builder, IUhtPropertyMemberContext context, string name, string nameSuffix, int tabs)
+		{
+			builder.AppendMetaDataDecl(ValueProperty, context, name, GetNameSuffix(nameSuffix, "_Inner"), tabs);
+			return base.AppendMetaDataDecl(builder, context, name, nameSuffix, tabs);
+		}
+
 		/// <inheritdoc/>
 		public override StringBuilder AppendMemberDecl(StringBuilder builder, IUhtPropertyMemberContext context, string name, string nameSuffix, int tabs)
 		{
diff --git a/Engine/Source/Programs/Shared/EpicGames.UHT/Types/Properties/UhtSetProperty.cs b/Engine/Source/Programs/Shared/EpicGames.UHT/Types/Properties/UhtSetProperty.cs
--- a/Engine/Source/Programs/Shared/EpicGames.UHT/Types/Properties/UhtSetProperty.cs
+++ b/Engine/Source/Programs/Shared/EpicGames.UHT/Types/Properties/UhtSetProperty.cs
@@ -126,6 +126,13 @@ namespace EpicGames.UHT.Types
 			return builder;
 		}
 
+		/// <inheritdoc/>
+		public override StringBuilder AppendMetaDataDecl(StringBuilder builder, IUhtPropertyMemberContext context, string name, string nameSuffix, int tabs)
+		{
+			builder.AppendMetaDataDecl(ValueProperty, context, name, GetNameSuffix(nameSuffix, "_ElementProp"), tabs);
+			return base.AppendMetaDataDecl(builder, context, name, nameSuffix, tabs);
+		}
+
 		/// <inheritdoc/>
 		public override StringBuilder AppendMemberDecl(StringBuilder builder, IUhtPropertyMemberContext context, string name, string nameSuffix, int tabs)
 		{
@@ -137,7 +144,6 @@ namespace EpicGames.UHT.Types
 		public override StringBuilder AppendMemberDef(StringBuilder builder, IUhtPropertyMemberContext context, string name, string nameSuffix, string? offset, int tabs)
 		{
 			builder.AppendMemberDef(ValueProperty, context, name, GetNameSuffix(nameSuffix, "_ElementProp"), "0", tabs);
-			builder.AppendMetaDataDef(this, context, name, nameSuffix, tabs);
 
 			if (ValueProperty is UhtStructProperty structProperty)
 			{
@@ -150,7 +156,7 @@ namespace EpicGames.UHT.Types
 					.Append("' is used in a TSet but does not have a GetValueTypeHash defined\");\r\n");
 			}
 
-			AppendMemberDefStart(builder, context, name, nameSuffix, offset, tabs, "FSetPropertyParams", "UECodeGen_Private::EPropertyGenFlags::Set", false);
+			AppendMemberDefStart(builder, context, name, nameSuffix, offset, tabs, "FSetPropertyParams", "UECodeGen_Private::EPropertyGenFlags::Set");
 			AppendMemberDefEnd(builder, context, name, nameSuffix);
 			return builder;
 		}
diff --git a/Engine/Source/Programs/Shared/EpicGames.UHT/Types/UhtProperty.cs b/Engine/Source/Programs/Shared/EpicGames.UHT/Types/UhtProperty.cs
--- a/Engine/Source/Programs/Shared/EpicGames.UHT/Types/UhtProperty.cs
+++ b/Engine/Source/Programs/Shared/EpicGames.UHT/Types/UhtProperty.cs
@@ -1164,6 +1164,20 @@ namespace EpicGames.UHT.Types
 			return builder;
 		}
 
+		/// <summary>
+		/// Append the required code to declare the properties meta data
+		/// </summary>
+		/// <param name="builder">Output builder</param>
+		/// <param name="context">Current context</param>
+		/// <param name="name">Name of the property.  This is needed in some cases where the name in the declarations doesn't match the property name.</param>
+		/// <param name="nameSuffix">Suffix to the property name</param>
+		/// <param name="tabs">Number of tabs prefix the line with</param>
+		/// <returns>Output builder</returns>
+		public virtual StringBuilder AppendMetaDataDecl(StringBuilder builder, IUhtPropertyMemberContext context, string name, string nameSuffix, int tabs)
+		{
+			return builder.AppendMetaDataDecl(this, context.NamePrefix, name, nameSuffix, context.MetaDataSuffix, tabs);
+		}
+
 		/// <summary>
 		/// Append the required code to declare the property as a member
 		/// </summary>
@@ -1184,14 +1198,9 @@ namespace EpicGames.UHT.Types
 		/// <param name="nameSuffix">Suffix to the property name</param>
 		/// <param name="tabs">Number of tabs prefix the line with</param>
 		/// <param name="paramsStructName">Structure name</param>
-		/// <param name="appendMetaDataDecl">If true, add the meta data decl prior to the member decl</param>
 		/// <returns>Output builder</returns>
-		public StringBuilder AppendMemberDecl(StringBuilder builder, IUhtPropertyMemberContext context, string name, string nameSuffix, int tabs, string paramsStructName, bool appendMetaDataDecl = true)
+		public static StringBuilder AppendMemberDecl(StringBuilder builder, IUhtPropertyMemberContext context, string name, string nameSuffix, int tabs, string paramsStructName)
 		{
-			if (appendMetaDataDecl)
-			{
-				builder.AppendMetaDataDecl(this, context, name, nameSuffix, tabs);
-			}
 			builder.AppendTabs(tabs).Append("static const UECodeGen_Private::").Append(paramsStructName).Append(' ').AppendNameDecl(context, name, nameSuffix).Append(";\r\n");
 			return builder;
 		}
@@ -1219,16 +1228,11 @@ namespace EpicGames.UHT.Types
 		/// <param name="tabs">Number of tabs prefix the line with</param>
 		/// <param name="paramsStructName">Structure name</param>
 		/// <param name="paramsGenFlags">Structure flags</param>
-		/// <param name="appendMetaDataDef">If true, add the meta data def prior to the member def</param>
 		/// <param name="appendOffset">If true, add the offset parameter</param>
 		/// <returns>Output builder</returns>
 		public StringBuilder AppendMemberDefStart(StringBuilder builder, IUhtPropertyMemberContext context, string name, string nameSuffix, string? offset, int tabs,
-			string paramsStructName, string paramsGenFlags, bool appendMetaDataDef = true, bool appendOffset = true)
+			string paramsStructName, string paramsGenFlags, bool appendOffset = true)
 		{
-			if (appendMetaDataDef)
-			{
-				builder.AppendMetaDataDef(this, context, name, nameSuffix, tabs);
-			}
 			builder
 				.AppendTabs(tabs)
 				.Append("const UECodeGen_Private::").Append(paramsStructName).Append(' ')
diff --git a/Engine/Source/Programs/Shared/EpicGames.UHT/Utils/UhtStringBuilderExtensions.cs b/Engine/Source/Programs/Shared/EpicGames.UHT/Utils/UhtStringBuilderExtensions.cs
--- a/Engine/Source/Programs/Shared/EpicGames.UHT/Utils/UhtStringBuilderExtensions.cs
+++ b/Engine/Source/Programs/Shared/EpicGames.UHT/Utils/UhtStringBuilderExtensions.cs
@@ -176,7 +176,7 @@ namespace EpicGames.UHT.Utils
 		/// <returns>Destination builder</returns>
 		public static StringBuilder AppendMetaDataParams(this StringBuilder builder, UhtProperty property, IUhtPropertyMemberContext context, string name, string nameSuffix)
 		{
-			return builder.AppendMetaDataParams(property, context.StaticsName, context.NamePrefix, name, nameSuffix, context.MetaDataSuffix);
+			return builder.AppendMetaDataParams(property, null, context.NamePrefix, name, nameSuffix, context.MetaDataSuffix);
 		}
 
 		/// <summary>
@@ -190,67 +190,13 @@ namespace EpicGames.UHT.Utils
 		/// <param name="metaNameSuffix">Optional meta data name suffix</param>
 		/// <param name="tabs">Number of tabs to indent</param>
 		/// <returns>Destination builder</returns>
-		private static StringBuilder AppendMetaDataDecl(this StringBuilder builder, UhtType type, string? namePrefix, string name, string? nameSuffix, string? metaNameSuffix, int tabs)
+		public static StringBuilder AppendMetaDataDecl(this StringBuilder builder, UhtType type, string? namePrefix, string name, string? nameSuffix, string? metaNameSuffix, int tabs)
 		{
 			if (!type.MetaData.IsEmpty())
 			{
-				builder.Append("#if WITH_METADATA\r\n");
-				builder.AppendTabs(tabs).Append("static const UECodeGen_Private::FMetaDataPairParam ").AppendNameDecl(namePrefix, name, nameSuffix).Append(metaNameSuffix).Append("[];\r\n");
-				builder.Append("#endif\r\n");
-			}
-			return builder;
-		}
-
-		/// <summary>
-		/// Append the meta data declaration
-		/// </summary>
-		/// <param name="builder">Destination builder</param>
-		/// <param name="type">Source type containing the meta data</param>
-		/// <param name="name">Name</param>
-		/// <param name="tabs">Number of tabs to indent</param>
-		/// <returns>Destination builder</returns>
-		public static StringBuilder AppendMetaDataDecl(this StringBuilder builder, UhtType type, string name, int tabs)
-		{
-			return builder.AppendMetaDataDecl(type, null, name, null, null, tabs);
-		}
-
-		/// <summary>
-		/// Append the meta data declaration
-		/// </summary>
-		/// <param name="builder">Destination builder</param>
-		/// <param name="property">Source type containing the meta data</param>
-		/// <param name="context">Property context used to get the statics name and name prefix</param>
-		/// <param name="name">Name</param>
-		/// <param name="nameSuffix">Optional name suffix</param>
-		/// <param name="tabs">Number of tabs to indent</param>
-		/// <returns>Destination builder</returns>
-		public static StringBuilder AppendMetaDataDecl(this StringBuilder builder, UhtProperty property, IUhtPropertyMemberContext context, string name, string nameSuffix, int tabs)
-		{
-			return builder.AppendMetaDataDecl(property, context.NamePrefix, name, nameSuffix, context.MetaDataSuffix, tabs);
-		}
-
-		/// <summary>
-		/// Append the meta data definition.
-		/// </summary>
-		/// <param name="builder">Destination builder</param>
-		/// <param name="type">Source type containing the meta data</param>
-		/// <param name="staticsName">Optional name of the statics block which will be output in the form of "StaticsName::"</param>
-		/// <param name="namePrefix">Name prefix</param>
-		/// <param name="name">Name</param>
-		/// <param name="nameSuffix">Optional name suffix</param>
-		/// <param name="metaNameSuffix">Optional meta data name suffix</param>
-		/// <param name="tabs">Number of tabs to indent</param>
-		/// <returns>Destination builder</returns>
-		private static StringBuilder AppendMetaDataDef(this StringBuilder builder, UhtType type, string? staticsName, string? namePrefix, string name, string? nameSuffix, string? metaNameSuffix, int tabs)
-		{
-			bool isPartOfEngine = type.Package.IsPartOfEngine;
-			if (!type.MetaData.IsEmpty())
-			{
+				bool isPartOfEngine = type.Package.IsPartOfEngine;
 				List<KeyValuePair<string, string>> sortedMetaData = type.MetaData.GetSorted();
-
-				builder.Append("#if WITH_METADATA\r\n");
-				builder.AppendTabs(tabs).Append("const UECodeGen_Private::FMetaDataPairParam ").AppendNameDef(staticsName, namePrefix, name, nameSuffix).Append(metaNameSuffix).Append("[] = {\r\n");
-
+				builder.AppendTabs(tabs).Append("static constexpr UECodeGen_Private::FMetaDataPairParam ").AppendNameDecl(namePrefix, name, nameSuffix).Append(metaNameSuffix).Append("[] = {\r\n");
 				foreach (KeyValuePair<string, string> kvp in sortedMetaData)
 				{
 					bool restricted = !isPartOfEngine && s_hiddenMetaDataNames.Contains(kvp.Key);
@@ -266,27 +212,12 @@ namespace EpicGames.UHT.Utils
 				}
 
 				builder.AppendTabs(tabs).Append("};\r\n");
-				builder.Append("#endif\r\n");
 			}
 			return builder;
 		}
 
 		/// <summary>
-		/// Append the meta data definition
-		/// </summary>
-		/// <param name="builder">Destination builder</param>
-		/// <param name="type">Source type containing the meta data</param>
-		/// <param name="staticsName">Optional name of the statics block which will be output in the form of "StaticsName::"</param>
-		/// <param name="name">Name</param>
-		/// <param name="tabs">Number of tabs to indent</param>
-		/// <returns>Destination builder</returns>
-		public static StringBuilder AppendMetaDataDef(this StringBuilder builder, UhtType type, string? staticsName, string name, int tabs)
-		{
-			return builder.AppendMetaDataDef(type, staticsName, null, name, null, null, tabs);
-		}
-
-		/// <summary>
-		/// Append the meta data definition
+		/// Append the meta data declaration
 		/// </summary>
 		/// <param name="builder">Destination builder</param>
 		/// <param name="property">Source type containing the meta data</param>
@@ -295,9 +226,9 @@ namespace EpicGames.UHT.Utils
 		/// <param name="nameSuffix">Optional name suffix</param>
 		/// <param name="tabs">Number of tabs to indent</param>
 		/// <returns>Destination builder</returns>
-		public static StringBuilder AppendMetaDataDef(this StringBuilder builder, UhtProperty property, IUhtPropertyMemberContext context, string name, string nameSuffix, int tabs)
+		public static StringBuilder AppendMetaDataDecl(this StringBuilder builder, UhtProperty property, IUhtPropertyMemberContext context, string name, string nameSuffix, int tabs)
 		{
-			return builder.AppendMetaDataDef(property, context.StaticsName, context.NamePrefix, name, nameSuffix, context.MetaDataSuffix, tabs);
+			return property.AppendMetaDataDecl(builder, context, name, nameSuffix, tabs);
 		}
 
 		/// <summary>
