Generated code for enums now use constexpr excluding the registration object.

diff --git a/Engine/Source/Programs/Shared/EpicGames.UHT/Exporters/CodeGen/UhtHeaderCodeGeneratorCppFile.cs b/Engine/Source/Programs/Shared/EpicGames.UHT/Exporters/CodeGen/UhtHeaderCodeGeneratorCppFile.cs
--- a/Engine/Source/Programs/Shared/EpicGames.UHT/Exporters/CodeGen/UhtHeaderCodeGeneratorCppFile.cs
+++ b/Engine/Source/Programs/Shared/EpicGames.UHT/Exporters/CodeGen/UhtHeaderCodeGeneratorCppFile.cs
@@ -445,40 +445,43 @@ namespace EpicGames.UHT.Exporters.CodeGen
 				int hashCodeBlockStart = builder.Length;
 
 				// Statics declaration
-				builder.Append("struct ").Append(staticsName).Append("\r\n");
-				builder.Append("{\r\n");
-				builder.Append("\tstatic const UECodeGen_Private::FEnumeratorParam Enumerators[];\r\n");
-				builder.AppendMetaDataDecl(enumObj, null, null, MetaDataParamsName, 1);
-				builder.Append("\tstatic const UECodeGen_Private::FEnumParams EnumParams;\r\n");
-				builder.Append("};\r\n");
-
-				// Enumerators
-				builder.Append("const UECodeGen_Private::FEnumeratorParam ").Append(staticsName).Append("::Enumerators[] = {\r\n");
-				int enumIndex = 0;
-				foreach (UhtEnumValue value in enumObj.EnumValues)
 				{
-					if (!enumObj.MetaData.TryGetValue("OverrideName", enumIndex, out string? keyName))
+					builder.Append("struct ").Append(staticsName).Append("\r\n");
+					builder.Append("{\r\n");
+					builder.AppendMetaDataDecl(enumObj, null, null, MetaDataParamsName, 1);
+
+					// Enumerators
+					builder.Append("\tstatic constexpr UECodeGen_Private::FEnumeratorParam Enumerators[] = {\r\n");
+					int enumIndex = 0;
+					foreach (UhtEnumValue value in enumObj.EnumValues)
 					{
-						keyName = value.Name.ToString();
+						if (!enumObj.MetaData.TryGetValue("OverrideName", enumIndex, out string? keyName))
+						{
+							keyName = value.Name.ToString();
+						}
+						builder.Append("\t\t{ ").AppendUTF8LiteralString(keyName).Append(", (int64)").Append(value.Name).Append(" },\r\n");
+						++enumIndex;
 					}
-					builder.Append("\t{ ").AppendUTF8LiteralString(keyName).Append(", (int64)").Append(value.Name).Append(" },\r\n");
-					++enumIndex;
+					builder.Append("\t};\r\n");
+					builder.Append("\tstatic const UECodeGen_Private::FEnumParams EnumParams;\r\n");
+					builder.Append("};\r\n");
 				}
-				builder.Append("};\r\n");
 
-				// Singleton parameters
-				builder.Append("const UECodeGen_Private::FEnumParams ").Append(staticsName).Append("::EnumParams = {\r\n");
-				builder.Append("\t(UObject*(*)())").Append(PackageSingletonName).Append(",\r\n");
-				builder.Append('\t').Append(enumDisplayNameFn).Append(",\r\n");
-				builder.Append('\t').AppendUTF8LiteralString(enumObj.SourceName).Append(",\r\n");
-				builder.Append('\t').AppendUTF8LiteralString(enumObj.CppType).Append(",\r\n");
-				builder.Append('\t').Append(staticsName).Append("::Enumerators,\r\n");
-				builder.Append('\t').Append(ObjectFlags).Append(",\r\n");
-				builder.Append("\tUE_ARRAY_COUNT(").Append(staticsName).Append("::Enumerators),\r\n");
-				builder.Append('\t').Append(enumObj.EnumFlags.HasAnyFlags(EEnumFlags.Flags) ? "EEnumFlags::Flags" : "EEnumFlags::None").Append(",\r\n");
-				builder.Append("\t(uint8)UEnum::ECppForm::").Append(enumObj.CppForm.ToString()).Append(",\r\n");
-				builder.Append('\t').AppendMetaDataParams(enumObj, staticsName, MetaDataParamsName).Append("\r\n");
-				builder.Append("};\r\n");
+				// Statics definition
+				{
+					builder.Append("const UECodeGen_Private::FEnumParams ").Append(staticsName).Append("::EnumParams = {\r\n");
+					builder.Append("\t(UObject*(*)())").Append(PackageSingletonName).Append(",\r\n");
+					builder.Append('\t').Append(enumDisplayNameFn).Append(",\r\n");
+					builder.Append('\t').AppendUTF8LiteralString(enumObj.SourceName).Append(",\r\n");
+					builder.Append('\t').AppendUTF8LiteralString(enumObj.CppType).Append(",\r\n");
+					builder.Append('\t').Append(staticsName).Append("::Enumerators,\r\n");
+					builder.Append('\t').Append(ObjectFlags).Append(",\r\n");
+					builder.Append("\tUE_ARRAY_COUNT(").Append(staticsName).Append("::Enumerators),\r\n");
+					builder.Append('\t').Append(enumObj.EnumFlags.HasAnyFlags(EEnumFlags.Flags) ? "EEnumFlags::Flags" : "EEnumFlags::None").Append(",\r\n");
+					builder.Append("\t(uint8)UEnum::ECppForm::").Append(enumObj.CppForm.ToString()).Append(",\r\n");
+					builder.Append('\t').AppendMetaDataParams(enumObj, staticsName, MetaDataParamsName).Append("\r\n");
+					builder.Append("};\r\n");
+				}
 
 				// Registration singleton
 				builder.Append("UEnum* ").Append(singletonName).Append("()\r\n");
