Horde: Add support for declaring resource needs in compute client

diff --git a/Engine/Source/Programs/Shared/EpicGames.Horde/Compute/Clients/AgentComputeClient.cs b/Engine/Source/Programs/Shared/EpicGames.Horde/Compute/Clients/AgentComputeClient.cs
--- a/Engine/Source/Programs/Shared/EpicGames.Horde/Compute/Clients/AgentComputeClient.cs
+++ b/Engine/Source/Programs/Shared/EpicGames.Horde/Compute/Clients/AgentComputeClient.cs
@@ -82,6 +82,12 @@ namespace EpicGames.Horde.Compute.Clients
 			return new LeaseImpl(source);
 		}
 
+		/// <inheritdoc/>
+		public Task DeclareResourceNeedsAsync(ClusterId clusterId, string pool, Dictionary<string, int> resourceNeeds, CancellationToken cancellationToken = default)
+		{
+			return Task.CompletedTask;
+		}
+
 		async IAsyncEnumerable<RemoteComputeSocket> ConnectAsync(ILogger logger, [EnumeratorCancellation] CancellationToken cancellationToken)
 		{
 			using Socket listener = new Socket(SocketType.Stream, ProtocolType.IP);
diff --git a/Engine/Source/Programs/Shared/EpicGames.Horde/Compute/Clients/LocalComputeClient.cs b/Engine/Source/Programs/Shared/EpicGames.Horde/Compute/Clients/LocalComputeClient.cs
--- a/Engine/Source/Programs/Shared/EpicGames.Horde/Compute/Clients/LocalComputeClient.cs
+++ b/Engine/Source/Programs/Shared/EpicGames.Horde/Compute/Clients/LocalComputeClient.cs
@@ -93,5 +93,11 @@ namespace EpicGames.Horde.Compute.Clients
 			return Task.FromResult<IComputeLease?>(new LeaseImpl(socket));
 #pragma warning restore CA2000 // Dispose objects before losing scope
 		}
+
+		/// <inheritdoc/>
+		public Task DeclareResourceNeedsAsync(ClusterId clusterId, string pool, Dictionary<string, int> resourceNeeds, CancellationToken cancellationToken = default)
+		{
+			return Task.CompletedTask;
+		}
 	}
 }
diff --git a/Engine/Source/Programs/Shared/EpicGames.Horde/Compute/Clients/ServerComputeClient.cs b/Engine/Source/Programs/Shared/EpicGames.Horde/Compute/Clients/ServerComputeClient.cs
--- a/Engine/Source/Programs/Shared/EpicGames.Horde/Compute/Clients/ServerComputeClient.cs
+++ b/Engine/Source/Programs/Shared/EpicGames.Horde/Compute/Clients/ServerComputeClient.cs
@@ -81,6 +81,7 @@ namespace EpicGames.Horde.Compute.Clients
 
 		readonly IHttpClientFactory _httpClientFactory;
 		readonly CancellationTokenSource _cancellationSource = new CancellationTokenSource();
+		readonly string _sessionId;
 		readonly ILogger _logger;
 
 		/// <summary>
@@ -88,9 +89,20 @@ namespace EpicGames.Horde.Compute.Clients
 		/// </summary>
 		/// <param name="httpClientFactory">Factory for constructing http client instances</param>
 		/// <param name="logger">Logger for diagnostic messages</param>
-		public ServerComputeClient(IHttpClientFactory httpClientFactory, ILogger logger)
+		public ServerComputeClient(IHttpClientFactory httpClientFactory, ILogger logger) : this(httpClientFactory, null, logger)
+		{
+		}
+		
+		/// <summary>
+		/// Constructor
+		/// </summary>
+		/// <param name="httpClientFactory">Factory for constructing http client instances</param>
+		/// <param name="sessionId">Arbitrary ID used for identifying this compute client. If not provided, a random one will be generated</param>
+		/// <param name="logger">Logger for diagnostic messages</param>
+		public ServerComputeClient(IHttpClientFactory httpClientFactory, string? sessionId, ILogger logger)
 		{
 			_httpClientFactory = httpClientFactory;
+			_sessionId = sessionId ?? Guid.NewGuid().ToString();
 			_logger = logger;
 		}
 
@@ -120,6 +132,14 @@ namespace EpicGames.Horde.Compute.Clients
 		}
 
 		/// <inheritdoc/>
+		public async Task DeclareResourceNeedsAsync(ClusterId clusterId, string pool, Dictionary<string, int> resourceNeeds, CancellationToken cancellationToken = default)
+		{
+			HttpClient client = _httpClientFactory.CreateClient(HordeHttpClient.HttpClientName);
+			ResourceNeedsMessage request = new () { SessionId = _sessionId, Pool = pool, ResourceNeeds = resourceNeeds };
+			using HttpResponseMessage response = await HordeHttpClient.PostAsync(client, $"api/v2/compute/{clusterId}/resource-needs", request, _cancellationSource.Token);
+			response.EnsureSuccessStatusCode();
+		}
+
 		async IAsyncEnumerable<LeaseInfo> ConnectAsync(ClusterId clusterId, Requirements? requirements, string? requestId, ILogger workerLogger, [EnumeratorCancellation] CancellationToken cancellationToken)
 		{
 			_logger.LogDebug("Requesting compute resource");
diff --git a/Engine/Source/Programs/Shared/EpicGames.Horde/Compute/IComputeClient.cs b/Engine/Source/Programs/Shared/EpicGames.Horde/Compute/IComputeClient.cs
--- a/Engine/Source/Programs/Shared/EpicGames.Horde/Compute/IComputeClient.cs
+++ b/Engine/Source/Programs/Shared/EpicGames.Horde/Compute/IComputeClient.cs
@@ -1,6 +1,7 @@
 // Copyright Epic Games, Inc. All Rights Reserved.
 
 using System;
+using System.Collections.Generic;
 using System.Threading;
 using System.Threading.Tasks;
 using Microsoft.Extensions.Logging;
@@ -21,6 +22,17 @@ namespace EpicGames.Horde.Compute
 		/// <param name="logger">Logger for output from this worker</param>
 		/// <param name="cancellationToken">Cancellation token for the operation</param>
 		public Task<IComputeLease?> TryAssignWorkerAsync(ClusterId clusterId, Requirements? requirements, string? requestId, ILogger logger, CancellationToken cancellationToken = default);
+		
+		/// <summary>
+		/// Declare resource needs for current client
+		/// Helps inform the server about current demand.
+		/// Can be called as often as necessary to keep needs up-to-date.
+		/// </summary>
+		/// <param name="clusterId">Cluster to execute the request</param>
+		/// <param name="pool">Which pool this applies to</param>
+		/// <param name="resourceNeeds">Properties with a target amount of each, such as CPU or RAM</param>
+		/// <param name="cancellationToken">Cancellation token for the operation</param>
+		public Task DeclareResourceNeedsAsync(ClusterId clusterId, string pool, Dictionary<string, int> resourceNeeds, CancellationToken cancellationToken = default);
 	}
 
 	/// <summary>
