Horde: Allow '+' characters in ref names, and include more detail in error messages for invalid refs.

diff --git a/Engine/Source/Programs/Shared/EpicGames.Horde/Storage/RefName.cs b/Engine/Source/Programs/Shared/EpicGames.Horde/Storage/RefName.cs
--- a/Engine/Source/Programs/Shared/EpicGames.Horde/Storage/RefName.cs
+++ b/Engine/Source/Programs/Shared/EpicGames.Horde/Storage/RefName.cs
@@ -5,6 +5,7 @@ using EpicGames.Serialization;
 using System;
 using System.ComponentModel;
 using System.Globalization;
+using System.Text;
 using System.Text.Json;
 using System.Text.Json.Serialization;
 
@@ -36,23 +37,23 @@ namespace EpicGames.Horde.Storage
 		public RefName(Utf8String text)
 		{
 			Text = text;
-			ValidatePathArgument(nameof(text), text.Span);
+			ValidatePathArgument(text.Span, nameof(text));
 		}
 
 		/// <summary>
 		/// Validates a given string as a blob id
 		/// </summary>
-		/// <param name="name">Name of the argument</param>
 		/// <param name="text">String to validate</param>
-		public static void ValidatePathArgument(string name, ReadOnlySpan<byte> text)
+		/// <param name="argumentName">Name of the argument</param>
+		public static void ValidatePathArgument(ReadOnlySpan<byte> text, string argumentName)
 		{
 			if (text.Length == 0)
 			{
-				throw new ArgumentException("Ref names cannot be empty", name);
+				throw new ArgumentException("Ref names cannot be empty", argumentName);
 			}
 			if (text[^1] == '/')
 			{
-				throw new ArgumentException("Ref names cannot start or end with a slash", name);
+				throw new ArgumentException($"{Encoding.UTF8.GetString(text)} is not a valid ref name (cannot start or end with a slash)", argumentName);
 			}
 
 			int lastSlashIdx = -1;
@@ -63,7 +64,7 @@ namespace EpicGames.Horde.Storage
 				{
 					if (lastSlashIdx == idx - 1)
 					{
-						throw new ArgumentException("Leading and consecutive slashes are not permitted in ref names", name);
+						throw new ArgumentException($"{Encoding.UTF8.GetString(text)} is not a valid ref name (leading and consecutive slashes are not permitted)", argumentName);
 					}
 					else
 					{
@@ -74,7 +75,7 @@ namespace EpicGames.Horde.Storage
 				{
 					if (!IsValidChar(text[idx]))
 					{
-						throw new ArgumentException($"'{(char)text[idx]} is not a valid ref name character", name);
+						throw new ArgumentException($"{Encoding.UTF8.GetString(text)} is not a valid ref name ('{(char)text[idx]}' is an invalid character)", argumentName);
 					}
 				}
 			}
@@ -84,7 +85,7 @@ namespace EpicGames.Horde.Storage
 
 		static uint[] CreateValidCharsArray()
 		{
-			const string ValidChars = "0123456789abcdefghijklmnopqrstuvwxyz_/-";
+			const string ValidChars = "0123456789abcdefghijklmnopqrstuvwxyz_/-+";
 
 			uint[] validChars = new uint[256 / 8];
 			for (int idx = 0; idx < ValidChars.Length; idx++)
