Generated code for structures now use constexpr excluding registration object.

diff --git a/Engine/Source/Programs/Shared/EpicGames.UHT/Exporters/CodeGen/UhtHeaderCodeGeneratorCppFile.cs b/Engine/Source/Programs/Shared/EpicGames.UHT/Exporters/CodeGen/UhtHeaderCodeGeneratorCppFile.cs
--- a/Engine/Source/Programs/Shared/EpicGames.UHT/Exporters/CodeGen/UhtHeaderCodeGeneratorCppFile.cs
+++ b/Engine/Source/Programs/Shared/EpicGames.UHT/Exporters/CodeGen/UhtHeaderCodeGeneratorCppFile.cs
@@ -656,28 +656,23 @@ namespace EpicGames.UHT.Exporters.CodeGen
 				// Meta data
 				builder.AppendMetaDataDecl(scriptStruct, context, properties, MetaDataParamsName, 1);
 
+				// Properties
+				AppendPropertiesDecl(builder, context, properties, 1);
+
 				// New struct ops
 				if (scriptStruct.ScriptStructFlags.HasAnyFlags(EStructFlags.Native))
 				{
-					builder.Append("\tstatic void* NewStructOps();\r\n");
+					builder.Append("\tstatic void* NewStructOps()\r\n");
+					builder.Append("\t{\r\n");
+					builder.Append("\t\treturn (UScriptStruct::ICppStructOps*)new UScriptStruct::TCppStructOps<").Append(scriptStruct.SourceName).Append(">();\r\n");
+					builder.Append("\t}\r\n");
 				}
-
-				AppendPropertiesDecl(builder, context, properties, 1);
-
 				builder.Append("\tstatic const UECodeGen_Private::FStructParams StructParams;\r\n");
 				builder.Append("};\r\n");
 			}
 
 			// Populate the elements of the static structure
 			{
-				if (scriptStruct.ScriptStructFlags.HasAnyFlags(EStructFlags.Native))
-				{
-					builder.Append("void* ").Append(staticsName).Append("::NewStructOps()\r\n");
-					builder.Append("{\r\n");
-					builder.Append("\treturn (UScriptStruct::ICppStructOps*)new UScriptStruct::TCppStructOps<").Append(scriptStruct.SourceName).Append(">();\r\n");
-					builder.Append("}\r\n");
-				}
-
 				AppendPropertiesDefs(builder, context, properties, 0);
 
 				builder.Append("const UECodeGen_Private::FStructParams ").Append(staticsName).Append("::StructParams = {\r\n");
@@ -836,6 +831,7 @@ namespace EpicGames.UHT.Exporters.CodeGen
 			AppendNativeFunctionHeader(builder, function, UhtPropertyTextType.EventFunctionArgOrRetVal, false, exportFunctionName, extraParameter, UhtFunctionExportFlags.None, 0, "\r\n");
 			AppendEventFunctionPrologue(builder, function, strippedFunctionName, tabs, "\r\n", true);
 			builder
+				.Append('\t')
 				.Append(strippedFunctionName)
 				.Append('.')
 				.Append(function.FunctionFlags.HasAnyFlags(EFunctionFlags.MulticastDelegate) ? "ProcessMulticastDelegate" : "ProcessDelegate")
