Horde: Allow setting aliases on blobs at the moment that we write them.

diff --git a/Engine/Source/Programs/Horde/Horde.Server/Ddc/BlobService.cs b/Engine/Source/Programs/Horde/Horde.Server/Ddc/BlobService.cs
--- a/Engine/Source/Programs/Horde/Horde.Server/Ddc/BlobService.cs
+++ b/Engine/Source/Programs/Horde/Horde.Server/Ddc/BlobService.cs
@@ -152,7 +152,7 @@ namespace Horde.Server.Ddc
 				using Stream stream = content.GetStream();
 				await stream.ReadAsync(memory, cancellationToken);
 
-				blobHandle = await writer.WriteNodeAsync(memory.Length, Array.Empty<BlobHandle>(), s_rawBlobType, cancellationToken);
+				blobHandle = await writer.WriteBlobAsync(memory.Length, Array.Empty<BlobHandle>(), s_rawBlobType, cancellationToken);
 				await writer.FlushAsync(cancellationToken);
 			}
 
diff --git a/Engine/Source/Programs/Shared/EpicGames.Horde/Storage/Bundles/BundleWriter.cs b/Engine/Source/Programs/Shared/EpicGames.Horde/Storage/Bundles/BundleWriter.cs
--- a/Engine/Source/Programs/Shared/EpicGames.Horde/Storage/Bundles/BundleWriter.cs
+++ b/Engine/Source/Programs/Shared/EpicGames.Horde/Storage/Bundles/BundleWriter.cs
@@ -166,10 +166,11 @@ namespace EpicGames.Horde.Storage.Bundles
 			public readonly int Offset;
 			public readonly int Length;
 			public readonly BundleNodeHandle[] Refs;
+			public readonly AliasInfo[] Aliases;
 
 			public PendingBundle? PendingBundle => _pendingBundle;
 
-			public PendingNode(BundleReader reader, NodeKey key, int packet, int offset, int length, IReadOnlyList<BundleNodeHandle> refs, PendingBundle pendingBundle)
+			public PendingNode(BundleReader reader, NodeKey key, int packet, int offset, int length, IReadOnlyList<BundleNodeHandle> refs, IReadOnlyList<AliasInfo> aliases, PendingBundle pendingBundle)
 				: base(key.Hash)
 			{
 				_reader = reader;
@@ -179,6 +180,7 @@ namespace EpicGames.Horde.Storage.Bundles
 				Offset = offset;
 				Length = length;
 				Refs = refs.ToArray();
+				Aliases = aliases.ToArray();
 
 				_pendingBundle = pendingBundle;
 			}
@@ -428,14 +430,14 @@ namespace EpicGames.Horde.Storage.Bundles
 			}
 
 			// Finish a node write
-			public PendingNode WriteNode(NodeKey nodeKey, int size, IReadOnlyList<BlobHandle> refs)
+			public PendingNode WriteNode(NodeKey nodeKey, int size, IReadOnlyList<BlobHandle> refs, IReadOnlyList<AliasInfo> aliases)
 			{
 				int offset = _currentPacketLength;
 
 				_currentPacketLength += size;
 				UncompressedLength += size;
 
-				PendingNode pendingNode = new PendingNode(_treeReader, nodeKey, _currentPacketIdx, offset, (int)size, refs.ConvertAll(x => (BundleNodeHandle)x), this);
+				PendingNode pendingNode = new PendingNode(_treeReader, nodeKey, _currentPacketIdx, offset, (int)size, refs.ConvertAll(x => (BundleNodeHandle)x), aliases, this);
 				_queue.Add(pendingNode);
 				_queuedRefs += refs.Count;
 				_nodeKeyToInfo.Add(nodeKey, pendingNode);
@@ -547,9 +549,16 @@ namespace EpicGames.Horde.Storage.Bundles
 
 					for (int idx = 0; idx < _queue.Count; idx++)
 					{
+						PendingNode node = _queue[idx];
+
 						BundleNodeLocator nodeLocator = new BundleNodeLocator(_queue[idx].Key.Hash, locator, idx);
 						traceLogger?.LogInformation("Updated pending node {Hash} with locator {Locator}", _queue[idx].Key.Hash, nodeLocator);
 						_queue[idx].MarkAsWritten(nodeLocator);
+
+						foreach (AliasInfo alias in node.Aliases)
+						{
+							await store.AddAliasAsync(alias.Name, nodeLocator, alias.Rank, alias.Data);
+						}
 					}
 
 					BlobWriteCallback? callback = Interlocked.Exchange(ref _callbacks, s_callbackSentinel);
@@ -851,9 +860,10 @@ namespace EpicGames.Horde.Storage.Bundles
 		/// <param name="size">Used size of the buffer</param>
 		/// <param name="references">References to other nodes</param>
 		/// <param name="type">Type of the node that was written</param>
+		/// <param name="aliases">Aliases for the node</param>
 		/// <param name="cancellationToken">Cancellation token for the operation</param>
 		/// <returns>Handle to the written node</returns>
-		public async ValueTask<BlobHandle> WriteNodeAsync(int size, IReadOnlyList<BlobHandle> references, BlobType type, CancellationToken cancellationToken = default)
+		public async ValueTask<BlobHandle> WriteBlobAsync(int size, IReadOnlyList<BlobHandle> references, BlobType type, IReadOnlyList<AliasInfo> aliases, CancellationToken cancellationToken = default)
 		{
 			PendingBundle currentBundle = GetCurrentBundle();
 
@@ -872,7 +882,7 @@ namespace EpicGames.Horde.Storage.Bundles
 			}
 
 			// Append this node data
-			PendingNode pendingNode = currentBundle.WriteNode(nodeKey, size, references);
+			PendingNode pendingNode = currentBundle.WriteNode(nodeKey, size, references, aliases);
 			_nodeCache.Add(nodeKey, pendingNode);
 			TraceLogger?.LogInformation("Added new node for {NodeKey} in bundle {BundleId}", nodeKey, currentBundle.BundleId);
 
diff --git a/Engine/Source/Programs/Shared/EpicGames.Horde/Storage/IStorageClient.cs b/Engine/Source/Programs/Shared/EpicGames.Horde/Storage/IStorageClient.cs
--- a/Engine/Source/Programs/Shared/EpicGames.Horde/Storage/IStorageClient.cs
+++ b/Engine/Source/Programs/Shared/EpicGames.Horde/Storage/IStorageClient.cs
@@ -69,6 +69,14 @@ namespace EpicGames.Horde.Storage
 		public bool? Extend { get; set; }
 	}
 
+	/// <summary>
+	/// Information about an alias
+	/// </summary>
+	/// <param name="Name">Name of the alias</param>
+	/// <param name="Rank">Rank of the alias</param>
+	/// <param name="Data">Inline data to be stored for the alias</param>
+	public record class AliasInfo(Utf8String Name, int Rank, ReadOnlyMemory<byte> Data);
+
 	/// <summary>
 	/// Interface for the storage system.
 	/// </summary>
@@ -183,9 +191,10 @@ namespace EpicGames.Horde.Storage
 		/// <param name="size">Used size of the buffer</param>
 		/// <param name="references">References to other nodes</param>
 		/// <param name="type">Type of the node that was written</param>
+		/// <param name="aliases">Aliases for this node</param>
 		/// <param name="cancellationToken">Cancellation token for the operation</param>
 		/// <returns>Handle to the written node</returns>
-		ValueTask<BlobHandle> WriteNodeAsync(int size, IReadOnlyList<BlobHandle> references, BlobType type, CancellationToken cancellationToken = default);
+		ValueTask<BlobHandle> WriteBlobAsync(int size, IReadOnlyList<BlobHandle> references, BlobType type, IReadOnlyList<AliasInfo> aliases, CancellationToken cancellationToken = default);
 
 		/// <summary>
 		/// Writes the reference using the given target node
@@ -357,4 +366,24 @@ namespace EpicGames.Horde.Storage
 			return stats;
 		}
 	}
+
+	/// <summary>
+	/// Extension methods for <see cref="IStorageWriter"/>
+	/// </summary>
+	public static class StorageWriterExtensions
+	{
+		/// <summary>
+		/// Finish writing a node.
+		/// </summary>
+		/// <param name="writer">Writer instance to manipulate</param>
+		/// <param name="size">Used size of the buffer</param>
+		/// <param name="references">References to other nodes</param>
+		/// <param name="type">Type of the node that was written</param>
+		/// <param name="cancellationToken">Cancellation token for the operation</param>
+		/// <returns>Handle to the written node</returns>
+		public static ValueTask<BlobHandle> WriteBlobAsync(this IStorageWriter writer, int size, IReadOnlyList<BlobHandle> references, BlobType type, CancellationToken cancellationToken = default)
+		{
+			return writer.WriteBlobAsync(size, references, type, Array.Empty<AliasInfo>(), cancellationToken);
+		}
+	}
 }
diff --git a/Engine/Source/Programs/Shared/EpicGames.Horde/Storage/NodeRef.cs b/Engine/Source/Programs/Shared/EpicGames.Horde/Storage/NodeRef.cs
--- a/Engine/Source/Programs/Shared/EpicGames.Horde/Storage/NodeRef.cs
+++ b/Engine/Source/Programs/Shared/EpicGames.Horde/Storage/NodeRef.cs
@@ -197,7 +197,7 @@ namespace EpicGames.Horde.Storage
 			node.Serialize(nodeWriter);
 
 			// Write the final data
-			BlobHandle handle = await writer.WriteNodeAsync(nodeWriter.Length, nodeWriter.References, node.NodeType, cancellationToken);
+			BlobHandle handle = await writer.WriteBlobAsync(nodeWriter.Length, nodeWriter.References, node.NodeType, cancellationToken);
 			return new NodeRef<TNode>(handle);
 		}
 
diff --git a/Engine/Source/Programs/Shared/EpicGames.Horde/Storage/Nodes/ChunkedDataNode.cs b/Engine/Source/Programs/Shared/EpicGames.Horde/Storage/Nodes/ChunkedDataNode.cs
--- a/Engine/Source/Programs/Shared/EpicGames.Horde/Storage/Nodes/ChunkedDataNode.cs
+++ b/Engine/Source/Programs/Shared/EpicGames.Horde/Storage/Nodes/ChunkedDataNode.cs
@@ -232,7 +232,7 @@ namespace EpicGames.Horde.Storage.Nodes
 					sizeSinceProgressUpdate = 0;
 				}
 
-				BlobHandle handle = await writer.WriteNodeAsync(nextLength, Array.Empty<BlobHandle>(), GetNodeType<LeafChunkedDataNode>(), cancellationToken);
+				BlobHandle handle = await writer.WriteBlobAsync(nextLength, Array.Empty<BlobHandle>(), GetNodeType<LeafChunkedDataNode>(), cancellationToken);
 				handles.Add(new NodeRef<ChunkedDataNode>(handle));
 
 				readBuffer.Memory.Slice(nextLength, size - nextLength).CopyTo(readBuffer.Memory);
diff --git a/Engine/Source/Programs/Shared/EpicGames.Horde/Storage/Nodes/ChunkedDataWriter.cs b/Engine/Source/Programs/Shared/EpicGames.Horde/Storage/Nodes/ChunkedDataWriter.cs
--- a/Engine/Source/Programs/Shared/EpicGames.Horde/Storage/Nodes/ChunkedDataWriter.cs
+++ b/Engine/Source/Programs/Shared/EpicGames.Horde/Storage/Nodes/ChunkedDataWriter.cs
@@ -345,7 +345,7 @@ namespace EpicGames.Horde.Storage.Nodes
 		/// <returns>Handle to the written leaf node</returns>
 		async ValueTask FlushLeafNodeAsync(CancellationToken cancellationToken)
 		{
-			BlobHandle handle = await _writer.WriteNodeAsync(_leafLength, Array.Empty<BlobHandle>(), s_leafNodeType, cancellationToken);
+			BlobHandle handle = await _writer.WriteBlobAsync(_leafLength, Array.Empty<BlobHandle>(), s_leafNodeType, cancellationToken);
 			_leafHandles.Add(new NodeRef<ChunkedDataNode>(handle));
 			ResetLeafState();
 		}
