Cleaned up minor formatting issues with functions and classes.

diff --git a/Engine/Source/Programs/Shared/EpicGames.UHT/Exporters/CodeGen/UhtHeaderCodeGenerator.cs b/Engine/Source/Programs/Shared/EpicGames.UHT/Exporters/CodeGen/UhtHeaderCodeGenerator.cs
--- a/Engine/Source/Programs/Shared/EpicGames.UHT/Exporters/CodeGen/UhtHeaderCodeGenerator.cs
+++ b/Engine/Source/Programs/Shared/EpicGames.UHT/Exporters/CodeGen/UhtHeaderCodeGenerator.cs
@@ -424,6 +424,11 @@ namespace EpicGames.UHT.Exporters.CodeGen
 				!function.FunctionExportFlags.HasAnyFlags(UhtFunctionExportFlags.CustomThunk);
 		}
 
+		protected static bool IsCallbackFunction(UhtFunction function)
+		{
+			return function.FunctionFlags.HasAnyFlags(EFunctionFlags.Event) && function.SuperFunction == null;
+		}
+
 		/// <summary>
 		/// Determines whether the glue version of the specified native function should be exported.
 		/// </summary>
diff --git a/Engine/Source/Programs/Shared/EpicGames.UHT/Exporters/CodeGen/UhtHeaderCodeGeneratorCppFile.cs b/Engine/Source/Programs/Shared/EpicGames.UHT/Exporters/CodeGen/UhtHeaderCodeGeneratorCppFile.cs
--- a/Engine/Source/Programs/Shared/EpicGames.UHT/Exporters/CodeGen/UhtHeaderCodeGeneratorCppFile.cs
+++ b/Engine/Source/Programs/Shared/EpicGames.UHT/Exporters/CodeGen/UhtHeaderCodeGeneratorCppFile.cs
@@ -848,7 +848,7 @@ namespace EpicGames.UHT.Exporters.CodeGen
 			const string MetaDataParamsName = "Function_MetaDataParams";
 			string singletonName = GetSingletonName(function, true);
 			string staticsName = singletonName + "_Statics";
-			bool paramsInStatic = isNoExport || !function.FunctionFlags.HasAnyFlags(EFunctionFlags.Event);
+			bool paramsInStatic = isNoExport || !IsCallbackFunction(function);
 			bool isNet = function.FunctionFlags.HasAnyFlags(EFunctionFlags.NetRequest | EFunctionFlags.NetResponse);
 
 			string strippedFunctionName = function.StrippedFunctionName;
@@ -880,10 +880,10 @@ namespace EpicGames.UHT.Exporters.CodeGen
 						AppendEventParameter(builder, function, strippedFunctionName, UhtPropertyTextType.EventParameterFunctionMember, false, 1, "\r\n");
 					}
 
-					AppendPropertiesDecl(builder, context, properties, 1);
-
 					builder.AppendMetaDataDecl(function, context, properties, MetaDataParamsName, 1);
 
+					AppendPropertiesDecl(builder, context, properties, 1);
+
 					builder.Append("\tstatic const UECodeGen_Private::FFunctionParams FuncParams;\r\n");
 
 					builder.Append("};\r\n");
@@ -1307,19 +1307,19 @@ namespace EpicGames.UHT.Exporters.CodeGen
 				string sparseDataType = sparseScriptStruct.EngineName;
 
 				builder.Append('F').Append(sparseDataType).Append("* ").Append(classObj.SourceName).Append("::Get").Append(sparseDataType).Append("() const \r\n");
-				builder.Append("{ \r\n");
-				builder.Append("\treturn static_cast<F").Append(sparseDataType).Append("*>(GetClass()->GetOrCreateSparseClassData()); \r\n");
-				builder.Append("} \r\n");
+				builder.Append("{\r\n");
+				builder.Append("\treturn static_cast<F").Append(sparseDataType).Append("*>(GetClass()->GetOrCreateSparseClassData());\r\n");
+				builder.Append("}\r\n");
 
-				builder.Append("const F").Append(sparseDataType).Append("* ").Append(classObj.SourceName).Append("::Get").Append(sparseDataType).Append("(EGetSparseClassDataMethod GetMethod) const \r\n");
-				builder.Append("{ \r\n");
-				builder.Append("\treturn static_cast<const F").Append(sparseDataType).Append("*>(GetClass()->GetSparseClassData(GetMethod)); \r\n");
-				builder.Append("} \r\n");
+				builder.Append("const F").Append(sparseDataType).Append("* ").Append(classObj.SourceName).Append("::Get").Append(sparseDataType).Append("(EGetSparseClassDataMethod GetMethod) const\r\n");
+				builder.Append("{\r\n");
+				builder.Append("\treturn static_cast<const F").Append(sparseDataType).Append("*>(GetClass()->GetSparseClassData(GetMethod));\r\n");
+				builder.Append("}\r\n");
 
 				builder.Append("UScriptStruct* ").Append(classObj.SourceName).Append("::StaticGet").Append(sparseDataType).Append("ScriptStruct()\r\n");
-				builder.Append("{ \r\n");
-				builder.Append("\treturn F").Append(sparseDataType).Append("::StaticStruct(); \r\n");
-				builder.Append("} \r\n");
+				builder.Append("{\r\n");
+				builder.Append("\treturn F").Append(sparseDataType).Append("::StaticStruct();\r\n");
+				builder.Append("}\r\n");
 			}
 			return builder;
 		}
@@ -1524,7 +1524,7 @@ namespace EpicGames.UHT.Exporters.CodeGen
 
 				if (hasInterfaces)
 				{
-					builder.Append("\tconst UECodeGen_Private::FImplementedInterfaceParams ").Append(staticsName).Append("::InterfaceParams[] = {\r\n");
+					builder.Append("const UECodeGen_Private::FImplementedInterfaceParams ").Append(staticsName).Append("::InterfaceParams[] = {\r\n");
 
 					foreach (UhtStruct structObj in classObj.Bases)
 					{
@@ -1533,7 +1533,7 @@ namespace EpicGames.UHT.Exporters.CodeGen
 							if (baseClass.ClassFlags.HasAnyFlags(EClassFlags.Interface))
 							{
 								builder
-									.Append("\t\t{ ")
+									.Append("\t{ ")
 									.Append(GetSingletonName(baseClass.AlternateObject, false))
 									.Append(", (int32)VTABLE_OFFSET(")
 									.Append(classObj.SourceName)
@@ -1546,7 +1546,7 @@ namespace EpicGames.UHT.Exporters.CodeGen
 						}
 					}
 
-					builder.Append("\t};\r\n");
+					builder.Append("};\r\n");
 
 					builder.Append("static_assert(UE_ARRAY_COUNT(").Append(staticsName).Append("::InterfaceParams) < 64);\r\n");
 				}
@@ -1563,12 +1563,11 @@ namespace EpicGames.UHT.Exporters.CodeGen
 				builder.Append("};\r\n");
 
 				EClassFlags classFlags = classObj.ClassFlags & EClassFlags.SaveInCompiledInClasses;
-
 				builder.Append("const UECodeGen_Private::FClassParams ").Append(staticsName).Append("::ClassParams = {\r\n");
 				builder.Append("\t&").Append(classObj.SourceName).Append("::StaticClass,\r\n");
 				if (classObj.Config.Length > 0)
 				{
-					builder.Append("\t\t").AppendUTF8LiteralString(classObj.Config).Append(",\r\n");
+					builder.Append('\t').AppendUTF8LiteralString(classObj.Config).Append(",\r\n");
 				}
 				else
 				{
@@ -1576,14 +1575,14 @@ namespace EpicGames.UHT.Exporters.CodeGen
 				}
 				builder.Append("\t&StaticCppClassTypeInfo,\r\n");
 				builder.Append("\tDependentSingletons,\r\n");
-				builder.AppendArrayPtrLine(functions, UhtDefineScopeNames.WithEditor, null, "FuncInfo", 2, ",\r\n");
-				builder.AppendArrayPtrLine(properties, UhtDefineScopeNames.Standard, staticsName, "PropPointers", 2, ",\r\n");
+				builder.AppendArrayPtrLine(functions, UhtDefineScopeNames.WithEditor, null, "FuncInfo", 1, ",\r\n");
+				builder.AppendArrayPtrLine(properties, UhtDefineScopeNames.Standard, staticsName, "PropPointers", 1, ",\r\n");
 				builder.Append('\t').Append(hasInterfaces ? "InterfaceParams" : "nullptr").Append(",\r\n");
 				builder.Append("\tUE_ARRAY_COUNT(DependentSingletons),\r\n");
-				builder.AppendArrayCountLine(functions, UhtDefineScopeNames.WithEditor, null, "FuncInfo", 2, ",\r\n");
-				builder.AppendArrayCountLine(properties, UhtDefineScopeNames.Standard, staticsName, "PropPointers", 2, ",\r\n");
+				builder.AppendArrayCountLine(functions, UhtDefineScopeNames.WithEditor, null, "FuncInfo", 1, ",\r\n");
+				builder.AppendArrayCountLine(properties, UhtDefineScopeNames.Standard, staticsName, "PropPointers", 1, ",\r\n");
 				builder.Append('\t').Append(hasInterfaces ? "UE_ARRAY_COUNT(InterfaceParams)" : "0").Append(",\r\n");
-				builder.Append($"\t\t0x{(uint)classFlags:X8}u,\r\n");
+				builder.Append($"\t0x{(uint)classFlags:X8}u,\r\n");
 				builder.Append('\t').AppendMetaDataParams(classObj, staticsName, MetaDataParamsName).Append("\r\n");
 				builder.Append("};\r\n");
 			}
