Cleaned up class code generation including moving class functions so they are no longer interspersed with class.

diff --git a/Engine/Source/Programs/Shared/EpicGames.UHT/Exporters/CodeGen/UhtHeaderCodeGeneratorCppFile.cs b/Engine/Source/Programs/Shared/EpicGames.UHT/Exporters/CodeGen/UhtHeaderCodeGeneratorCppFile.cs
--- a/Engine/Source/Programs/Shared/EpicGames.UHT/Exporters/CodeGen/UhtHeaderCodeGeneratorCppFile.cs
+++ b/Engine/Source/Programs/Shared/EpicGames.UHT/Exporters/CodeGen/UhtHeaderCodeGeneratorCppFile.cs
@@ -225,8 +225,18 @@ namespace EpicGames.UHT.Exporters.CodeGen
 					{
 						if (!classObj.ClassFlags.HasAnyFlags(EClassFlags.Intrinsic))
 						{
+							// Collect the functions to be exported
+							UhtUsedDefineScopes<UhtFunction> functions = new(classObj.Functions);
+							functions.Instances.Sort((x, y) => StringComparerUE.OrdinalIgnoreCase.Compare(x.EngineName, y.EngineName));
+
+							// Output any functions
+							foreach (UhtFunction classFunction in functions.Instances)
+							{
+								AppendClassFunction(builder, classObj, classFunction);
+							}
+
 							using UhtCodeBlockComment blockComment = new(builder, field);
-							AppendClass(builder, classObj);
+							AppendClass(builder, classObj, functions);
 							classes.Add(classObj);
 						}
 					}
@@ -1046,11 +1056,187 @@ namespace EpicGames.UHT.Exporters.CodeGen
 			return builder;
 		}
 
-		private StringBuilder AppendClass(StringBuilder builder, UhtClass classObj)
+		private StringBuilder AppendClassFunction(StringBuilder builder, UhtClass classObj, UhtFunction function)
+		{
+			bool isNotDelegate = !function.FunctionFlags.HasAnyFlags(EFunctionFlags.Delegate);
+			bool isRpc = IsRpcFunction(function) && ShouldExportFunction(function);
+			bool isCallback = IsCallbackFunction(function);
+			if (isNotDelegate || isRpc || isCallback)
+			{
+				using UhtCodeBlockComment blockCommand = new(builder, classObj, function);
+				if (isCallback)
+				{
+					AppendEventParameter(builder, function, function.StrippedFunctionName, UhtPropertyTextType.EventParameterMember, true, 0, "\r\n");
+					AppendClassFunctionCallback(builder, classObj, function);
+					if (classObj.ClassFlags.HasAnyFlags(EClassFlags.Interface))
+					{
+						AppendInterfaceCallFunction(builder, classObj, function);
+					}
+				}
+				if (!function.FunctionFlags.HasAnyFlags(EFunctionFlags.Delegate))
+				{
+					AppendFunction(builder, function, classObj.ClassExportFlags.HasAnyFlags(UhtClassExportFlags.NoExport));
+				}
+				if (isRpc)
+				{
+					using UhtMacroBlockEmitter blockEmitter = new(builder, UhtDefineScopeNames.WithEditor, function.DefineScope);
+					builder.Append("DEFINE_FUNCTION(").AppendClassSourceNameOrInterfaceName(classObj).Append("::").Append(function.UnMarshalAndCallName).Append(")\r\n");
+					builder.Append("{\r\n");
+					AppendFunctionThunk(builder, function);
+					builder.Append("}\r\n");
+				}
+			}
+			return builder;
+		}
+
+		private StringBuilder AppendClassFunctionCallback(StringBuilder builder, UhtClass classObj, UhtFunction function)
 		{
-			// Output the RPC methods
-			AppendRpcFunctions(builder, classObj);
+			// Net response functions don't go into the VM
+			if (classObj.ClassExportFlags.HasAnyFlags(UhtClassExportFlags.NoExport) || function.FunctionFlags.HasAnyFlags(EFunctionFlags.NetResponse))
+			{
+				return builder;
+
+			}
+
+			bool isInterfaceClass = classObj.ClassFlags.HasAnyFlags(EClassFlags.Interface);
+			using UhtMacroBlockEmitter blockEmitter = new(builder, UhtDefineScopeNames.WithEditor, function.DefineScope);
+
+			if (!isInterfaceClass)
+			{
+				builder.Append("static FName NAME_").Append(classObj.SourceName).Append('_').Append(function.EngineName).Append(" = FName(TEXT(\"").Append(function.EngineName).Append("\"));\r\n");
+			}
+
+			AppendNativeFunctionHeader(builder, function, UhtPropertyTextType.EventFunctionArgOrRetVal, false, null, null, UhtFunctionExportFlags.None, 0, "\r\n");
+
+			if (isInterfaceClass)
+			{
+				builder.Append("{\r\n");
+
+				// assert if this is ever called directly
+				builder
+					.Append("\tcheck(0 && \"Do not directly call Event functions in Interfaces. Call Execute_")
+					.Append(function.EngineName)
+					.Append(" instead.\");\r\n");
 
+				// satisfy compiler if it's expecting a return value
+				if (function.ReturnProperty != null)
+				{
+					string eventParmStructName = GetEventStructParametersName(classObj, function.EngineName);
+					builder.Append('\t').Append(eventParmStructName).Append(" Parms;\r\n");
+					builder.Append("\treturn Parms.ReturnValue;\r\n");
+				}
+				builder.Append("}\r\n");
+			}
+			else
+			{
+				AppendEventFunctionPrologue(builder, function, function.EngineName, 0, "\r\n", false);
+
+				// Cast away const just in case, because ProcessEvent isn't const
+				builder.Append('\t');
+				if (function.FunctionFlags.HasAnyFlags(EFunctionFlags.Const))
+				{
+					builder.Append("const_cast<").Append(classObj.SourceName).Append("*>(this)->");
+				}
+				builder
+					.Append("ProcessEvent(FindFunctionChecked(")
+					.Append("NAME_")
+					.Append(classObj.SourceName)
+					.Append('_')
+					.Append(function.EngineName)
+					.Append("),")
+					.Append(function.Children.Count > 0 ? "&Parms" : "NULL")
+					.Append(");\r\n");
+
+				AppendEventFunctionEpilogue(builder, function, 0, "\r\n");
+			}
+			return builder;
+		}
+
+		private StringBuilder AppendInterfaceCallFunction(StringBuilder builder, UhtClass classObj, UhtFunction function)
+		{
+			builder.Append("static FName NAME_").Append(function.Outer?.SourceName).Append('_').Append(function.SourceName).Append(" = FName(TEXT(\"").Append(function.EngineName).Append("\"));\r\n");
+			string extraParameter = function.FunctionFlags.HasAnyFlags(EFunctionFlags.Const) ? "const UObject* O" : "UObject* O";
+			AppendNativeFunctionHeader(builder, function, UhtPropertyTextType.InterfaceFunctionArgOrRetVal, false, null, extraParameter, UhtFunctionExportFlags.None, 0, "\r\n");
+			builder.Append("{\r\n");
+			builder.Append("\tcheck(O != NULL);\r\n");
+			builder.Append("\tcheck(O->GetClass()->ImplementsInterface(").Append(classObj.SourceName).Append("::StaticClass()));\r\n");
+			if (function.Children.Count > 0)
+			{
+				builder.Append('\t').Append(GetEventStructParametersName(classObj, function.StrippedFunctionName)).Append(" Parms;\r\n");
+			}
+			builder.Append("\tUFunction* const Func = O->FindFunction(NAME_").Append(function.Outer?.SourceName).Append('_').Append(function.SourceName).Append(");\r\n");
+			builder.Append("\tif (Func)\r\n");
+			builder.Append("\t{\r\n");
+			foreach (UhtType parameter in function.ParameterProperties.Span)
+			{
+				if (parameter is UhtProperty property)
+				{
+					builder.Append("\t\tParms.").Append(property.SourceName).Append('=').Append(property.SourceName).Append(";\r\n");
+				}
+			}
+			builder
+				.Append("\t\t")
+				.Append(function.FunctionFlags.HasAnyFlags(EFunctionFlags.Const) ? "const_cast<UObject*>(O)" : "O")
+				.Append("->ProcessEvent(Func, ")
+				.Append(function.Children.Count > 0 ? "&Parms" : "NULL")
+				.Append(");\r\n");
+			foreach (UhtType parameter in function.ParameterProperties.Span)
+			{
+				if (parameter is UhtProperty property)
+				{
+					if (property.PropertyFlags.HasExactFlags(EPropertyFlags.OutParm | EPropertyFlags.ConstParm, EPropertyFlags.OutParm))
+					{
+						builder.Append("\t\t").Append(property.SourceName).Append("=Parms.").Append(property.SourceName).Append(";\r\n");
+					}
+				}
+			}
+			builder.Append("\t}\r\n");
+
+			// else clause to call back into native if it's a BlueprintNativeEvent
+			if (function.FunctionFlags.HasAnyFlags(EFunctionFlags.Native))
+			{
+				builder
+					.Append("\telse if (auto I = (")
+					.Append(function.FunctionFlags.HasAnyFlags(EFunctionFlags.Const) ? "const I" : "I")
+					.Append(classObj.EngineName)
+					.Append("*)(O->GetNativeInterfaceAddress(U")
+					.Append(classObj.EngineName)
+					.Append("::StaticClass())))\r\n");
+				builder.Append("\t{\r\n");
+				builder.Append("\t\t");
+				if (function.HasReturnProperty)
+				{
+					builder.Append("Parms.ReturnValue = ");
+				}
+				builder.Append("I->").Append(function.SourceName).Append("_Implementation(");
+
+				bool first = true;
+				foreach (UhtType parameter in function.ParameterProperties.Span)
+				{
+					if (parameter is UhtProperty property)
+					{
+						if (!first)
+						{
+							builder.Append(',');
+						}
+						first = false;
+						builder.Append(property.SourceName);
+					}
+				}
+				builder.Append(");\r\n");
+				builder.Append("\t}\r\n");
+			}
+
+			if (function.HasReturnProperty)
+			{
+				builder.Append("\treturn Parms.ReturnValue;\r\n");
+			}
+			builder.Append("}\r\n");
+			return builder;
+		}
+
+		private StringBuilder AppendClass(StringBuilder builder, UhtClass classObj, UhtUsedDefineScopes<UhtFunction> functions)
+		{
 			// Add the auto getters/setters
 			AppendAutoGettersSetters(builder, classObj);
 
@@ -1060,25 +1246,12 @@ namespace EpicGames.UHT.Exporters.CodeGen
 			// Add sparse accessors
 			AppendSparseAccessors(builder, classObj);
 
-			// Collect the callback function and sort by name to make the order stable
-			List<UhtFunction> callbackFunctions = new(classObj.Functions.Where(x => x.FunctionFlags.HasAnyFlags(EFunctionFlags.Event) && x.SuperFunction == null));
-			callbackFunctions.Sort((x, y) => StringComparerUE.OrdinalIgnoreCase.Compare(x.EngineName, y.EngineName));
-
-			// Generate the callback parameter structures
-			AppendCallbackParametersDecls(builder, callbackFunctions);
-
-			// VM -> C++ proxies (events and delegates).
-			if (!classObj.ClassExportFlags.HasAnyFlags(UhtClassExportFlags.NoExport))
-			{
-				AppendCallbackFunctions(builder, classObj, callbackFunctions);
-			}
-
 			if (!classObj.ClassExportFlags.HasAnyFlags(UhtClassExportFlags.NoExport))
 			{
 				AppendNatives(builder, classObj);
 			}
 
-			AppendNativeGeneratedInitCode(builder, classObj);
+			AppendNativeGeneratedInitCode(builder, classObj, functions);
 
 			if (!classObj.ClassExportFlags.HasAnyFlags(UhtClassExportFlags.HasConstructor))
 			{
@@ -1115,11 +1288,6 @@ namespace EpicGames.UHT.Exporters.CodeGen
 				AppendSerializer(builder, classObj, UhtSerializerArchiveType.Archive, "IMPLEMENT_FARCHIVE_SERIALIZER");
 				AppendSerializer(builder, classObj, UhtSerializerArchiveType.StructuredArchiveRecord, "IMPLEMENT_FSTRUCTUREDARCHIVE_SERIALIZER");
 			}
-
-			if (classObj.ClassFlags.HasAnyFlags(EClassFlags.Interface))
-			{
-				AppendInterfaceCallFunctions(builder, classObj, callbackFunctions);
-			}
 			return builder;
 		}
 
@@ -1324,92 +1492,6 @@ namespace EpicGames.UHT.Exporters.CodeGen
 			return builder;
 		}
 
-		private StringBuilder AppendInterfaceCallFunctions(StringBuilder builder, UhtClass classObj, List<UhtFunction> callbackFunctions)
-		{
-			foreach (UhtFunction function in callbackFunctions)
-			{
-				builder.Append("static FName NAME_").Append(function.Outer?.SourceName).Append('_').Append(function.SourceName).Append(" = FName(TEXT(\"").Append(function.EngineName).Append("\"));\r\n");
-				string extraParameter = function.FunctionFlags.HasAnyFlags(EFunctionFlags.Const) ? "const UObject* O" : "UObject* O";
-				AppendNativeFunctionHeader(builder, function, UhtPropertyTextType.InterfaceFunctionArgOrRetVal, false, null, extraParameter, UhtFunctionExportFlags.None, 0, "\r\n");
-				builder.Append("{\r\n");
-				builder.Append("\tcheck(O != NULL);\r\n");
-				builder.Append("\tcheck(O->GetClass()->ImplementsInterface(").Append(classObj.SourceName).Append("::StaticClass()));\r\n");
-				if (function.Children.Count > 0)
-				{
-					builder.Append('\t').Append(GetEventStructParametersName(classObj, function.StrippedFunctionName)).Append(" Parms;\r\n");
-				}
-				builder.Append("\tUFunction* const Func = O->FindFunction(NAME_").Append(function.Outer?.SourceName).Append('_').Append(function.SourceName).Append(");\r\n");
-				builder.Append("\tif (Func)\r\n");
-				builder.Append("\t{\r\n");
-				foreach (UhtType parameter in function.ParameterProperties.Span)
-				{
-					if (parameter is UhtProperty property)
-					{
-						builder.Append("\t\tParms.").Append(property.SourceName).Append('=').Append(property.SourceName).Append(";\r\n");
-					}
-				}
-				builder
-					.Append("\t\t")
-					.Append(function.FunctionFlags.HasAnyFlags(EFunctionFlags.Const) ? "const_cast<UObject*>(O)" : "O")
-					.Append("->ProcessEvent(Func, ")
-					.Append(function.Children.Count > 0 ? "&Parms" : "NULL")
-					.Append(");\r\n");
-				foreach (UhtType parameter in function.ParameterProperties.Span)
-				{
-					if (parameter is UhtProperty property)
-					{
-						if (property.PropertyFlags.HasExactFlags(EPropertyFlags.OutParm | EPropertyFlags.ConstParm, EPropertyFlags.OutParm))
-						{
-							builder.Append("\t\t").Append(property.SourceName).Append("=Parms.").Append(property.SourceName).Append(";\r\n");
-						}
-					}
-				}
-				builder.Append("\t}\r\n");
-
-				// else clause to call back into native if it's a BlueprintNativeEvent
-				if (function.FunctionFlags.HasAnyFlags(EFunctionFlags.Native))
-				{
-					builder
-						.Append("\telse if (auto I = (")
-						.Append(function.FunctionFlags.HasAnyFlags(EFunctionFlags.Const) ? "const I" : "I")
-						.Append(classObj.EngineName)
-						.Append("*)(O->GetNativeInterfaceAddress(U")
-						.Append(classObj.EngineName)
-						.Append("::StaticClass())))\r\n");
-					builder.Append("\t{\r\n");
-					builder.Append("\t\t");
-					if (function.HasReturnProperty)
-					{
-						builder.Append("Parms.ReturnValue = ");
-					}
-					builder.Append("I->").Append(function.SourceName).Append("_Implementation(");
-
-					bool first = true;
-					foreach (UhtType parameter in function.ParameterProperties.Span)
-					{
-						if (parameter is UhtProperty property)
-						{
-							if (!first)
-							{
-								builder.Append(',');
-							}
-							first = false;
-							builder.Append(property.SourceName);
-						}
-					}
-					builder.Append(");\r\n");
-					builder.Append("\t}\r\n");
-				}
-
-				if (function.HasReturnProperty)
-				{
-					builder.Append("\treturn Parms.ReturnValue;\r\n");
-				}
-				builder.Append("}\r\n");
-			}
-			return builder;
-		}
-
 		private static StringBuilder AppendSerializer(StringBuilder builder, UhtClass classObj, UhtSerializerArchiveType serializerType, string macroText)
 		{
 			if (!classObj.SerializerArchiveType.HasAnyFlags(serializerType))
@@ -1420,7 +1502,7 @@ namespace EpicGames.UHT.Exporters.CodeGen
 			return builder;
 		}
 
-		private StringBuilder AppendNativeGeneratedInitCode(StringBuilder builder, UhtClass classObj)
+		private StringBuilder AppendNativeGeneratedInitCode(StringBuilder builder, UhtClass classObj, UhtUsedDefineScopes<UhtFunction> functions)
 		{
 			const string MetaDataParamsName = "Class_MetaDataParams";
 			string singletonName = GetSingletonName(classObj, true);
@@ -1432,20 +1514,6 @@ namespace EpicGames.UHT.Exporters.CodeGen
 
 			bool hasInterfaces = classObj.Bases.Any(x => x is UhtClass baseClass && baseClass.ClassFlags.HasAnyFlags(EClassFlags.Interface));
 
-			// Collect the functions to be exported
-			UhtUsedDefineScopes<UhtFunction> functions = new(classObj.Functions);
-			functions.Instances.Sort((x, y) => StringComparerUE.OrdinalIgnoreCase.Compare(x.EngineName, y.EngineName));
-
-			// Output any function
-			foreach (UhtFunction function in functions.Instances)
-			{
-				if (!function.FunctionFlags.HasAnyFlags(EFunctionFlags.Delegate))
-				{
-					using UhtCodeBlockComment blockCommand = new(builder, classObj, function);
-					AppendFunction(builder, function, classObj.ClassExportFlags.HasAnyFlags(UhtClassExportFlags.NoExport));
-				}
-			}
-
 			builder.Append("IMPLEMENT_CLASS_NO_AUTO_REGISTRATION(").Append(classObj.SourceName).Append(");\r\n");
 
 			// Everything from this point on will be part of the definition hash
@@ -1466,20 +1534,53 @@ namespace EpicGames.UHT.Exporters.CodeGen
 			{
 				builder.Append("struct ").Append(staticsName).Append("\r\n");
 				builder.Append("{\r\n");
-				builder.Append("\tstatic UObject* (*const DependentSingletons[])();\r\n");
-
-				builder.AppendIfInstances(functions, UhtDefineScopeNames.WithEditor, builder => builder.Append("\tstatic const FClassFunctionLinkInfo FuncInfo[];\r\n"));
 
 				builder.AppendMetaDataDecl(classObj, context, properties, MetaDataParamsName, 1);
 
 				AppendPropertiesDecl(builder, context, properties, 1);
 
+				builder.Append("\tstatic UObject* (*const DependentSingletons[])();\r\n");
+
+				//builder.AppendIfInstances(functions, UhtDefineScopeNames.WithEditor, builder => builder.Append("\t\tstatic const FClassFunctionLinkInfo FuncInfo[];\r\n"));
+				// Functions
+				builder.AppendInstances(functions, UhtDefineScopeNames.WithEditor,
+					builder =>
+					{
+						builder.Append("\tstatic constexpr FClassFunctionLinkInfo ").Append("FuncInfo[] = {\r\n");
+					},
+					(builder, function) =>
+					{
+						builder
+							.Append("\t\t{ &")
+							.Append(GetSingletonName(function, true))
+							.Append(", ")
+							.AppendUTF8LiteralString(function.EngineName)
+							.Append(" },")
+							.AppendObjectHash(classObj, context, function)
+							.Append("\r\n");
+					},
+					builder =>
+					{
+						builder.Append("\t};\r\n");
+						builder.Append("\tstatic_assert(UE_ARRAY_COUNT(").Append("FuncInfo) < 2048);\r\n");
+					});
+
+
 				if (hasInterfaces)
 				{
 					builder.Append("\tstatic const UECodeGen_Private::FImplementedInterfaceParams InterfaceParams[];\r\n");
 				}
 
-				builder.Append("\tstatic const FCppClassTypeInfoStatic StaticCppClassTypeInfo;\r\n");
+				builder.Append("\tstatic constexpr FCppClassTypeInfoStatic StaticCppClassTypeInfo = {\r\n");
+				if (classObj.ClassFlags.HasAnyFlags(EClassFlags.Interface))
+				{
+					builder.Append("\t\tTCppClassTypeTraits<I").Append(classObj.EngineName).Append(">::IsAbstract,\r\n");
+				}
+				else
+				{
+					builder.Append("\t\tTCppClassTypeTraits<").Append(classObj.SourceName).Append(">::IsAbstract,\r\n");
+				}
+				builder.Append("\t};\r\n");
 
 				builder.Append("\tstatic const UECodeGen_Private::FClassParams ClassParams;\r\n");
 
@@ -1488,6 +1589,9 @@ namespace EpicGames.UHT.Exporters.CodeGen
 
 			// Define the statics object
 			{
+				AppendPropertiesDefs(builder, context, properties, 0);
+
+				// Dependent singletons
 				builder.Append("UObject* (*const ").Append(staticsName).Append("::DependentSingletons[])() = {\r\n");
 				if (classObj.SuperClass != null && classObj.SuperClass != classObj)
 				{
@@ -1495,33 +1599,9 @@ namespace EpicGames.UHT.Exporters.CodeGen
 				}
 				builder.Append("\t(UObject* (*)())").Append(GetSingletonName(Package, true)).Append(",\r\n");
 				builder.Append("};\r\n");
-
 				builder.Append("static_assert(UE_ARRAY_COUNT(").Append(staticsName).Append("::DependentSingletons) < 16);\r\n");
 
-				builder.AppendInstances(functions, UhtDefineScopeNames.WithEditor,
-					builder =>
-					{
-						builder.Append("const FClassFunctionLinkInfo ").Append(staticsName).Append("::FuncInfo[] = {\r\n");
-					},
-					(builder, function) =>
-					{
-						builder
-							.Append("\t{ &")
-							.Append(GetSingletonName(function, true))
-							.Append(", ")
-							.AppendUTF8LiteralString(function.EngineName)
-							.Append(" },")
-							.AppendObjectHash(classObj, context, function)
-							.Append("\r\n");
-					},
-					builder =>
-					{
-						builder.Append("};\r\n");
-						builder.Append("static_assert(UE_ARRAY_COUNT(").Append(staticsName).Append("::FuncInfo) < 2048);\r\n");
-					});
-
-				AppendPropertiesDefs(builder, context, properties, 0);
-
+				// Implemented interfaces
 				if (hasInterfaces)
 				{
 					builder.Append("const UECodeGen_Private::FImplementedInterfaceParams ").Append(staticsName).Append("::InterfaceParams[] = {\r\n");
@@ -1547,21 +1627,9 @@ namespace EpicGames.UHT.Exporters.CodeGen
 					}
 
 					builder.Append("};\r\n");
-
-					builder.Append("static_assert(UE_ARRAY_COUNT(").Append(staticsName).Append("::InterfaceParams) < 64);\r\n");
 				}
 
-				builder.Append("const FCppClassTypeInfoStatic ").Append(staticsName).Append("::StaticCppClassTypeInfo = {\r\n");
-				if (classObj.ClassFlags.HasAnyFlags(EClassFlags.Interface))
-				{
-					builder.Append("\tTCppClassTypeTraits<I").Append(classObj.EngineName).Append(">::IsAbstract,\r\n");
-				}
-				else
-				{
-					builder.Append("\tTCppClassTypeTraits<").Append(classObj.SourceName).Append(">::IsAbstract,\r\n");
-				}
-				builder.Append("};\r\n");
-
+				// Class parameters
 				EClassFlags classFlags = classObj.ClassFlags & EClassFlags.SaveInCompiledInClasses;
 				builder.Append("const UECodeGen_Private::FClassParams ").Append(staticsName).Append("::ClassParams = {\r\n");
 				builder.Append("\t&").Append(classObj.SourceName).Append("::StaticClass,\r\n");
@@ -1744,104 +1812,6 @@ namespace EpicGames.UHT.Exporters.CodeGen
 			return builder;
 		}
 
-		private static StringBuilder AppendCallbackParametersDecls(StringBuilder builder, List<UhtFunction> callbackFunctions)
-		{
-			foreach (UhtFunction function in callbackFunctions)
-			{
-				AppendEventParameter(builder, function, function.StrippedFunctionName, UhtPropertyTextType.EventParameterMember, true, 1, "\r\n");
-			}
-			return builder;
-		}
-
-		private StringBuilder AppendCallbackFunctions(StringBuilder builder, UhtClass classObj, List<UhtFunction> callbackFunctions)
-		{
-			if (callbackFunctions.Count > 0)
-			{
-				bool isInterfaceClass = classObj.ClassFlags.HasAnyFlags(EClassFlags.Interface);
-				{
-					using UhtMacroBlockEmitter blockEmitter = new(builder, UhtDefineScopeNames.WithEditor, UhtDefineScope.None);
-					foreach (UhtFunction function in callbackFunctions)
-					{
-						// Net response functions don't go into the VM
-						if (function.FunctionFlags.HasAnyFlags(EFunctionFlags.NetResponse))
-						{
-							continue;
-						}
-
-						blockEmitter.Set(function.DefineScope);
-
-						if (!isInterfaceClass)
-						{
-							builder.Append("static FName NAME_").Append(classObj.SourceName).Append('_').Append(function.EngineName).Append(" = FName(TEXT(\"").Append(function.EngineName).Append("\"));\r\n");
-						}
-
-						AppendNativeFunctionHeader(builder, function, UhtPropertyTextType.EventFunctionArgOrRetVal, false, null, null, UhtFunctionExportFlags.None, 0, "\r\n");
-
-						if (isInterfaceClass)
-						{
-							builder.Append("{\r\n");
-
-							// assert if this is ever called directly
-							builder
-								.Append("\tcheck(0 && \"Do not directly call Event functions in Interfaces. Call Execute_")
-								.Append(function.EngineName)
-								.Append(" instead.\");\r\n");
-
-							// satisfy compiler if it's expecting a return value
-							if (function.ReturnProperty != null)
-							{
-								string eventParmStructName = GetEventStructParametersName(classObj, function.EngineName);
-								builder.Append('\t').Append(eventParmStructName).Append(" Parms;\r\n");
-								builder.Append("\treturn Parms.ReturnValue;\r\n");
-							}
-							builder.Append("}\r\n");
-						}
-						else
-						{
-							AppendEventFunctionPrologue(builder, function, function.EngineName, 1, "\r\n", false);
-
-							// Cast away const just in case, because ProcessEvent isn't const
-							builder.Append('\t');
-							if (function.FunctionFlags.HasAnyFlags(EFunctionFlags.Const))
-							{
-								builder.Append("const_cast<").Append(classObj.SourceName).Append("*>(this)->");
-							}
-							builder
-								.Append("ProcessEvent(FindFunctionChecked(")
-								.Append("NAME_")
-								.Append(classObj.SourceName)
-								.Append('_')
-								.Append(function.EngineName)
-								.Append("),")
-								.Append(function.Children.Count > 0 ? "&Parms" : "NULL")
-								.Append(");\r\n");
-
-							AppendEventFunctionEpilogue(builder, function, 1, "\r\n");
-						}
-					}
-				}
-			}
-			return builder;
-		}
-
-		private static StringBuilder AppendRpcFunctions(StringBuilder builder, UhtClass classObj)
-		{
-			// Collect the functions in reversed order
-			UhtUsedDefineScopes<UhtFunction> functions = new(classObj.Functions.Where(x => IsRpcFunction(x) && ShouldExportFunction(x)));
-			functions.Instances.Reverse();
-			functions.OrderByDefineScope();
-
-			// Output the RPC methods
-			builder.AppendInstances(functions, UhtDefineScopeNames.WithEditor, (builder, function) => 
-			{
-				builder.Append("DEFINE_FUNCTION(").AppendClassSourceNameOrInterfaceName(classObj).Append("::").Append(function.UnMarshalAndCallName).Append(")\r\n");
-				builder.Append("{\r\n");
-				AppendFunctionThunk(builder, function);
-				builder.Append("}\r\n");
-			});
-			return builder;
-		}
-
 		private static StringBuilder AppendFunctionThunk(StringBuilder builder, UhtFunction function)
 		{
 			// Export the GET macro for the parameters
