Fix threading issue in AsyncThreadPoolWorkQueueTest.

diff --git a/Engine/Source/Programs/Shared/EpicGames.Core.Tests/AsyncThreadPoolWorkQueueTests.cs b/Engine/Source/Programs/Shared/EpicGames.Core.Tests/AsyncThreadPoolWorkQueueTests.cs
--- a/Engine/Source/Programs/Shared/EpicGames.Core.Tests/AsyncThreadPoolWorkQueueTests.cs
+++ b/Engine/Source/Programs/Shared/EpicGames.Core.Tests/AsyncThreadPoolWorkQueueTests.cs
@@ -23,13 +23,13 @@ public sealed class AsyncThreadPoolWorkQueueTests : IDisposable
 	{
 		int c = 0;
 		await _queue.EnqueueAsync((_) => 
-		{ 
-			c++; 
+		{
+			Interlocked.Increment(ref c);
 			return Task.CompletedTask; 
 		});
 		await _queue.EnqueueAsync((_) => 
-		{ 
-			c++; 
+		{
+			Interlocked.Increment(ref c);
 			return Task.CompletedTask; 
 		});
 		await _queue.ExecuteAsync();
