Horde: Use the EpicGames.Core.JsonLogEvent type to represent log events inside Horde. Work towards allowing IssueHandlers to be implemented outside Horde server.

diff --git a/Engine/Source/Programs/Horde/Horde.Server/Authentication/HordeJwtBearerHandler.cs b/Engine/Source/Programs/Horde/Horde.Server/Authentication/HordeJwtBearerHandler.cs
--- a/Engine/Source/Programs/Horde/Horde.Server/Authentication/HordeJwtBearerHandler.cs
+++ b/Engine/Source/Programs/Horde/Horde.Server/Authentication/HordeJwtBearerHandler.cs
@@ -10,6 +10,7 @@ using System.Security.Claims;
 using System.Text.Json;
 using System.Threading;
 using System.Threading.Tasks;
+using EpicGames.Core;
 using Horde.Server.Users;
 using Horde.Server.Utilities;
 using Microsoft.AspNetCore.Authentication;
diff --git a/Engine/Source/Programs/Horde/Horde.Server/Issues/External/JiraService.cs b/Engine/Source/Programs/Horde/Horde.Server/Issues/External/JiraService.cs
--- a/Engine/Source/Programs/Horde/Horde.Server/Issues/External/JiraService.cs
+++ b/Engine/Source/Programs/Horde/Horde.Server/Issues/External/JiraService.cs
@@ -12,7 +12,6 @@ using Microsoft.Extensions.Logging;
 using Microsoft.Extensions.Options;
 using Polly;
 using Polly.Extensions.Http;
-using Horde.Server.Utilities;
 using Microsoft.Extensions.Hosting;
 using HordeCommon;
 using System.Threading;
@@ -20,6 +19,7 @@ using System.Linq;
 using Horde.Server.Streams;
 using Horde.Server.Users;
 using Horde.Server.Server;
+using EpicGames.Core;
 
 namespace Horde.Server.Issues.External
 {
diff --git a/Engine/Source/Programs/Horde/Horde.Server/Issues/Handlers/GauntletIssueHandler.cs b/Engine/Source/Programs/Horde/Horde.Server/Issues/Handlers/GauntletIssueHandler.cs
--- a/Engine/Source/Programs/Horde/Horde.Server/Issues/Handlers/GauntletIssueHandler.cs
+++ b/Engine/Source/Programs/Horde/Horde.Server/Issues/Handlers/GauntletIssueHandler.cs
@@ -2,10 +2,9 @@
 
 using System.Collections.Generic;
 using System.Text;
+using System.Text.Json;
 using System.Text.RegularExpressions;
 using EpicGames.Core;
-using Horde.Server.Logs;
-using Horde.Server.Utilities;
 using Microsoft.Extensions.Logging;
 
 namespace Horde.Server.Issues.Handlers
@@ -95,12 +94,14 @@ namespace Horde.Server.Issues.Handlers
 		/// <param name="metadata"></param>
 		private static void GetNames(IssueEvent issueEvent, HashSet<IssueKey> testNames, HashSet<IssueMetadata> metadata)
 		{
-			foreach (ILogEventLine line in issueEvent.Lines)
+			foreach (JsonLogEvent line in issueEvent.Lines)
 			{
+				JsonDocument document = JsonDocument.Parse(line.Data);
+
 				string? name = null;
 
 				string? value;
-				if (line.Data.TryGetNestedProperty("properties.Name", out value))
+				if (document.RootElement.TryGetNestedProperty("properties.Name", out value))
 				{
 					name = value;
 				}
@@ -124,16 +125,18 @@ namespace Horde.Server.Issues.Handlers
 		{
 			if(issueEvent.EventId == KnownLogEvents.Gauntlet_BuildDropEvent)
 			{
-				foreach (ILogEventLine line in issueEvent.Lines)
+				foreach (JsonLogEvent line in issueEvent.Lines)
 				{
+					JsonDocument document = JsonDocument.Parse(line.Data);
+
 					string? path = null;
 
 					string? value;
-					if (line.Data.TryGetNestedProperty("properties.File", out value))
+					if (document.RootElement.TryGetNestedProperty("properties.File", out value))
 					{
 						path = value;
 					}
-					else if (line.Data.TryGetNestedProperty("properties.Directory", out value))
+					else if (document.RootElement.TryGetNestedProperty("properties.Directory", out value))
 					{
 						path = value;
 					}
diff --git a/Engine/Source/Programs/Horde/Horde.Server/Issues/Handlers/LocalizationIssueHandler.cs b/Engine/Source/Programs/Horde/Horde.Server/Issues/Handlers/LocalizationIssueHandler.cs
--- a/Engine/Source/Programs/Horde/Horde.Server/Issues/Handlers/LocalizationIssueHandler.cs
+++ b/Engine/Source/Programs/Horde/Horde.Server/Issues/Handlers/LocalizationIssueHandler.cs
@@ -2,9 +2,8 @@
 
 using System;
 using System.Collections.Generic;
+using System.Text.Json;
 using EpicGames.Core;
-using Horde.Server.Logs;
-using Horde.Server.Utilities;
 using Microsoft.Extensions.Logging;
 
 namespace Horde.Server.Issues.Handlers
@@ -44,10 +43,12 @@ namespace Horde.Server.Issues.Handlers
 		/// <param name="sourceFiles">List of source files</param>
 		public static void GetSourceFiles(IssueEvent issueEvent, HashSet<IssueKey> sourceFiles)
 		{
-			foreach (ILogEventLine line in issueEvent.Lines)
+			foreach (JsonLogEvent line in issueEvent.Lines)
 			{
+				JsonDocument document = JsonDocument.Parse(line.Data);
+
 				string? relativePath;
-				if (line.Data.TryGetNestedProperty("properties.file.relativePath", out relativePath) || line.Data.TryGetNestedProperty("properties.file", out relativePath))
+				if (document.RootElement.TryGetNestedProperty("properties.file.relativePath", out relativePath) || document.RootElement.TryGetNestedProperty("properties.file", out relativePath))
 				{
 					if (!relativePath.EndsWith(".manifest", StringComparison.OrdinalIgnoreCase))
 					{
diff --git a/Engine/Source/Programs/Horde/Horde.Server/Issues/Handlers/ScopedIssueHandler.cs b/Engine/Source/Programs/Horde/Horde.Server/Issues/Handlers/ScopedIssueHandler.cs
--- a/Engine/Source/Programs/Horde/Horde.Server/Issues/Handlers/ScopedIssueHandler.cs
+++ b/Engine/Source/Programs/Horde/Horde.Server/Issues/Handlers/ScopedIssueHandler.cs
@@ -3,10 +3,9 @@
 using System;
 using System.Collections.Generic;
 using System.Text;
+using System.Text.Json;
 using System.Text.RegularExpressions;
 using EpicGames.Core;
-using Horde.Server.Logs;
-using Horde.Server.Utilities;
 
 namespace Horde.Server.Issues.Handlers
 {
@@ -29,11 +28,13 @@ namespace Horde.Server.Issues.Handlers
 		{
 			string? scope = null;
 
-			foreach (ILogEventLine line in logEvent.Lines)
+			foreach (JsonLogEvent line in logEvent.Lines)
 			{
+				JsonDocument document = JsonDocument.Parse(line.Data);
+
 				string? channelType;
 				string? channelText;
-				if (line.Data.TryGetNestedProperty("properties.channel.$type", out channelType) && line.Data.TryGetNestedProperty("properties.channel.$text", out channelText))
+				if (document.RootElement.TryGetNestedProperty("properties.channel.$type", out channelType) && document.RootElement.TryGetNestedProperty("properties.channel.$text", out channelText))
 				{
 					if (channelType != "Channel" || !channelText.StartsWith("Log", StringComparison.Ordinal))
 					{
diff --git a/Engine/Source/Programs/Horde/Horde.Server/Issues/Handlers/SystemicIssueHandler.cs b/Engine/Source/Programs/Horde/Horde.Server/Issues/Handlers/SystemicIssueHandler.cs
--- a/Engine/Source/Programs/Horde/Horde.Server/Issues/Handlers/SystemicIssueHandler.cs
+++ b/Engine/Source/Programs/Horde/Horde.Server/Issues/Handlers/SystemicIssueHandler.cs
@@ -2,7 +2,6 @@
 
 using System.Collections.Generic;
 using EpicGames.Core;
-using Horde.Server.Logs;
 using HordeCommon;
 using Microsoft.Extensions.Logging;
 
@@ -40,7 +39,7 @@ namespace Horde.Server.Issues.Handlers
 		{
 			if (logEvent.EventId != null)
 			{
-				if (logEvent.EventId != null && IsMatchingEventId(logEvent.EventId.Value))
+				if (IsMatchingEventId(logEvent.EventId.Value))
 				{
 					if (_nonSystemicError)
 					{
diff --git a/Engine/Source/Programs/Horde/Horde.Server/Issues/IssueEvent.cs b/Engine/Source/Programs/Horde/Horde.Server/Issues/IssueEvent.cs
--- a/Engine/Source/Programs/Horde/Horde.Server/Issues/IssueEvent.cs
+++ b/Engine/Source/Programs/Horde/Horde.Server/Issues/IssueEvent.cs
@@ -47,7 +47,7 @@ namespace Horde.Server.Issues
 		/// <summary>
 		/// Gets this event data as a BSON document
 		/// </summary>
-		public IReadOnlyList<ILogEventLine> Lines => EventData.Lines;
+		public IReadOnlyList<JsonLogEvent> Lines => EventData.Lines;
 
 		/// <summary>
 		/// Constructor
diff --git a/Engine/Source/Programs/Horde/Horde.Server/Issues/IssueKey.cs b/Engine/Source/Programs/Horde/Horde.Server/Issues/IssueKey.cs
--- a/Engine/Source/Programs/Horde/Horde.Server/Issues/IssueKey.cs
+++ b/Engine/Source/Programs/Horde/Horde.Server/Issues/IssueKey.cs
@@ -5,9 +5,7 @@ using System.Collections.Generic;
 using System.Linq;
 using System.Text.Json;
 using EpicGames.Core;
-using Horde.Server.Logs;
 using Horde.Server.Streams;
-using Horde.Server.Utilities;
 
 namespace Horde.Server.Issues
 {
@@ -135,10 +133,12 @@ namespace Horde.Server.Issues
 		/// <param name="issueEvent">The log event to parse</param>
 		public static void AddAssets(this HashSet<IssueKey> keys, IssueEvent issueEvent)
 		{
-			foreach (ILogEventLine line in issueEvent.Lines)
+			foreach (JsonLogEvent line in issueEvent.Lines)
 			{
+				JsonDocument document = JsonDocument.Parse(line.Data);
+
 				string? relativePath;
-				if (line.Data.TryGetNestedProperty("properties.asset.relativePath", out relativePath) || line.Data.TryGetNestedProperty("properties.asset.$text", out relativePath))
+				if (document.RootElement.TryGetNestedProperty("properties.asset.relativePath", out relativePath) || document.RootElement.TryGetNestedProperty("properties.asset.$text", out relativePath))
 				{
 					int endIdx = relativePath.LastIndexOfAny(new char[] { '/', '\\' }) + 1;
 					string fileName = relativePath.Substring(endIdx);
@@ -185,10 +185,12 @@ namespace Horde.Server.Issues
 		/// <param name="issueEvent">The event data</param>
 		public static void AddSourceFiles(this HashSet<IssueKey> keys, IssueEvent issueEvent)
 		{
-			foreach (ILogEventLine line in issueEvent.Lines)
+			foreach (JsonLogEvent line in issueEvent.Lines)
 			{
+				JsonDocument document = JsonDocument.Parse(line.Data);
+
 				JsonElement properties;
-				if (line.Data.TryGetProperty("properties", out properties) && properties.ValueKind == JsonValueKind.Object)
+				if (document.RootElement.TryGetProperty("properties", out properties) && properties.ValueKind == JsonValueKind.Object)
 				{
 					IssueKeyType type = IssueKeyType.File;
 					if (properties.TryGetProperty("note", out JsonElement noteElement) && noteElement.GetBoolean())
@@ -234,10 +236,12 @@ namespace Horde.Server.Issues
 		/// <param name="eventData">The log event data</param>
 		public static void AddSymbols(this HashSet<IssueKey> keys, IssueEvent eventData)
 		{
-			foreach (ILogEventLine line in eventData.Lines)
+			foreach (JsonLogEvent line in eventData.Lines)
 			{
+				JsonDocument document = JsonDocument.Parse(line.Data);
+
 				string? identifier;
-				if (line.Data.TryGetNestedProperty("properties.symbol.identifier", out identifier))
+				if (document.RootElement.TryGetNestedProperty("properties.symbol.identifier", out identifier))
 				{
 					IssueKey key = new IssueKey(identifier, IssueKeyType.Symbol);
 					keys.Add(key);
diff --git a/Engine/Source/Programs/Horde/Horde.Server/Logs/ILogEvent.cs b/Engine/Source/Programs/Horde/Horde.Server/Logs/ILogEvent.cs
--- a/Engine/Source/Programs/Horde/Horde.Server/Logs/ILogEvent.cs
+++ b/Engine/Source/Programs/Horde/Horde.Server/Logs/ILogEvent.cs
@@ -10,32 +10,6 @@ using MongoDB.Bson;
 
 namespace Horde.Server.Logs
 {
-	/// <summary>
-	/// Interface for a log event line
-	/// </summary>
-	public interface ILogEventLine
-	{
-		/// <summary>
-		/// Log level
-		/// </summary>
-		LogLevel Level { get; }
-
-		/// <summary>
-		/// The event id
-		/// </summary>
-		EventId? EventId { get; }
-
-		/// <summary>
-		/// The complete rendered message, in plaintext
-		/// </summary>
-		string Message { get; }
-
-		/// <summary>
-		/// Data for the line
-		/// </summary>
-		JsonElement Data { get; }
-	}
-
 	/// <summary>
 	/// Interface for event data
 	/// </summary>
@@ -57,9 +31,9 @@ namespace Horde.Server.Logs
 		string Message { get; }
 
 		/// <summary>
-		/// Gets this event data as a BSON document
+		/// Gets this event data as a JSON objects
 		/// </summary>
-		IReadOnlyList<ILogEventLine> Lines { get; }
+		IReadOnlyList<JsonLogEvent> Lines { get; }
 	}
 
 	/// <summary>
@@ -137,41 +111,7 @@ namespace Horde.Server.Logs
 		/// <returns></returns>
 		public static IEnumerable<JsonProperty> FindPropertiesOfType(this ILogEventData data, Utf8String type)
 		{
-			return data.Lines.SelectMany(x => FindPropertiesOfType(x, type));
-		}
-
-		/// <summary>
-		/// Find all properties of the given type in a particular log line
-		/// </summary>
-		/// <param name="line">Line data</param>
-		/// <param name="type">Type of property to return</param>
-		/// <returns></returns>
-		public static IEnumerable<JsonProperty> FindPropertiesOfType(this ILogEventLine line, Utf8String type)
-		{
-			JsonElement properties;
-			if (line.Data.TryGetProperty("properties", out properties) && properties.ValueKind == JsonValueKind.Object)
-			{
-				foreach (JsonProperty property in properties.EnumerateObject())
-				{
-					if (property.Value.ValueKind == JsonValueKind.Object)
-					{
-						foreach (JsonProperty subProperty in property.Value.EnumerateObject())
-						{
-							if (subProperty.NameEquals(LogEventPropertyName.Type.Span))
-							{
-								if (subProperty.Value.ValueKind == JsonValueKind.String && subProperty.Value.ValueEquals(type.Span))
-								{
-									yield return property;
-								}
-								else
-								{
-									break;
-								}
-							}
-						}
-					}
-				}
-			}
+			return data.Lines.SelectMany(x => x.FindPropertiesOfType(type));
 		}
 	}
 }
diff --git a/Engine/Source/Programs/Horde/Horde.Server/Logs/LogEventMessages.cs b/Engine/Source/Programs/Horde/Horde.Server/Logs/LogEventMessages.cs
--- a/Engine/Source/Programs/Horde/Horde.Server/Logs/LogEventMessages.cs
+++ b/Engine/Source/Programs/Horde/Horde.Server/Logs/LogEventMessages.cs
@@ -55,7 +55,7 @@ namespace Horde.Server.Logs
 			LineIndex = logEvent.LineIndex;
 			LineCount = logEvent.LineCount;
 			IssueId = issueId;
-			Lines = eventData.Lines.ConvertAll(x => x.Data);
+			Lines = eventData.Lines.ConvertAll(x => JsonDocument.Parse(x.Data).RootElement);
 		}
 	}
 }
diff --git a/Engine/Source/Programs/Horde/Horde.Server/Logs/LogFileService.cs b/Engine/Source/Programs/Horde/Horde.Server/Logs/LogFileService.cs
--- a/Engine/Source/Programs/Horde/Horde.Server/Logs/LogFileService.cs
+++ b/Engine/Source/Programs/Horde/Horde.Server/Logs/LogFileService.cs
@@ -969,63 +969,26 @@ namespace Horde.Server.Logs
 			return _logEvents.FindAsync(logFile.Id, spanId, index, count);
 		}
 
-		class LogEventLine : ILogEventLine
+		class LogEventData : ILogEventData
 		{
-			readonly LogLevel _level;
-			public EventId? EventId { get; }
-			public string Message { get; }
-			public JsonElement Data { get; }
+			public string? _message;
+			public IReadOnlyList<JsonLogEvent> Lines { get; }
 
-			LogLevel ILogEventLine.Level => _level;
+			EventId? ILogEventData.EventId => (Lines.Count > 0) ? Lines[0].EventId : null;
+			EventSeverity ILogEventData.Severity => (Lines.Count == 0) ? EventSeverity.Information : (Lines[0].Level == LogLevel.Warning) ? EventSeverity.Warning : EventSeverity.Error;
 
-			public LogEventLine(ReadOnlySpan<byte> data)
-				: this(JsonSerializer.Deserialize<JsonElement>(data))
+			public LogEventData(IReadOnlyList<JsonLogEvent> lines)
 			{
+				Lines = lines;
 			}
 
-			public LogEventLine(JsonElement data)
+			string ILogEventData.Message
 			{
-				Data = data;
-
-				JsonElement levelElement;
-				if (!data.TryGetProperty("level", out levelElement) || !Enum.TryParse(levelElement.GetString(), out _level))
-				{
-					_level = LogLevel.Information;
-				}
-
-				JsonElement idElement;
-				if (data.TryGetProperty("id", out idElement))
+				get
 				{
-					int idValue;
-					if (idElement.TryGetInt32(out idValue))
-					{
-						EventId = idValue;
-					}
+					_message ??= String.Join("\n", Lines.Select(x => x.GetRenderedMessage().ToString()));
+					return _message;
 				}
-
-				JsonElement messageElement;
-				if (data.TryGetProperty("renderedMessage", out messageElement) || data.TryGetProperty("message", out messageElement))
-				{
-					Message = messageElement.GetString() ?? "(Invalid)";
-				}
-				else
-				{
-					Message = "(Missing message or renderedMessage field)";
-				}
-			}
-		}
-
-		class LogEventData : ILogEventData
-		{
-			public IReadOnlyList<ILogEventLine> Lines { get; }
-
-			EventId? ILogEventData.EventId => (Lines.Count > 0) ? Lines[0].EventId : null;
-			EventSeverity ILogEventData.Severity => (Lines.Count == 0) ? EventSeverity.Information : (Lines[0].Level == LogLevel.Warning) ? EventSeverity.Warning : EventSeverity.Error;
-			string ILogEventData.Message => String.Join("\n", Lines.Select(x => x.Message));
-
-			public LogEventData(IReadOnlyList<ILogEventLine> lines)
-			{
-				Lines = lines;
 			}
 		}
 
@@ -1050,13 +1013,13 @@ namespace Horde.Server.Logs
 			span.SetAttribute("lineCount", lineCount);
 
 			List<Utf8String> lines = await ReadLinesAsync(logFile, lineIndex, lineCount, cancellationToken);
-			List<LogEventLine> eventLines = new List<LogEventLine>(lines.Count);
+			List<JsonLogEvent> jsonLines = new List<JsonLogEvent>(lines.Count);
 
 			foreach (Utf8String line in lines)
 			{
 				try
 				{
-					eventLines.Add(new LogEventLine(line.Span));
+					jsonLines.Add(JsonLogEvent.Parse(line.Memory));
 				}
 				catch (JsonException ex)
 				{
@@ -1064,7 +1027,7 @@ namespace Horde.Server.Logs
 				}
 			}
 
-			return new LogEventData(eventLines);
+			return new LogEventData(jsonLines);
 		}
 
 		/// <inheritdoc/>
diff --git a/Engine/Source/Programs/Horde/Horde.Server/Utilities/JsonElementExtensions.cs b/Engine/Source/Programs/Horde/Horde.Server/Utilities/JsonElementExtensions.cs
deleted file mode 100644
diff --git a/Engine/Source/Programs/Shared/EpicGames.Core/JsonElementExtensions.cs b/Engine/Source/Programs/Shared/EpicGames.Core/JsonElementExtensions.cs
new file mode 100644
--- /dev/null
+++ b/Engine/Source/Programs/Shared/EpicGames.Core/JsonElementExtensions.cs
@@ -0,0 +1,202 @@
+// Copyright Epic Games, Inc. All Rights Reserved.
+
+using System;
+using System.Diagnostics.CodeAnalysis;
+using System.Text;
+using System.Text.Json;
+
+namespace EpicGames.Core
+{
+	/// <summary>
+	/// Extensions for parsing values out of generic dictionary objects
+	/// </summary>
+	public static class JsonElementExtensions
+	{
+		/// <summary>
+		/// Checks if the element has a property with the given value
+		/// </summary>
+		/// <param name="element">Element to search</param>
+		/// <param name="name">Name of the property</param>
+		/// <param name="value">Expected value of the property</param>
+		/// <returns>True if the property exists and matches</returns>
+		public static bool HasStringProperty(this JsonElement element, ReadOnlySpan<byte> name, string value)
+		{
+			JsonElement property;
+			return element.ValueKind == JsonValueKind.Object && element.TryGetProperty(name, out property) && property.ValueEquals(value);
+		}
+
+		/// <summary>
+		/// Checks if the element has a property with the given value
+		/// </summary>
+		/// <param name="element">Element to search</param>
+		/// <param name="name">Name of the property</param>
+		/// <param name="value">Expected value of the property</param>
+		/// <returns>True if the property exists and matches</returns>
+		public static bool HasStringProperty(this JsonElement element, string name, string value)
+		{
+			JsonElement property;
+			return element.ValueKind == JsonValueKind.Object && element.TryGetProperty(name, out property) && property.ValueEquals(value);
+		}
+
+		/// <summary>
+		/// Gets a property value of a certain type
+		/// </summary>
+		/// <param name="element">The element to get a property from</param>
+		/// <param name="name">Name of the element</param>
+		/// <param name="valueKind">The required type of property</param>
+		/// <param name="value">Value of the property</param>
+		/// <returns>True if the property exists and was a string</returns>
+		public static bool TryGetProperty(this JsonElement element, ReadOnlySpan<byte> name, JsonValueKind valueKind, [NotNullWhen(true)] out JsonElement value)
+		{
+			JsonElement property;
+			if (element.TryGetProperty(name, out property) && property.ValueKind == valueKind)
+			{
+				value = property;
+				return true;
+			}
+			else
+			{
+				value = new JsonElement();
+				return false;
+			}
+		}
+
+		/// <summary>
+		/// Gets a property value of a certain type
+		/// </summary>
+		/// <param name="element">The element to get a property from</param>
+		/// <param name="name">Name of the element</param>
+		/// <param name="valueKind">The required type of property</param>
+		/// <param name="value">Value of the property</param>
+		/// <returns>True if the property exists and was a string</returns>
+		public static bool TryGetProperty(this JsonElement element, string name, JsonValueKind valueKind, [NotNullWhen(true)] out JsonElement value)
+		{
+			JsonElement property;
+			if (element.TryGetProperty(name, out property) && property.ValueKind == valueKind)
+			{
+				value = property;
+				return true;
+			}
+			else
+			{
+				value = new JsonElement();
+				return false;
+			}
+		}
+
+		/// <summary>
+		/// Gets a string property value
+		/// </summary>
+		/// <param name="element">The element to get a property from</param>
+		/// <param name="name">Name of the element</param>
+		/// <param name="value">Value of the property</param>
+		/// <returns>True if the property exists and was a string</returns>
+		public static bool TryGetStringProperty(this JsonElement element, ReadOnlySpan<byte> name, [NotNullWhen(true)] out string? value)
+		{
+			JsonElement property;
+			if (element.TryGetProperty(name, out property) && property.ValueKind == JsonValueKind.String)
+			{
+				value = property.GetString()!;
+				return true;
+			}
+			else
+			{
+				value = null;
+				return false;
+			}
+		}
+
+		/// <summary>
+		/// Gets a string property value
+		/// </summary>
+		/// <param name="element">The element to get a property from</param>
+		/// <param name="name">Name of the element</param>
+		/// <param name="value">Value of the property</param>
+		/// <returns>True if the property exists and was a string</returns>
+		public static bool TryGetStringProperty(this JsonElement element, string name, [NotNullWhen(true)] out string? value)
+		{
+			return TryGetStringProperty(element, Encoding.UTF8.GetBytes(name).AsSpan(), out value);
+		}
+
+		/// <summary>
+		/// Gets a property value from a document or subdocument, indicated with dotted notation
+		/// </summary>
+		/// <param name="element">Document to get a property for</param>
+		/// <param name="name">Name of the property</param>
+		/// <param name="outElement">Receives the nexted element</param>
+		/// <returns>True if the property exists and was of the correct type</returns>
+		public static bool TryGetNestedProperty(this JsonElement element, ReadOnlySpan<byte> name, [NotNullWhen(true)] out JsonElement outElement)
+		{
+			int dotIdx = name.IndexOf((byte)'.');
+			if (dotIdx == -1)
+			{
+				return element.TryGetProperty(name, out outElement);
+			}
+
+			JsonElement docValue;
+			if (element.TryGetProperty(name.Slice(0, dotIdx), out docValue) && docValue.ValueKind == JsonValueKind.Object)
+			{
+				return TryGetNestedProperty(docValue, name.Slice(dotIdx + 1), out outElement);
+			}
+
+			outElement = new JsonElement();
+			return false;
+		}
+
+		/// <summary>
+		/// Gets a property value from a document or subdocument, indicated with dotted notation
+		/// </summary>
+		/// <param name="element">Document to get a property for</param>
+		/// <param name="name">Name of the property</param>
+		/// <param name="outElement">Receives the nexted element</param>
+		/// <returns>True if the property exists and was of the correct type</returns>
+		public static bool TryGetNestedProperty(this JsonElement element, string name, [NotNullWhen(true)] out JsonElement outElement)
+		{
+			return TryGetNestedProperty(element, Encoding.UTF8.GetBytes(name).AsSpan(), out outElement);
+		}
+
+		/// <summary>
+		/// Gets an int32 value from the document
+		/// </summary>
+		/// <param name="element">Document to get a property for</param>
+		/// <param name="name">Name of the property</param>
+		/// <param name="outValue">Receives the property value</param>
+		/// <returns>True if the property was retrieved</returns>
+		public static bool TryGetNestedProperty(this JsonElement element, string name, out int outValue)
+		{
+			JsonElement value;
+			if (element.TryGetNestedProperty(name, out value) && value.ValueKind == JsonValueKind.Number)
+			{
+				outValue = value.GetInt32();
+				return true;
+			}
+			else
+			{
+				outValue = 0;
+				return false;
+			}
+		}
+
+		/// <summary>
+		/// Gets a string value from the document
+		/// </summary>
+		/// <param name="element">Document to get a property for</param>
+		/// <param name="name">Name of the property</param>
+		/// <param name="outValue">Receives the property value</param>
+		/// <returns>True if the property was retrieved</returns>
+		public static bool TryGetNestedProperty(this JsonElement element, string name, [NotNullWhen(true)] out string? outValue)
+		{
+			JsonElement value;
+			if (element.TryGetNestedProperty(name, out value) && value.ValueKind == JsonValueKind.String)
+			{
+				outValue = value.GetString();
+				return outValue != null;
+			}
+			else
+			{
+				outValue = null;
+				return false;
+			}
+		}
+	}
+}
diff --git a/Engine/Source/Programs/Shared/EpicGames.Core/JsonLogEvent.cs b/Engine/Source/Programs/Shared/EpicGames.Core/JsonLogEvent.cs
--- a/Engine/Source/Programs/Shared/EpicGames.Core/JsonLogEvent.cs
+++ b/Engine/Source/Programs/Shared/EpicGames.Core/JsonLogEvent.cs
@@ -186,6 +186,30 @@ namespace EpicGames.Core
 
 		static readonly Utf8String s_newlineEscaped = "\\n";
 
+		/// <summary>
+		/// Gets the rendered message from the event data
+		/// </summary>
+		public Utf8String GetRenderedMessage()
+		{
+			Utf8JsonReader reader = new Utf8JsonReader(Data.Span);
+			if (reader.Read() && reader.TokenType == JsonTokenType.StartObject)
+			{
+				while (reader.Read() && reader.TokenType == JsonTokenType.PropertyName)
+				{
+					ReadOnlySpan<byte> propertyName = reader.ValueSpan;
+					if (!reader.Read())
+					{
+						break;
+					}
+					else if (propertyName.SequenceEqual(LogEventPropertyName.Message) && reader.TokenType == JsonTokenType.String)
+					{
+						return new Utf8String(reader.GetUtf8String().ToArray());
+					}
+				}
+			}
+			return Utf8String.Empty;
+		}
+
 		/// <summary>
 		/// Count the number of lines in the message field of a log event
 		/// </summary>
@@ -248,6 +272,51 @@ namespace EpicGames.Core
 			return LogEvent.Read(state.Data.Span).ToString();
 		}
 
+		/// <summary>
+		/// Find all properties of the given type in a particular log line
+		/// </summary>
+		/// <param name="type">Type of property to return</param>
+		/// <returns></returns>
+		public IEnumerable<JsonProperty> FindPropertiesOfType(Utf8String type)
+		{
+			JsonDocument document = JsonDocument.Parse(Data);
+			return FindPropertiesOfType(document.RootElement, type);
+		}
+
+		/// <summary>
+		/// Find all properties of the given type in a particular log line
+		/// </summary>
+		/// <param name="line">Line data</param>
+		/// <param name="type">Type of property to return</param>
+		/// <returns></returns>
+		public static IEnumerable<JsonProperty> FindPropertiesOfType(JsonElement line, Utf8String type)
+		{
+			JsonElement properties;
+			if (line.TryGetProperty("properties", out properties) && properties.ValueKind == JsonValueKind.Object)
+			{
+				foreach (JsonProperty property in properties.EnumerateObject())
+				{
+					if (property.Value.ValueKind == JsonValueKind.Object)
+					{
+						foreach (JsonProperty subProperty in property.Value.EnumerateObject())
+						{
+							if (subProperty.NameEquals(LogEventPropertyName.Type.Span))
+							{
+								if (subProperty.Value.ValueKind == JsonValueKind.String && subProperty.Value.ValueEquals(type.Span))
+								{
+									yield return property;
+								}
+								else
+								{
+									break;
+								}
+							}
+						}
+					}
+				}
+			}
+		}
+
 		/// <inheritdoc/>
 		public override string ToString() => Encoding.UTF8.GetString(Data.ToArray());
 
