Horde: Move issue handlers into EpicGames.Horde.

diff --git a/Engine/Source/Programs/Horde/Horde.Server.Tests/IssueServiceTests.cs b/Engine/Source/Programs/Horde/Horde.Server.Tests/IssueServiceTests.cs
--- a/Engine/Source/Programs/Horde/Horde.Server.Tests/IssueServiceTests.cs
+++ b/Engine/Source/Programs/Horde/Horde.Server.Tests/IssueServiceTests.cs
@@ -9,6 +9,7 @@ using System.Runtime.InteropServices;
 using System.Text.Json;
 using System.Threading.Tasks;
 using EpicGames.Core;
+using EpicGames.Horde.Issues;
 using HordeAgent.Horde.Agent.Parser;
 using HordeAgent.Horde.Agent.Utility;
 using Horde.Server.Jobs;
diff --git a/Engine/Source/Programs/Horde/Horde.Server/Issues/Handlers/BuildGraphIssueHandler.cs b/Engine/Source/Programs/Horde/Horde.Server/Issues/Handlers/BuildGraphIssueHandler.cs
deleted file mode 100644
diff --git a/Engine/Source/Programs/Horde/Horde.Server/Issues/Handlers/CompileIssueHandler.cs b/Engine/Source/Programs/Horde/Horde.Server/Issues/Handlers/CompileIssueHandler.cs
deleted file mode 100644
diff --git a/Engine/Source/Programs/Horde/Horde.Server/Issues/Handlers/ContentIssueHandler.cs b/Engine/Source/Programs/Horde/Horde.Server/Issues/Handlers/ContentIssueHandler.cs
deleted file mode 100644
diff --git a/Engine/Source/Programs/Horde/Horde.Server/Issues/Handlers/CopyrightIssueHandler.cs b/Engine/Source/Programs/Horde/Horde.Server/Issues/Handlers/CopyrightIssueHandler.cs
deleted file mode 100644
diff --git a/Engine/Source/Programs/Horde/Horde.Server/Issues/Handlers/DefaultIssueHandler.cs b/Engine/Source/Programs/Horde/Horde.Server/Issues/Handlers/DefaultIssueHandler.cs
deleted file mode 100644
diff --git a/Engine/Source/Programs/Horde/Horde.Server/Issues/Handlers/GauntletIssueHandler.cs b/Engine/Source/Programs/Horde/Horde.Server/Issues/Handlers/GauntletIssueHandler.cs
deleted file mode 100644
diff --git a/Engine/Source/Programs/Horde/Horde.Server/Issues/Handlers/HashedIssueHandler.cs b/Engine/Source/Programs/Horde/Horde.Server/Issues/Handlers/HashedIssueHandler.cs
deleted file mode 100644
diff --git a/Engine/Source/Programs/Horde/Horde.Server/Issues/Handlers/LocalizationIssueHandler.cs b/Engine/Source/Programs/Horde/Horde.Server/Issues/Handlers/LocalizationIssueHandler.cs
deleted file mode 100644
diff --git a/Engine/Source/Programs/Horde/Horde.Server/Issues/Handlers/PerforceCaseIssueHandler.cs b/Engine/Source/Programs/Horde/Horde.Server/Issues/Handlers/PerforceCaseIssueHandler.cs
deleted file mode 100644
diff --git a/Engine/Source/Programs/Horde/Horde.Server/Issues/Handlers/ScopedIssueHandler.cs b/Engine/Source/Programs/Horde/Horde.Server/Issues/Handlers/ScopedIssueHandler.cs
deleted file mode 100644
diff --git a/Engine/Source/Programs/Horde/Horde.Server/Issues/Handlers/ShaderIssueHandler.cs b/Engine/Source/Programs/Horde/Horde.Server/Issues/Handlers/ShaderIssueHandler.cs
deleted file mode 100644
diff --git a/Engine/Source/Programs/Horde/Horde.Server/Issues/Handlers/SymbolIssueHandler.cs b/Engine/Source/Programs/Horde/Horde.Server/Issues/Handlers/SymbolIssueHandler.cs
deleted file mode 100644
diff --git a/Engine/Source/Programs/Horde/Horde.Server/Issues/Handlers/SystemicIssueHandler.cs b/Engine/Source/Programs/Horde/Horde.Server/Issues/Handlers/SystemicIssueHandler.cs
deleted file mode 100644
diff --git a/Engine/Source/Programs/Horde/Horde.Server/Issues/Handlers/UnacceptableWordsIssueHandler.cs b/Engine/Source/Programs/Horde/Horde.Server/Issues/Handlers/UnacceptableWordsIssueHandler.cs
deleted file mode 100644
diff --git a/Engine/Source/Programs/Horde/Horde.Server/Issues/IIssueCollection.cs b/Engine/Source/Programs/Horde/Horde.Server/Issues/IIssueCollection.cs
--- a/Engine/Source/Programs/Horde/Horde.Server/Issues/IIssueCollection.cs
+++ b/Engine/Source/Programs/Horde/Horde.Server/Issues/IIssueCollection.cs
@@ -13,6 +13,7 @@ using Horde.Server.Streams;
 using MongoDB.Bson;
 using MongoDB.Driver;
 using EpicGames.Horde.Api;
+using EpicGames.Horde.Issues;
 
 namespace Horde.Server.Issues
 {
diff --git a/Engine/Source/Programs/Horde/Horde.Server/Issues/IIssueFingerprint.cs b/Engine/Source/Programs/Horde/Horde.Server/Issues/IIssueFingerprint.cs
--- a/Engine/Source/Programs/Horde/Horde.Server/Issues/IIssueFingerprint.cs
+++ b/Engine/Source/Programs/Horde/Horde.Server/Issues/IIssueFingerprint.cs
@@ -3,6 +3,7 @@
 using System;
 using System.Collections.Generic;
 using System.Linq;
+using EpicGames.Horde.Issues;
 
 namespace Horde.Server.Issues
 {
diff --git a/Engine/Source/Programs/Horde/Horde.Server/Issues/IssueChangeFilter.cs b/Engine/Source/Programs/Horde/Horde.Server/Issues/IssueChangeFilter.cs
deleted file mode 100644
diff --git a/Engine/Source/Programs/Horde/Horde.Server/Issues/IssueCollection.cs b/Engine/Source/Programs/Horde/Horde.Server/Issues/IssueCollection.cs
--- a/Engine/Source/Programs/Horde/Horde.Server/Issues/IssueCollection.cs
+++ b/Engine/Source/Programs/Horde/Horde.Server/Issues/IssueCollection.cs
@@ -9,6 +9,7 @@ using System.Linq;
 using System.Threading.Tasks;
 using EpicGames.Core;
 using EpicGames.Horde.Api;
+using EpicGames.Horde.Issues;
 using EpicGames.Redis.Utility;
 using Horde.Server.Auditing;
 using Horde.Server.Jobs;
diff --git a/Engine/Source/Programs/Horde/Horde.Server/Issues/IssueEvent.cs b/Engine/Source/Programs/Horde/Horde.Server/Issues/IssueEvent.cs
deleted file mode 100644
diff --git a/Engine/Source/Programs/Horde/Horde.Server/Issues/IssueHandler.cs b/Engine/Source/Programs/Horde/Horde.Server/Issues/IssueHandler.cs
deleted file mode 100644
diff --git a/Engine/Source/Programs/Horde/Horde.Server/Issues/IssueKey.cs b/Engine/Source/Programs/Horde/Horde.Server/Issues/IssueKey.cs
deleted file mode 100644
diff --git a/Engine/Source/Programs/Horde/Horde.Server/Issues/IssueMetadata.cs b/Engine/Source/Programs/Horde/Horde.Server/Issues/IssueMetadata.cs
deleted file mode 100644
diff --git a/Engine/Source/Programs/Horde/Horde.Server/Issues/IssueService.cs b/Engine/Source/Programs/Horde/Horde.Server/Issues/IssueService.cs
--- a/Engine/Source/Programs/Horde/Horde.Server/Issues/IssueService.cs
+++ b/Engine/Source/Programs/Horde/Horde.Server/Issues/IssueService.cs
@@ -11,7 +11,7 @@ using System.Threading;
 using System.Threading.Tasks;
 using EpicGames.Core;
 using EpicGames.Horde.Api;
-using Horde.Server.Issues.Handlers;
+using EpicGames.Horde.Issues;
 using Horde.Server.Jobs;
 using Horde.Server.Jobs.Graphs;
 using Horde.Server.Jobs.Templates;
@@ -275,7 +275,7 @@ namespace Horde.Server.Issues
 
 			// Find all the issue handler types
 			List<(Type Type, IssueHandlerAttribute Attribute)> handlerTypes = new List<(Type, IssueHandlerAttribute)>();
-			foreach (Type type in Assembly.GetExecutingAssembly().GetTypes())
+			foreach (Type type in typeof(IssueHandler).Assembly.GetTypes())
 			{
 				IssueHandlerAttribute? attribute = type.GetCustomAttribute<IssueHandlerAttribute>();
 				if (attribute != null)
@@ -1437,7 +1437,7 @@ namespace Horde.Server.Issues
 			}
 			if (name.Equals("LegacySymbolIssueHandler", StringComparison.OrdinalIgnoreCase))
 			{
-				summary.Append(SymbolIssueHandler.GetSummaryStatic(fingerprint, severity));
+				summary.Append(GetSymbolIssueSummary(fingerprint, severity));
 				return true;
 			}
 
@@ -1452,6 +1452,38 @@ namespace Horde.Server.Issues
 			return false;
 		}
 
+		static readonly IssueMetadata s_duplicateEventIdMetadata = new IssueMetadata("EventId", KnownLogEvents.Linker_DuplicateSymbol.Id.ToString());
+
+		static string GetSymbolIssueSummary(IIssueFingerprint fingerprint, IssueSeverity severity)
+		{
+			HashSet<string> symbols = new HashSet<string>(fingerprint.Keys.Where(x => x.Type == IssueKeyType.Symbol).Select(x => x.Name));
+			if (symbols.Count == 0)
+			{
+				string[] nodes = fingerprint.Metadata?.FindValues("Node").ToArray() ?? Array.Empty<string>();
+
+				StringBuilder summary = new StringBuilder("Linker ");
+				summary.Append((severity == IssueSeverity.Warning) ? "warnings" : "errors");
+				if (nodes.Length > 0)
+				{
+					summary.Append($" in {StringUtils.FormatList(nodes, 2)}");
+				}
+
+				return summary.ToString();
+			}
+			else
+			{
+				string problemType = (fingerprint.Metadata?.Contains(s_duplicateEventIdMetadata) == true) ? "Duplicate" : "Undefined";
+				if (symbols.Count == 1)
+				{
+					return $"{problemType} symbol '{symbols.First()}'";
+				}
+				else
+				{
+					return $"{problemType} symbols: {StringUtils.FormatList(symbols.ToArray(), 3)}";
+				}
+			}
+		}
+
 		/// <summary>
 		/// Update the sentinels for the given list of issues
 		/// </summary>
diff --git a/Engine/Source/Programs/Horde/Horde.Server/Jobs/Templates/TemplateId.cs b/Engine/Source/Programs/Horde/Horde.Server/Jobs/Templates/TemplateId.cs
deleted file mode 100644
diff --git a/Engine/Source/Programs/Horde/Horde.Server/Streams/StreamId.cs b/Engine/Source/Programs/Horde/Horde.Server/Streams/StreamId.cs
deleted file mode 100644
diff --git a/Engine/Source/Programs/Shared/EpicGames.Horde/Issues/Handlers/BuildGraphIssueHandler.cs b/Engine/Source/Programs/Shared/EpicGames.Horde/Issues/Handlers/BuildGraphIssueHandler.cs
new file mode 100644
--- /dev/null
+++ b/Engine/Source/Programs/Shared/EpicGames.Horde/Issues/Handlers/BuildGraphIssueHandler.cs
@@ -0,0 +1,55 @@
+// Copyright Epic Games, Inc. All Rights Reserved.
+
+using System.Collections.Generic;
+using EpicGames.Core;
+using Microsoft.Extensions.Logging;
+
+namespace EpicGames.Horde.Issues.Handlers
+{
+	/// <summary>
+	/// Instance of a particular BuildGraph script error
+	/// </summary>
+	[IssueHandler(Priority = 10)]
+	public class BuildGraphIssueHandler : IssueHandler
+	{
+		readonly List<IssueEventGroup> _issues = new List<IssueEventGroup>();
+
+		/// <summary>
+		/// Determines if the given event id matches
+		/// </summary>
+		/// <param name="eventId">The event id to compare</param>
+		/// <returns>True if the given event id matches</returns>
+		public static bool IsMatchingEventId(EventId eventId)
+		{
+			return eventId == KnownLogEvents.AutomationTool_BuildGraphScript;
+		}
+
+		static bool IsMaskedEventId(EventId id) => id == KnownLogEvents.ExitCode || id == KnownLogEvents.Systemic_Xge_BuildFailed;
+
+		/// <inheritdoc/>
+		public override bool HandleEvent(IssueEvent issueEvent)
+		{
+			if (issueEvent.EventId.HasValue)
+			{
+				EventId eventId = issueEvent.EventId.Value;
+				if (IsMatchingEventId(eventId))
+				{
+					IssueEventGroup issue = new IssueEventGroup("BuildGraph", "BuildGraph {Severity} in {Files}", IssueChangeFilter.All);
+					issue.Events.Add(issueEvent);
+					issue.Keys.AddSourceFiles(issueEvent);
+					_issues.Add(issue);
+
+					return true;
+				}
+				else if (_issues.Count > 0 && IsMaskedEventId(eventId))
+				{
+					return true;
+				}
+			}
+			return false;
+		}
+
+		/// <inheritdoc/>
+		public override IEnumerable<IssueEventGroup> GetIssues() => _issues;
+	}
+}
diff --git a/Engine/Source/Programs/Shared/EpicGames.Horde/Issues/Handlers/CompileIssueHandler.cs b/Engine/Source/Programs/Shared/EpicGames.Horde/Issues/Handlers/CompileIssueHandler.cs
new file mode 100644
--- /dev/null
+++ b/Engine/Source/Programs/Shared/EpicGames.Horde/Issues/Handlers/CompileIssueHandler.cs
@@ -0,0 +1,87 @@
+// Copyright Epic Games, Inc. All Rights Reserved.
+
+using System.Collections.Generic;
+using EpicGames.Core;
+using Microsoft.Extensions.Logging;
+
+namespace EpicGames.Horde.Issues.Handlers
+{
+	/// <summary>
+	/// Instance of a particular compile error
+	/// </summary>
+	[IssueHandler(Priority = 10)]
+	public class CompileIssueHandler : IssueHandler
+	{
+		/// <summary>
+		/// Annotation describing the compile type
+		/// </summary>
+		const string CompileTypeAnnotation = "CompileType";
+
+		/// <summary>
+		/// Annotation specifying a group for compile issues from this node
+		/// </summary>
+		const string CompileGroupAnnotation = "CompileGroup";
+
+		readonly IssueHandlerContext _context;
+		readonly List<IssueEventGroup> _issues = new List<IssueEventGroup>();
+
+		/// <summary>
+		/// Constructor
+		/// </summary>
+		public CompileIssueHandler(IssueHandlerContext context)
+		{
+			_context = context;
+		}
+
+		/// <summary>
+		/// Determines if the given event id matches
+		/// </summary>
+		/// <param name="eventId">The event id to compare</param>
+		/// <returns>True if the given event id matches</returns>
+		public static bool IsMatchingEventId(EventId eventId)
+		{
+			return eventId == KnownLogEvents.Compiler || eventId == KnownLogEvents.UHT || eventId == KnownLogEvents.AutomationTool_SourceFileLine || eventId == KnownLogEvents.MSBuild;
+		}
+
+		static bool IsMaskedEventId(EventId id) => id == KnownLogEvents.ExitCode || id == KnownLogEvents.Systemic_Xge_BuildFailed || id == KnownLogEvents.Compiler_Summary;
+
+		/// <inheritdoc/>
+		public override bool HandleEvent(IssueEvent issueEvent)
+		{
+			if (issueEvent.EventId.HasValue)
+			{
+				EventId eventId = issueEvent.EventId.Value;
+				if (IsMatchingEventId(eventId))
+				{
+					string compileType = "Compile";
+					if (_context.NodeAnnotations.TryGetValue(CompileTypeAnnotation, out string? type))
+					{
+						compileType = type;
+					}
+
+					string fingerprintType = "Compile";
+					if (_context.NodeAnnotations.TryGetValue(CompileGroupAnnotation, out string? group))
+					{
+						fingerprintType = $"{fingerprintType}:{group}";
+					}
+
+					IssueEventGroup issue = new IssueEventGroup(fingerprintType, "{Meta:CompileType} {Severity} in {Files}", IssueChangeFilter.Code);
+					issue.Events.Add(issueEvent);
+					issue.Keys.AddSourceFiles(issueEvent);
+					issue.Metadata.Add(CompileTypeAnnotation, compileType);
+					_issues.Add(issue);
+
+					return true;
+				}
+				else if (_issues.Count > 0 && IsMaskedEventId(eventId))
+				{
+					return true;
+				}
+			}
+			return false;
+		}
+
+		/// <inheritdoc/>
+		public override IEnumerable<IssueEventGroup> GetIssues() => _issues;
+	}
+}
diff --git a/Engine/Source/Programs/Shared/EpicGames.Horde/Issues/Handlers/ContentIssueHandler.cs b/Engine/Source/Programs/Shared/EpicGames.Horde/Issues/Handlers/ContentIssueHandler.cs
new file mode 100644
--- /dev/null
+++ b/Engine/Source/Programs/Shared/EpicGames.Horde/Issues/Handlers/ContentIssueHandler.cs
@@ -0,0 +1,46 @@
+// Copyright Epic Games, Inc. All Rights Reserved.
+
+using System.Collections.Generic;
+using EpicGames.Core;
+using Microsoft.Extensions.Logging;
+
+namespace EpicGames.Horde.Issues.Handlers
+{
+	/// <summary>
+	/// Instance of a particular compile error
+	/// </summary>
+	[IssueHandler(Priority = 10)]
+	public class ContentIssueHandler : IssueHandler
+	{
+		readonly List<IssueEventGroup> _issues = new List<IssueEventGroup>();
+
+		static bool IsMatchingEventId(EventId eventId) => eventId == KnownLogEvents.Engine_AssetLog;
+		static bool IsMaskedEventId(EventId eventId) => eventId == KnownLogEvents.ExitCode;
+
+		/// <inheritdoc/>
+		public override bool HandleEvent(IssueEvent issueEvent)
+		{
+			if (issueEvent.EventId != null)
+			{
+				EventId eventId = issueEvent.EventId.Value;
+				if (IsMatchingEventId(eventId))
+				{
+					IssueEventGroup issue = new IssueEventGroup("Content", "{Severity} in {Files}", IssueChangeFilter.Content);
+					issue.Events.Add(issueEvent);
+					issue.Keys.AddAssets(issueEvent);
+					_issues.Add(issue);
+
+					return true;
+				}
+				else if (_issues.Count > 0 && IsMaskedEventId(eventId))
+				{
+					return true;
+				}
+			}
+			return false;
+		}
+
+		/// <inheritdoc/>
+		public override IEnumerable<IssueEventGroup> GetIssues() => _issues;
+	}
+}
diff --git a/Engine/Source/Programs/Shared/EpicGames.Horde/Issues/Handlers/CopyrightIssueHandler.cs b/Engine/Source/Programs/Shared/EpicGames.Horde/Issues/Handlers/CopyrightIssueHandler.cs
new file mode 100644
--- /dev/null
+++ b/Engine/Source/Programs/Shared/EpicGames.Horde/Issues/Handlers/CopyrightIssueHandler.cs
@@ -0,0 +1,34 @@
+// Copyright Epic Games, Inc. All Rights Reserved.
+
+using System.Collections.Generic;
+using EpicGames.Core;
+
+namespace EpicGames.Horde.Issues.Handlers
+{
+	/// <summary>
+	/// Instance of a particular compile error
+	/// </summary>
+	[IssueHandler(Priority = 10)]
+	public class CopyrightIssueHandler : IssueHandler
+	{
+		readonly List<IssueEventGroup> _issues = new List<IssueEventGroup>();
+
+		/// <inheritdoc/>
+		public override bool HandleEvent(IssueEvent issueEvent)
+		{
+			if (issueEvent.EventId == KnownLogEvents.AutomationTool_MissingCopyright)
+			{
+				IssueEventGroup issue = new IssueEventGroup("Copyright", "Missing copyright notice in {Files}", IssueChangeFilter.Code);
+				issue.Events.Add(issueEvent);
+				issue.Keys.AddSourceFiles(issueEvent);
+				_issues.Add(issue);
+
+				return true;
+			}
+			return false;
+		}
+
+		/// <inheritdoc/>
+		public override IEnumerable<IssueEventGroup> GetIssues() => _issues;
+	}
+}
diff --git a/Engine/Source/Programs/Shared/EpicGames.Horde/Issues/Handlers/DefaultIssueHandler.cs b/Engine/Source/Programs/Shared/EpicGames.Horde/Issues/Handlers/DefaultIssueHandler.cs
new file mode 100644
--- /dev/null
+++ b/Engine/Source/Programs/Shared/EpicGames.Horde/Issues/Handlers/DefaultIssueHandler.cs
@@ -0,0 +1,44 @@
+// Copyright Epic Games, Inc. All Rights Reserved.
+
+using System.Collections.Generic;
+
+namespace EpicGames.Horde.Issues.Handlers
+{
+	/// <summary>
+	/// Default handler for log events not matched by any other handler
+	/// </summary>
+	[IssueHandler(Priority = 0)]
+	public class DefaultIssueHandler : IssueHandler
+	{
+		readonly IssueHandlerContext _context;
+		IssueEventGroup? _issue;
+
+		/// <summary>
+		/// Constructor
+		/// </summary>
+		public DefaultIssueHandler(IssueHandlerContext context) => _context = context;
+
+		/// <inheritdoc/>
+		public override bool HandleEvent(IssueEvent issueEvent)
+		{
+			if (_issue == null)
+			{
+				_issue = new IssueEventGroup("Default", "{Severity} in {Meta:Node}", IssueChangeFilter.All);
+				_issue.Metadata.Add("Node", _context.NodeName);
+				_issue.Keys.Add(IssueKey.FromStep(_context.StreamId, _context.TemplateId, _context.NodeName));
+			}
+
+			_issue.Events.Add(issueEvent);
+			return true;
+		}
+
+		/// <inheritdoc/>
+		public override IEnumerable<IssueEventGroup> GetIssues()
+		{
+			if (_issue != null)
+			{
+				yield return _issue;
+			}
+		}
+	}
+}
diff --git a/Engine/Source/Programs/Shared/EpicGames.Horde/Issues/Handlers/GauntletIssueHandler.cs b/Engine/Source/Programs/Shared/EpicGames.Horde/Issues/Handlers/GauntletIssueHandler.cs
new file mode 100644
--- /dev/null
+++ b/Engine/Source/Programs/Shared/EpicGames.Horde/Issues/Handlers/GauntletIssueHandler.cs
@@ -0,0 +1,219 @@
+// Copyright Epic Games, Inc. All Rights Reserved.
+
+using System.Collections.Generic;
+using System.Text;
+using System.Text.Json;
+using System.Text.RegularExpressions;
+using EpicGames.Core;
+using Microsoft.Extensions.Logging;
+
+namespace EpicGames.Horde.Issues.Handlers
+{
+	/// <summary>
+	/// Instance of a particular Gauntlet error
+	/// </summary>
+	[IssueHandler(Priority = 10)]
+	public class GauntletIssueHandler : IssueHandler
+	{
+		/// <summary>
+		/// Prefix for framework keys
+		/// </summary>
+		const string FrameworkPrefix = "framework";
+
+		/// <summary>
+		/// Prefix for test keys
+		/// </summary>
+		const string TestPrefix = "test";
+
+		/// <summary>
+		/// Prefix for device keys
+		/// </summary>
+		const string DevicePrefix = "device";
+
+		/// <summary>
+		/// Prefix for access keys
+		/// </summary>
+		const string AccessPrefix = "access";
+
+		/// <summary>
+		/// Prefix for fatal failure keys
+		/// </summary>
+		const string FatalPrefix = "fatal";
+
+		/// <summary>
+		/// Max Message Length to hash
+		/// </summary>
+		const int MaxMessageLength = 2000;
+
+		readonly IssueHandlerContext _context;
+		readonly List<IssueEventGroup> _issues = new List<IssueEventGroup>();
+
+		/// <summary>
+		///  Known Gauntlet events
+		/// </summary>
+		static readonly Dictionary<EventId, string> s_knownGauntletEvents = new Dictionary<EventId, string> 
+		{
+			{ KnownLogEvents.Gauntlet, FrameworkPrefix},
+			{ KnownLogEvents.Gauntlet_TestEvent, TestPrefix},
+			{ KnownLogEvents.Gauntlet_DeviceEvent, DevicePrefix},
+			{ KnownLogEvents.Gauntlet_UnrealEngineTestEvent, TestPrefix},
+			{ KnownLogEvents.Gauntlet_BuildDropEvent, AccessPrefix},
+			{ KnownLogEvents.Gauntlet_FatalEvent, FatalPrefix}
+		};
+
+		/// <summary>
+		/// Constructor
+		/// </summary>
+		public GauntletIssueHandler(IssueHandlerContext context) => _context = context;
+
+		/// <summary>
+		/// Determines if the given event id matches
+		/// </summary>
+		/// <param name="eventId">The event id to compare</param>
+		/// <returns>True if the given event id matches</returns>
+		public static bool IsMatchingEventId(EventId eventId)
+		{
+			return s_knownGauntletEvents.ContainsKey(eventId);
+		}
+
+		/// <summary>
+		/// Return the prefix string associate with the event id
+		/// </summary>
+		/// <param name="eventId">The event id to get the information from</param>
+		/// <returns>The corresponding prefix as a string</returns>
+		public static string GetEventPrefix(EventId eventId)
+		{
+			return s_knownGauntletEvents[eventId];
+		}
+
+		/// <summary>
+		/// Parses symbol names from a log event
+		/// </summary>
+		/// <param name="issueEvent">The log event data</param>
+		/// <param name="testNames">Receives a set of the test names</param>
+		/// <param name="metadata"></param>
+		private static void GetNames(IssueEvent issueEvent, HashSet<IssueKey> testNames, HashSet<IssueMetadata> metadata)
+		{
+			foreach (JsonLogEvent line in issueEvent.Lines)
+			{
+				JsonDocument document = JsonDocument.Parse(line.Data);
+
+				string? name = null;
+
+				string? value;
+				if (document.RootElement.TryGetNestedProperty("properties.Name", out value))
+				{
+					name = value;
+				}
+
+				if (name != null)
+				{
+					string prefix = GetEventPrefix(issueEvent.EventId!.Value);
+					testNames.Add($"{prefix}:{name}", IssueKeyType.None);
+					metadata.Add("Context", $"with {name}");
+				}
+			}
+		}
+
+		/// <summary>
+		/// Parses symbol file or directory from a log event
+		/// </summary>
+		/// <param name="issueEvent">The log event data</param>
+		/// <param name="paths">Receives a set of the paths</param>
+		/// <param name="metadata"></param>
+		private static void GetPaths(IssueEvent issueEvent, HashSet<IssueKey> paths, HashSet<IssueMetadata> metadata)
+		{
+			if(issueEvent.EventId == KnownLogEvents.Gauntlet_BuildDropEvent)
+			{
+				foreach (JsonLogEvent line in issueEvent.Lines)
+				{
+					JsonDocument document = JsonDocument.Parse(line.Data);
+
+					string? path = null;
+
+					string? value;
+					if (document.RootElement.TryGetNestedProperty("properties.File", out value))
+					{
+						path = value;
+					}
+					else if (document.RootElement.TryGetNestedProperty("properties.Directory", out value))
+					{
+						path = value;
+					}
+
+					if (path != null)
+					{
+						paths.Add($"{AccessPrefix}:{path}", IssueKeyType.None);
+						metadata.Add("Context", $"with {path}");
+					}
+				}
+			}
+		}
+
+		/// <summary>
+		/// Produce a hash from error message
+		/// </summary>
+		/// <param name="message">The log event message</param>
+		/// <param name="keys">Receives a set of the keys</param>
+		/// <param name="metadata">Receives a set of metadata</param>
+		private void GetHash(string message, HashSet<IssueKey> keys, HashSet<IssueMetadata> metadata)
+		{
+			string error = message.Length > MaxMessageLength? message.Substring(0, MaxMessageLength) : message;
+
+			if (TryGetHash(error, out Md5Hash hash))
+			{
+				keys.Add($"hash:{hash}:stream:{_context.StreamId}", IssueKeyType.None);
+			}
+			else
+			{
+				keys.Add($"{_context.NodeName}", IssueKeyType.None);
+			}
+			metadata.Add("Context", $"in {_context.NodeName}");
+		}
+
+		private static bool TryGetHash(string message, out Md5Hash hash)
+		{
+			string sanitized = message.ToUpperInvariant();
+			sanitized = Regex.Replace(sanitized, @"(?<![a-zA-Z])(?:[A-Z]:|/)[^ :]+[/\\]SYNC[/\\]", "{root}/"); // Redact things that look like workspace roots; may be different between agents
+			sanitized = Regex.Replace(sanitized, @"0[xX][0-9a-fA-F]+", "H"); // Redact hex strings
+			sanitized = Regex.Replace(sanitized, @"\d[\d.,:]*", "n"); // Redact numbers and timestamp like things
+
+			if (sanitized.Length > 30)
+			{
+				hash = Md5Hash.Compute(Encoding.UTF8.GetBytes(sanitized));
+				return true;
+			}
+			else
+			{
+				hash = Md5Hash.Zero;
+				return false;
+			}
+		}
+
+		/// <inheritdoc/>
+		public override bool HandleEvent(IssueEvent issueEvent)
+		{
+			if (issueEvent.EventId != null && IsMatchingEventId(issueEvent.EventId.Value))
+			{
+				IssueEventGroup issue = new IssueEventGroup("Gauntlet", "Gauntlet {Meta:Type} {Severity} {Meta:Context}", IssueChangeFilter.Code);
+				issue.Events.Add(issueEvent);
+
+				GetNames(issueEvent, issue.Keys, issue.Metadata);
+				GetPaths(issueEvent, issue.Keys, issue.Metadata);
+				if (issue.Keys.Count == 0)
+				{
+					GetHash(issueEvent.Message, issue.Keys, issue.Metadata);
+				}
+
+				issue.Metadata.Add(new IssueMetadata("type", GetEventPrefix(issueEvent.EventId.Value)));
+				_issues.Add(issue);
+
+				return true;
+			}
+			return false;
+		}
+
+		/// <inheritdoc/>
+		public override IEnumerable<IssueEventGroup> GetIssues() => _issues;
+	}
+}
diff --git a/Engine/Source/Programs/Shared/EpicGames.Horde/Issues/Handlers/HashedIssueHandler.cs b/Engine/Source/Programs/Shared/EpicGames.Horde/Issues/Handlers/HashedIssueHandler.cs
new file mode 100644
--- /dev/null
+++ b/Engine/Source/Programs/Shared/EpicGames.Horde/Issues/Handlers/HashedIssueHandler.cs
@@ -0,0 +1,118 @@
+// Copyright Epic Games, Inc. All Rights Reserved.
+
+using System.Collections.Generic;
+using System.Linq;
+using System.Text;
+using System.Text.RegularExpressions;
+using EpicGames.Core;
+using Microsoft.Extensions.Logging;
+
+namespace EpicGames.Horde.Issues.Handlers
+{
+	/// <summary>
+	/// Instance of a particular compile error
+	/// </summary>
+	[IssueHandler(Priority = 1)]
+	public class HashedIssueHandler : IssueHandler
+	{
+		readonly IssueHandlerContext _context;
+		readonly List<IssueEvent> _issueEvents = new List<IssueEvent>();
+
+		/// <summary>
+		/// Known general events
+		/// </summary>
+		static readonly HashSet<EventId> s_knownGeneralEvents = new HashSet<EventId> { KnownLogEvents.Generic, KnownLogEvents.ExitCode, KnownLogEvents.Horde, KnownLogEvents.Horde_InvalidPreflight };
+
+		/// <summary>
+		/// Constructor
+		/// </summary>
+		public HashedIssueHandler(IssueHandlerContext context) => _context = context;
+
+		/// <summary>
+		/// Determines if the given event is general and should be salted to make it unique
+		/// </summary>
+		/// <param name="eventId">The event id to compare</param>
+		/// <returns>True if the given event id matches</returns>
+		public static bool IsGeneralEventId(EventId eventId)
+		{
+			return s_knownGeneralEvents.Contains(eventId) || (eventId.Id >= KnownLogEvents.Systemic.Id && eventId.Id <= KnownLogEvents.Systemic_Max.Id);
+		}
+
+		/// <inheritdoc/>
+		public override bool HandleEvent(IssueEvent logEvent)
+		{
+			_issueEvents.Add(logEvent);
+			return true;
+		}
+
+		/// <inheritdoc/>
+		public override IEnumerable<IssueEventGroup> GetIssues()
+		{
+			List<IssueEventGroup> issues = new List<IssueEventGroup>();
+
+			IssueEventGroup? genericFingerprint = null;
+			HashSet<Md5Hash> hashes = new HashSet<Md5Hash>();
+
+			// keep hash consistent when only have general, non-unique events
+			bool allGeneral = _issueEvents.FirstOrDefault(stepEvent => stepEvent.EventId == null || !IsGeneralEventId(stepEvent.EventId.Value)) == null;
+
+			foreach (IssueEvent stepEvent in _issueEvents)
+			{
+				string hashSource = stepEvent.Message;
+
+				if (!allGeneral && stepEvent.EventId != null)
+				{
+					// If the event is general, salt the hash with the stream id, template, otherwise it will be aggressively matched.
+					// Consider salting with node name, though template id should be enough and have better grouping
+					if (IsGeneralEventId(stepEvent.EventId.Value))
+					{
+						hashSource += $"step:{_context.StreamId}:{_context.TemplateId}";
+					}
+				}
+
+				if (hashes.Count < 25 && TryGetHash(hashSource, out Md5Hash hash))
+				{
+					hashes.Add(hash);
+
+					IssueEventGroup issue = new IssueEventGroup("Hashed", "{Severity} in {Meta:Node}", IssueChangeFilter.All);
+					issue.Events.Add(stepEvent);
+					issue.Keys.AddHash(hash);
+					issue.Metadata.Add("Node", _context.NodeName);
+					issues.Add(issue);
+				}
+				else
+				{
+					if (genericFingerprint == null)
+					{
+						genericFingerprint = new IssueEventGroup("Hashed", "{Severity} in {Meta:Node}", IssueChangeFilter.All);
+						genericFingerprint.Keys.Add(IssueKey.FromStep(_context.StreamId, _context.TemplateId, _context.NodeName));
+						genericFingerprint.Metadata.Add("Node", _context.NodeName);
+						issues.Add(genericFingerprint);
+					}
+					genericFingerprint.Events.Add(stepEvent);
+				}
+			}
+
+			return issues;
+		}
+
+		static bool TryGetHash(string message, out Md5Hash hash)
+		{
+			string sanitized = message.ToUpperInvariant();
+			sanitized = Regex.Replace(sanitized, @"(?<![a-zA-Z])(?:[A-Z]:|/)[^ :]+[/\\]SYNC[/\\]", "{root}/"); // Redact things that look like workspace roots; may be different between agents
+			sanitized = Regex.Replace(sanitized, @"0[xX][0-9a-fA-F]+", "H"); // Redact hex strings
+			sanitized = Regex.Replace(sanitized, @"\d[\d.,:]*", "n"); // Redact numbers and timestamp like things
+
+			if (sanitized.Length > 30)
+			{
+				hash = Md5Hash.Compute(Encoding.UTF8.GetBytes(sanitized));
+				return true;
+			}
+			else
+			{
+				hash = Md5Hash.Zero;
+				return false;
+			}
+		}
+	}
+}
diff --git a/Engine/Source/Programs/Shared/EpicGames.Horde/Issues/Handlers/LocalizationIssueHandler.cs b/Engine/Source/Programs/Shared/EpicGames.Horde/Issues/Handlers/LocalizationIssueHandler.cs
new file mode 100644
--- /dev/null
+++ b/Engine/Source/Programs/Shared/EpicGames.Horde/Issues/Handlers/LocalizationIssueHandler.cs
@@ -0,0 +1,93 @@
+// Copyright Epic Games, Inc. All Rights Reserved.
+
+using System;
+using System.Collections.Generic;
+using System.Text.Json;
+using EpicGames.Core;
+using Microsoft.Extensions.Logging;
+
+namespace EpicGames.Horde.Issues.Handlers
+{
+	/// <summary>
+	/// Instance of a localization error
+	/// </summary>
+	[IssueHandler(Priority = 10)]
+	public class LocalizationIssueHandler : IssueHandler
+	{
+		readonly List<IssueEventGroup> _issues = new List<IssueEventGroup>();
+
+		/// <summary>
+		/// Determines if the given event id matches
+		/// </summary>
+		/// <param name="eventId">The event id to compare</param>
+		/// <returns>True if the given event id matches</returns>
+		public static bool IsMatchingEventId(EventId? eventId)
+		{
+			return eventId == KnownLogEvents.Engine_Localization;
+		}
+
+		/// <summary>
+		/// Determines if an event should be masked by this 
+		/// </summary>
+		/// <param name="eventId"></param>
+		/// <returns></returns>
+		static bool IsMaskedEventId(EventId eventId)
+		{
+			return eventId == KnownLogEvents.ExitCode;
+		}
+
+		/// <summary>
+		/// Extracts a list of source files from an event
+		/// </summary>
+		/// <param name="issueEvent">The event data</param>
+		/// <param name="sourceFiles">List of source files</param>
+		public static void GetSourceFiles(IssueEvent issueEvent, HashSet<IssueKey> sourceFiles)
+		{
+			foreach (JsonLogEvent line in issueEvent.Lines)
+			{
+				JsonDocument document = JsonDocument.Parse(line.Data);
+
+				string? relativePath;
+				if (document.RootElement.TryGetNestedProperty("properties.file.relativePath", out relativePath) || document.RootElement.TryGetNestedProperty("properties.file", out relativePath))
+				{
+					if (!relativePath.EndsWith(".manifest", StringComparison.OrdinalIgnoreCase))
+					{
+						int endIdx = relativePath.LastIndexOfAny(new char[] { '/', '\\' }) + 1;
+						string fileName = relativePath.Substring(endIdx);
+						sourceFiles.Add(new IssueKey(fileName, IssueKeyType.File));
+					}
+				}
+			}
+		}
+
+		/// <inheritdoc/>
+		public override bool HandleEvent(IssueEvent logEvent)
+		{
+			if (logEvent.EventId != null)
+			{
+				EventId eventId = logEvent.EventId.Value;
+				if (IsMatchingEventId(eventId))
+				{
+					IssueEventGroup issue = new IssueEventGroup("Localization", "Localization {Severity} in {Files}", IssueChangeFilter.Code);
+					issue.Events.Add(logEvent);
+					GetSourceFiles(logEvent, issue.Keys);
+
+					if (issue.Keys.Count > 0)
+					{
+						_issues.Add(issue);
+					}
+
+					return true;
+				}
+				else if(_issues.Count > 0 && IsMaskedEventId(eventId))
+				{
+					return true;
+				}
+			}
+			return false;
+		}
+
+		/// <inheritdoc/>
+		public override IEnumerable<IssueEventGroup> GetIssues() => _issues;
+	}
+}
diff --git a/Engine/Source/Programs/Shared/EpicGames.Horde/Issues/Handlers/PerforceCaseIssueHandler.cs b/Engine/Source/Programs/Shared/EpicGames.Horde/Issues/Handlers/PerforceCaseIssueHandler.cs
new file mode 100644
--- /dev/null
+++ b/Engine/Source/Programs/Shared/EpicGames.Horde/Issues/Handlers/PerforceCaseIssueHandler.cs
@@ -0,0 +1,45 @@
+// Copyright Epic Games, Inc. All Rights Reserved.
+
+using System.Collections.Generic;
+using EpicGames.Core;
+using Microsoft.Extensions.Logging;
+
+namespace EpicGames.Horde.Issues.Handlers
+{
+	/// <summary>
+	/// Instance of a Perforce case mismatch error
+	/// </summary>
+	[IssueHandler(Priority = 10)]
+	public class PerforceCaseIssueHandler : IssueHandler
+	{
+		readonly List<IssueEventGroup> _issues = new List<IssueEventGroup>();
+
+		/// <summary>
+		/// Determines if the given event id matches
+		/// </summary>
+		/// <param name="eventId">The event id to compare</param>
+		/// <returns>True if the given event id matches</returns>
+		public static bool IsMatchingEventId(EventId eventId)
+		{
+			return eventId == KnownLogEvents.AutomationTool_PerforceCase;
+		}
+
+		/// <inheritdoc/>
+		public override bool HandleEvent(IssueEvent issueEvent)
+		{
+			if (issueEvent.EventId != null && IsMatchingEventId(issueEvent.EventId.Value))
+			{
+				IssueEventGroup issue = new IssueEventGroup("PerforceCase", "Inconsistent case for {Files}", IssueChangeFilter.All);
+				issue.Events.Add(issueEvent);
+				issue.Keys.AddDepotPaths(issueEvent);
+				_issues.Add(issue);
+
+				return true;
+			}
+			return false;
+		}
+
+		/// <inheritdoc/>
+		public override IEnumerable<IssueEventGroup> GetIssues() => _issues;
+	}
+}
diff --git a/Engine/Source/Programs/Shared/EpicGames.Horde/Issues/Handlers/ScopedIssueHandler.cs b/Engine/Source/Programs/Shared/EpicGames.Horde/Issues/Handlers/ScopedIssueHandler.cs
new file mode 100644
--- /dev/null
+++ b/Engine/Source/Programs/Shared/EpicGames.Horde/Issues/Handlers/ScopedIssueHandler.cs
@@ -0,0 +1,96 @@
+// Copyright Epic Games, Inc. All Rights Reserved.
+
+using System;
+using System.Collections.Generic;
+using System.Text;
+using System.Text.Json;
+using System.Text.RegularExpressions;
+using EpicGames.Core;
+
+namespace EpicGames.Horde.Issues.Handlers
+{
+	/// <summary>
+	/// Instance of a particular compile error
+	/// </summary>
+	[IssueHandler(Priority = 2, Tag = "Scoped")]
+	public class ScopedIssueHandler : IssueHandler
+	{
+		const string NodeName = "Node";
+		const string ScopeName = "Scope";
+
+		readonly IssueHandlerContext _context;
+		readonly List<IssueEventGroup> _issues = new List<IssueEventGroup>();
+
+		/// <summary>
+		/// Constructor
+		/// </summary>
+		public ScopedIssueHandler(IssueHandlerContext context) => _context = context;
+
+		/// <inheritdoc/>
+		public override bool HandleEvent(IssueEvent logEvent)
+		{
+			string? scope = null;
+
+			foreach (JsonLogEvent line in logEvent.Lines)
+			{
+				JsonDocument document = JsonDocument.Parse(line.Data);
+
+				string? channelType;
+				string? channelText;
+				if (document.RootElement.TryGetNestedProperty("properties.channel.$type", out channelType) && document.RootElement.TryGetNestedProperty("properties.channel.$text", out channelText))
+				{
+					if (channelType != "Channel" || !channelText.StartsWith("Log", StringComparison.Ordinal))
+					{
+						scope = null;
+						break;
+					}
+						
+					if (scope != null && scope != channelText)
+					{
+						scope = null;
+						break;
+					}
+
+					scope = channelText;
+				}
+			}
+
+			if (scope == null)
+			{
+				return false;
+			}
+
+			string fingerprintType = $"Scoped:{scope}";
+
+			string hashSource = logEvent.Message;
+				
+			if (TryGetHash(hashSource, out Md5Hash hash))
+			{
+				IssueEventGroup issue = new IssueEventGroup(fingerprintType, "{Severity} in {Meta:Node} - {Meta:Scope}", IssueChangeFilter.All);
+				issue.Events.Add(logEvent);
+				issue.Keys.Add(IssueKey.FromHash(hash));
+				issue.Metadata.Add(NodeName, _context.NodeName);
+				issue.Metadata.Add(ScopeName, scope);
+				_issues.Add(issue);
+
+				return true;
+			}
+
+			return false;
+		}
+
+		/// <inheritdoc/>
+		public override IEnumerable<IssueEventGroup> GetIssues() => _issues;
+
+		static bool TryGetHash(string message, out Md5Hash hash)
+		{
+			string sanitized = message.ToUpperInvariant();
+			sanitized = Regex.Replace(sanitized, @"(?<![a-zA-Z])(?:[A-Z]:|/)[^ :]+[/\\]SYNC[/\\]", "{root}/"); // Redact things that look like workspace roots; may be different between agents
+			sanitized = Regex.Replace(sanitized, @"0[xX][0-9a-fA-F]+", "H"); // Redact hex strings
+			sanitized = Regex.Replace(sanitized, @"\d[\d.,:]*", "n"); // Redact numbers and timestamp like things
+
+			hash = Md5Hash.Compute(Encoding.UTF8.GetBytes(sanitized));
+			return true;
+		}
+	}
+}
diff --git a/Engine/Source/Programs/Shared/EpicGames.Horde/Issues/Handlers/ShaderIssueHandler.cs b/Engine/Source/Programs/Shared/EpicGames.Horde/Issues/Handlers/ShaderIssueHandler.cs
new file mode 100644
--- /dev/null
+++ b/Engine/Source/Programs/Shared/EpicGames.Horde/Issues/Handlers/ShaderIssueHandler.cs
@@ -0,0 +1,46 @@
+// Copyright Epic Games, Inc. All Rights Reserved.
+
+using System.Collections.Generic;
+using EpicGames.Core;
+using Microsoft.Extensions.Logging;
+
+namespace EpicGames.Horde.Issues.Handlers
+{
+	/// <summary>
+	/// Instance of a particular shader compile error
+	/// </summary>
+	[IssueHandler(Priority = 10)]
+	public class ShaderIssueHandler : IssueHandler
+	{
+		readonly List<IssueEventGroup> _issues = new List<IssueEventGroup>();
+
+		static bool IsMatchingEventId(EventId id) => id == KnownLogEvents.Engine_ShaderCompiler;
+		static bool IsMaskedEventId(EventId id) => id == KnownLogEvents.Generic || id == KnownLogEvents.ExitCode || id == KnownLogEvents.Systemic_Xge_BuildFailed || id == KnownLogEvents.Engine_Crash || id == KnownLogEvents.AutomationTool_CrashExitCode || id == KnownLogEvents.Engine_AppError;
+
+		/// <inheritdoc/>
+		public override bool HandleEvent(IssueEvent logEvent)
+		{
+			if (logEvent.EventId.HasValue)
+			{
+				EventId eventId = logEvent.EventId.Value;
+				if (IsMatchingEventId(eventId))
+				{
+					IssueEventGroup issue = new IssueEventGroup("Shader", "Shader compile {Severity} in {Files}", IssueChangeFilter.Code);
+					issue.Events.Add(logEvent);
+					issue.Keys.AddSourceFiles(logEvent);
+					_issues.Add(issue);
+
+					return true;
+				}
+				else if (_issues.Count > 0 && IsMaskedEventId(eventId))
+				{
+					return true;
+				}
+			}
+			return false;
+		}
+
+		/// <inheritdoc/>
+		public override IEnumerable<IssueEventGroup> GetIssues() => _issues;
+	}
+}
diff --git a/Engine/Source/Programs/Shared/EpicGames.Horde/Issues/Handlers/SymbolIssueHandler.cs b/Engine/Source/Programs/Shared/EpicGames.Horde/Issues/Handlers/SymbolIssueHandler.cs
new file mode 100644
--- /dev/null
+++ b/Engine/Source/Programs/Shared/EpicGames.Horde/Issues/Handlers/SymbolIssueHandler.cs
@@ -0,0 +1,73 @@
+// Copyright Epic Games, Inc. All Rights Reserved.
+
+using System.Collections.Generic;
+using EpicGames.Core;
+using Microsoft.Extensions.Logging;
+
+namespace EpicGames.Horde.Issues.Handlers
+{
+	/// <summary>
+	/// Instance of a particular compile error
+	/// </summary>
+	[IssueHandler(Priority = 10)]
+	public class SymbolIssueHandler : IssueHandler
+	{
+		const string EventIdName = "EventId";
+
+		readonly List<IssueEventGroup> _issues = new List<IssueEventGroup>();
+
+		/// <summary>
+		/// Determines if the given event id matches
+		/// </summary>
+		/// <param name="eventId">The event id to compare</param>
+		/// <returns>True if the given event id matches</returns>
+		public static bool IsMatchingEventId(EventId eventId)
+		{
+			return eventId == KnownLogEvents.Linker_UndefinedSymbol || eventId == KnownLogEvents.Linker_DuplicateSymbol || eventId == KnownLogEvents.Linker;
+		}
+
+		/// <summary>
+		/// Determines if an event should be masked by this 
+		/// </summary>
+		/// <param name="eventId"></param>
+		/// <returns></returns>
+		static bool IsMaskedEventId(EventId eventId)
+		{
+			return eventId == KnownLogEvents.ExitCode || eventId == KnownLogEvents.Systemic_Xge_BuildFailed;
+		}
+
+		/// <inheritdoc/>
+		public override bool HandleEvent(IssueEvent issueEvent)
+		{
+			if (issueEvent.EventId != null)
+			{
+				EventId eventId = issueEvent.EventId.Value;
+				if (IsMatchingEventId(eventId))
+				{
+					IssueEventGroup issue = new IssueEventGroup("Symbol", "{LegacySymbolIssueHandler}", IssueChangeFilter.Code);
+					issue.Events.Add(issueEvent);
+					issue.Keys.AddSymbols(issueEvent);
+					issue.Metadata.Add(EventIdName, eventId.Id.ToString());
+
+					if (issue.Keys.Count > 0)
+					{
+						_issues.Add(issue);
+						return true;
+					}
+					if (_issues.Count > 0)
+					{
+						return true;
+					}
+				}
+				else if (_issues.Count > 0 && IsMaskedEventId(eventId))
+				{
+					return true;
+				}
+			}
+			return false;
+		}
+
+		/// <inheritdoc/>
+		public override IEnumerable<IssueEventGroup> GetIssues() => _issues;
+	}
+}
diff --git a/Engine/Source/Programs/Shared/EpicGames.Horde/Issues/Handlers/SystemicIssueHandler.cs b/Engine/Source/Programs/Shared/EpicGames.Horde/Issues/Handlers/SystemicIssueHandler.cs
new file mode 100644
--- /dev/null
+++ b/Engine/Source/Programs/Shared/EpicGames.Horde/Issues/Handlers/SystemicIssueHandler.cs
@@ -0,0 +1,72 @@
+// Copyright Epic Games, Inc. All Rights Reserved.
+
+using System.Collections.Generic;
+using EpicGames.Core;
+using Microsoft.Extensions.Logging;
+
+namespace EpicGames.Horde.Issues.Handlers
+{
+	/// <summary>
+	/// Instance of a particular systemic error
+	/// </summary>
+	[IssueHandler(Priority = 10)]
+	public class SystemicIssueHandler : IssueHandler
+	{
+		/// <summary>
+		///  Known systemic errors
+		/// </summary>
+		static readonly HashSet<EventId> s_knownSystemic = new HashSet<EventId> { KnownLogEvents.Horde, KnownLogEvents.Horde_InvalidPreflight };
+
+		/// <summary>
+		/// Determines if the given event id matches
+		/// </summary>
+		/// <param name="eventId">The event id to compare</param>
+		/// <returns>True if the given event id matches</returns>
+		public static bool IsMatchingEventId(EventId eventId)
+		{
+			return s_knownSystemic.Contains(eventId) || (eventId.Id >= KnownLogEvents.Systemic.Id && eventId.Id <= KnownLogEvents.Systemic_Max.Id);
+		}
+
+		bool _nonSystemicError = false;
+		readonly IssueHandlerContext _context;
+		readonly List<IssueEventGroup> _issues = new List<IssueEventGroup>();
+
+		/// <summary>
+		/// Constructor
+		/// </summary>
+		public SystemicIssueHandler(IssueHandlerContext context) => _context = context;
+
+		/// <inheritdoc/>
+		public override bool HandleEvent(IssueEvent logEvent)
+		{
+			if (logEvent.EventId != null)
+			{
+				if (IsMatchingEventId(logEvent.EventId.Value))
+				{
+					if (_nonSystemicError)
+					{
+						return true;
+					}
+					else
+					{
+						IssueEventGroup issue = new IssueEventGroup("Systemic", "Systemic {Severity} in {Nodes}", IssueChangeFilter.None);
+						issue.Keys.Add(IssueKey.FromStep(_context.StreamId, _context.TemplateId, _context.NodeName));
+						issue.Events.Add(logEvent);
+						_issues.Add(issue);
+
+						return true;
+					}
+				}
+				else if (logEvent.Severity >= LogLevel.Error)
+				{
+					// We've seen a non-systemic error event, so ignore this systemic event to prevent superfluous issues from being created
+					_nonSystemicError = true;
+				}
+			}
+			return false;
+		}
+
+		/// <inheritdoc/>
+		public override IEnumerable<IssueEventGroup> GetIssues() => _issues;
+	}
+}
diff --git a/Engine/Source/Programs/Shared/EpicGames.Horde/Issues/Handlers/UnacceptableWordsIssueHandler.cs b/Engine/Source/Programs/Shared/EpicGames.Horde/Issues/Handlers/UnacceptableWordsIssueHandler.cs
new file mode 100644
--- /dev/null
+++ b/Engine/Source/Programs/Shared/EpicGames.Horde/Issues/Handlers/UnacceptableWordsIssueHandler.cs
@@ -0,0 +1,34 @@
+// Copyright Epic Games, Inc. All Rights Reserved.
+
+using System.Collections.Generic;
+using EpicGames.Core;
+
+namespace EpicGames.Horde.Issues.Handlers
+{
+	/// <summary>
+	/// Instance of a Perforce case mismatch error
+	/// </summary>
+	[IssueHandler(Priority = 8)]
+	public class UnacceptableWordsIssueHandler : IssueHandler
+	{
+		readonly List<IssueEventGroup> _issues = new List<IssueEventGroup>();
+
+		/// <inheritdoc/>
+		public override bool HandleEvent(IssueEvent issueEvent)
+		{
+			if(issueEvent.EventId == KnownLogEvents.AutomationTool_UnacceptableWords)
+			{
+				IssueEventGroup issue = new IssueEventGroup("UnacceptableWords", "Unacceptable words in {Files}", IssueChangeFilter.Code);
+				issue.Events.Add(issueEvent);
+				issue.Keys.AddSourceFiles(issueEvent);
+				_issues.Add(issue);
+
+				return true;
+			}
+			return false;
+		}
+
+		/// <inheritdoc/>
+		public override IEnumerable<IssueEventGroup> GetIssues() => _issues;
+	}
+}
diff --git a/Engine/Source/Programs/Shared/EpicGames.Horde/Issues/IssueChangeFilter.cs b/Engine/Source/Programs/Shared/EpicGames.Horde/Issues/IssueChangeFilter.cs
new file mode 100644
--- /dev/null
+++ b/Engine/Source/Programs/Shared/EpicGames.Horde/Issues/IssueChangeFilter.cs
@@ -0,0 +1,58 @@
+// Copyright Epic Games, Inc. All Rights Reserved.
+
+using System;
+using System.Collections.Generic;
+
+namespace EpicGames.Horde.Issues
+{
+	/// <summary>
+	/// Set of rules for filtering the list of suspects for an issue
+	/// </summary>
+	public static class IssueChangeFilter
+	{
+		/// <summary>
+		/// Filter exclude all changes
+		/// </summary>
+		public static IReadOnlyList<string> None { get; } = Array.Empty<string>();
+
+		/// <summary>
+		/// Filter including all changes
+		/// </summary>
+		public static IReadOnlyList<string> All { get; } = new[] { "..." };
+
+		/// <summary>
+		/// Set of extensions to treat as code
+		/// </summary>
+		public static IReadOnlyList<string> Code { get; } = new []
+		{
+			"*.c",
+			"*.cc",
+			"*.cpp",
+			"*.inl",
+			"*.m",
+			"*.mm",
+			"*.rc",
+			"*.cs",
+			"*.csproj",
+			"*.h",
+			"*.hpp",
+			"*.inl",
+			"*.usf",
+			"*.ush",
+			"*.uproject",
+			"*.uplugin",
+			"*.sln",
+			"*.verse"
+		};
+
+		/// <summary>
+		/// Set of file extensions to treat as content
+		/// </summary>
+		public static IReadOnlyList<string> Content { get; } = new[]
+		{
+			"*.uasset",
+			"*.umap",
+			"*.ini"
+		};
+	}
+}
diff --git a/Engine/Source/Programs/Shared/EpicGames.Horde/Issues/IssueEvent.cs b/Engine/Source/Programs/Shared/EpicGames.Horde/Issues/IssueEvent.cs
new file mode 100644
--- /dev/null
+++ b/Engine/Source/Programs/Shared/EpicGames.Horde/Issues/IssueEvent.cs
@@ -0,0 +1,103 @@
+// Copyright Epic Games, Inc. All Rights Reserved.
+
+using System.Collections.Generic;
+using EpicGames.Core;
+using Microsoft.Extensions.Logging;
+
+namespace EpicGames.Horde.Issues
+{
+	/// <summary>
+	/// Wraps a log event and allows it to be tagged by issue handlers
+	/// </summary>
+	public class IssueEvent
+	{
+		/// <summary>
+		/// Index of the line within this log
+		/// </summary>
+		public int LineIndex { get; }
+
+		/// <summary>
+		/// Severity of the event
+		/// </summary>
+		public LogLevel Severity { get; }
+
+		/// <summary>
+		/// The type of event
+		/// </summary>
+		public EventId? EventId { get; }
+
+		/// <summary>
+		/// The complete rendered message, in plaintext
+		/// </summary>
+		public string Message { get; }
+
+		/// <summary>
+		/// Gets this event data as a BSON document
+		/// </summary>
+		public IReadOnlyList<JsonLogEvent> Lines { get; }
+
+		/// <summary>
+		/// Constructor
+		/// </summary>
+		public IssueEvent(int lineIndex, LogLevel severity, EventId? eventId, string message, IReadOnlyList<JsonLogEvent> lines)
+		{
+			LineIndex = lineIndex;
+			Severity = severity;
+			EventId = eventId;
+			Message = message;
+			Lines = lines;
+		}
+
+		/// <inheritdoc/>
+		public override string ToString() => $"[{LineIndex}] {Message}";
+	}
+
+	/// <summary>
+	/// A group of <see cref="IssueEvent"/> objects with their fingerprint
+	/// </summary>
+	public class IssueEventGroup
+	{
+		/// <summary>
+		/// The type of issue, which defines the handler to use for it
+		/// </summary>
+		public string Type { get; set; }
+
+		/// <summary>
+		/// Template string for the issue summary
+		/// </summary>
+		public string SummaryTemplate { get; set; }
+
+		/// <summary>
+		/// List of keys which identify this issue.
+		/// </summary>
+		public HashSet<IssueKey> Keys { get; } = new HashSet<IssueKey>();
+
+		/// <summary>
+		/// Collection of additional metadata added by the handler
+		/// </summary>
+		public HashSet<IssueMetadata> Metadata { get; } = new HashSet<IssueMetadata>();
+
+		/// <summary>
+		/// Filter for changes that should be included in this issue
+		/// </summary>
+		public IReadOnlyList<string> ChangeFilter { get; set; }
+
+		/// <summary>
+		/// Individual log events
+		/// </summary>
+		public List<IssueEvent> Events { get; } = new List<IssueEvent>();
+
+		/// <summary>
+		/// Constructor
+		/// </summary>
+		/// <param name="type">The type of issue</param>
+		/// <param name="summaryTemplate">Template for the summary string to display for the issue</param>
+		/// <param name="changeFilter">Filter for changes covered by this issue</param>
+		public IssueEventGroup(string type, string summaryTemplate, IReadOnlyList<string> changeFilter)
+		{
+			Type = type;
+			SummaryTemplate = summaryTemplate;
+			ChangeFilter = changeFilter;
+		}
+	}
+}
diff --git a/Engine/Source/Programs/Shared/EpicGames.Horde/Issues/IssueHandler.cs b/Engine/Source/Programs/Shared/EpicGames.Horde/Issues/IssueHandler.cs
new file mode 100644
--- /dev/null
+++ b/Engine/Source/Programs/Shared/EpicGames.Horde/Issues/IssueHandler.cs
@@ -0,0 +1,82 @@
+// Copyright Epic Games, Inc. All Rights Reserved.
+
+using System;
+using System.Collections.Generic;
+using Horde.Server.Streams;
+using Horde.Server.Jobs.Templates;
+
+namespace EpicGames.Horde.Issues
+{
+	/// <summary>
+	/// Marks an issue handler that should be automatically inserted into the pipeline
+	/// </summary>
+	[AttributeUsage(AttributeTargets.Class)]
+	public sealed class IssueHandlerAttribute : Attribute
+	{
+		/// <summary> 
+		/// Priority of this handler
+		/// </summary>
+		public int Priority { get; set; }
+
+		/// <summary>
+		/// Class of handler which can be explicitly enabled via a workflow
+		/// </summary>
+		public string? Tag { get; set; }
+	}
+
+	/// <summary>
+	/// Context object for issue handlers
+	/// </summary>
+	public class IssueHandlerContext
+	{
+		/// <summary>
+		/// Identifier for the current stream 
+		/// </summary>
+		public StreamId StreamId { get; }
+
+		/// <summary>
+		/// Identififer of the template
+		/// </summary>
+		public TemplateId TemplateId { get; }
+
+		/// <summary>
+		/// Identifier for the current node name 
+		/// </summary>
+		public string NodeName { get; }
+
+		/// <summary>
+		/// Annotations for this node
+		/// </summary>
+		public IReadOnlyDictionary<string, string> NodeAnnotations { get; }
+
+		/// <summary>
+		/// Constructor
+		/// </summary>
+		public IssueHandlerContext(StreamId streamId, TemplateId templateId, string nodeName, IReadOnlyDictionary<string, string> nodeAnnotations)
+		{
+			StreamId = streamId;
+			TemplateId = templateId;
+			NodeName = nodeName;
+			NodeAnnotations = nodeAnnotations;
+		}
+	}
+
+	/// <summary>
+	/// Interface for issue matchers
+	/// </summary>
+	public abstract class IssueHandler
+	{
+		/// <summary>
+		/// Attempts to assign a log event to an issue
+		/// </summary>
+		/// <param name="issueEvent">Events to process</param>
+		/// <returns>Issue definition for this log event</returns>
+		public abstract bool HandleEvent(IssueEvent issueEvent);
+
+		/// <summary>
+		/// Gets all the issues created by this handler
+		/// </summary>
+		/// <returns></returns>
+		public abstract IEnumerable<IssueEventGroup> GetIssues();
+	}
+}
diff --git a/Engine/Source/Programs/Shared/EpicGames.Horde/Issues/IssueKey.cs b/Engine/Source/Programs/Shared/EpicGames.Horde/Issues/IssueKey.cs
new file mode 100644
--- /dev/null
+++ b/Engine/Source/Programs/Shared/EpicGames.Horde/Issues/IssueKey.cs
@@ -0,0 +1,253 @@
+// Copyright Epic Games, Inc. All Rights Reserved.
+
+using System;
+using System.Collections.Generic;
+using System.Linq;
+using System.Text.Json;
+using EpicGames.Core;
+using Horde.Server.Jobs.Templates;
+using Horde.Server.Streams;
+
+namespace EpicGames.Horde.Issues
+{
+	/// <summary>
+	/// Type of a key in an issue, used for grouping.
+	/// </summary>
+	public enum IssueKeyType
+	{
+		/// <summary>
+		/// Unknown type
+		/// </summary>
+		None,
+
+		/// <summary>
+		/// Filename
+		/// </summary>
+		File,
+
+		/// <summary>
+		/// Secondary file
+		/// </summary>
+		Note,
+
+		/// <summary>
+		/// Name of a symbol
+		/// </summary>
+		Symbol,
+
+		/// <summary>
+		/// Hash of a particular error
+		/// </summary>
+		Hash,
+
+		/// <summary>
+		/// Identifier for a particular step
+		/// </summary>
+		Step,
+	}
+
+	/// <summary>
+	/// Defines a key which can be used to group an issue with other issues
+	/// </summary>
+	public class IssueKey : IEquatable<IssueKey>
+	{
+		/// <summary>
+		/// Name of the key
+		/// </summary>
+		public string Name { get; }
+
+		/// <summary>
+		/// Type of the key
+		/// </summary>
+		public IssueKeyType Type { get; }
+
+		/// <summary>
+		/// Arbitrary string that can be used to discriminate between otherwise identical keys, limiting the issues that it can merge with.
+		/// </summary>
+		public string? Scope { get; }
+
+		/// <summary>
+		/// Constructor
+		/// </summary>
+		public IssueKey(string name, IssueKeyType type)
+		{
+			Name = name;
+			Type = type;
+		}
+
+		/// <summary>
+		/// Constructor
+		/// </summary>
+		public IssueKey(string name, IssueKeyType type, string? scope)
+		{
+			Name = name;
+			Type = type;
+			Scope = scope;
+		}
+
+		/// <summary>
+		/// Creates an issue key for a particular hash
+		/// </summary>
+		public static IssueKey FromHash(Md5Hash hash) => new IssueKey(hash.ToString(), IssueKeyType.Hash);
+
+		/// <summary>
+		/// Creates an issue key for a particular step
+		/// </summary>
+		public static IssueKey FromStep(StreamId streamId, TemplateId templateId, string nodeName) => new IssueKey($"{streamId}:{templateId}:{nodeName}", IssueKeyType.Step);
+
+		/// <inheritdoc/>
+		public bool Equals(IssueKey? other) => other is not null && other.Name.Equals(Name, StringComparison.OrdinalIgnoreCase) && other.Type == Type && String.Equals(Scope, other.Scope, StringComparison.Ordinal);
+
+		/// <inheritdoc/>
+		public override bool Equals(object? obj) => obj is IssueKey other && Equals(other);
+
+		/// <inheritdoc/>
+		public override int GetHashCode() => HashCode.Combine(String.GetHashCode(Name, StringComparison.OrdinalIgnoreCase), Type, Scope?.GetHashCode(StringComparison.Ordinal) ?? 0);
+
+		/// <inheritdoc/>
+		public override string ToString() => Name;
+
+		/// <inheritdoc/>
+		public static bool operator ==(IssueKey left, IssueKey right) => left.Equals(right);
+
+		/// <inheritdoc/>
+		public static bool operator !=(IssueKey left, IssueKey right) => !left.Equals(right);
+	}
+
+	/// <summary>
+	/// Extension methods for issue keys
+	/// </summary>
+	public static class IssueKeyExtensions
+	{
+		/// <summary>
+		/// Adds a new entry to a set
+		/// </summary>
+		public static void Add(this HashSet<IssueKey> entries, string name, IssueKeyType type, string? scope = null)
+		{
+			entries.Add(new IssueKey(name, type, scope));
+		}
+
+		/// <summary>
+		/// Adds all the assets from the given log event
+		/// </summary>
+		/// <param name="keys">Set of keys</param>
+		/// <param name="issueEvent">The log event to parse</param>
+		public static void AddAssets(this HashSet<IssueKey> keys, IssueEvent issueEvent)
+		{
+			foreach (JsonLogEvent line in issueEvent.Lines)
+			{
+				JsonDocument document = JsonDocument.Parse(line.Data);
+
+				string? relativePath;
+				if (document.RootElement.TryGetNestedProperty("properties.asset.relativePath", out relativePath) || document.RootElement.TryGetNestedProperty("properties.asset.$text", out relativePath))
+				{
+					int endIdx = relativePath.LastIndexOfAny(new char[] { '/', '\\' }) + 1;
+					string fileName = relativePath.Substring(endIdx);
+					IssueKey issueKey = new IssueKey(fileName, IssueKeyType.File);
+					keys.Add(issueKey);
+				}
+			}
+		}
+
+		/// <summary>
+		/// Extracts a list of source files from an event
+		/// </summary>
+		/// <param name="keys">Set of keys</param>
+		/// <param name="issueEvent">The event data</param>
+		public static void AddDepotPaths(this HashSet<IssueKey> keys, IssueEvent issueEvent)
+		{
+			foreach (JsonProperty property in issueEvent.Lines.SelectMany(x => x.FindPropertiesOfType(LogValueType.DepotPath)))
+			{
+				JsonElement value;
+				if (property.Value.TryGetProperty(LogEventPropertyName.Text.Span, out value) && value.ValueKind == JsonValueKind.String)
+				{
+					string path = value.GetString() ?? String.Empty;
+					string fileName = path.Substring(path.LastIndexOf('/') + 1);
+					keys.Add(new IssueKey(fileName, IssueKeyType.File));
+				}
+			}
+		}
+
+		/// <summary>
+		/// Extracts a list of source files from an event
+		/// </summary>
+		/// <param name="keys">Set of keys</param>
+		/// <param name="hash">The event hash</param>
+		/// <param name="scope">Scope for merging this hash value</param>
+		public static void AddHash(this HashSet<IssueKey> keys, Md5Hash hash, string? scope = null)
+		{
+			keys.Add(hash.ToString(), IssueKeyType.Hash, scope);
+		}
+
+		/// <summary>
+		/// Extracts a list of source files from an event
+		/// </summary>
+		/// <param name="keys">Set of keys</param>
+		/// <param name="issueEvent">The event data</param>
+		public static void AddSourceFiles(this HashSet<IssueKey> keys, IssueEvent issueEvent)
+		{
+			foreach (JsonLogEvent line in issueEvent.Lines)
+			{
+				JsonDocument document = JsonDocument.Parse(line.Data);
+
+				JsonElement properties;
+				if (document.RootElement.TryGetProperty("properties", out properties) && properties.ValueKind == JsonValueKind.Object)
+				{
+					IssueKeyType type = IssueKeyType.File;
+					if (properties.TryGetProperty("note", out JsonElement noteElement) && noteElement.GetBoolean())
+					{
+						type = IssueKeyType.Note;
+					}
+
+					foreach (JsonProperty property in properties.EnumerateObject())
+					{
+						if (property.NameEquals("file") && property.Value.ValueKind == JsonValueKind.String)
+						{
+							keys.AddSourceFile(property.Value.GetString()!, type);
+						}
+						if (property.Value.HasStringProperty("$type", "SourceFile") && property.Value.TryGetStringProperty("relativePath", out string? value))
+						{
+							keys.AddSourceFile(value, type);
+						}
+					}
+				}
+			}
+		}
+
+		/// <summary>
+		/// Add a new source file to a list of unique source files
+		/// </summary>
+		/// <param name="keys">List of source files</param>
+		/// <param name="relativePath">File to add</param>
+		/// <param name="type">Type of key to add</param>
+		public static void AddSourceFile(this HashSet<IssueKey> keys, string relativePath, IssueKeyType type)
+		{
+			int endIdx = relativePath.LastIndexOfAny(new char[] { '/', '\\' }) + 1;
+
+			string fileName = relativePath.Substring(endIdx);
+			IssueKey key = new IssueKey(fileName, type);
+
+			keys.Add(key);
+		}
+
+		/// <summary>
+		/// Parses symbol names from a log event
+		/// </summary>
+		/// <param name="keys">List of source files</param>
+		/// <param name="eventData">The log event data</param>
+		public static void AddSymbols(this HashSet<IssueKey> keys, IssueEvent eventData)
+		{
+			foreach (JsonLogEvent line in eventData.Lines)
+			{
+				JsonDocument document = JsonDocument.Parse(line.Data);
+
+				string? identifier;
+				if (document.RootElement.TryGetNestedProperty("properties.symbol.identifier", out identifier))
+				{
+					IssueKey key = new IssueKey(identifier, IssueKeyType.Symbol);
+					keys.Add(key);
+				}
+			}
+		}
+	}
+}
diff --git a/Engine/Source/Programs/Shared/EpicGames.Horde/Issues/IssueMetadata.cs b/Engine/Source/Programs/Shared/EpicGames.Horde/Issues/IssueMetadata.cs
new file mode 100644
--- /dev/null
+++ b/Engine/Source/Programs/Shared/EpicGames.Horde/Issues/IssueMetadata.cs
@@ -0,0 +1,75 @@
+// Copyright Epic Games, Inc. All Rights Reserved.
+
+using System;
+using System.Collections.Generic;
+
+namespace EpicGames.Horde.Issues
+{
+	/// <summary>
+	/// Entry in an issue's metadata collection. Implemented as a case-insensitive key and value
+	/// </summary>
+	public class IssueMetadata : IEquatable<IssueMetadata>
+	{
+		/// <summary>
+		/// Key for the metadata item
+		/// </summary>
+		public string Key { get; }
+
+		/// <summary>
+		/// Value for the metadata
+		/// </summary>
+		public string Value { get; }
+
+		/// <summary>
+		/// Constructor
+		/// </summary>
+		public IssueMetadata(string key, string value)
+		{
+			Key = key;
+			Value = value;
+		}
+
+		/// <inheritdoc/>
+		public override int GetHashCode() => HashCode.Combine(Key.GetHashCode(StringComparison.OrdinalIgnoreCase), Value.GetHashCode(StringComparison.OrdinalIgnoreCase));
+
+		/// <inheritdoc/>
+		public override bool Equals(object? obj) => obj is IssueMetadata entry && Equals(entry);
+
+		/// <inheritdoc/>
+		public bool Equals(IssueMetadata? entry) => entry is not null && Key.Equals(entry.Key, StringComparison.OrdinalIgnoreCase) && Value.Equals(entry.Value, StringComparison.OrdinalIgnoreCase);
+
+		/// <inheritdoc/>
+		public override string ToString() => $"{Key}={Value}";
+	}
+
+	/// <summary>
+	/// Extension methods for <see cref="IssueMetadata"/>
+	/// </summary>
+	public static class IssueMetadataEntryExtensions
+	{
+		/// <summary>
+		/// Adds a new entry to a set
+		/// </summary>
+		public static void Add(this HashSet<IssueMetadata> entries, string key, string value)
+		{
+			entries.Add(new IssueMetadata(key, value));
+		}
+
+		/// <summary>
+		/// Gets all the metadata values with a given key
+		/// </summary>
+		/// <param name="entries">Set of entries to search</param>
+		/// <param name="key">Key name to search for</param>
+		/// <returns>All values with the given key</returns>
+		public static IEnumerable<string> FindValues(this IReadOnlySet<IssueMetadata> entries, string key)
+		{
+			foreach (IssueMetadata entry in entries)
+			{
+				if (entry.Key.Equals(key, StringComparison.OrdinalIgnoreCase))
+				{
+					yield return entry.Value;
+				}
+			}
+		}
+	}
+}
diff --git a/Engine/Source/Programs/Shared/EpicGames.Horde/Jobs/Templates/TemplateId.cs b/Engine/Source/Programs/Shared/EpicGames.Horde/Jobs/Templates/TemplateId.cs
new file mode 100644
--- /dev/null
+++ b/Engine/Source/Programs/Shared/EpicGames.Horde/Jobs/Templates/TemplateId.cs
@@ -0,0 +1,45 @@
+// Copyright Epic Games, Inc. All Rights Reserved.
+
+using System.ComponentModel;
+using EpicGames.Core;
+using EpicGames.Horde;
+using EpicGames.Serialization;
+
+namespace Horde.Server.Jobs.Templates
+{
+	/// <summary>
+	/// Identifier for a job template
+	/// </summary>
+	/// <param name="Id">Id to construct from</param>
+	[JsonSchemaString]
+	[TypeConverter(typeof(StringIdTypeConverter<TemplateId, TemplateIdConverter>))]
+	[StringIdConverter(typeof(TemplateIdConverter))]
+	[CbConverter(typeof(StringIdCbConverter<TemplateId, TemplateIdConverter>))]
+	public record struct TemplateId(StringId Id)
+	{
+		/// <summary>
+		/// Constructor
+		/// </summary>
+		public TemplateId(string id) : this(new StringId(id))
+		{
+		}
+
+		/// <inheritdoc cref="StringId.IsEmpty"/>
+		public bool IsEmpty => Id.IsEmpty;
+
+		/// <inheritdoc/>
+		public override string ToString() => Id.ToString();
+	}
+
+	/// <summary>
+	/// Converter to and from <see cref="StringId"/> instances.
+	/// </summary>
+	class TemplateIdConverter : StringIdConverter<TemplateId>
+	{
+		/// <inheritdoc/>
+		public override TemplateId FromStringId(StringId id) => new TemplateId(id);
+
+		/// <inheritdoc/>
+		public override StringId ToStringId(TemplateId value) => value.Id;
+	}
+}
diff --git a/Engine/Source/Programs/Shared/EpicGames.Horde/Streams/StreamId.cs b/Engine/Source/Programs/Shared/EpicGames.Horde/Streams/StreamId.cs
new file mode 100644
--- /dev/null
+++ b/Engine/Source/Programs/Shared/EpicGames.Horde/Streams/StreamId.cs
@@ -0,0 +1,45 @@
+// Copyright Epic Games, Inc. All Rights Reserved.
+
+using System.ComponentModel;
+using EpicGames.Core;
+using EpicGames.Horde;
+using EpicGames.Serialization;
+
+namespace Horde.Server.Streams
+{
+	/// <summary>
+	/// Identifier for a stream
+	/// </summary>
+	/// <param name="Id">Id to construct from</param>
+	[JsonSchemaString]
+	[TypeConverter(typeof(StringIdTypeConverter<StreamId, StreamIdConverter>))]
+	[StringIdConverter(typeof(StreamIdConverter))]
+	[CbConverter(typeof(StringIdCbConverter<StreamId, StreamIdConverter>))]
+	public record struct StreamId(StringId Id)
+	{
+		/// <summary>
+		/// Constructor
+		/// </summary>
+		public StreamId(string id) : this(new StringId(id))
+		{
+		}
+
+		/// <inheritdoc cref="StringId.IsEmpty"/>
+		public bool IsEmpty => Id.IsEmpty;
+
+		/// <inheritdoc/>
+		public override string ToString() => Id.ToString();
+	}
+
+	/// <summary>
+	/// Converter to and from <see cref="StringId"/> instances.
+	/// </summary>
+	public class StreamIdConverter : StringIdConverter<StreamId>
+	{
+		/// <inheritdoc/>
+		public override StreamId FromStringId(StringId id) => new StreamId(id);
+
+		/// <inheritdoc/>
+		public override StringId ToStringId(StreamId value) => value.Id;
+	}
+}
