Horde: Missing files.

diff --git a/Engine/Source/Programs/Shared/EpicGames.Horde/Common/Condition.cs b/Engine/Source/Programs/Shared/EpicGames.Horde/Common/Condition.cs
--- a/Engine/Source/Programs/Shared/EpicGames.Horde/Common/Condition.cs
+++ b/Engine/Source/Programs/Shared/EpicGames.Horde/Common/Condition.cs
@@ -653,7 +653,7 @@ namespace EpicGames.Horde.Common
 			}
 			else
 			{
-				writer.WriteUtf8StringValue(value.Text);
+				writer.WriteUtf8StringValue(new Utf8String(value.Text));
 			}
 		}
 
@@ -662,7 +662,7 @@ namespace EpicGames.Horde.Common
 		{
 			if (value != null)
 			{
-				writer.WriteUtf8String(name, value.Text);
+				writer.WriteUtf8String(name, new Utf8String(value.Text));
 			}
 		}
 	}
diff --git a/Engine/Source/Programs/Shared/EpicGames.Horde/Storage/Bundles/BundleNodeLocator.cs b/Engine/Source/Programs/Shared/EpicGames.Horde/Storage/Bundles/BundleNodeLocator.cs
--- a/Engine/Source/Programs/Shared/EpicGames.Horde/Storage/Bundles/BundleNodeLocator.cs
+++ b/Engine/Source/Programs/Shared/EpicGames.Horde/Storage/Bundles/BundleNodeLocator.cs
@@ -145,10 +145,10 @@ namespace EpicGames.Horde.Storage.Bundles
 		public override BundleNodeLocator Read(CbField field) => BundleNodeLocator.Parse(field.AsString());
 
 		/// <inheritdoc/>
-		public override void Write(CbWriter writer, BundleNodeLocator value) => writer.WriteUtf8StringValue(value.ToString());
+		public override void Write(CbWriter writer, BundleNodeLocator value) => writer.WriteUtf8StringValue(new Utf8String(value.ToString()));
 
 		/// <inheritdoc/>
-		public override void WriteNamed(CbWriter writer, Utf8String name, BundleNodeLocator value) => writer.WriteUtf8String(name, value.ToString());
+		public override void WriteNamed(CbWriter writer, Utf8String name, BundleNodeLocator value) => writer.WriteUtf8String(name, new Utf8String(value.ToString()));
 	}
 
 	/// <summary>
diff --git a/Engine/Source/Programs/Shared/EpicGames.Serialization.Tests/JupiterTests.cs b/Engine/Source/Programs/Shared/EpicGames.Serialization.Tests/JupiterTests.cs
--- a/Engine/Source/Programs/Shared/EpicGames.Serialization.Tests/JupiterTests.cs
+++ b/Engine/Source/Programs/Shared/EpicGames.Serialization.Tests/JupiterTests.cs
@@ -33,7 +33,7 @@ namespace EpicGames.Serialization.Tests
 			CbWriter writer1 = new CbWriter();
 			writer1.BeginObject();
 			writer1.WriteInteger("a", 1);
-			writer1.WriteUtf8String("b", "hello");
+			writer1.WriteUtf8String("b", new Utf8String("hello"));
 			writer1.EndObject();
 
 			CbObject object1 = writer1.ToObject();
@@ -148,7 +148,7 @@ namespace EpicGames.Serialization.Tests
 
 			CbWriter writer = new CbWriter();
 			writer.BeginObject();
-			writer.WriteUtf8String("string", "test");
+			writer.WriteUtf8String("string", new Utf8String("test"));
 			writer.WriteBinaryAttachment("hash", hash1);
 			writer.EndObject();
 
