Horde: Fix serialization tests.

diff --git a/Engine/Source/Programs/Shared/EpicGames.Redis.Tests/SerializationTests.cs b/Engine/Source/Programs/Shared/EpicGames.Redis.Tests/SerializationTests.cs
--- a/Engine/Source/Programs/Shared/EpicGames.Redis.Tests/SerializationTests.cs
+++ b/Engine/Source/Programs/Shared/EpicGames.Redis.Tests/SerializationTests.cs
@@ -28,7 +28,7 @@ namespace EpicGames.Redis.Tests
 			TestRecord output = RedisSerializer.Deserialize<TestRecord>(value);
 			Assert.AreEqual(123, output.Number);
 			Assert.AreEqual("hello", output.String);
-			Assert.AreEqual("world", output.StringU8);
+			Assert.AreEqual(new Utf8String("world"), output.StringU8);
 		}
 
 		[TestMethod]
@@ -43,7 +43,7 @@ namespace EpicGames.Redis.Tests
 			TestRecord output = RedisSerializer.Deserialize<TestRecord>(value);
 			Assert.AreEqual(123, output.Number);
 			Assert.AreEqual("|||", output.String);
-			Assert.AreEqual("\\", output.StringU8);
+			Assert.AreEqual(new Utf8String("\\"), output.StringU8);
 		}
 
 		[TypeConverter(typeof(TestStringTypeConverter))]
diff --git a/Engine/Source/Programs/Shared/EpicGames.Serialization.Tests/EngineTests.cs b/Engine/Source/Programs/Shared/EpicGames.Serialization.Tests/EngineTests.cs
--- a/Engine/Source/Programs/Shared/EpicGames.Serialization.Tests/EngineTests.cs
+++ b/Engine/Source/Programs/Shared/EpicGames.Serialization.Tests/EngineTests.cs
@@ -144,7 +144,7 @@ namespace EpicGames.Serialization.Tests
 				byte[] noneBytes = { (byte)fieldType, 4, (byte)'N', (byte)'a', (byte)'m', (byte)'e' };
 				CbField noneField = new CbField(noneBytes);
 				Assert.AreEqual(noneField.GetSize(), noneBytes.Length);
-				Assert.AreEqual(noneField.GetName(), "Name");
+				Assert.AreEqual(noneField.GetName(), new Utf8String("Name"));
 				Assert.IsTrue(noneField.HasName());
 				Assert.IsFalse(noneField.HasValue());
 				Assert.AreEqual(noneField.GetHash(), Blake3Hash.Compute(noneBytes));
@@ -176,7 +176,7 @@ namespace EpicGames.Serialization.Tests
 				byte[] noneBytes = { (byte)fieldType, 4, (byte)'N', (byte)'a', (byte)'m', (byte)'e' };
 				CbField noneField = new CbField(noneBytes.AsMemory(1), fieldType);
 				Assert.AreEqual(noneField.GetSize(), noneBytes.Length);
-				Assert.AreEqual(noneField.GetName(), "Name");
+				Assert.AreEqual(noneField.GetName(), new Utf8String("Name"));
 				Assert.IsTrue(noneField.HasName());
 				Assert.IsFalse(noneField.HasValue());
 				Assert.AreEqual(noneField.GetHash(), Blake3Hash.Compute(noneBytes));
@@ -192,8 +192,8 @@ namespace EpicGames.Serialization.Tests
 				CbFieldType fieldType = CbFieldType.None | CbFieldType.HasFieldName;
 				byte[] noneBytes = { (byte)fieldType, 0 };
 				CbField noneField = new CbField(noneBytes.AsMemory(1), fieldType);
-				Assert.AreEqual(noneField.GetSize(), noneBytes.Length);
-				Assert.AreEqual(noneField.GetName(), "");
+				Assert.AreEqual(noneBytes.Length, noneField.GetSize());
+				Assert.AreEqual(Utf8String.Empty, noneField.GetName());
 				Assert.IsTrue(noneField.HasName());
 				Assert.IsFalse(noneField.HasValue());
 				Assert.AreEqual(noneField.GetHash(), Blake3Hash.Compute(noneBytes));
@@ -341,7 +341,7 @@ namespace EpicGames.Serialization.Tests
 				for (CbFieldIterator it = objectClone.CreateIterator(); it; ++it)
 				{
 					CbField field = it.Current;
-					Assert.AreEqual(field.GetName(), "F");
+					Assert.AreEqual(field.GetName(), new Utf8String("F"));
 					Assert.AreEqual(field.AsInt32(), 8);
 				}
 				for (CbFieldIterator it = objectClone.CreateIterator(), end = new CbFieldIterator(); it != end; ++it)
@@ -1621,22 +1621,22 @@ namespace EpicGames.Serialization.Tests
 				for (CbFieldIterator it = obj.CreateIterator(); it;)
 				{
 					CbFieldIterator last = it;
-					if (it.Current.GetName().Equals("A"))
+					if (it.Current.GetName() == new Utf8String("A"))
 					{
 						a = it.Current.AsUInt32();
 						++it;
 					}
-					if (it.Current.GetName().Equals("B"))
+					if (it.Current.GetName() == new Utf8String("B"))
 					{
 						b = it.Current.AsUInt32();
 						++it;
 					}
-					if (it.Current.GetName().Equals("C"))
+					if (it.Current.GetName() == new Utf8String("C"))
 					{
 						c = it.Current.AsUInt32();
 						++it;
 					}
-					if (it.Current.GetName().Equals("D"))
+					if (it.Current.GetName() == new Utf8String("D"))
 					{
 						d = it.Current.AsUInt32();
 						++it;
diff --git a/Engine/Source/Programs/Shared/EpicGames.Serialization.Tests/JupiterTests.cs b/Engine/Source/Programs/Shared/EpicGames.Serialization.Tests/JupiterTests.cs
--- a/Engine/Source/Programs/Shared/EpicGames.Serialization.Tests/JupiterTests.cs
+++ b/Engine/Source/Programs/Shared/EpicGames.Serialization.Tests/JupiterTests.cs
@@ -55,22 +55,22 @@ namespace EpicGames.Serialization.Tests
 			ReadOnlyMemory<byte> memory = new ReadOnlyMemory<byte>(bytes);
 			CbField o = new CbField(memory);
 
-			Assert.AreEqual("BuildAction", o.Name);
+			Assert.AreEqual(new Utf8String("BuildAction"), o.Name);
 			List<CbField> buildActionFields = o.ToList();
 			Assert.AreEqual(3, buildActionFields.Count);
-			Assert.AreEqual("Function", buildActionFields[0].Name);
-			Assert.AreEqual("Constants", buildActionFields[1].Name);
-			Assert.AreEqual("Inputs", buildActionFields[2].Name);
+			Assert.AreEqual(new Utf8String("Function"), buildActionFields[0].Name);
+			Assert.AreEqual(new Utf8String("Constants"), buildActionFields[1].Name);
+			Assert.AreEqual(new Utf8String("Inputs"), buildActionFields[2].Name);
 
 			List<CbField> constantsFields = buildActionFields[1].ToList();
 			Assert.AreEqual(3, constantsFields.Count);
-			Assert.AreEqual("TextureBuildSettings", constantsFields[0].Name);
-			Assert.AreEqual("TextureOutputSettings", constantsFields[1].Name);
-			Assert.AreEqual("TextureSource", constantsFields[2].Name);
+			Assert.AreEqual(new Utf8String("TextureBuildSettings"), constantsFields[0].Name);
+			Assert.AreEqual(new Utf8String("TextureOutputSettings"), constantsFields[1].Name);
+			Assert.AreEqual(new Utf8String("TextureSource"), constantsFields[2].Name);
 
 			List<CbField> inputsFields = buildActionFields[2].ToList();
 			Assert.AreEqual(1, inputsFields.Count);
-			Assert.AreEqual("7587B323422942733DDD048A91709FDE", inputsFields[0].Name);
+			Assert.AreEqual(new Utf8String("7587B323422942733DDD048A91709FDE"), inputsFields[0].Name);
 			Assert.IsTrue(inputsFields[0].IsBinaryAttachment());
 			Assert.IsTrue(inputsFields[0].IsAttachment());
 			Assert.IsFalse(inputsFields[0].IsObjectAttachment());
@@ -85,7 +85,7 @@ namespace EpicGames.Serialization.Tests
 			ReadOnlyMemory<byte> memory = new ReadOnlyMemory<byte>(bytes);
 			CbField o = new CbField(memory);
 
-			Assert.AreEqual("BuildOutput", o.Name);
+			Assert.AreEqual(new Utf8String("BuildOutput"), o.Name);
 			List<CbField> buildActionFields = o.ToList();
 			Assert.AreEqual(1, buildActionFields.Count);
 			CbField payloads = buildActionFields[0];
@@ -103,7 +103,7 @@ namespace EpicGames.Serialization.Tests
 			byte[] bytes = File.ReadAllBytes("CompactBinaryObjects/compact_binary");
 
 			CbObject o = new CbObject(bytes);
-			Assert.AreEqual("", o.AsField().Name);
+			Assert.AreEqual(Utf8String.Empty, o.AsField().Name);
 			List<CbField> buildActionFields = o.ToList();
 			Assert.AreEqual(3, buildActionFields.Count);
 			CbField payloads = buildActionFields[0];
@@ -132,7 +132,7 @@ namespace EpicGames.Serialization.Tests
 			CbField o = new CbField(memory);
 
 			// the top object has no name
-			Assert.AreEqual("", o.Name);
+			Assert.AreEqual(Utf8String.Empty, o.Name);
 			List<CbField> fields = o.ToList();
 			Assert.AreEqual(1, fields.Count);
 			CbField? needs = o["needs"];
@@ -157,12 +157,12 @@ namespace EpicGames.Serialization.Tests
 			CbField o = new CbField(memory);
 
 			// the object has no name and 2 fields
-			Assert.AreEqual("", o.Name);
+			Assert.AreEqual(Utf8String.Empty, o.Name);
 			List<CbField> fields = o.ToList();
 			Assert.AreEqual(2, fields.Count);
 
 			CbField? stringField = o["string"];
-			Assert.AreEqual("test", stringField!.AsUtf8String());
+			Assert.AreEqual(new Utf8String("test"), stringField!.AsUtf8String());
 
 			CbField? hashField = o["hash"];
 			Assert.AreEqual(hash1, hashField!.AsAttachment());
