Horde: Expose additional stats for bundle cache performance.

diff --git a/Engine/Source/Programs/Shared/EpicGames.Horde/Storage/Bundles/BundleReader.cs b/Engine/Source/Programs/Shared/EpicGames.Horde/Storage/Bundles/BundleReader.cs
--- a/Engine/Source/Programs/Shared/EpicGames.Horde/Storage/Bundles/BundleReader.cs
+++ b/Engine/Source/Programs/Shared/EpicGames.Horde/Storage/Bundles/BundleReader.cs
@@ -138,6 +138,7 @@ namespace EpicGames.Horde.Storage.Bundles
 		int _numHeaderReads;
 		int _numPacketReads;
 		long _numBytesRead;
+		long _decodeTimeTicks;
 
 		/// <summary>
 		/// Accessor for the cache
@@ -238,7 +239,9 @@ namespace EpicGames.Horde.Storage.Bundles
 			BundlePacket packet = header.Packets[packetIdx];
 
 			byte[] decodedPacket = new byte[packet.DecodedLength];
+			Stopwatch decodeTimer = Stopwatch.StartNew();
 			BundleData.Decompress(packet.CompressionFormat, encodedPacket, decodedPacket);
+			Interlocked.Add(ref _decodeTimeTicks, decodeTimer.ElapsedTicks);
 			_cache.AddCachedDecodedPacket(locator, packetIdx, decodedPacket);
 
 			lock (_queueLock)
@@ -596,6 +599,7 @@ namespace EpicGames.Horde.Storage.Bundles
 			stats.Add("Num bytes read", _numBytesRead);
 			stats.Add("Num header reads", _numHeaderReads);
 			stats.Add("Num packet reads", _numPacketReads);
+			stats.Add("Decode time (ms)", (_decodeTimeTicks * 1000) / Stopwatch.Frequency);
 		}
 	}
 }
diff --git a/Engine/Source/Programs/Shared/EpicGames.Horde/Storage/Clients/BundleStorageClient.cs b/Engine/Source/Programs/Shared/EpicGames.Horde/Storage/Clients/BundleStorageClient.cs
--- a/Engine/Source/Programs/Shared/EpicGames.Horde/Storage/Clients/BundleStorageClient.cs
+++ b/Engine/Source/Programs/Shared/EpicGames.Horde/Storage/Clients/BundleStorageClient.cs
@@ -245,6 +245,10 @@ namespace EpicGames.Horde.Storage.Clients
 		#endregion
 
 		/// <inheritdoc/>
-		public void GetStats(StorageStats stats) => _bundleReader.GetStats(stats);
+		public void GetStats(StorageStats stats)
+		{
+			_backend.GetStats(stats);
+			_bundleReader.GetStats(stats);
+		}
 	}
 }
diff --git a/Engine/Source/Programs/Shared/EpicGames.Horde/Storage/StorageBackendCache.cs b/Engine/Source/Programs/Shared/EpicGames.Horde/Storage/StorageBackendCache.cs
--- a/Engine/Source/Programs/Shared/EpicGames.Horde/Storage/StorageBackendCache.cs
+++ b/Engine/Source/Programs/Shared/EpicGames.Horde/Storage/StorageBackendCache.cs
@@ -6,6 +6,7 @@ using Microsoft.Extensions.Logging;
 using Microsoft.Extensions.Logging.Abstractions;
 using System;
 using System.Collections.Generic;
+using System.Diagnostics;
 using System.IO;
 using System.Linq;
 using System.Threading;
@@ -94,7 +95,11 @@ namespace EpicGames.Horde.Storage
 			public ValueTask<Uri?> TryGetReadRedirectAsync(string path, CancellationToken cancellationToken = default) => _inner.TryGetReadRedirectAsync(path, cancellationToken);
 			public ValueTask<(string, Uri)?> TryGetWriteRedirectAsync(string? prefix = null, CancellationToken cancellationToken = default) => _inner.TryGetWriteRedirectAsync(prefix, cancellationToken);
 
-			public void GetStats(StorageStats stats) => _inner.GetStats(stats);
+			public void GetStats(StorageStats stats)
+			{
+				_inner.GetStats(stats);
+				_cacheStorage.GetStats(stats);
+			}
 		}
 
 		object LockObject => _items;
@@ -108,6 +113,11 @@ namespace EpicGames.Horde.Storage
 		readonly Dictionary<string, PendingItem> _pathToPendingItem = new Dictionary<string, PendingItem>(StringComparer.Ordinal);
 
 		long _size;
+		long _cleanCount;
+		long _cleanTimeTicks;
+		long _fetchTimeTicks;
+		long _writeTimeTicks;
+		long _fetchBytes;
 
 		internal IEnumerable<string> Items => _items.Select(x => x.Key);
 
@@ -212,9 +222,13 @@ namespace EpicGames.Horde.Storage
 
 		async Task ReadIntoCacheAsync(string key, Func<CancellationToken, Task<Stream>> createStreamAsync, CancellationToken cancellationToken)
 		{
+			long openStartTicks = Stopwatch.GetTimestamp();
 			using Stream stream = await createStreamAsync(cancellationToken);
 			long totalLength = stream.Length;
+			long openFinishTicks = Stopwatch.GetTimestamp();
+			Interlocked.Add(ref _fetchTimeTicks, openFinishTicks - openStartTicks);
 
+			long cleanStartTicks = openFinishTicks;
 			lock (LockObject)
 			{
 				_size += totalLength;
@@ -234,6 +248,7 @@ namespace EpicGames.Horde.Storage
 						_items.Remove(node);
 
 						_size -= item.Length;
+						_cleanCount++;
 					}
 					catch (Exception ex)
 					{
@@ -241,7 +256,10 @@ namespace EpicGames.Horde.Storage
 					}
 				}
 			}
+			long cleanFinishTicks = Stopwatch.GetTimestamp();
+			Interlocked.Add(ref _cleanTimeTicks, cleanFinishTicks - cleanStartTicks);
 
+			long writeStartTicks = cleanFinishTicks;
 			string path = await _backend.WriteAsync(stream, cancellationToken: cancellationToken);
 			lock (LockObject)
 			{
@@ -251,6 +269,21 @@ namespace EpicGames.Horde.Storage
 
 				_pathToPendingItem.Remove(key);
 			}
+			Interlocked.Add(ref _fetchBytes, totalLength);
+			long writeFinishTicks = Stopwatch.GetTimestamp();
+			Interlocked.Add(ref _writeTimeTicks, writeFinishTicks - writeStartTicks);
+		}
+
+		/// <summary>
+		/// Get stats for the cache operation
+		/// </summary>
+		public void GetStats(StorageStats stats)
+		{
+			stats.Add("Cache clean count", _cleanCount);
+			stats.Add("Cache fetch time (ms)", (_fetchTimeTicks * 1000) / Stopwatch.Frequency);
+			stats.Add("Cache clean time (ms)", (_cleanTimeTicks * 1000) / Stopwatch.Frequency);
+			stats.Add("Cache write time (ms)", (_writeTimeTicks * 1000) / Stopwatch.Frequency);
+			stats.Add("Cache fetch bytes", _fetchBytes);
 		}
 	}
 }
