Horde: Include the type of child nodes in chunk data streams. This allows descending through the hierarchy until we reach leaf nodes without having to actually expand those nodes, which aids in prefetching of bulk data.

diff --git a/Engine/Source/Programs/Shared/EpicGames.Horde/Storage/Node.cs b/Engine/Source/Programs/Shared/EpicGames.Horde/Storage/Node.cs
--- a/Engine/Source/Programs/Shared/EpicGames.Horde/Storage/Node.cs
+++ b/Engine/Source/Programs/Shared/EpicGames.Horde/Storage/Node.cs
@@ -270,7 +270,7 @@ namespace EpicGames.Horde.Storage
 		/// <summary>
 		/// Type to deserialize
 		/// </summary>
-		public BlobType Type => _nodeData.Type;
+		public BlobType Type => _blobData.Type;
 
 		/// <summary>
 		/// Version of the current node, as specified via <see cref="NodeTypeAttribute"/>
@@ -280,42 +280,42 @@ namespace EpicGames.Horde.Storage
 		/// <summary>
 		/// Total length of the data in this node
 		/// </summary>
-		public int Length => _nodeData.Data.Length;
+		public int Length => _blobData.Data.Length;
 
 		/// <summary>
 		/// Hash of the node being deserialized
 		/// </summary>
-		public IoHash Hash => _nodeData.Hash;
+		public IoHash Hash => _blobData.Hash;
 
 		/// <summary>
 		/// 
 		/// </summary>
-		public ReadOnlyMemory<byte> Data => _nodeData.Data;
+		public ReadOnlyMemory<byte> Data => _blobData.Data;
 
 		/// <summary>
 		/// Locations of all referenced nodes.
 		/// </summary>
-		public IReadOnlyList<BlobHandle> References => _nodeData.Refs;
+		public IReadOnlyList<BlobHandle> References => _blobData.Refs;
 
-		readonly BlobData _nodeData;
+		readonly BlobData _blobData;
 		int _refIdx;
 
 		/// <summary>
 		/// Constructor
 		/// </summary>
-		public NodeReader(BlobData nodeData)
-			: base(nodeData.Data)
+		public NodeReader(BlobData blobData)
+			: base(blobData.Data)
 		{
-			_nodeData = nodeData;
+			_blobData = blobData;
 		}
 
 		/// <summary>
 		/// Reads the next reference to another node
 		/// </summary>
-		public BlobHandle ReadNodeHandle()
+		public BlobHandle ReadBlobHandle()
 		{
 			IoHash hash = this.ReadIoHash();
-			return GetNodeHandle(_refIdx++, hash);
+			return GetBlobHandle(_refIdx++, hash);
 		}
 
 		/// <summary>
@@ -324,10 +324,10 @@ namespace EpicGames.Horde.Storage
 		/// <param name="index"></param>
 		/// <param name="hash"></param>
 		/// <returns></returns>
-		public BlobHandle GetNodeHandle(int index, IoHash hash)
+		public BlobHandle GetBlobHandle(int index, IoHash hash)
 		{
-			Debug.Assert(_nodeData.Refs[index].Hash == hash);
-			return _nodeData.Refs[index];
+			Debug.Assert(_blobData.Refs[index].Hash == hash);
+			return _blobData.Refs[index];
 		}
 	}
 
diff --git a/Engine/Source/Programs/Shared/EpicGames.Horde/Storage/NodeRef.cs b/Engine/Source/Programs/Shared/EpicGames.Horde/Storage/NodeRef.cs
--- a/Engine/Source/Programs/Shared/EpicGames.Horde/Storage/NodeRef.cs
+++ b/Engine/Source/Programs/Shared/EpicGames.Horde/Storage/NodeRef.cs
@@ -31,7 +31,7 @@ namespace EpicGames.Horde.Storage
 		/// Deserialization constructor
 		/// </summary>
 		/// <param name="reader"></param>
-		public NodeRef(NodeReader reader) : this(reader.ReadNodeHandle())
+		public NodeRef(NodeReader reader) : this(reader.ReadBlobHandle())
 		{
 		}
 
@@ -80,7 +80,7 @@ namespace EpicGames.Horde.Storage
 		/// <summary>
 		/// Constructor
 		/// </summary>
-		public NodeRef(NodeReader reader) : base(reader.ReadNodeHandle())
+		public NodeRef(NodeReader reader) : base(reader.ReadBlobHandle())
 		{
 		}
 
diff --git a/Engine/Source/Programs/Shared/EpicGames.Horde/Storage/Nodes/CbNode.cs b/Engine/Source/Programs/Shared/EpicGames.Horde/Storage/Nodes/CbNode.cs
--- a/Engine/Source/Programs/Shared/EpicGames.Horde/Storage/Nodes/CbNode.cs
+++ b/Engine/Source/Programs/Shared/EpicGames.Horde/Storage/Nodes/CbNode.cs
@@ -27,7 +27,7 @@ namespace EpicGames.Horde.Storage.Nodes
 			{
 				if (field.IsAttachment())
 				{
-					BlobHandle handle = _reader.GetNodeHandle(_refs.Count, field.AsAttachment());
+					BlobHandle handle = _reader.GetBlobHandle(_refs.Count, field.AsAttachment());
 					_refs.Add(new NodeRef(handle));
 				}
 				else if (field.IsArray())
diff --git a/Engine/Source/Programs/Shared/EpicGames.Horde/Storage/Nodes/ChunkedDataNode.cs b/Engine/Source/Programs/Shared/EpicGames.Horde/Storage/Nodes/ChunkedDataNode.cs
--- a/Engine/Source/Programs/Shared/EpicGames.Horde/Storage/Nodes/ChunkedDataNode.cs
+++ b/Engine/Source/Programs/Shared/EpicGames.Horde/Storage/Nodes/ChunkedDataNode.cs
@@ -85,12 +85,70 @@ namespace EpicGames.Horde.Storage.Nodes
 		}
 	}
 
+	/// <summary>
+	/// Type of a chunked data node
+	/// </summary>
+	public enum ChunkedDataNodeType
+	{
+		/// <summary>
+		/// Unknown node type
+		/// </summary>
+		Unknown,
+
+		/// <summary>
+		/// Leaf node
+		/// </summary>
+		Leaf,
+
+		/// <summary>
+		/// An interior node
+		/// </summary>
+		Interior
+	}
+
+	/// <summary>
+	/// Reference to a chunked data node
+	/// </summary>
+	public class ChunkedDataNodeRef : NodeRef<ChunkedDataNode>
+	{
+		/// <summary>
+		/// Type of the referenced node
+		/// </summary>
+		public ChunkedDataNodeType Type { get; }
+
+		/// <summary>
+		/// Constructor
+		/// </summary>
+		public ChunkedDataNodeRef(NodeReader reader)
+			: base(reader)
+		{
+			Type = (ChunkedDataNodeType)reader.ReadUnsignedVarInt();
+		}
+
+		/// <summary>
+		/// Constructor
+		/// </summary>
+		public ChunkedDataNodeRef(ChunkedDataNodeType type, BlobHandle handle)
+			: base(handle)
+		{
+			Type = type;
+		}
+
+		/// <inheritdoc/>
+		public override void Serialize(NodeWriter writer)
+		{
+			base.Serialize(writer);
+
+			writer.WriteUnsignedVarInt((ulong)Type);
+		}
+	}
+
 	/// <summary>
 	/// Stores the flat list of chunks produced from chunking a single data stream
 	/// </summary>
 	/// <param name="Hash">Hash of the data</param>
 	/// <param name="LeafHandles">Handles to the leaf chunks</param>
-	public record class LeafChunkedData(IoHash Hash, List<NodeRef<ChunkedDataNode>> LeafHandles);
+	public record class LeafChunkedData(IoHash Hash, List<ChunkedDataNodeRef> LeafHandles);
 
 	/// <summary>
 	/// File node that contains a chunk of data
@@ -207,14 +265,14 @@ namespace EpicGames.Horde.Storage.Nodes
 			using Blake3.Hasher hasher = Blake3.Hasher.New();
 			using IMemoryOwner<byte> readBuffer = MemoryPool<byte>.Shared.Rent(options.MaxSize);
 
-			List<NodeRef<ChunkedDataNode>> handles = new List<NodeRef<ChunkedDataNode>>();
+			List<ChunkedDataNodeRef> leafNodeRefs = new List<ChunkedDataNodeRef>();
 
 			int size = 0;
 			int sizeSinceProgressUpdate = 0;
 			for (; ; )
 			{
 				size += await stream.ReadGreedyAsync(readBuffer.Memory.Slice(size), cancellationToken);
-				if (size == 0 && handles.Count > 0)
+				if (size == 0 && leafNodeRefs.Count > 0)
 				{
 					break;
 				}
@@ -233,7 +291,7 @@ namespace EpicGames.Horde.Storage.Nodes
 				}
 
 				BlobHandle handle = await writer.WriteBlobAsync(nextLength, Array.Empty<BlobHandle>(), GetNodeType<LeafChunkedDataNode>(), cancellationToken);
-				handles.Add(new NodeRef<ChunkedDataNode>(handle));
+				leafNodeRefs.Add(new ChunkedDataNodeRef(ChunkedDataNodeType.Leaf, handle));
 
 				readBuffer.Memory.Slice(nextLength, size - nextLength).CopyTo(readBuffer.Memory);
 				size -= nextLength;
@@ -242,7 +300,7 @@ namespace EpicGames.Horde.Storage.Nodes
 			copyStats?.Update(1, sizeSinceProgressUpdate);
 
 			IoHash hash = IoHash.FromBlake3(hasher);
-			return new LeafChunkedData(hash, handles);
+			return new LeafChunkedData(hash, leafNodeRefs);
 		}
 
 		/// <summary>
@@ -318,7 +376,7 @@ namespace EpicGames.Horde.Storage.Nodes
 	/// <summary>
 	/// An interior file node
 	/// </summary>
-	[NodeType("{F4DEDDBC-70CB-4C7A-8347-F011AFCCCDB9}", 1)]
+	[NodeType("{F4DEDDBC-70CB-4C7A-8347-F011AFCCCDB9}", 2)]
 	public class InteriorChunkedDataNode : ChunkedDataNode
 	{
 		/// <summary>
@@ -329,13 +387,13 @@ namespace EpicGames.Horde.Storage.Nodes
 		/// <summary>
 		/// Child nodes
 		/// </summary>
-		public IReadOnlyList<NodeRef<ChunkedDataNode>> Children { get; }
+		public IReadOnlyList<ChunkedDataNodeRef> Children { get; }
 
 		/// <summary>
 		/// Constructor
 		/// </summary>
 		/// <param name="children"></param>
-		public InteriorChunkedDataNode(IReadOnlyList<NodeRef<ChunkedDataNode>> children)
+		public InteriorChunkedDataNode(IReadOnlyList<ChunkedDataNodeRef> children)
 		{
 			Children = children;
 		}
@@ -346,10 +404,17 @@ namespace EpicGames.Horde.Storage.Nodes
 		public InteriorChunkedDataNode(NodeReader reader)
 		{
 			// Keep this code in sync with CopyToStreamAsync
-			NodeRef<ChunkedDataNode>[] children = new NodeRef<ChunkedDataNode>[reader.Length / IoHash.NumBytes];
-			for (int idx = 0; idx < children.Length; idx++)
+			List<ChunkedDataNodeRef> children = new List<ChunkedDataNodeRef>();
+			while (reader.RemainingMemory.Length > 0)
 			{
-				children[idx] = reader.ReadNodeRef<ChunkedDataNode>();
+				if (reader.Version >= 2)
+				{
+					children.Add(new ChunkedDataNodeRef(reader));
+				}
+				else
+				{
+					children.Add(new ChunkedDataNodeRef(ChunkedDataNodeType.Unknown, reader.ReadBlobHandle()));
+				}
 			}
 			Children = children;
 		}
@@ -357,7 +422,7 @@ namespace EpicGames.Horde.Storage.Nodes
 		/// <inheritdoc/>
 		public override void Serialize(NodeWriter writer)
 		{
-			foreach (NodeRef<ChunkedDataNode> child in Children)
+			foreach (ChunkedDataNodeRef child in Children)
 			{
 				writer.WriteNodeRef(child);
 			}
@@ -380,50 +445,50 @@ namespace EpicGames.Horde.Storage.Nodes
 		/// <summary>
 		/// Create a tree of nodes from the given list of handles, splitting nodes in each layer based on the hash of the last node.
 		/// </summary>
-		/// <param name="handles">List of leaf handles</param>
+		/// <param name="nodeRefs">List of leaf nodes</param>
 		/// <param name="options">Options for splitting the tree</param>
 		/// <param name="writer">Output writer for new interior nodes</param>
 		/// <param name="cancellationToken">Cancellation token for the operation</param>
 		/// <returns>Handle to the root node of the tree</returns>
-		public static async Task<NodeRef<ChunkedDataNode>> CreateTreeAsync(List<NodeRef<ChunkedDataNode>> handles, InteriorChunkedDataNodeOptions options, IStorageWriter writer, CancellationToken cancellationToken)
+		public static async Task<NodeRef<ChunkedDataNode>> CreateTreeAsync(List<ChunkedDataNodeRef> nodeRefs, InteriorChunkedDataNodeOptions options, IStorageWriter writer, CancellationToken cancellationToken)
 		{
-			List<NodeRef<ChunkedDataNode>> handleBuffer = new List<NodeRef<ChunkedDataNode>>();
+			List<ChunkedDataNodeRef> handleBuffer = new List<ChunkedDataNodeRef>();
 
 			List<InteriorChunkedDataNode> interiorNodes = new List<InteriorChunkedDataNode>();
-			while (handles.Count > 1)
+			while (nodeRefs.Count > 1)
 			{
 				interiorNodes.Clear();
-				CreateTreeLayer(handles, options, interiorNodes);
+				CreateTreeLayer(nodeRefs, options, interiorNodes);
 
 				handleBuffer.Clear();
 				foreach (InteriorChunkedDataNode interiorNode in interiorNodes)
 				{
-					NodeRef<ChunkedDataNode> handle = await writer.WriteNodeAsync<ChunkedDataNode>(interiorNode, cancellationToken);
-					handleBuffer.Add(handle);
+					NodeRef<ChunkedDataNode> newNodeRef = await writer.WriteNodeAsync<ChunkedDataNode>(interiorNode, cancellationToken);
+					handleBuffer.Add(new ChunkedDataNodeRef(ChunkedDataNodeType.Interior, newNodeRef.Handle));
 				}
 
-				handles = handleBuffer;
+				nodeRefs = handleBuffer;
 			}
 
-			return handles[0];
+			return nodeRefs[0];
 		}
 		
 		/// <summary>
-		/// Split a list of handles into a layer of interior nodes
+		/// Split a list of leaf handles into a layer of interior nodes
 		/// </summary>
-		static void CreateTreeLayer(List<NodeRef<ChunkedDataNode>> handles, InteriorChunkedDataNodeOptions options, List<InteriorChunkedDataNode> interiorNodes)
+		static void CreateTreeLayer(List<ChunkedDataNodeRef> nodeRefs, InteriorChunkedDataNodeOptions options, List<InteriorChunkedDataNode> interiorNodes)
 		{
 			Span<byte> buffer = stackalloc byte[IoHash.NumBytes];
 
-			for (int index = 0; index < handles.Count; )
+			for (int index = 0; index < nodeRefs.Count; )
 			{
 				int minIndex = index;
-				int maxIndex = Math.Min(minIndex + options.MaxChildCount, handles.Count);
+				int maxIndex = Math.Min(minIndex + options.MaxChildCount, nodeRefs.Count);
 
-				index = Math.Min(index + options.MinChildCount, handles.Count);
+				index = Math.Min(index + options.MinChildCount, nodeRefs.Count);
 				for (; index < maxIndex; index++)
 				{
-					handles[index].Handle.Hash.CopyTo(buffer);
+					nodeRefs[index].Handle.Hash.CopyTo(buffer);
 
 					uint value = BinaryPrimitives.ReadUInt32LittleEndian(buffer);
 					if (value < options.SliceThreshold)
@@ -432,10 +497,10 @@ namespace EpicGames.Horde.Storage.Nodes
 					}
 				}
 
-				NodeRef<ChunkedDataNode>[] children = new NodeRef<ChunkedDataNode>[index - minIndex];
+				ChunkedDataNodeRef[] children = new ChunkedDataNodeRef[index - minIndex];
 				for (int childIndex = minIndex; childIndex < index; childIndex++)
 				{
-					children[childIndex - minIndex] = new NodeRef<ChunkedDataNode>(handles[childIndex]);
+					children[childIndex - minIndex] = nodeRefs[childIndex];
 				}
 
 				interiorNodes.Add(new InteriorChunkedDataNode(children));
@@ -463,8 +528,12 @@ namespace EpicGames.Horde.Storage.Nodes
 			NodeReader nodeReader = new NodeReader(nodeData);
 			while (nodeReader.GetMemory(0).Length > 0)
 			{
-				NodeRef nodeRef = nodeReader.ReadNodeRef();
-				await ChunkedDataNode.CopyToStreamAsync(nodeRef.Handle!, outputStream, cancellationToken);
+				BlobHandle handle = nodeReader.ReadBlobHandle();
+				if (nodeReader.Version >= 2)
+				{
+					_ = nodeReader.ReadUnsignedVarInt();
+				}
+				await ChunkedDataNode.CopyToStreamAsync(handle, outputStream, cancellationToken);
 			}
 		}
 	}
diff --git a/Engine/Source/Programs/Shared/EpicGames.Horde/Storage/Nodes/ChunkedDataWriter.cs b/Engine/Source/Programs/Shared/EpicGames.Horde/Storage/Nodes/ChunkedDataWriter.cs
--- a/Engine/Source/Programs/Shared/EpicGames.Horde/Storage/Nodes/ChunkedDataWriter.cs
+++ b/Engine/Source/Programs/Shared/EpicGames.Horde/Storage/Nodes/ChunkedDataWriter.cs
@@ -78,7 +78,7 @@ namespace EpicGames.Horde.Storage.Nodes
 
 		// Tree state
 		long _totalLength;
-		readonly List<NodeRef<ChunkedDataNode>> _leafHandles = new List<NodeRef<ChunkedDataNode>>();
+		readonly List<ChunkedDataNodeRef> _leafHandles = new List<ChunkedDataNodeRef>();
 
 		// Leaf node state
 		uint _leafHash;
@@ -346,7 +346,7 @@ namespace EpicGames.Horde.Storage.Nodes
 		async ValueTask FlushLeafNodeAsync(CancellationToken cancellationToken)
 		{
 			BlobHandle handle = await _writer.WriteBlobAsync(_leafLength, Array.Empty<BlobHandle>(), s_leafNodeType, cancellationToken);
-			_leafHandles.Add(new NodeRef<ChunkedDataNode>(handle));
+			_leafHandles.Add(new ChunkedDataNodeRef(ChunkedDataNodeType.Leaf, handle));
 			ResetLeafState();
 		}
 	}
diff --git a/Engine/Source/Programs/Shared/EpicGames.Horde/Storage/Nodes/DdcRefNode.cs b/Engine/Source/Programs/Shared/EpicGames.Horde/Storage/Nodes/DdcRefNode.cs
--- a/Engine/Source/Programs/Shared/EpicGames.Horde/Storage/Nodes/DdcRefNode.cs
+++ b/Engine/Source/Programs/Shared/EpicGames.Horde/Storage/Nodes/DdcRefNode.cs
@@ -37,7 +37,7 @@ namespace Horde.Server.Ddc
 		public DdcRefNode(NodeReader reader)
 		{
 			RootHash = reader.ReadIoHash();
-			References = reader.ReadList(x => (reader.ReadIoHash(), reader.ReadNodeHandle()));
+			References = reader.ReadList(x => (reader.ReadIoHash(), reader.ReadBlobHandle()));
 		}
 
 		/// <inheritdoc/>
