Horde: Have BlobData instances implement IDisposable, allowing us to manage lifetimes of the underlying buffers data they reference.

diff --git a/Engine/Source/Programs/Horde/Horde.Server/Ddc/BlobService.cs b/Engine/Source/Programs/Horde/Horde.Server/Ddc/BlobService.cs
--- a/Engine/Source/Programs/Horde/Horde.Server/Ddc/BlobService.cs
+++ b/Engine/Source/Programs/Horde/Horde.Server/Ddc/BlobService.cs
@@ -104,7 +104,7 @@ namespace Horde.Server.Ddc
 				throw new BlobNotFoundException(ns, blob);
 			}
 
-			BlobData data = await alias.Target.ReadAsync(cancellationToken);
+			using BlobData data = await alias.Target.ReadAsync(cancellationToken);
 			return new BlobContents(data.Data.ToArray());
 		}
 
diff --git a/Engine/Source/Programs/Horde/Horde.Server/Ddc/RefService.cs b/Engine/Source/Programs/Horde/Horde.Server/Ddc/RefService.cs
--- a/Engine/Source/Programs/Horde/Horde.Server/Ddc/RefService.cs
+++ b/Engine/Source/Programs/Horde/Horde.Server/Ddc/RefService.cs
@@ -50,7 +50,7 @@ namespace Horde.Server.Ddc
 			}
 
 			BlobHandle blobHandle = blobAlias.Target;
-			BlobData blobContents = await blobHandle.ReadAsync(cancellationToken);
+			using BlobData blobContents = await blobHandle.ReadAsync(cancellationToken);
 			CbObject payload = new CbObject(blobContents.Data);
 
 			BlobId[] referencedBlobs = Array.Empty<BlobId>();
diff --git a/Engine/Source/Programs/Horde/Horde.Server/Storage/StorageController.cs b/Engine/Source/Programs/Horde/Horde.Server/Storage/StorageController.cs
--- a/Engine/Source/Programs/Horde/Horde.Server/Storage/StorageController.cs
+++ b/Engine/Source/Programs/Horde/Horde.Server/Storage/StorageController.cs
@@ -580,7 +580,7 @@ namespace Horde.Server.Storage
 
 			object content;
 
-			BlobData nodeData = await storageClient.ReadNodeDataAsync(new BundleNodeLocator(export.Hash, locator, exportIdx), cancellationToken);
+			using BlobData nodeData = await storageClient.ReadNodeDataAsync(new BundleNodeLocator(export.Hash, locator, exportIdx), cancellationToken);
 
 			Node node = Node.Deserialize(nodeData);
 			switch (node)
diff --git a/Engine/Source/Programs/Shared/EpicGames.Horde/Storage/BlobData.cs b/Engine/Source/Programs/Shared/EpicGames.Horde/Storage/BlobData.cs
--- a/Engine/Source/Programs/Shared/EpicGames.Horde/Storage/BlobData.cs
+++ b/Engine/Source/Programs/Shared/EpicGames.Horde/Storage/BlobData.cs
@@ -7,7 +7,54 @@ using EpicGames.Core;
 namespace EpicGames.Horde.Storage
 {
 	/// <summary>
-	/// Data for an individual node
+	/// Data for an individual node. Must be disposed after use.
 	/// </summary>
-	public record struct BlobData(BlobType Type, IoHash Hash, ReadOnlyMemory<byte> Data, IReadOnlyList<BlobHandle> Refs);
+	public class BlobData : IDisposable
+	{
+		/// <summary>
+		/// Type of the blob
+		/// </summary>
+		public BlobType Type { get; }
+
+		/// <summary>
+		/// Hash of the node data
+		/// </summary>
+		public IoHash Hash { get; }
+
+		/// <summary>
+		/// Raw data for the blob. Lifetime of this data is tied to the lifetime of the <see cref="BlobData"/> object; consumers must not retain references to it.
+		/// </summary>
+		public ReadOnlyMemory<byte> Data { get; }
+
+		/// <summary>
+		/// Handles to referenced blobs
+		/// </summary>
+		public IReadOnlyList<BlobHandle> Refs { get; }
+
+		/// <summary>
+		/// Constructor
+		/// </summary>
+		public BlobData(BlobType type, IoHash hash, ReadOnlyMemory<byte> data, IReadOnlyList<BlobHandle> refs)
+		{
+			Type = type;
+			Hash = hash;
+			Data = data;
+			Refs = refs;
+		}
+
+		/// <inheritdoc/>
+		public void Dispose()
+		{
+			Dispose(true);
+			GC.SuppressFinalize(this);
+		}
+
+		/// <summary>
+		/// Overridable dispose method
+		/// </summary>
+		/// <param name="disposing">True if derived instances should dispose managed resources. False when called from a finalizer.</param>
+		protected virtual void Dispose(bool disposing)
+		{
+		}
+	}
 }
diff --git a/Engine/Source/Programs/Shared/EpicGames.Horde/Storage/BlobHandle.cs b/Engine/Source/Programs/Shared/EpicGames.Horde/Storage/BlobHandle.cs
--- a/Engine/Source/Programs/Shared/EpicGames.Horde/Storage/BlobHandle.cs
+++ b/Engine/Source/Programs/Shared/EpicGames.Horde/Storage/BlobHandle.cs
@@ -30,7 +30,7 @@ namespace EpicGames.Horde.Storage
 		/// <param name="cancellationToken">Cancellation token for the operation</param>
 		public virtual async ValueTask<BlobType> GetTypeAsync(CancellationToken cancellationToken = default)
 		{
-			BlobData data = await ReadAsync(cancellationToken);
+			using BlobData data = await ReadAsync(cancellationToken);
 			return data.Type;
 		}
 
@@ -40,12 +40,12 @@ namespace EpicGames.Horde.Storage
 		/// <param name="cancellationToken">Cancellation token for the operation</param>
 		public virtual async ValueTask<IReadOnlyList<BlobHandle>> GetRefsAsync(CancellationToken cancellationToken = default)
 		{
-			BlobData data = await ReadAsync(cancellationToken);
+			using BlobData data = await ReadAsync(cancellationToken);
 			return data.Refs;
 		}
 
 		/// <summary>
-		/// Creates a reader for this node's data
+		/// Reads the blob's data
 		/// </summary>
 		/// <param name="cancellationToken">Cancellation token for the operation</param>
 		public abstract ValueTask<BlobData> ReadAsync(CancellationToken cancellationToken = default);
@@ -59,7 +59,7 @@ namespace EpicGames.Horde.Storage
 		/// <returns>Number of bytes that were read</returns>
 		public virtual async ValueTask<int> ReadPartialAsync(int offset, Memory<byte> buffer, CancellationToken cancellationToken = default)
 		{
-			BlobData data = await ReadAsync(cancellationToken);
+			using BlobData data = await ReadAsync(cancellationToken);
 
 			int length = data.Data.Length - offset;
 			if (length < 0)
diff --git a/Engine/Source/Programs/Shared/EpicGames.Horde/Storage/Bundles/BundleReader.cs b/Engine/Source/Programs/Shared/EpicGames.Horde/Storage/Bundles/BundleReader.cs
--- a/Engine/Source/Programs/Shared/EpicGames.Horde/Storage/Bundles/BundleReader.cs
+++ b/Engine/Source/Programs/Shared/EpicGames.Horde/Storage/Bundles/BundleReader.cs
@@ -579,8 +579,8 @@ namespace EpicGames.Horde.Storage.Bundles
 		/// <returns>Node data read from the given bundle</returns>
 		public async ValueTask<Node> ReadNodeAsync(BundleNodeLocator locator, CancellationToken cancellationToken = default)
 		{
-			BlobData nodeData = await ReadNodeDataAsync(locator, cancellationToken);
-			return Node.Deserialize(nodeData);
+			using BlobData blobData = await ReadNodeDataAsync(locator, cancellationToken);
+			return Node.Deserialize(blobData);
 		}
 
 		/// <summary>
diff --git a/Engine/Source/Programs/Shared/EpicGames.Horde/Storage/Node.cs b/Engine/Source/Programs/Shared/EpicGames.Horde/Storage/Node.cs
--- a/Engine/Source/Programs/Shared/EpicGames.Horde/Storage/Node.cs
+++ b/Engine/Source/Programs/Shared/EpicGames.Horde/Storage/Node.cs
@@ -402,8 +402,8 @@ namespace EpicGames.Horde.Storage
 		/// <returns></returns>
 		public static async ValueTask<TNode> ReadNodeAsync<TNode>(this BlobHandle handle, CancellationToken cancellationToken = default) where TNode : Node
 		{
-			BlobData nodeData = await handle.ReadAsync(cancellationToken);
-			return Node.Deserialize<TNode>(nodeData);
+			using BlobData blobData = await handle.ReadAsync(cancellationToken);
+			return Node.Deserialize<TNode>(blobData);
 		}
 
 		/// <summary>
@@ -439,8 +439,8 @@ namespace EpicGames.Horde.Storage
 				return null;
 			}
 
-			BlobData nodeData = await refTarget.ReadAsync(cancellationToken);
-			return Node.Deserialize<TNode>(nodeData);
+			using BlobData blobData = await refTarget.ReadAsync(cancellationToken);
+			return Node.Deserialize<TNode>(blobData);
 		}
 
 		/// <summary>
diff --git a/Engine/Source/Programs/Shared/EpicGames.Horde/Storage/NodeRef.cs b/Engine/Source/Programs/Shared/EpicGames.Horde/Storage/NodeRef.cs
--- a/Engine/Source/Programs/Shared/EpicGames.Horde/Storage/NodeRef.cs
+++ b/Engine/Source/Programs/Shared/EpicGames.Horde/Storage/NodeRef.cs
@@ -52,8 +52,8 @@ namespace EpicGames.Horde.Storage
 		/// <returns></returns>
 		public async ValueTask<Node> ExpandAsync(CancellationToken cancellationToken = default)
 		{
-			BlobData nodeData = await Handle!.ReadAsync(cancellationToken);
-			return Node.Deserialize(nodeData);
+			using BlobData blobData = await Handle!.ReadAsync(cancellationToken);
+			return Node.Deserialize(blobData);
 		}
 	}
 
diff --git a/Engine/Source/Programs/Shared/EpicGames.Horde/Storage/Nodes/ChunkedDataNode.cs b/Engine/Source/Programs/Shared/EpicGames.Horde/Storage/Nodes/ChunkedDataNode.cs
--- a/Engine/Source/Programs/Shared/EpicGames.Horde/Storage/Nodes/ChunkedDataNode.cs
+++ b/Engine/Source/Programs/Shared/EpicGames.Horde/Storage/Nodes/ChunkedDataNode.cs
@@ -37,14 +37,14 @@ namespace EpicGames.Horde.Storage.Nodes
 		/// <param name="cancellationToken">Cancellation token for the operation</param>
 		public static async Task CopyToStreamAsync(BlobHandle handle, Stream outputStream, CancellationToken cancellationToken)
 		{
-			BlobData nodeData = await handle.ReadAsync(cancellationToken);
-			if (nodeData.Type.Guid == s_leafNodeGuid)
+			using BlobData blobData = await handle.ReadAsync(cancellationToken);
+			if (blobData.Type.Guid == s_leafNodeGuid)
 			{
-				await LeafChunkedDataNode.CopyToStreamAsync(nodeData, outputStream, cancellationToken);
+				await LeafChunkedDataNode.CopyToStreamAsync(blobData, outputStream, cancellationToken);
 			}
-			else if (nodeData.Type.Guid == s_interiorNodeGuid)
+			else if (blobData.Type.Guid == s_interiorNodeGuid)
 			{
-				await InteriorChunkedDataNode.CopyToStreamAsync(nodeData, outputStream, cancellationToken);
+				await InteriorChunkedDataNode.CopyToStreamAsync(blobData, outputStream, cancellationToken);
 			}
 			else
 			{
diff --git a/Engine/Source/Programs/Shared/EpicGames.Horde/Storage/Workspace.cs b/Engine/Source/Programs/Shared/EpicGames.Horde/Storage/Workspace.cs
--- a/Engine/Source/Programs/Shared/EpicGames.Horde/Storage/Workspace.cs
+++ b/Engine/Source/Programs/Shared/EpicGames.Horde/Storage/Workspace.cs
@@ -755,8 +755,8 @@ namespace EpicGames.Horde.Storage
 			}
 
 			// Otherwise 
-			BlobData nodeData = await nodeRef.Handle.ReadAsync(cancellationToken);
-			await LeafChunkedDataNode.CopyToStreamAsync(nodeData, outputStream, cancellationToken);
+			using BlobData blobData = await nodeRef.Handle.ReadAsync(cancellationToken);
+			await LeafChunkedDataNode.CopyToStreamAsync(blobData, outputStream, cancellationToken);
 		}
 
 		FileState? TryMoveCachedDataAsync(IoHash hash, DirectoryState targetDirState, string targetName)
