EpicGames.Core: Add a BackgroundTask implementation that returns a result. Also remove IDisposable implementation; callers should always prefer the IAsyncDisposable implementation.

diff --git a/Engine/Source/Programs/Horde/Horde.Agent/Services/StatusService.cs b/Engine/Source/Programs/Horde/Horde.Agent/Services/StatusService.cs
--- a/Engine/Source/Programs/Horde/Horde.Agent/Services/StatusService.cs
+++ b/Engine/Source/Programs/Horde/Horde.Agent/Services/StatusService.cs
@@ -59,7 +59,7 @@ namespace Horde.Agent.Services
 		/// <inheritdoc/>
 		public async Task StopAsync(CancellationToken cancellationToken)
 		{
-			await _task.StopAsync();
+			await _task.StopAsync(cancellationToken);
 		}
 
 		/// <summary>
diff --git a/Engine/Source/Programs/Horde/Horde.Server.Tests/ControllerIntegrationTest.cs b/Engine/Source/Programs/Horde/Horde.Server.Tests/ControllerIntegrationTest.cs
--- a/Engine/Source/Programs/Horde/Horde.Server.Tests/ControllerIntegrationTest.cs
+++ b/Engine/Source/Programs/Horde/Horde.Server.Tests/ControllerIntegrationTest.cs
@@ -72,7 +72,7 @@ public class TestWebApplicationFactory<TStartup> : WebApplicationFactory<TStartu
 	}
 }
 
-public class ControllerIntegrationTest : IDisposable
+public class ControllerIntegrationTest : IAsyncDisposable
 {
 	private readonly Lazy<Task<Fixture>> _fixture;
 
@@ -91,16 +91,12 @@ public class ControllerIntegrationTest : IDisposable
 
 	protected IServiceProvider ServiceProvider => Factory.Services;
 
-	public void Dispose()
+	public virtual async ValueTask DisposeAsync()
 	{
-		Dispose(true);
-		GC.SuppressFinalize(this);
-	}
-
-	protected virtual void Dispose(bool disposing)
-	{
-		Factory.Dispose();
+		await Factory.DisposeAsync();
 		MongoDbInstance.Dispose();
+
+		GC.SuppressFinalize(this);
 	}
 
 	public Task<Fixture> GetFixtureAsync()
diff --git a/Engine/Source/Programs/Horde/Horde.Server.Tests/DatabaseIntegrationTest.cs b/Engine/Source/Programs/Horde/Horde.Server.Tests/DatabaseIntegrationTest.cs
--- a/Engine/Source/Programs/Horde/Horde.Server.Tests/DatabaseIntegrationTest.cs
+++ b/Engine/Source/Programs/Horde/Horde.Server.Tests/DatabaseIntegrationTest.cs
@@ -38,7 +38,7 @@ namespace Horde.Server.Tests
             CurrentValue = currentValue;
         }
 
-        public T Get(string name)
+        public T Get(string? name)
         {
             return CurrentValue;
         }
diff --git a/Engine/Source/Programs/Horde/Horde.Server.Tests/IssueServiceTests.cs b/Engine/Source/Programs/Horde/Horde.Server.Tests/IssueServiceTests.cs
--- a/Engine/Source/Programs/Horde/Horde.Server.Tests/IssueServiceTests.cs
+++ b/Engine/Source/Programs/Horde/Horde.Server.Tests/IssueServiceTests.cs
@@ -62,7 +62,7 @@ namespace Horde.Server.Tests
 				}
 			}
 
-			public IDisposable BeginScope<TState>(TState state) => null!;
+			public IDisposable? BeginScope<TState>(TState state) where TState : notnull => null!;
 
 			public bool IsEnabled(LogLevel logLevel) => true;
 
diff --git a/Engine/Source/Programs/Horde/Horde.Server/Agents/Fleet/AwsAutoScalingLifecycle.cs b/Engine/Source/Programs/Horde/Horde.Server/Agents/Fleet/AwsAutoScalingLifecycle.cs
--- a/Engine/Source/Programs/Horde/Horde.Server/Agents/Fleet/AwsAutoScalingLifecycle.cs
+++ b/Engine/Source/Programs/Horde/Horde.Server/Agents/Fleet/AwsAutoScalingLifecycle.cs
@@ -133,7 +133,7 @@ public sealed class AwsAutoScalingLifecycleService : IHostedService, IAsyncDispo
 	{
 		foreach (BackgroundTask bgTask in _lifecycleEventListenerTasks)
 		{
-			await bgTask.StopAsync();
+			await bgTask.StopAsync(cancellationToken);
 		}
 
 		await _updateLifecyclesTicker.StopAsync();
diff --git a/Engine/Source/Programs/Horde/Horde.Server/Configuration/ConfigService.cs b/Engine/Source/Programs/Horde/Horde.Server/Configuration/ConfigService.cs
--- a/Engine/Source/Programs/Horde/Horde.Server/Configuration/ConfigService.cs
+++ b/Engine/Source/Programs/Horde/Horde.Server/Configuration/ConfigService.cs
@@ -155,7 +155,7 @@ namespace Horde.Server.Configuration
 		/// <inheritdoc/>
 		public async ValueTask DisposeAsync()
 		{
-			_updateTask.Dispose();
+			await _updateTask.DisposeAsync();
 			await _ticker.DisposeAsync();
 		}
 
@@ -281,7 +281,7 @@ namespace Horde.Server.Configuration
 		/// <inheritdoc/>
 		public async Task StopAsync(CancellationToken cancellationToken)
 		{
-			await _updateTask.StopAsync();
+			await _updateTask.StopAsync(cancellationToken);
 			if (_serverSettings.IsRunModeActive(RunMode.Worker))
 			{
 				await _ticker.StopAsync();
diff --git a/Engine/Source/Programs/Horde/Horde.Server/Notifications/NotificationService.cs b/Engine/Source/Programs/Horde/Horde.Server/Notifications/NotificationService.cs
--- a/Engine/Source/Programs/Horde/Horde.Server/Notifications/NotificationService.cs
+++ b/Engine/Source/Programs/Horde/Horde.Server/Notifications/NotificationService.cs
@@ -187,7 +187,7 @@ namespace Horde.Server.Notifications
 		public async Task StopAsync(CancellationToken cancellationToken)
 		{
 			await _ticker.StopAsync();
-			await _backgroundTask.StopAsync();
+			await _backgroundTask.StopAsync(cancellationToken);
 		}
 
 		/// <inheritdoc/>
diff --git a/Engine/Source/Programs/Horde/Horde.Server/Notifications/Sinks/SlackNotificationSink.cs b/Engine/Source/Programs/Horde/Horde.Server/Notifications/Sinks/SlackNotificationSink.cs
--- a/Engine/Source/Programs/Horde/Horde.Server/Notifications/Sinks/SlackNotificationSink.cs
+++ b/Engine/Source/Programs/Horde/Horde.Server/Notifications/Sinks/SlackNotificationSink.cs
@@ -307,7 +307,7 @@ namespace Horde.Server.Notifications.Sinks
 		{
 			await _escalateTicker.StopAsync();
 			await _issueQueueTicker.StopAsync();
-			await _backgroundTask.StopAsync();
+			await _backgroundTask.StopAsync(cancellationToken);
 		}
 
 		#region Avatars
diff --git a/Engine/Source/Programs/Horde/Horde.Server/Server/MongoService.cs b/Engine/Source/Programs/Horde/Horde.Server/Server/MongoService.cs
--- a/Engine/Source/Programs/Horde/Horde.Server/Server/MongoService.cs
+++ b/Engine/Source/Programs/Horde/Horde.Server/Server/MongoService.cs
@@ -262,11 +262,20 @@ namespace Horde.Server.Server
 		readonly Dictionary<string, Task> _collectionUpgradeTasks = new Dictionary<string, Task>(StringComparer.Ordinal);
 		readonly Task<bool> _setSchemaVersionTask;
 
+		static string? s_existingInstance;
+
 		/// <summary>
 		/// Constructor
 		/// </summary>
 		public MongoService(IOptions<ServerSettings> settingsSnapshot, RedisService redisService, Tracer tracer, ILogger<MongoService> logger, ILoggerFactory loggerFactory)
 		{
+			if (s_existingInstance != null)
+			{
+				throw new Exception("Existing instance on MongoService!");
+			}
+
+			s_existingInstance = Environment.StackTrace;
+
 			Settings = settingsSnapshot.Value;
 			_redisService = redisService;
 			_tracer = tracer;
@@ -369,6 +378,7 @@ namespace Horde.Server.Server
 				_mongoProcessGroup = null;
 			}
 			//_upgradeSema.Dispose();
+			s_existingInstance = null;
 		}
 
 		/// <summary>
diff --git a/Engine/Source/Programs/Shared/EpicGames.Core/BackgroundTask.cs b/Engine/Source/Programs/Shared/EpicGames.Core/BackgroundTask.cs
--- a/Engine/Source/Programs/Shared/EpicGames.Core/BackgroundTask.cs
+++ b/Engine/Source/Programs/Shared/EpicGames.Core/BackgroundTask.cs
@@ -9,16 +9,16 @@ namespace EpicGames.Core
 	/// <summary>
 	/// Runs a task in the background and allows stopping it on demand
 	/// </summary>
-	public sealed class BackgroundTask : IDisposable, IAsyncDisposable
+	public sealed class BackgroundTask : IAsyncDisposable
 	{
 		readonly Func<CancellationToken, Task> _runTask;
-		Task _task = Task.CompletedTask;
+		Task? _task;
 		readonly CancellationTokenSource _cancellationTokenSource = new CancellationTokenSource();
 
 		/// <summary>
-		/// Accessor for the inner task
+		/// Accessor for the underlying task
 		/// </summary>
-		public Task Task => _task;
+		public Task? Task => _task;
 
 		/// <summary>
 		/// Constructor. Note that the task does not start until <see cref="Start"/> is called.
@@ -29,17 +29,6 @@ namespace EpicGames.Core
 			_runTask = runTask;
 		}
 
-		/// <inheritdoc/>
-		public void Dispose()
-		{
-			if (!_task.IsCompleted)
-			{
-				_cancellationTokenSource.Cancel();
-				StopAsync().Wait();
-			}
-			_cancellationTokenSource.Dispose();
-		}
-
 		/// <inheritdoc/>
 		public async ValueTask DisposeAsync()
 		{
@@ -59,12 +48,24 @@ namespace EpicGames.Core
 			return task;
 		}
 
+		/// <summary>
+		/// Creates and starts a new background task instance
+		/// </summary>
+		/// <param name="runTask"></param>
+		/// <returns>New background task instance</returns>
+		public static BackgroundTask<T> StartNew<T>(Func<CancellationToken, Task<T>> runTask)
+		{
+			BackgroundTask<T> task = new BackgroundTask<T>(runTask);
+			task.Start();
+			return task;
+		}
+
 		/// <summary>
 		/// Starts the task
 		/// </summary>
 		public void Start()
 		{
-			if (!_task.IsCompleted)
+			if (_task != null)
 			{
 				throw new InvalidOperationException("Background task is already running");
 			}
@@ -74,20 +75,107 @@ namespace EpicGames.Core
 		/// <summary>
 		/// Signals the cancellation token and waits for the task to finish
 		/// </summary>
-		/// <returns></returns>
-		public async Task StopAsync()
+		/// <param name="cancellationToken">Cancellation token for the operation</param>
+		public async Task StopAsync(CancellationToken cancellationToken = default)
 		{
-			if (!_task.IsCompleted)
+			if (_task != null && !_task.IsCompleted)
 			{
 				try
 				{
 					_cancellationTokenSource.Cancel();
-					await _task;
+					await _task.WaitAsync(cancellationToken);
 				}
 				catch (OperationCanceledException)
 				{
 				}
 			}
 		}
+
+		/// <summary>
+		/// Waits for the task to complete
+		/// </summary>
+		/// <param name="cancellationToken">Cancellation token for the operation</param>
+		public async Task WaitAsync(CancellationToken cancellationToken)
+		{
+			if (_task == null)
+			{
+				throw new InvalidOperationException("Task has not been started.");
+			}
+			await _task.WaitAsync(cancellationToken);
+		}
+	}
+
+	/// <summary>
+	/// Runs a task in the background and allows stopping it on demand
+	/// </summary>
+	public sealed class BackgroundTask<T> : IAsyncDisposable
+	{
+		readonly Func<CancellationToken, Task<T>> _runTask;
+		Task<T>? _task;
+		readonly CancellationTokenSource _cancellationTokenSource = new CancellationTokenSource();
+
+		/// <summary>
+		/// Accessor for the underlying task
+		/// </summary>
+		public Task<T>? Task => _task;
+
+		/// <summary>
+		/// Constructor. Note that the task does not start until <see cref="Start"/> is called.
+		/// </summary>
+		/// <param name="runTask"></param>
+		public BackgroundTask(Func<CancellationToken, Task<T>> runTask)
+		{
+			_runTask = runTask;
+		}
+
+		/// <inheritdoc/>
+		public async ValueTask DisposeAsync()
+		{
+			await StopAsync();
+			_cancellationTokenSource.Dispose();
+		}
+
+		/// <summary>
+		/// Starts the task
+		/// </summary>
+		public void Start()
+		{
+			if (_task != null)
+			{
+				throw new InvalidOperationException("Background task is already running");
+			}
+			_task = System.Threading.Tasks.Task.Run(() => _runTask(_cancellationTokenSource.Token), _cancellationTokenSource.Token);
+		}
+
+		/// <summary>
+		/// Signals the cancellation token and waits for the task to finish
+		/// </summary>
+		/// <param name="cancellationToken">Cancellation token for the operation</param>
+		public async Task StopAsync(CancellationToken cancellationToken = default)
+		{
+			if (_task != null && !_task.IsCompleted)
+			{
+				try
+				{
+					_cancellationTokenSource.Cancel();
+					await _task.WaitAsync(cancellationToken);
+				}
+				catch (OperationCanceledException)
+				{
+				}
+			}
+		}
+
+		/// <summary>
+		/// Waits for the task to complete
+		/// </summary>
+		public Task<T> WaitAsync(CancellationToken cancellationToken = default)
+		{
+			if (_task == null)
+			{
+				throw new InvalidOperationException("Task has not been started.");
+			}
+			return _task.WaitAsync(cancellationToken);
+		}
 	}
 }
diff --git a/Engine/Source/Programs/Shared/EpicGames.Horde/Compute/Clients/AgentComputeClient.cs b/Engine/Source/Programs/Shared/EpicGames.Horde/Compute/Clients/AgentComputeClient.cs
--- a/Engine/Source/Programs/Shared/EpicGames.Horde/Compute/Clients/AgentComputeClient.cs
+++ b/Engine/Source/Programs/Shared/EpicGames.Horde/Compute/Clients/AgentComputeClient.cs
@@ -94,7 +94,7 @@ namespace EpicGames.Horde.Compute.Clients
 			listener.Bind(new IPEndPoint(IPAddress.Loopback, _port));
 			listener.Listen();
 
-			using BackgroundTask agentTask = BackgroundTask.StartNew(ctx => RunAgentAsync(_hordeAgentAssembly, _port, logger, ctx));
+			await using BackgroundTask agentTask = BackgroundTask.StartNew(ctx => RunAgentAsync(_hordeAgentAssembly, _port, logger, ctx));
 			using Socket tcpSocket = await listener.AcceptAsync(cancellationToken);
 
 			await using RemoteComputeSocket socket = new RemoteComputeSocket(new TcpTransport(tcpSocket), _logger);
diff --git a/Engine/Source/Programs/Shared/EpicGames.Horde/Compute/ComputeSocket.cs b/Engine/Source/Programs/Shared/EpicGames.Horde/Compute/ComputeSocket.cs
--- a/Engine/Source/Programs/Shared/EpicGames.Horde/Compute/ComputeSocket.cs
+++ b/Engine/Source/Programs/Shared/EpicGames.Horde/Compute/ComputeSocket.cs
@@ -639,7 +639,7 @@ namespace EpicGames.Horde.Compute
 				_recvBuffers.Add(channelId, new RecvBuffer(recvBuffer.CreateWriter()));
 
 				// Only start the receive task once we have a buffer to receive data, otherwise we discard data from the remote
-				if (_recvTask.Task.IsCompleted)
+				if (_recvTask.Task == null)
 				{
 					_recvTask.Start();
 				}
diff --git a/Engine/Source/Programs/Shared/EpicGames.Horde/Storage/StorageBackendCache.cs b/Engine/Source/Programs/Shared/EpicGames.Horde/Storage/StorageBackendCache.cs
--- a/Engine/Source/Programs/Shared/EpicGames.Horde/Storage/StorageBackendCache.cs
+++ b/Engine/Source/Programs/Shared/EpicGames.Horde/Storage/StorageBackendCache.cs
@@ -51,7 +51,7 @@ namespace EpicGames.Horde.Storage
 				}
 			}
 
-			public Task WaitAsync(CancellationToken cancellationToken) => _readTask.Task.WaitAsync(cancellationToken);
+			public Task WaitAsync(CancellationToken cancellationToken) => _readTask.WaitAsync(cancellationToken);
 		}
 
 		sealed class BackendWrapper : IStorageBackend
