Horde: Retry HTTP requests that fail due to a timeout.

diff --git a/Engine/Source/Programs/Horde/Horde.Agent/AgentApp.cs b/Engine/Source/Programs/Horde/Horde.Agent/AgentApp.cs
--- a/Engine/Source/Programs/Horde/Horde.Agent/AgentApp.cs
+++ b/Engine/Source/Programs/Horde/Horde.Agent/AgentApp.cs
@@ -24,6 +24,8 @@ using Microsoft.Extensions.Options;
 using OpenTracing.Util;
 using Polly;
 using Polly.Extensions.Http;
+using Polly.Retry;
+using Polly.Timeout;
 
 namespace Horde.Agent
 {
@@ -173,32 +175,13 @@ namespace Horde.Agent
 				return builder.WaitAndRetryAsync(new[] { TimeSpan.FromSeconds(1), TimeSpan.FromSeconds(5), TimeSpan.FromSeconds(10) });
 			});
 
-			services.AddHttpClient(HordeHttpClient.HttpClientName, builder => 
-				{
-					builder.BaseAddress = serverProfile.Url;
-					builder.Timeout = TimeSpan.FromSeconds(240); // Global timeout
-				})
+			services.AddHordeHttpClient(client => client.BaseAddress = serverProfile.Url, useAuthChallenge: false)
 				.ConfigurePrimaryHttpMessageHandler(() => new SocketsHttpHandler
-					{
-						MaxConnectionsPerServer = 16,
-						PooledConnectionIdleTimeout = TimeSpan.FromMinutes(15),
-					})
-				.AddPolicyHandler((serviceProvider, request) => Policy.TimeoutAsync<HttpResponseMessage>(30,
-					(outcome, timespan, context) => { 
-						serviceProvider.GetRequiredService<ILogger<HttpStorageClient>>().LogWarning("Http request timed out after {Time}s.", (int)timespan.TotalSeconds); 
-						return Task.CompletedTask; 
-					}))
-				.AddPolicyHandler((serviceProvider, request) => HttpPolicyExtensions.HandleTransientHttpError()
-					.WaitAndRetryAsync(new[]
-						{
-							TimeSpan.FromSeconds(1),
-							TimeSpan.FromSeconds(5),
-							TimeSpan.FromSeconds(10),
-							TimeSpan.FromSeconds(30),
-							TimeSpan.FromSeconds(30),
-						},
-						(outcome, timespan, retryAttempt, context) => serviceProvider.GetRequiredService<ILogger<HttpStorageClient>>().LogWarning("Http request failed. Delaying for {DelayMs}ms (attempt #{RetryNum}).", timespan.TotalMilliseconds, retryAttempt)
-					));
+				{
+					MaxConnectionsPerServer = 16,
+					PooledConnectionIdleTimeout = TimeSpan.FromMinutes(15),
+				});
+
 			services.AddHttpClient(AwsInstanceLifecycleService.HttpClientName);
 			services.AddSingleton<AwsInstanceLifecycleService>();
 			if (settings.EnableAwsEc2Support)
diff --git a/Engine/Source/Programs/Shared/EpicGames.Horde/HordeHttpClient.cs b/Engine/Source/Programs/Shared/EpicGames.Horde/HordeHttpClient.cs
--- a/Engine/Source/Programs/Shared/EpicGames.Horde/HordeHttpClient.cs
+++ b/Engine/Source/Programs/Shared/EpicGames.Horde/HordeHttpClient.cs
@@ -16,9 +16,13 @@ using EpicGames.Horde.Dashboard;
 using EpicGames.Horde.Projects;
 using EpicGames.Horde.Secrets;
 using EpicGames.Horde.Server;
+using EpicGames.Horde.Storage.Clients;
 using Microsoft.Extensions.DependencyInjection;
+using Microsoft.Extensions.Logging;
 using Polly;
 using Polly.Extensions.Http;
+using Polly.Retry;
+using Polly.Timeout;
 
 namespace EpicGames.Horde
 {
@@ -359,9 +363,9 @@ namespace EpicGames.Horde
 		/// </summary>
 		/// <param name="services">Service collection to add services to</param>
 		/// <param name="useAuthChallenge">Whether to prompt the user to authenticate if necessary</param>
-		public static void AddHordeHttpClient(this IServiceCollection services, bool useAuthChallenge = true)
+		public static IHttpClientBuilder AddHordeHttpClient(this IServiceCollection services, bool useAuthChallenge = true)
 		{
-			services.AddHordeHttpClient((sp, client) => { }, useAuthChallenge);
+			return services.AddHordeHttpClient((sp, client) => { }, useAuthChallenge);
 		}
 
 		/// <summary>
@@ -370,9 +374,9 @@ namespace EpicGames.Horde
 		/// <param name="services">Service collection to add services to</param>
 		/// <param name="configureClient">Callback to modify options for the http client</param>
 		/// <param name="useAuthChallenge">Whether to prompt the user to authenticate if necessary</param>
-		public static void AddHordeHttpClient(this IServiceCollection services, Action<HttpClient> configureClient, bool useAuthChallenge = true)
+		public static IHttpClientBuilder AddHordeHttpClient(this IServiceCollection services, Action<HttpClient> configureClient, bool useAuthChallenge = true)
 		{
-			services.AddHordeHttpClient((sp, client) => configureClient(client), useAuthChallenge);
+			return services.AddHordeHttpClient((sp, client) => configureClient(client), useAuthChallenge);
 		}
 
 		/// <summary>
@@ -381,11 +385,13 @@ namespace EpicGames.Horde
 		/// <param name="services">Service collection to add services to</param>
 		/// <param name="configureClient">Callback to modify options for the http client</param>
 		/// <param name="useAuthChallenge">Whether to prompt the user to authenticate if necessary</param>
-		public static void AddHordeHttpClient(this IServiceCollection services, Action<IServiceProvider, HttpClient> configureClient, bool useAuthChallenge = true)
+		public static IHttpClientBuilder AddHordeHttpClient(this IServiceCollection services, Action<IServiceProvider, HttpClient> configureClient, bool useAuthChallenge = true)
 		{
 			// Sets defaults from the environment before calling the user provided configuration method
 			void ConfigureClientFromEnvironment(IServiceProvider serviceProvider, HttpClient httpClient)
 			{
+				httpClient.Timeout = TimeSpan.FromSeconds(240); // Global timeout
+
 				configureClient(serviceProvider, httpClient);
 
 				// Only use the token from the environment if the configured base address is missing or matches the one configured in the environment
@@ -410,9 +416,8 @@ namespace EpicGames.Horde
 			}
 
 			IHttpClientBuilder builder = services.AddHttpClient<HordeHttpClient>(HordeHttpClient.HttpClientName, ConfigureClientFromEnvironment)
-				.AddPolicyHandler(HttpPolicyExtensions
-					.HandleTransientHttpError()
-					.WaitAndRetryAsync(new[] { TimeSpan.FromSeconds(2.0), TimeSpan.FromSeconds(5.0), TimeSpan.FromSeconds(10), TimeSpan.FromSeconds(30) }));
+				.AddPolicyHandler((serviceProvider, request) => CreateDefaultTimeoutRetryPolicy(serviceProvider.GetRequiredService<ILogger<HttpStorageClient>>()))
+				.AddPolicyHandler((serviceProvider, request) => CreateDefaultTransientErrorPolicy(serviceProvider.GetRequiredService<ILogger<HttpStorageClient>>()));
 
 			if (useAuthChallenge)
 			{
@@ -420,8 +425,50 @@ namespace EpicGames.Horde
 				services.AddTransient<HordeHttpAuthHandler>();
 				services.AddHttpClient(HordeHttpAuthHandlerState.HttpClientName, configureClient);
 
-				builder.AddHttpMessageHandler<HordeHttpAuthHandler>();
+				builder = builder.AddHttpMessageHandler<HordeHttpAuthHandler>();
+			}
+
+			return builder;
+		}
+
+		/// <summary>
+		/// Create a default timeout retry policy
+		/// </summary>
+		public static IAsyncPolicy<HttpResponseMessage> CreateDefaultTimeoutRetryPolicy(ILogger logger)
+		{
+			// Wait 30 seconds for operations to timeout
+			Task OnTimeoutAsync(Context context, TimeSpan timespan, Task timeoutTask)
+			{
+				logger.LogWarning("Http request timed out after {Time}s.", (int)timespan.TotalSeconds);
+				return Task.CompletedTask;
+			}
+
+			AsyncTimeoutPolicy<HttpResponseMessage> timeoutPolicy = Policy.TimeoutAsync<HttpResponseMessage>(30, OnTimeoutAsync);
+
+			// Retry twice after a timeout
+			void OnRetry(Exception ex, TimeSpan timespan)
+			{
+				logger.LogWarning("Retrying http call after {Time}s.", timespan.TotalSeconds);
 			}
+
+			TimeSpan[] retryTimes = new[] { TimeSpan.FromSeconds(5.0), TimeSpan.FromSeconds(10.0) };
+			AsyncRetryPolicy retryPolicy = Policy.Handle<TimeoutRejectedException>().WaitAndRetryAsync(retryTimes, OnRetry);
+			return retryPolicy.WrapAsync(timeoutPolicy);
+		}
+
+		/// <summary>
+		/// Create a default timeout retry policy
+		/// </summary>
+		public static IAsyncPolicy<HttpResponseMessage> CreateDefaultTransientErrorPolicy(ILogger logger)
+		{
+			Task OnTimeoutAsync(DelegateResult<HttpResponseMessage> outcome, TimeSpan timespan, int retryAttempt, Context context)
+			{
+				logger.LogWarning("Http request failed. Delaying for {DelayMs}ms (attempt #{RetryNum}).", timespan.TotalMilliseconds, retryAttempt);
+				return Task.CompletedTask;
+			}
+
+			TimeSpan[] retryTimes = new[] { TimeSpan.FromSeconds(1.0), TimeSpan.FromSeconds(5.0), TimeSpan.FromSeconds(10.0), TimeSpan.FromSeconds(30.0), TimeSpan.FromSeconds(30.0) };
+			return HttpPolicyExtensions.HandleTransientHttpError().WaitAndRetryAsync(retryTimes, OnTimeoutAsync);
 		}
 	}
 }
