Horde: Move IStorageWriter and DedupeStorageWriter into a separate file.

diff --git a/Engine/Source/Programs/Shared/EpicGames.Horde/Storage/Bundles/BundleWriter.cs b/Engine/Source/Programs/Shared/EpicGames.Horde/Storage/Bundles/BundleWriter.cs
--- a/Engine/Source/Programs/Shared/EpicGames.Horde/Storage/Bundles/BundleWriter.cs
+++ b/Engine/Source/Programs/Shared/EpicGames.Horde/Storage/Bundles/BundleWriter.cs
@@ -4,7 +4,6 @@ using System;
 using System.Buffers;
 using System.Collections.Generic;
 using System.Diagnostics;
-using System.Diagnostics.CodeAnalysis;
 using System.Linq;
 using System.Threading;
 using System.Threading.Tasks;
@@ -78,134 +77,6 @@ namespace EpicGames.Horde.Storage.Bundles
 		public override ValueTask FlushAsync(CancellationToken cancellationToken = default) => new ValueTask();
 	}
 
-	/// <summary>
-	/// Index of known nodes that can be used for deduplication.
-	/// </summary>
-	public sealed class DedupeStorageWriter : IStorageWriter
-	{
-		record BlobKey(IoHash Hash, BlobType Type);
-
-		class DedupeCache
-		{
-			readonly int _maxKeys;
-			readonly Queue<BlobKey> _blobKeys = new Queue<BlobKey>();
-			readonly Dictionary<BlobKey, BlobHandle> _blobKeyToHandle = new Dictionary<BlobKey, BlobHandle>();
-
-			public DedupeCache(int maxKeys)
-			{
-				_maxKeys = maxKeys;
-				_blobKeys = new Queue<BlobKey>(maxKeys);
-				_blobKeyToHandle = new Dictionary<BlobKey, BlobHandle>(maxKeys);
-			}
-
-			internal void Add(BlobKey key, BlobHandle handle)
-			{
-				BlobKey? prevKey;
-				if (_blobKeys.Count == _maxKeys && _blobKeys.TryDequeue(out prevKey))
-				{
-					_blobKeyToHandle.Remove(prevKey);
-				}
-				_blobKeyToHandle.TryAdd(key, handle);
-			}
-
-			internal bool TryGetValue(BlobKey key, [NotNullWhen(true)] out BlobHandle? handle) => _blobKeyToHandle.TryGetValue(key, out handle);
-		}
-
-		class WrappedHandle : BlobHandle
-		{
-			public object _lockObject = new object();
-			public BlobHandle? _inner;
-
-			public override ValueTask FlushAsync(CancellationToken cancellationToken)
-			{
-				if (_inner == null)
-				{
-					throw new InvalidOperationException();
-				}
-				else
-				{
-					return _inner.FlushAsync(cancellationToken);
-				}
-			}
-
-			public override BlobHandle Unwrap()
-			{
-				return _inner?.Unwrap() ?? this;
-			}
-
-			public override ValueTask<BlobData> ReadAsync(CancellationToken cancellationToken = default)
-			{
-				if (_inner == null)
-				{
-					throw new InvalidOperationException();
-				}
-				else
-				{
-					return _inner.ReadAsync(cancellationToken);
-				}
-			}
-		}
-
-		readonly IStorageWriter _inner;
-		readonly DedupeCache _cache;
-
-		/// <summary>
-		/// Constructor
-		/// </summary>
-		/// <param name="inner"></param>
-		/// <param name="maxKeys"></param>
-		public DedupeStorageWriter(IStorageWriter inner, int maxKeys = 64 * 1024)
-		{
-			_inner = inner;
-			_cache = new DedupeCache(maxKeys);
-		}
-
-		private DedupeStorageWriter(IStorageWriter inner, DedupeCache cache)
-		{
-			_inner = inner;
-			_cache = cache;
-		}
-
-		/// <inheritdoc/>
-		public ValueTask DisposeAsync() => _inner.DisposeAsync();
-
-		/// <inheritdoc/>
-		public Task FlushAsync(CancellationToken cancellationToken = default) => _inner.FlushAsync(cancellationToken);
-
-		/// <inheritdoc/>
-		public IStorageWriter Fork() => new DedupeStorageWriter(_inner.Fork(), _cache);
-
-		/// <inheritdoc/>
-		public Memory<byte> GetOutputBuffer(int usedSize, int desiredSize) => _inner.GetOutputBuffer(usedSize, desiredSize);
-
-		/// <inheritdoc/>
-		public async ValueTask<BlobHandle> WriteBlobAsync(int size, IReadOnlyList<BlobHandle> references, BlobType type, IReadOnlyList<AliasInfo> aliases, CancellationToken cancellationToken = default)
-		{
-			ReadOnlyMemory<byte> data = _inner.GetOutputBuffer(size, size).Slice(0, size);
-			IoHash hash = IoHash.Compute(data.Span);
-			BlobKey key = new BlobKey(hash, type);
-
-			WrappedHandle? wrappedHandle;
-			lock (_cache)
-			{
-				BlobHandle? handle;
-				if (_cache.TryGetValue(key, out handle))
-				{
-					return handle;
-				}
-
-				wrappedHandle = new WrappedHandle();
-				_cache.Add(key, wrappedHandle);
-			}
-
-			wrappedHandle._inner = await _inner.WriteBlobAsync(size, references, type, aliases, cancellationToken);
-			return wrappedHandle;
-		}
-
-		/// <inheritdoc/>
-		public ValueTask WriteRefAsync(BlobHandle target, RefOptions? options = null, CancellationToken cancellationToken = default) => _inner.WriteRefAsync(target, options, cancellationToken);
-	}
-
 	/// <summary>
 	/// Writes nodes of a tree to an <see cref="IStorageClient"/>, packed into bundles. Each <see cref="BundleWriter"/> instance is single threaded,
 	/// but multiple instances may be written to in parallel.
diff --git a/Engine/Source/Programs/Shared/EpicGames.Horde/Storage/IStorageClient.cs b/Engine/Source/Programs/Shared/EpicGames.Horde/Storage/IStorageClient.cs
--- a/Engine/Source/Programs/Shared/EpicGames.Horde/Storage/IStorageClient.cs
+++ b/Engine/Source/Programs/Shared/EpicGames.Horde/Storage/IStorageClient.cs
@@ -160,51 +160,6 @@ namespace EpicGames.Horde.Storage
 		void GetStats(StorageStats stats);
 	}
 
-	/// <summary>
-	/// Interface for writing new nodes to the store
-	/// </summary>
-	public interface IStorageWriter : IAsyncDisposable
-	{
-		/// <summary>
-		/// Create another writer instance, allowing multiple threads to write in parallel.
-		/// </summary>
-		/// <returns>New writer instance</returns>
-		IStorageWriter Fork();
-
-		/// <summary>
-		/// Flush any pending nodes to storage
-		/// </summary>
-		/// <param name="cancellationToken">Cancellation token for the operation</param>
-		Task FlushAsync(CancellationToken cancellationToken = default);
-
-		/// <summary>
-		/// Gets an output buffer for writing.
-		/// </summary>
-		/// <param name="usedSize">Current size in the existing buffer that has been written to</param>
-		/// <param name="desiredSize">Desired size of the returned buffer</param>
-		/// <returns>Buffer to be written into.</returns>
-		Memory<byte> GetOutputBuffer(int usedSize, int desiredSize);
-
-		/// <summary>
-		/// Finish writing a node.
-		/// </summary>
-		/// <param name="size">Used size of the buffer</param>
-		/// <param name="references">References to other nodes</param>
-		/// <param name="type">Type of the node that was written</param>
-		/// <param name="aliases">Aliases for this node</param>
-		/// <param name="cancellationToken">Cancellation token for the operation</param>
-		/// <returns>Handle to the written node</returns>
-		ValueTask<BlobHandle> WriteBlobAsync(int size, IReadOnlyList<BlobHandle> references, BlobType type, IReadOnlyList<AliasInfo> aliases, CancellationToken cancellationToken = default);
-
-		/// <summary>
-		/// Writes the reference using the given target node
-		/// </summary>
-		/// <param name="target">The target node</param>
-		/// <param name="options">Options for the new ref</param>
-		/// <param name="cancellationToken">Cancellation token for the operation</param>
-		ValueTask WriteRefAsync(BlobHandle target, RefOptions? options = null, CancellationToken cancellationToken = default);
-	}
-
 	/// <summary>
 	/// Allows creating storage clients for different namespaces
 	/// </summary>
diff --git a/Engine/Source/Programs/Shared/EpicGames.Horde/Storage/IStorageWriter.cs b/Engine/Source/Programs/Shared/EpicGames.Horde/Storage/IStorageWriter.cs
new file mode 100644
--- /dev/null
+++ b/Engine/Source/Programs/Shared/EpicGames.Horde/Storage/IStorageWriter.cs
@@ -0,0 +1,200 @@
+// Copyright Epic Games, Inc. All Rights Reserved.
+
+using System;
+using System.Collections.Generic;
+using System.Diagnostics.CodeAnalysis;
+using System.Threading;
+using System.Threading.Tasks;
+using EpicGames.Core;
+
+namespace EpicGames.Horde.Storage
+{
+	/// <summary>
+	/// Interface for writing new nodes to the store
+	/// </summary>
+	public interface IStorageWriter : IAsyncDisposable
+	{
+		/// <summary>
+		/// Create another writer instance, allowing multiple threads to write in parallel.
+		/// </summary>
+		/// <returns>New writer instance</returns>
+		IStorageWriter Fork();
+
+		/// <summary>
+		/// Flush any pending nodes to storage
+		/// </summary>
+		/// <param name="cancellationToken">Cancellation token for the operation</param>
+		Task FlushAsync(CancellationToken cancellationToken = default);
+
+		/// <summary>
+		/// Gets an output buffer for writing.
+		/// </summary>
+		/// <param name="usedSize">Current size in the existing buffer that has been written to</param>
+		/// <param name="desiredSize">Desired size of the returned buffer</param>
+		/// <returns>Buffer to be written into.</returns>
+		Memory<byte> GetOutputBuffer(int usedSize, int desiredSize);
+
+		/// <summary>
+		/// Finish writing a node.
+		/// </summary>
+		/// <param name="size">Used size of the buffer</param>
+		/// <param name="references">References to other nodes</param>
+		/// <param name="type">Type of the node that was written</param>
+		/// <param name="aliases">Aliases for this node</param>
+		/// <param name="cancellationToken">Cancellation token for the operation</param>
+		/// <returns>Handle to the written node</returns>
+		ValueTask<BlobHandle> WriteBlobAsync(int size, IReadOnlyList<BlobHandle> references, BlobType type, IReadOnlyList<AliasInfo> aliases, CancellationToken cancellationToken = default);
+
+		/// <summary>
+		/// Writes the reference using the given target node
+		/// </summary>
+		/// <param name="target">The target node</param>
+		/// <param name="options">Options for the new ref</param>
+		/// <param name="cancellationToken">Cancellation token for the operation</param>
+		ValueTask WriteRefAsync(BlobHandle target, RefOptions? options = null, CancellationToken cancellationToken = default);
+	}
+
+	/// <summary>
+	/// Index of known nodes that can be used for deduplication.
+	/// </summary>
+	public sealed class DedupeStorageWriter : IStorageWriter
+	{
+		record BlobKey(IoHash Hash, BlobType Type);
+
+		class DedupeCache
+		{
+			readonly int _maxKeys;
+			readonly Queue<BlobKey> _blobKeys = new Queue<BlobKey>();
+			readonly Dictionary<BlobKey, BlobHandle> _blobKeyToHandle = new Dictionary<BlobKey, BlobHandle>();
+
+			public DedupeCache(int maxKeys)
+			{
+				_maxKeys = maxKeys;
+				_blobKeys = new Queue<BlobKey>(maxKeys);
+				_blobKeyToHandle = new Dictionary<BlobKey, BlobHandle>(maxKeys);
+			}
+
+			internal void Add(BlobKey key, BlobHandle handle)
+			{
+				BlobKey? prevKey;
+				if (_blobKeys.Count == _maxKeys && _blobKeys.TryDequeue(out prevKey))
+				{
+					_blobKeyToHandle.Remove(prevKey);
+				}
+				_blobKeyToHandle.TryAdd(key, handle);
+			}
+
+			internal bool TryGetValue(BlobKey key, [NotNullWhen(true)] out BlobHandle? handle) => _blobKeyToHandle.TryGetValue(key, out handle);
+		}
+
+		class WrappedHandle : BlobHandle
+		{
+			public object _lockObject = new object();
+			public BlobHandle? _inner;
+
+			public override ValueTask FlushAsync(CancellationToken cancellationToken)
+			{
+				if (_inner == null)
+				{
+					throw new InvalidOperationException();
+				}
+				else
+				{
+					return _inner.FlushAsync(cancellationToken);
+				}
+			}
+
+			public override BlobHandle Unwrap()
+			{
+				return _inner?.Unwrap() ?? this;
+			}
+
+			public override ValueTask<BlobData> ReadAsync(CancellationToken cancellationToken = default)
+			{
+				if (_inner == null)
+				{
+					throw new InvalidOperationException();
+				}
+				else
+				{
+					return _inner.ReadAsync(cancellationToken);
+				}
+			}
+		}
+
+		/// <summary>
+		/// Default value for maximum number of keys
+		/// </summary>
+		public const int DefaultMaxKeys = 64 * 1024;
+
+		readonly IStorageWriter _inner;
+		readonly DedupeCache _cache;
+
+		/// <summary>
+		/// Constructor
+		/// </summary>
+		/// <param name="inner"></param>
+		/// <param name="maxKeys"></param>
+		public DedupeStorageWriter(IStorageWriter inner, int maxKeys = DefaultMaxKeys)
+		{
+			_inner = inner;
+			_cache = new DedupeCache(maxKeys);
+		}
+
+		private DedupeStorageWriter(IStorageWriter inner, DedupeCache cache)
+		{
+			_inner = inner;
+			_cache = cache;
+		}
+
+		/// <inheritdoc/>
+		public ValueTask DisposeAsync() => _inner.DisposeAsync();
+
+		/// <inheritdoc/>
+		public Task FlushAsync(CancellationToken cancellationToken = default) => _inner.FlushAsync(cancellationToken);
+
+		/// <inheritdoc/>
+		public IStorageWriter Fork() => new DedupeStorageWriter(_inner.Fork(), _cache);
+
+		/// <inheritdoc/>
+		public Memory<byte> GetOutputBuffer(int usedSize, int desiredSize) => _inner.GetOutputBuffer(usedSize, desiredSize);
+
+		/// <inheritdoc/>
+		public async ValueTask<BlobHandle> WriteBlobAsync(int size, IReadOnlyList<BlobHandle> references, BlobType type, IReadOnlyList<AliasInfo> aliases, CancellationToken cancellationToken = default)
+		{
+			ReadOnlyMemory<byte> data = _inner.GetOutputBuffer(size, size).Slice(0, size);
+			IoHash hash = IoHash.Compute(data.Span);
+			BlobKey key = new BlobKey(hash, type);
+
+			WrappedHandle? wrappedHandle;
+			lock (_cache)
+			{
+				BlobHandle? handle;
+				if (_cache.TryGetValue(key, out handle))
+				{
+					return handle;
+				}
+
+				wrappedHandle = new WrappedHandle();
+				_cache.Add(key, wrappedHandle);
+			}
+
+			wrappedHandle._inner = await _inner.WriteBlobAsync(size, references, type, aliases, cancellationToken);
+			return wrappedHandle;
+		}
+
+		/// <inheritdoc/>
+		public ValueTask WriteRefAsync(BlobHandle target, RefOptions? options = null, CancellationToken cancellationToken = default) => _inner.WriteRefAsync(target, options, cancellationToken);
+	}
+
+	/// <summary>
+	/// Extension methods for <see cref="IStorageWriter"/>
+	/// </summary>
+	public static class StorageWriter
+	{
+		/// <summary>
+		/// Wraps a <see cref="IStorageWriter"/> with a <see cref="DedupeStorageWriter"/>
+		/// </summary>
+		public static DedupeStorageWriter WithDedupe(this IStorageWriter writer, int maxKeys = DedupeStorageWriter.DefaultMaxKeys) => new DedupeStorageWriter(writer, maxKeys);
+	}
+}
