Horde: Move IStorageWriter extension methods to same file as interface definition.

diff --git a/Engine/Source/Programs/Shared/EpicGames.Horde/Storage/IStorageClient.cs b/Engine/Source/Programs/Shared/EpicGames.Horde/Storage/IStorageClient.cs
--- a/Engine/Source/Programs/Shared/EpicGames.Horde/Storage/IStorageClient.cs
+++ b/Engine/Source/Programs/Shared/EpicGames.Horde/Storage/IStorageClient.cs
@@ -5,7 +5,6 @@ using System.Collections.Generic;
 using System.Linq;
 using System.Threading;
 using System.Threading.Tasks;
-using EpicGames.Core;
 using Microsoft.Extensions.Logging;
 
 namespace EpicGames.Horde.Storage
@@ -321,40 +320,4 @@ namespace EpicGames.Horde.Storage
 			return stats;
 		}
 	}
-
-	/// <summary>
-	/// Extension methods for <see cref="IStorageWriter"/>
-	/// </summary>
-	public static class StorageWriterExtensions
-	{
-		/// <summary>
-		/// Finish writing a node.
-		/// </summary>
-		/// <param name="writer">Writer instance to manipulate</param>
-		/// <param name="size">Used size of the buffer</param>
-		/// <param name="references">References to other nodes</param>
-		/// <param name="type">Type of the node that was written</param>
-		/// <param name="cancellationToken">Cancellation token for the operation</param>
-		/// <returns>Handle to the written node</returns>
-		public static ValueTask<BlobHandle> WriteBlobAsync(this IStorageWriter writer, int size, IReadOnlyList<BlobHandle> references, BlobType type, CancellationToken cancellationToken = default)
-		{
-			return writer.WriteBlobAsync(size, references, type, Array.Empty<AliasInfo>(), cancellationToken);
-		}
-
-		/// <summary>
-		/// Finish writing a node.
-		/// </summary>
-		/// <param name="writer">Writer instance to manipulate</param>
-		/// <param name="size">Used size of the buffer</param>
-		/// <param name="references">References to other nodes</param>
-		/// <param name="type">Type of the node that was written</param>
-		/// <param name="cancellationToken">Cancellation token for the operation</param>
-		/// <returns>Handle to the written node</returns>
-		public static async ValueTask<HashedNodeRef<T>> WriteHashedNodeRefAsync<T>(this IStorageWriter writer, int size, IReadOnlyList<BlobHandle> references, BlobType type, CancellationToken cancellationToken = default) where T : Node
-		{
-			IoHash hash = IoHash.Compute(writer.GetOutputBuffer(size, size).Span.Slice(0, size));
-			BlobHandle blobHandle = await WriteBlobAsync(writer, size, references, type, cancellationToken);
-			return new HashedNodeRef<T>(hash, blobHandle);
-		}
-	}
 }
diff --git a/Engine/Source/Programs/Shared/EpicGames.Horde/Storage/IStorageWriter.cs b/Engine/Source/Programs/Shared/EpicGames.Horde/Storage/IStorageWriter.cs
--- a/Engine/Source/Programs/Shared/EpicGames.Horde/Storage/IStorageWriter.cs
+++ b/Engine/Source/Programs/Shared/EpicGames.Horde/Storage/IStorageWriter.cs
@@ -190,11 +190,41 @@ namespace EpicGames.Horde.Storage
 	/// <summary>
 	/// Extension methods for <see cref="IStorageWriter"/>
 	/// </summary>
-	public static class StorageWriter
+	public static class StorageWriterExtensions
 	{
 		/// <summary>
 		/// Wraps a <see cref="IStorageWriter"/> with a <see cref="DedupeStorageWriter"/>
 		/// </summary>
 		public static DedupeStorageWriter WithDedupe(this IStorageWriter writer, int maxKeys = DedupeStorageWriter.DefaultMaxKeys) => new DedupeStorageWriter(writer, maxKeys);
+
+		/// <summary>
+		/// Finish writing a node.
+		/// </summary>
+		/// <param name="writer">Writer instance to manipulate</param>
+		/// <param name="size">Used size of the buffer</param>
+		/// <param name="references">References to other nodes</param>
+		/// <param name="type">Type of the node that was written</param>
+		/// <param name="cancellationToken">Cancellation token for the operation</param>
+		/// <returns>Handle to the written node</returns>
+		public static ValueTask<BlobHandle> WriteBlobAsync(this IStorageWriter writer, int size, IReadOnlyList<BlobHandle> references, BlobType type, CancellationToken cancellationToken = default)
+		{
+			return writer.WriteBlobAsync(size, references, type, Array.Empty<AliasInfo>(), cancellationToken);
+		}
+
+		/// <summary>
+		/// Finish writing a node.
+		/// </summary>
+		/// <param name="writer">Writer instance to manipulate</param>
+		/// <param name="size">Used size of the buffer</param>
+		/// <param name="references">References to other nodes</param>
+		/// <param name="type">Type of the node that was written</param>
+		/// <param name="cancellationToken">Cancellation token for the operation</param>
+		/// <returns>Handle to the written node</returns>
+		public static async ValueTask<HashedNodeRef<T>> WriteHashedNodeRefAsync<T>(this IStorageWriter writer, int size, IReadOnlyList<BlobHandle> references, BlobType type, CancellationToken cancellationToken = default) where T : Node
+		{
+			IoHash hash = IoHash.Compute(writer.GetOutputBuffer(size, size).Span.Slice(0, size));
+			BlobHandle blobHandle = await WriteBlobAsync(writer, size, references, type, cancellationToken);
+			return new HashedNodeRef<T>(hash, blobHandle);
+		}
 	}
 }
