Horde: Warn if executable or working dir does not exist for compute task

diff --git a/Engine/Source/Programs/Shared/EpicGames.Horde/Compute/AgentMessageHandler.cs b/Engine/Source/Programs/Shared/EpicGames.Horde/Compute/AgentMessageHandler.cs
--- a/Engine/Source/Programs/Shared/EpicGames.Horde/Compute/AgentMessageHandler.cs
+++ b/Engine/Source/Programs/Shared/EpicGames.Horde/Compute/AgentMessageHandler.cs
@@ -358,6 +358,16 @@ namespace EpicGames.Horde.Compute
 					}
 				}
 
+				if (!File.Exists(resolvedExecutable))
+				{
+					_logger.LogWarning("Executable {Path} does not exist", resolvedExecutable);	
+				}
+				
+				if (!File.Exists(resolvedWorkingDir))
+				{
+					_logger.LogWarning("Working dir {Path} does not exist", resolvedWorkingDir);	
+				}
+
 				using (ManagedProcessGroup group = new ManagedProcessGroup())
 				{
 					using (ManagedProcess process = new ManagedProcess(group, resolvedExecutable, resolvedCommandLine, resolvedWorkingDir, resolvedEnvVars, null, ProcessPriorityClass.Normal))
