Fix issue with JsonWriter to support characters the should be escaped in a value name.

diff --git a/Engine/Source/Programs/Shared/EpicGames.Core/JsonWriter.cs b/Engine/Source/Programs/Shared/EpicGames.Core/JsonWriter.cs
--- a/Engine/Source/Programs/Shared/EpicGames.Core/JsonWriter.cs
+++ b/Engine/Source/Programs/Shared/EpicGames.Core/JsonWriter.cs
@@ -121,8 +121,7 @@ namespace EpicGames.Core
 		{
 			WriteCommaNewline();
 
-			string space = (_style == JsonWriterStyle.Readable) ? " " : "";
-			_writer.Write("{0}\"{1}\":{2}", _indent, objectName, space);
+			WriteName(objectName);
 
 			_bRequiresComma = false;
 
@@ -164,8 +163,8 @@ namespace EpicGames.Core
 		{
 			WriteCommaNewline();
 
-			string space = (_style == JsonWriterStyle.Readable) ? " " : "";
-			_writer.Write("{0}\"{1}\":{2}[", _indent, arrayName, space);
+			WriteName(arrayName);
+			_writer.Write('[');
 
 			IncreaseIndent();
 			_bRequiresComma = false;
@@ -266,8 +265,7 @@ namespace EpicGames.Core
 		{
 			WriteCommaNewline();
 
-			string space = (_style == JsonWriterStyle.Readable) ? " " : "";
-			_writer.Write("{0}\"{1}\":{2}", _indent, name, space);
+			WriteName(name);
 			WriteEscapedString(value);
 
 			_bRequiresComma = true;
@@ -326,12 +324,20 @@ namespace EpicGames.Core
 			}
 		}
 
+		void WriteName(string name)
+		{
+			string space = (_style == JsonWriterStyle.Readable) ? " " : "";
+			_writer.Write(_indent);
+			WriteEscapedString(name);
+			_writer.Write(":{0}", space);
+		}
+
 		void WriteValueInternal(string name, string value)
 		{
 			WriteCommaNewline();
 
-			string space = (_style == JsonWriterStyle.Readable) ? " " : "";
-			_writer.Write("{0}\"{1}\":{2}{3}", _indent, name, space, value);
+			WriteName(name);
+			_writer.Write(value);
 
 			_bRequiresComma = true;
 		}
@@ -354,8 +360,28 @@ namespace EpicGames.Core
 		/// <returns>The escaped string</returns>
 		public static string EscapeString(string value)
 		{
+
+			// Prescan the string looking for things to escape.  If not found, we don't need to
+			// create the string builder
+			int idx = 0;
+			for (; idx < value.Length; idx++)
+			{
+				char c = value[idx];
+				if (c == '\"' || c == '\\' || Char.IsControl(c))
+				{
+					break;
+				}
+			}
+			if (idx == value.Length)
+			{
+				return value;
+			}
+
+			// Otherwise, create the string builder, append the known portion that doesn't have an escape
+			// and continue processing the string starting at the first character needing to be escaped.
 			StringBuilder result = new StringBuilder();
-			for (int idx = 0; idx < value.Length; idx++)
+			result.Append(value.AsSpan(0, idx));
+			for (; idx < value.Length; idx++)
 			{
 				switch (value[idx])
 				{
