Horde: Fix exception running Box remote execution jobs due to agents running an older version of the code that does an out-of-bounds read of the BundleNodeLocator struct.

diff --git a/Engine/Source/Programs/Shared/EpicGames.Horde/Compute/AgentMessage.cs b/Engine/Source/Programs/Shared/EpicGames.Horde/Compute/AgentMessage.cs
--- a/Engine/Source/Programs/Shared/EpicGames.Horde/Compute/AgentMessage.cs
+++ b/Engine/Source/Programs/Shared/EpicGames.Horde/Compute/AgentMessage.cs
@@ -383,7 +383,7 @@ namespace EpicGames.Horde.Compute
 			using (IAgentMessageBuilder request = await channel.CreateMessageAsync(AgentMessageType.WriteFiles, cancellationToken))
 			{
 				request.WriteString(path);
-				request.WriteBlobLocator(locator);
+				request.WriteString($"{IoHash.Zero}@{locator}"); // HACK: Currently deployed agents have a hash check in BundleNodeLocator.Parse() which does not check length before checking for the '@' character separating the hash from locator.
 				request.Send();
 			}
 
