EpicGames.Core: Implement DefaultLogger.BeginScope with NullScope

#rnx

diff --git a/Engine/Source/Programs/Shared/EpicGames.Core/Log.cs b/Engine/Source/Programs/Shared/EpicGames.Core/Log.cs
--- a/Engine/Source/Programs/Shared/EpicGames.Core/Log.cs
+++ b/Engine/Source/Programs/Shared/EpicGames.Core/Log.cs
@@ -783,16 +783,20 @@ namespace EpicGames.Core
 		}
 	}
 
+	/// <summary>
+	/// NullScope which does nothing
+	/// </summary>
+	internal sealed class NullScope : IDisposable
+	{
+		/// <inheritdoc/>
+		public void Dispose() { }
+	}
+
 	/// <summary>
 	/// Logger which captures the output for rendering later
 	/// </summary>
 	public class CaptureLogger : ILogger
 	{
-		class NullScope : IDisposable
-		{
-			public void Dispose() { }
-		}
-
 		/// <summary>
 		/// List of captured events
 		/// </summary>
@@ -1170,16 +1174,16 @@ namespace EpicGames.Core
 			return false;
 		}
 
-		public IDisposable BeginScope<TState>(TState state)
-		{
-			throw new NotImplementedException();
-		}
+		/// <inheritdoc/>
+		public IDisposable BeginScope<TState>(TState state) => new NullScope();
 
+		/// <inheritdoc/>
 		public bool IsEnabled(LogLevel logLevel)
 		{
 			return logLevel >= OutputLevel;
 		}
 
+		/// <inheritdoc/>
 		public void Log<TState>(LogLevel logLevel, EventId eventId, TState state, Exception? exception, Func<TState, Exception?, string> formatter)
 		{
 			if (!IsEnabled(logLevel))
