Horde: Add methods to IStorageClient to allow reading/writing individual blobs. While sub-optimal for small blobs, these methods can be used to write bundles themselves.

diff --git a/Engine/Source/Programs/Horde/Horde.Server/Commands/Bundles/BundleWriteTest.cs b/Engine/Source/Programs/Horde/Horde.Server/Commands/Bundles/BundleWriteTest.cs
--- a/Engine/Source/Programs/Horde/Horde.Server/Commands/Bundles/BundleWriteTest.cs
+++ b/Engine/Source/Programs/Horde/Horde.Server/Commands/Bundles/BundleWriteTest.cs
@@ -1,6 +1,7 @@
 // Copyright Epic Games, Inc. All Rights Reserved.
 
 using System;
+using System.Collections.Generic;
 using System.Diagnostics;
 using System.Security.Cryptography;
 using System.Threading;
@@ -24,6 +25,8 @@ namespace Horde.Server.Commands.Bundles
 
 			public void Dispose() { }
 
+			public ValueTask<BlobData> ReadBlobAsync(BlobLocator locator, CancellationToken cancellationToken = default) => throw new NotImplementedException();
+			public ValueTask<BlobHandle> WriteBlobAsync(BlobType type, ReadOnlyMemory<byte> data, IReadOnlyList<BlobHandle> references, CancellationToken cancellationToken = default) => throw new NotImplementedException();
 			public BlobHandle CreateBlobHandle(BlobLocator locator) => throw new NotImplementedException();
 			public Task<bool> DeleteRefAsync(RefName name, CancellationToken cancellationToken = default) => Task.FromResult(true);
 			public Task AddAliasAsync(string name, BlobHandle handle, int rank = 0, ReadOnlyMemory<byte> data = default, CancellationToken cancellationToken = default) => throw new NotImplementedException();
diff --git a/Engine/Source/Programs/Horde/Horde.Server/Storage/StorageService.cs b/Engine/Source/Programs/Horde/Horde.Server/Storage/StorageService.cs
--- a/Engine/Source/Programs/Horde/Horde.Server/Storage/StorageService.cs
+++ b/Engine/Source/Programs/Horde/Horde.Server/Storage/StorageService.cs
@@ -284,7 +284,7 @@ namespace Horde.Server.Storage
 			public bool SupportsRedirects => _impl.SupportsRedirects;
 			public IStorageBackend Backend => _impl.Backend;
 
-			#region Nodes
+			#region Blobs
 
 			public BlobHandle CreateBlobHandle(BlobLocator locator) => _impl.CreateBlobHandle(locator);
 			public BundleNodeHandle CreateNodeHandle(BundleNodeLocator locator) => _impl.CreateNodeHandle(locator);
@@ -292,6 +292,9 @@ namespace Horde.Server.Storage
 			public BundleWriter CreateWriter(RefName refName = default, BundleOptions? options = null) => _impl.CreateWriter(refName, options);
 			IStorageWriter IStorageClient.CreateWriter(RefName refName) => ((IStorageClient)_impl).CreateWriter(refName);
 
+			public ValueTask<BlobData> ReadBlobAsync(BlobLocator locator, CancellationToken cancellationToken = default) => _impl.ReadBlobAsync(locator, cancellationToken);
+			public ValueTask<BlobHandle> WriteBlobAsync(BlobType type, ReadOnlyMemory<byte> data, IReadOnlyList<BlobHandle> references, CancellationToken cancellationToken = default) => _impl.WriteBlobAsync(type, data, references, cancellationToken);
+
 			#endregion
 
 			#region Bundles
diff --git a/Engine/Source/Programs/Shared/EpicGames.Horde/Storage/Clients/BundleStorageClient.cs b/Engine/Source/Programs/Shared/EpicGames.Horde/Storage/Clients/BundleStorageClient.cs
--- a/Engine/Source/Programs/Shared/EpicGames.Horde/Storage/Clients/BundleStorageClient.cs
+++ b/Engine/Source/Programs/Shared/EpicGames.Horde/Storage/Clients/BundleStorageClient.cs
@@ -1,6 +1,7 @@
 // Copyright Epic Games, Inc. All Rights Reserved.
 
 using System;
+using System.Collections.Generic;
 using System.IO;
 using System.Threading;
 using System.Threading.Tasks;
@@ -153,6 +154,23 @@ namespace EpicGames.Horde.Storage.Clients
 		/// <inheritdoc/>
 		IStorageWriter IStorageClient.CreateWriter(RefName refName) => CreateWriter(refName);
 
+		/// <inheritdoc/>
+		public async ValueTask<BlobData> ReadBlobAsync(BlobLocator locator, CancellationToken cancellationToken = default)
+		{
+			BlobHandle handle = CreateBlobHandle(locator);
+			return await handle.ReadAsync(cancellationToken);
+		}
+
+		/// <inheritdoc/>
+		public async ValueTask<BlobHandle> WriteBlobAsync(BlobType type, ReadOnlyMemory<byte> data, IReadOnlyList<BlobHandle> references, CancellationToken cancellationToken = default)
+		{
+			await using (IStorageWriter writer = CreateWriter())
+			{
+				data.CopyTo(writer.GetOutputBuffer(0, data.Length));
+				return await writer.WriteBlobAsync(data.Length, references, type, cancellationToken);
+			}
+		}
+
 		#endregion
 
 		#region Aliases
diff --git a/Engine/Source/Programs/Shared/EpicGames.Horde/Storage/IStorageClient.cs b/Engine/Source/Programs/Shared/EpicGames.Horde/Storage/IStorageClient.cs
--- a/Engine/Source/Programs/Shared/EpicGames.Horde/Storage/IStorageClient.cs
+++ b/Engine/Source/Programs/Shared/EpicGames.Horde/Storage/IStorageClient.cs
@@ -82,6 +82,13 @@ namespace EpicGames.Horde.Storage
 	{
 		#region Blobs
 
+		/// <summary>
+		/// Creates a new blob handle by parsing a blob id
+		/// </summary>
+		/// <param name="blobId">Path to the blob</param>
+		/// <returns>New handle to the blob</returns>
+		BlobHandle CreateBlobHandle(BlobLocator blobId);
+
 		/// <summary>
 		/// Creates a new writer for storage blobs
 		/// </summary>
@@ -90,11 +97,20 @@ namespace EpicGames.Horde.Storage
 		IStorageWriter CreateWriter(RefName refName = default);
 
 		/// <summary>
-		/// Creates a new blob handle by parsing a blob id
+		/// Read a blob from the underlying storage system. Calling <see cref="BlobHandle.ReadAsync(CancellationToken)"/> is more efficient than calling this method repeatedly for small blobs.
 		/// </summary>
-		/// <param name="blobId">Path to the blob</param>
-		/// <returns>New handle to the blob</returns>
-		BlobHandle CreateBlobHandle(BlobLocator blobId);
+		/// <param name="locator">Locator for the blob</param>
+		/// <param name="cancellationToken">Cancellation token for the operation</param>
+		ValueTask<BlobData> ReadBlobAsync(BlobLocator locator, CancellationToken cancellationToken = default);
+
+		/// <summary>
+		/// Write a blob to the underlying storage system. Using the writer instance returned from <see cref="CreateWriter(RefName)"/> is more efficient than calling this method repeatedly for small blobs.
+		/// </summary>
+		/// <param name="type">Type of the blob</param>
+		/// <param name="data">Data to be stored</param>
+		/// <param name="references">References to other blobs</param>
+		/// <param name="cancellationToken">Cancellation token for the operation</param>
+		ValueTask<BlobHandle> WriteBlobAsync(BlobType type, ReadOnlyMemory<byte> data, IReadOnlyList<BlobHandle> references, CancellationToken cancellationToken = default);
 
 		#endregion
 
