Horde: Add a storage node which can be used to redirect to other nodes.

diff --git a/Engine/Source/Programs/Shared/EpicGames.Horde/Storage/Nodes/RedirectNode.cs b/Engine/Source/Programs/Shared/EpicGames.Horde/Storage/Nodes/RedirectNode.cs
new file mode 100644
--- /dev/null
+++ b/Engine/Source/Programs/Shared/EpicGames.Horde/Storage/Nodes/RedirectNode.cs
@@ -0,0 +1,47 @@
+// Copyright Epic Games, Inc. All Rights Reserved.
+
+namespace EpicGames.Horde.Storage.Nodes
+{
+	/// <summary>
+	/// A node containing arbitrary compact binary data
+	/// </summary>
+	[NodeType("{BE09E54F-7A6B-47CA-BCAF-972A38153B18}", 1)]
+	public class RedirectNode : Node 
+	{
+		/// <summary>
+		/// The target handle
+		/// </summary>
+		public BlobHandle Handle { get; }
+
+		/// <summary>
+		/// Constructor
+		/// </summary>
+		/// <param name="target">Target node for the redirect</param>
+		public RedirectNode(NodeRef target) => Handle = target.Handle;
+
+		/// <summary>
+		/// Constructor
+		/// </summary>
+		/// <param name="handle">Target node for the redirect</param>
+		public RedirectNode(BlobHandle handle) => Handle = handle;
+
+		/// <summary>
+		/// Deserialization constructor
+		/// </summary>
+		/// <param name="reader">Reader to deserialize from</param>
+		public RedirectNode(NodeReader reader) => Handle = reader.ReadBlobReference();
+
+		/// <summary>
+		/// Gets a typed reference to the target node
+		/// </summary>
+		public NodeRef ToNodeRef() => new NodeRef(Handle);
+
+		/// <summary>
+		/// Gets a typed reference to the target node
+		/// </summary>
+		public NodeRef<TTarget> ToNodeRef<TTarget>() where TTarget : Node => new NodeRef<TTarget>(Handle);
+
+		/// <inheritdoc/>
+		public override void Serialize(NodeWriter writer) => writer.WriteBlobReference(Handle);
+	}
+}
