Horde: Reorder arguments to WriteBlob/WriteRef.

diff --git a/Engine/Source/Programs/Horde/Horde.Server/Ddc/BlobService.cs b/Engine/Source/Programs/Horde/Horde.Server/Ddc/BlobService.cs
--- a/Engine/Source/Programs/Horde/Horde.Server/Ddc/BlobService.cs
+++ b/Engine/Source/Programs/Horde/Horde.Server/Ddc/BlobService.cs
@@ -152,7 +152,7 @@ namespace Horde.Server.Ddc
 				using Stream stream = content.GetStream();
 				await stream.ReadAsync(memory, cancellationToken);
 
-				blobHandle = await writer.WriteBlobAsync(memory.Length, Array.Empty<BlobHandle>(), s_rawBlobType, cancellationToken);
+				blobHandle = await writer.WriteBlobAsync(s_rawBlobType, memory.Length, Array.Empty<BlobHandle>(), cancellationToken);
 				await writer.FlushAsync(cancellationToken);
 			}
 
diff --git a/Engine/Source/Programs/Shared/EpicGames.Horde.Tests/BundleStoreTests.cs b/Engine/Source/Programs/Shared/EpicGames.Horde.Tests/BundleStoreTests.cs
--- a/Engine/Source/Programs/Shared/EpicGames.Horde.Tests/BundleStoreTests.cs
+++ b/Engine/Source/Programs/Shared/EpicGames.Horde.Tests/BundleStoreTests.cs
@@ -216,7 +216,7 @@ namespace EpicGames.Horde.Tests
 			RefName refName = new RefName("test");
 			await using (IStorageWriter writer = store.CreateWriter(refName))
 			{
-				BlobHandle target = await writer.WriteBlobAsync(0, Array.Empty<BlobHandle>(), new BlobType(Guid.NewGuid(), 0));
+				BlobHandle target = await writer.WriteBlobAsync(new BlobType(Guid.NewGuid(), 0), 0, Array.Empty<BlobHandle>());
 				await writer.WriteRefAsync(target, new byte[] { 4, 5, 6 });
 			}
 
diff --git a/Engine/Source/Programs/Shared/EpicGames.Horde/Storage/Bundles/BundleWriter.cs b/Engine/Source/Programs/Shared/EpicGames.Horde/Storage/Bundles/BundleWriter.cs
--- a/Engine/Source/Programs/Shared/EpicGames.Horde/Storage/Bundles/BundleWriter.cs
+++ b/Engine/Source/Programs/Shared/EpicGames.Horde/Storage/Bundles/BundleWriter.cs
@@ -363,7 +363,7 @@ namespace EpicGames.Horde.Storage.Bundles
 			}
 
 			// Finish a node write
-			public PendingNode WriteNode(int size, BlobType blobType, IReadOnlyList<BlobHandle> refs, IReadOnlyList<AliasInfo> aliases)
+			public PendingNode WriteNode(BlobType blobType, int size, IReadOnlyList<BlobHandle> refs, IReadOnlyList<AliasInfo> aliases)
 			{
 				PendingNode pendingNode = new PendingNode(_treeReader, blobType, _currentPacketIdx, _currentPacketLength, (int)size, refs, aliases, this);
 				_currentPacketLength += pendingNode.Length;
@@ -778,18 +778,18 @@ namespace EpicGames.Horde.Storage.Bundles
 		/// <summary>
 		/// Finish writing a node.
 		/// </summary>
+		/// <param name="type">Type of the node that was written</param>
 		/// <param name="size">Used size of the buffer</param>
 		/// <param name="references">References to other nodes</param>
-		/// <param name="type">Type of the node that was written</param>
 		/// <param name="aliases">Aliases for the node</param>
 		/// <param name="cancellationToken">Cancellation token for the operation</param>
 		/// <returns>Handle to the written node</returns>
-		public async ValueTask<BlobHandle> WriteBlobAsync(int size, IReadOnlyList<BlobHandle> references, BlobType type, IReadOnlyList<AliasInfo> aliases, CancellationToken cancellationToken = default)
+		public async ValueTask<BlobHandle> WriteBlobAsync(BlobType type, int size, IReadOnlyList<BlobHandle> references, IReadOnlyList<AliasInfo> aliases, CancellationToken cancellationToken = default)
 		{
 			PendingBundle currentBundle = GetCurrentBundle();
 
 			// Append this node data
-			PendingNode pendingNode = currentBundle.WriteNode(size, type, references, aliases);
+			PendingNode pendingNode = currentBundle.WriteNode(type, size, references, aliases);
 			TraceLogger?.LogInformation("Added new node for {NodeKey} in bundle {BundleId}", pendingNode, currentBundle.BundleId);
 
 			// Add dependencies on all bundles containing a dependent node
diff --git a/Engine/Source/Programs/Shared/EpicGames.Horde/Storage/Clients/BundleStorageClient.cs b/Engine/Source/Programs/Shared/EpicGames.Horde/Storage/Clients/BundleStorageClient.cs
--- a/Engine/Source/Programs/Shared/EpicGames.Horde/Storage/Clients/BundleStorageClient.cs
+++ b/Engine/Source/Programs/Shared/EpicGames.Horde/Storage/Clients/BundleStorageClient.cs
@@ -167,7 +167,7 @@ namespace EpicGames.Horde.Storage.Clients
 			await using (IStorageWriter writer = CreateWriter())
 			{
 				data.CopyTo(writer.GetOutputBuffer(0, data.Length));
-				return await writer.WriteBlobAsync(data.Length, references, type, cancellationToken);
+				return await writer.WriteBlobAsync(type, data.Length, references, cancellationToken);
 			}
 		}
 
diff --git a/Engine/Source/Programs/Shared/EpicGames.Horde/Storage/IStorageWriter.cs b/Engine/Source/Programs/Shared/EpicGames.Horde/Storage/IStorageWriter.cs
--- a/Engine/Source/Programs/Shared/EpicGames.Horde/Storage/IStorageWriter.cs
+++ b/Engine/Source/Programs/Shared/EpicGames.Horde/Storage/IStorageWriter.cs
@@ -45,13 +45,13 @@ namespace EpicGames.Horde.Storage
 		/// <summary>
 		/// Finish writing a node.
 		/// </summary>
+		/// <param name="type">Type of the node that was written</param>
 		/// <param name="size">Used size of the buffer</param>
 		/// <param name="references">References to other nodes</param>
-		/// <param name="type">Type of the node that was written</param>
 		/// <param name="aliases">Aliases for this node</param>
 		/// <param name="cancellationToken">Cancellation token for the operation</param>
 		/// <returns>Handle to the written node</returns>
-		ValueTask<BlobHandle> WriteBlobAsync(int size, IReadOnlyList<BlobHandle> references, BlobType type, IReadOnlyList<AliasInfo> aliases, CancellationToken cancellationToken = default);
+		ValueTask<BlobHandle> WriteBlobAsync(BlobType type, int size, IReadOnlyList<BlobHandle> references, IReadOnlyList<AliasInfo> aliases, CancellationToken cancellationToken = default);
 
 		/// <summary>
 		/// Writes the reference using the given target node
@@ -180,7 +180,7 @@ namespace EpicGames.Horde.Storage
 		public Memory<byte> GetOutputBuffer(int usedSize, int desiredSize) => _inner.GetOutputBuffer(usedSize, desiredSize);
 
 		/// <inheritdoc/>
-		public async ValueTask<BlobHandle> WriteBlobAsync(int size, IReadOnlyList<BlobHandle> references, BlobType type, IReadOnlyList<AliasInfo> aliases, CancellationToken cancellationToken = default)
+		public async ValueTask<BlobHandle> WriteBlobAsync(BlobType type, int size, IReadOnlyList<BlobHandle> references, IReadOnlyList<AliasInfo> aliases, CancellationToken cancellationToken = default)
 		{
 			ReadOnlyMemory<byte> data = _inner.GetOutputBuffer(size, size).Slice(0, size);
 			IoHash hash = IoHash.Compute(data.Span);
@@ -199,7 +199,7 @@ namespace EpicGames.Horde.Storage
 				_cache.Add(key, wrappedHandle);
 			}
 
-			wrappedHandle._inner = await _inner.WriteBlobAsync(size, references.ConvertAll(x => ((WrappedHandle)x)._inner!), type, aliases, cancellationToken);
+			wrappedHandle._inner = await _inner.WriteBlobAsync(type, size, references.ConvertAll(x => ((WrappedHandle)x)._inner!), aliases, cancellationToken);
 			return wrappedHandle;
 		}
 
@@ -221,29 +221,29 @@ namespace EpicGames.Horde.Storage
 		/// Finish writing a node.
 		/// </summary>
 		/// <param name="writer">Writer instance to manipulate</param>
+		/// <param name="type">Type of the node that was written</param>
 		/// <param name="size">Used size of the buffer</param>
 		/// <param name="references">References to other nodes</param>
-		/// <param name="type">Type of the node that was written</param>
 		/// <param name="cancellationToken">Cancellation token for the operation</param>
 		/// <returns>Handle to the written node</returns>
-		public static ValueTask<BlobHandle> WriteBlobAsync(this IStorageWriter writer, int size, IReadOnlyList<BlobHandle> references, BlobType type, CancellationToken cancellationToken = default)
+		public static ValueTask<BlobHandle> WriteBlobAsync(this IStorageWriter writer, BlobType type, int size, IReadOnlyList<BlobHandle> references, CancellationToken cancellationToken = default)
 		{
-			return writer.WriteBlobAsync(size, references, type, Array.Empty<AliasInfo>(), cancellationToken);
+			return writer.WriteBlobAsync(type, size, references, Array.Empty<AliasInfo>(), cancellationToken);
 		}
 
 		/// <summary>
 		/// Finish writing a node.
 		/// </summary>
 		/// <param name="writer">Writer instance to manipulate</param>
+		/// <param name="type">Type of the node that was written</param>
 		/// <param name="size">Used size of the buffer</param>
 		/// <param name="references">References to other nodes</param>
-		/// <param name="type">Type of the node that was written</param>
 		/// <param name="cancellationToken">Cancellation token for the operation</param>
 		/// <returns>Handle to the written node</returns>
-		public static async ValueTask<HashedNodeRef<T>> WriteHashedNodeRefAsync<T>(this IStorageWriter writer, int size, IReadOnlyList<BlobHandle> references, BlobType type, CancellationToken cancellationToken = default) where T : Node
+		public static async ValueTask<HashedNodeRef<T>> WriteHashedNodeRefAsync<T>(this IStorageWriter writer, BlobType type, int size, IReadOnlyList<BlobHandle> references, CancellationToken cancellationToken = default) where T : Node
 		{
 			IoHash hash = IoHash.Compute(writer.GetOutputBuffer(size, size).Span.Slice(0, size));
-			BlobHandle blobHandle = await WriteBlobAsync(writer, size, references, type, cancellationToken);
+			BlobHandle blobHandle = await WriteBlobAsync(writer, type, size, references, cancellationToken);
 			return new HashedNodeRef<T>(hash, blobHandle);
 		}
 	}
diff --git a/Engine/Source/Programs/Shared/EpicGames.Horde/Storage/NodeRef.cs b/Engine/Source/Programs/Shared/EpicGames.Horde/Storage/NodeRef.cs
--- a/Engine/Source/Programs/Shared/EpicGames.Horde/Storage/NodeRef.cs
+++ b/Engine/Source/Programs/Shared/EpicGames.Horde/Storage/NodeRef.cs
@@ -292,7 +292,7 @@ namespace EpicGames.Horde.Storage
 			IoHash hash = nodeWriter.ComputeHash();
 
 			// Write the final data
-			BlobHandle handle = await writer.WriteBlobAsync(nodeWriter.Length, nodeWriter.References, node.NodeType, cancellationToken);
+			BlobHandle handle = await writer.WriteBlobAsync(node.NodeType, nodeWriter.Length, nodeWriter.References, cancellationToken);
 			return new HashedNodeRef<TNode>(hash, handle);
 		}
 
diff --git a/Engine/Source/Programs/Shared/EpicGames.Horde/Storage/Nodes/ChunkedDataNode.cs b/Engine/Source/Programs/Shared/EpicGames.Horde/Storage/Nodes/ChunkedDataNode.cs
--- a/Engine/Source/Programs/Shared/EpicGames.Horde/Storage/Nodes/ChunkedDataNode.cs
+++ b/Engine/Source/Programs/Shared/EpicGames.Horde/Storage/Nodes/ChunkedDataNode.cs
@@ -299,7 +299,7 @@ namespace EpicGames.Horde.Storage.Nodes
 					sizeSinceProgressUpdate = 0;
 				}
 
-				HashedNodeRef<ChunkedDataNode> nodeRef = await writer.WriteHashedNodeRefAsync<ChunkedDataNode>(nextLength, Array.Empty<BlobHandle>(), GetNodeType<LeafChunkedDataNode>(), cancellationToken);
+				HashedNodeRef<ChunkedDataNode> nodeRef = await writer.WriteHashedNodeRefAsync<ChunkedDataNode>(GetNodeType<LeafChunkedDataNode>(), nextLength, Array.Empty<BlobHandle>(), cancellationToken);
 				leafNodeRefs.Add(new ChunkedDataNodeRef(ChunkedDataNodeType.Leaf, nodeRef));
 
 				readBuffer.Memory.Slice(nextLength, size - nextLength).CopyTo(readBuffer.Memory);
diff --git a/Engine/Source/Programs/Shared/EpicGames.Horde/Storage/Nodes/ChunkedDataWriter.cs b/Engine/Source/Programs/Shared/EpicGames.Horde/Storage/Nodes/ChunkedDataWriter.cs
--- a/Engine/Source/Programs/Shared/EpicGames.Horde/Storage/Nodes/ChunkedDataWriter.cs
+++ b/Engine/Source/Programs/Shared/EpicGames.Horde/Storage/Nodes/ChunkedDataWriter.cs
@@ -345,7 +345,7 @@ namespace EpicGames.Horde.Storage.Nodes
 		/// <returns>Handle to the written leaf node</returns>
 		async ValueTask FlushLeafNodeAsync(CancellationToken cancellationToken)
 		{
-			HashedNodeRef<ChunkedDataNode> nodeRef  = await _writer.WriteHashedNodeRefAsync<ChunkedDataNode>(_leafLength, Array.Empty<BlobHandle>(), s_leafNodeType, cancellationToken);
+			HashedNodeRef<ChunkedDataNode> nodeRef  = await _writer.WriteHashedNodeRefAsync<ChunkedDataNode>(s_leafNodeType, _leafLength, Array.Empty<BlobHandle>(), cancellationToken);
 			_leafHandles.Add(new ChunkedDataNodeRef(ChunkedDataNodeType.Leaf, nodeRef));
 			ResetLeafState();
 		}
