Horde: Fix out of bounds exception when writing bundles.

diff --git a/Engine/Source/Programs/Shared/EpicGames.Horde/Storage/Bundles/BundleWriter.cs b/Engine/Source/Programs/Shared/EpicGames.Horde/Storage/Bundles/BundleWriter.cs
--- a/Engine/Source/Programs/Shared/EpicGames.Horde/Storage/Bundles/BundleWriter.cs
+++ b/Engine/Source/Programs/Shared/EpicGames.Horde/Storage/Bundles/BundleWriter.cs
@@ -397,7 +397,7 @@ namespace EpicGames.Horde.Storage.Bundles
 					}
 					else if (packetIdx == _currentPacketIdx)
 					{
-						return _currentPacket!.Memory.Slice(offset, length).Slice(offset, length).ToArray();
+						return _currentPacket!.Memory.Slice(offset, length).ToArray();
 					}
 					else
 					{
