Horde: Use strings rather than ref names to specify base paths for uploaded blobs, and remove ability to write refs from IStorageWriter.

diff --git a/Engine/Source/Programs/Horde/Horde.Server.Tests/GcServiceTests.cs b/Engine/Source/Programs/Horde/Horde.Server.Tests/GcServiceTests.cs
--- a/Engine/Source/Programs/Horde/Horde.Server.Tests/GcServiceTests.cs
+++ b/Engine/Source/Programs/Horde/Horde.Server.Tests/GcServiceTests.cs
@@ -4,7 +4,6 @@ using System;
 using System.Collections.Generic;
 using System.Linq;
 using System.Threading.Tasks;
-using EpicGames.Core;
 using Microsoft.VisualStudio.TestTools.UnitTesting;
 using EpicGames.Horde.Storage;
 using EpicGames.Horde.Storage.Bundles;
@@ -108,7 +107,7 @@ namespace Horde.Server.Tests
 				List<BundleLocator> imports = children[idx].ConvertAll(x => locators[x]);
 				BundleHeader header = new BundleHeader(types.ToArray(), imports.ToArray(), Array.Empty<BundleExport>(), Array.Empty<BundlePacket>());
 				Bundle bundle = new Bundle(header, Array.Empty<ReadOnlyMemory<byte>>());
-				locators[idx] = await store.WriteBundleAsync(bundle, prefix: new Utf8String("gctest"));
+				locators[idx] = await store.WriteBundleAsync(bundle, basePath: "gctest");
 			}
 
 			return locators;
diff --git a/Engine/Source/Programs/Horde/Horde.Server/Commands/Bundles/BundleWriteTest.cs b/Engine/Source/Programs/Horde/Horde.Server/Commands/Bundles/BundleWriteTest.cs
--- a/Engine/Source/Programs/Horde/Horde.Server/Commands/Bundles/BundleWriteTest.cs
+++ b/Engine/Source/Programs/Horde/Horde.Server/Commands/Bundles/BundleWriteTest.cs
@@ -1,7 +1,6 @@
 // Copyright Epic Games, Inc. All Rights Reserved.
 
 using System;
-using System.Collections.Generic;
 using System.Diagnostics;
 using System.Security.Cryptography;
 using System.Threading;
@@ -25,8 +24,6 @@ namespace Horde.Server.Commands.Bundles
 
 			public void Dispose() { }
 
-			public ValueTask<BlobData> ReadBlobAsync(BlobLocator locator, CancellationToken cancellationToken = default) => throw new NotImplementedException();
-			public ValueTask<BlobHandle> WriteBlobAsync(BlobType type, ReadOnlyMemory<byte> data, IReadOnlyList<BlobHandle> references, CancellationToken cancellationToken = default) => throw new NotImplementedException();
 			public BlobHandle CreateBlobHandle(BlobLocator locator) => throw new NotImplementedException();
 			public Task<bool> DeleteRefAsync(RefName name, CancellationToken cancellationToken = default) => Task.FromResult(true);
 			public Task AddAliasAsync(string name, BlobHandle handle, int rank = 0, ReadOnlyMemory<byte> data = default, CancellationToken cancellationToken = default) => throw new NotImplementedException();
@@ -35,13 +32,13 @@ namespace Horde.Server.Commands.Bundles
 			public Task<RefValue?> TryReadRefAsync(RefName name, RefCacheTime cacheTime = default, CancellationToken cancellationToken = default) => throw new NotImplementedException();
 			public Task WriteRefAsync(RefName name, BlobHandle target, ReadOnlyMemory<byte> data = default, RefOptions? options = null, CancellationToken cancellationToken = default) => Task.CompletedTask;
 
-			public BundleWriter CreateWriter(RefName refName = default, BundleOptions? options = null)
+			public BundleWriter CreateWriter(string? basePath = null, BundleOptions? options = null)
 			{
-				_ = (this, refName, options);
+				_ = (this, basePath, options);
 				throw new NotImplementedException();// new BundleWriter(this, _reader, refName, options);
 			}
 
-			IStorageWriter IStorageClient.CreateWriter(RefName refName) => CreateWriter(refName);
+			IStorageWriter IStorageClient.CreateWriter(string? basePath) => CreateWriter(basePath);
 
 			public void GetStats(StorageStats stats) { }
 		}
diff --git a/Engine/Source/Programs/Horde/Horde.Server/Storage/StorageController.cs b/Engine/Source/Programs/Horde/Horde.Server/Storage/StorageController.cs
--- a/Engine/Source/Programs/Horde/Horde.Server/Storage/StorageController.cs
+++ b/Engine/Source/Programs/Horde/Horde.Server/Storage/StorageController.cs
@@ -104,7 +104,7 @@ namespace Horde.Server.Storage
 					bundle = await Bundle.FromStreamAsync(stream, cancellationToken);
 				}
 
-				BundleLocator locator = await storageClient.WriteBundleAsync(bundle, prefix: (prefix == null) ? Utf8String.Empty : new Utf8String(prefix), cancellationToken: cancellationToken);
+				BundleLocator locator = await storageClient.WriteBundleAsync(bundle, prefix, cancellationToken: cancellationToken);
 				return new WriteBlobResponse { Blob = locator.ToString(), SupportsRedirects = storageClient?.Backend.SupportsRedirects };
 			}
 		}
diff --git a/Engine/Source/Programs/Horde/Horde.Server/Storage/StorageService.cs b/Engine/Source/Programs/Horde/Horde.Server/Storage/StorageService.cs
--- a/Engine/Source/Programs/Horde/Horde.Server/Storage/StorageService.cs
+++ b/Engine/Source/Programs/Horde/Horde.Server/Storage/StorageService.cs
@@ -303,11 +303,8 @@ namespace Horde.Server.Storage
 			public BlobHandle CreateBlobHandle(BlobLocator locator) => _impl.CreateBlobHandle(locator);
 			public BundleNodeHandle CreateNodeHandle(BundleNodeLocator locator) => _impl.CreateNodeHandle(locator);
 
-			public BundleWriter CreateWriter(RefName refName = default, BundleOptions? options = null) => _impl.CreateWriter(refName, options);
-			IStorageWriter IStorageClient.CreateWriter(RefName refName) => ((IStorageClient)_impl).CreateWriter(refName);
-
-			public ValueTask<BlobData> ReadBlobAsync(BlobLocator locator, CancellationToken cancellationToken = default) => _impl.ReadBlobAsync(locator, cancellationToken);
-			public ValueTask<BlobHandle> WriteBlobAsync(BlobType type, ReadOnlyMemory<byte> data, IReadOnlyList<BlobHandle> references, CancellationToken cancellationToken = default) => _impl.WriteBlobAsync(type, data, references, cancellationToken);
+			public BundleWriter CreateWriter(string? basePath = null, BundleOptions? options = null) => _impl.CreateWriter(basePath, options);
+			IStorageWriter IStorageClient.CreateWriter(string? basePath) => ((IStorageClient)_impl).CreateWriter(basePath);
 
 			#endregion
 
diff --git a/Engine/Source/Programs/Shared/EpicGames.Horde.Tests/BundleStoreTests.cs b/Engine/Source/Programs/Shared/EpicGames.Horde.Tests/BundleStoreTests.cs
--- a/Engine/Source/Programs/Shared/EpicGames.Horde.Tests/BundleStoreTests.cs
+++ b/Engine/Source/Programs/Shared/EpicGames.Horde.Tests/BundleStoreTests.cs
@@ -140,7 +140,7 @@ namespace EpicGames.Horde.Tests
 		{
 			// Generate a tree
 			{
-				await using IStorageWriter writer = store.CreateWriter(new RefName("test"), options);
+				await using IStorageWriter writer = store.CreateWriter("test", options);
 
 				SimpleNode node1 = new SimpleNode(new ReadOnlySequence<byte>(new byte[] { 1 }), Array.Empty<HashedNodeRef<SimpleNode>>());
 				SimpleNode node2 = new SimpleNode(new ReadOnlySequence<byte>(new byte[] { 2 }), new[] { await writer.WriteHashedNodeAsync(node1) });
@@ -208,24 +208,6 @@ namespace EpicGames.Horde.Tests
 			Assert.AreEqual(123, node.Data.FirstSpan[0]);
 		}
 
-		[TestMethod]
-		public async Task SimpleNodeWithRefDataAsync()
-		{
-			using MemoryStorageClient store = new MemoryStorageClient();
-
-			RefName refName = new RefName("test");
-			await using (IStorageWriter writer = store.CreateWriter(refName))
-			{
-				BlobHandle target = await writer.WriteBlobAsync(new BlobType(Guid.NewGuid(), 0), 0, Array.Empty<BlobHandle>());
-				await writer.WriteRefAsync(target, new byte[] { 4, 5, 6 });
-			}
-
-			RefValue? value = await ((IStorageClient)store).TryReadRefAsync(refName);
-
-			Assert.IsNotNull(value);
-			Assert.IsTrue(value.Data.Span.SequenceEqual(new byte[] { 4, 5, 6 }));			
-		}
-
 		[TestMethod]
 		public async Task DirectoryNodesAsync()
 		{
diff --git a/Engine/Source/Programs/Shared/EpicGames.Horde.Tests/ChunkingTests.cs b/Engine/Source/Programs/Shared/EpicGames.Horde.Tests/ChunkingTests.cs
--- a/Engine/Source/Programs/Shared/EpicGames.Horde.Tests/ChunkingTests.cs
+++ b/Engine/Source/Programs/Shared/EpicGames.Horde.Tests/ChunkingTests.cs
@@ -51,7 +51,8 @@ namespace EpicGames.Horde.Tests
 		{
 			using MemoryStorageClient store = new MemoryStorageClient();
 
-			await using (IStorageWriter writer = store.CreateWriter(new RefName("hello")))
+			const string RefName = "hello";
+			await using (IStorageWriter writer = store.CreateWriter(RefName))
 			{
 				ChunkingOptions options = new ChunkingOptions();
 				options.LeafOptions = new LeafChunkedDataNodeOptions(64, 64, 64);
@@ -65,7 +66,7 @@ namespace EpicGames.Horde.Tests
 				directory.AddFile("test.foo", FileEntryFlags.None, 0, chunkedData);
 
 				HashedNodeRef<DirectoryNode> directoryRef = await writer.WriteHashedNodeAsync(directory);
-				await writer.WriteRefAsync(directoryRef.Handle);
+				await store.WriteRefTargetAsync(RefName, directoryRef.Handle);
 			}
 		}
 
diff --git a/Engine/Source/Programs/Shared/EpicGames.Horde/Storage/Bundles/BundleWriter.cs b/Engine/Source/Programs/Shared/EpicGames.Horde/Storage/Bundles/BundleWriter.cs
--- a/Engine/Source/Programs/Shared/EpicGames.Horde/Storage/Bundles/BundleWriter.cs
+++ b/Engine/Source/Programs/Shared/EpicGames.Horde/Storage/Bundles/BundleWriter.cs
@@ -465,7 +465,7 @@ namespace EpicGames.Horde.Storage.Bundles
 			}
 
 			// Mark the bundle as complete
-			public async Task WriteAsync(BundleStorageClient store, Utf8String prefix, ILogger? traceLogger)
+			public async Task WriteAsync(BundleStorageClient store, string? basePath, ILogger? traceLogger)
 			{
 				traceLogger?.LogInformation("Marking bundle {BundleId} as complete ({NumNodes} nodes); adding to write queue.", BundleId, _queue.Count);
 
@@ -474,7 +474,7 @@ namespace EpicGames.Horde.Storage.Bundles
 				try
 				{
 					Bundle bundle = CreateBundle();
-					BundleLocator locator = await store.WriteBundleAsync(bundle, prefix);
+					BundleLocator locator = await store.WriteBundleAsync(bundle, basePath);
 					traceLogger?.LogInformation("Written bundle {BundleId} as {Locator}", BundleId, locator);
 
 					for (int idx = 0; idx < _queue.Count; idx++)
@@ -600,17 +600,17 @@ namespace EpicGames.Horde.Storage.Bundles
 		{
 			long _memoryFootprint;
 			readonly BundleStorageClient _store;
-			readonly Utf8String _prefix;
+			readonly string? _basePath;
 			readonly long _maxMemoryFootprint;
 			readonly ILogger? _traceLogger;
 			readonly AsyncEvent _completeEvent = new AsyncEvent();
 			int _refCount;
 			readonly List<Task> _writeTasks = new List<Task>();
 
-			public WriteQueue(BundleStorageClient store, Utf8String prefix, long maxMemoryFootprint, ILogger? traceLogger)
+			public WriteQueue(BundleStorageClient store, string? basePath, long maxMemoryFootprint, ILogger? traceLogger)
 			{
 				_store = store;
-				_prefix = prefix;
+				_basePath = basePath;
 				_maxMemoryFootprint = maxMemoryFootprint;
 				_refCount = 1;
 				_traceLogger = traceLogger;
@@ -673,7 +673,7 @@ namespace EpicGames.Horde.Storage.Bundles
 
 			async Task WriteAsync(PendingBundle pendingBundle)
 			{
-				await pendingBundle.WriteAsync(_store, _prefix, _traceLogger);
+				await pendingBundle.WriteAsync(_store, _basePath, _traceLogger);
 				Interlocked.Add(ref _memoryFootprint, -pendingBundle.CompressedLength);
 				pendingBundle.Dispose();
 				_completeEvent.Pulse();
@@ -685,7 +685,7 @@ namespace EpicGames.Horde.Storage.Bundles
 		readonly BundleStorageClient _store;
 		readonly BundleReader _reader;
 		readonly BundleOptions _options;
-		readonly RefName _refName;
+		readonly string? _basePath;
 
 		readonly WriteQueue _writeQueue;
 
@@ -700,11 +700,11 @@ namespace EpicGames.Horde.Storage.Bundles
 		/// </summary>
 		/// <param name="store">Store to write data to</param>
 		/// <param name="reader">Reader for serialized node data</param>
-		/// <param name="refName">Name of the ref being written</param>
+		/// <param name="basePath">Base path for new nodes</param>
 		/// <param name="options">Options for the writer</param>
 		/// <param name="traceLogger">Optional logger for trace information</param>
-		public BundleWriter(BundleStorageClient store, BundleReader reader, RefName refName, BundleOptions? options = null, ILogger? traceLogger = null)
-			: this(store, reader, refName, options, null, traceLogger)
+		public BundleWriter(BundleStorageClient store, BundleReader reader, string? basePath, BundleOptions? options = null, ILogger? traceLogger = null)
+			: this(store, reader, basePath, options, null, traceLogger)
 		{
 		}
 
@@ -713,7 +713,7 @@ namespace EpicGames.Horde.Storage.Bundles
 		/// </summary>
 		/// <param name="other"></param>
 		public BundleWriter(BundleWriter other)
-			: this(other._store, other._reader, other._refName, other._options, other._writeQueue, other.TraceLogger)
+			: this(other._store, other._reader, other._basePath, other._options, other._writeQueue, other.TraceLogger)
 		{
 			_writeQueue.AddRef();
 		}
@@ -721,13 +721,13 @@ namespace EpicGames.Horde.Storage.Bundles
 		/// <summary>
 		/// Internal constructor
 		/// </summary>
-		private BundleWriter(BundleStorageClient store, BundleReader reader, RefName refName, BundleOptions? options, WriteQueue? writeQueue, ILogger? traceLogger = null)
+		private BundleWriter(BundleStorageClient store, BundleReader reader, string? basePath, BundleOptions? options, WriteQueue? writeQueue, ILogger? traceLogger = null)
 		{
 			_store = store;
 			_reader = reader;
-			_refName = refName;
+			_basePath = basePath;
 			_options = options ?? s_defaultOptions;
-			_writeQueue = writeQueue ?? new WriteQueue(store, refName.Text, _options.MaxWriteQueueLength, traceLogger);
+			_writeQueue = writeQueue ?? new WriteQueue(store, basePath, _options.MaxWriteQueueLength, traceLogger);
 			TraceLogger = traceLogger;
 		}
 
@@ -847,16 +847,5 @@ namespace EpicGames.Horde.Storage.Bundles
 		{
 			return (BundleNodeHandle)await NodeRefExtensions.FlushAsync(this, root, cancellationToken);
 		}
-
-		/// <inheritdoc/>
-		ValueTask IStorageWriter.WriteRefAsync(BlobHandle target, ReadOnlyMemory<byte> data, RefOptions? options, CancellationToken cancellationToken) => WriteRefAsync((BundleNodeHandle)target, data, options, cancellationToken);
-
-		/// <inheritdoc/>
-		public async ValueTask WriteRefAsync(BundleNodeHandle target, ReadOnlyMemory<byte> data, RefOptions? options = null, CancellationToken cancellationToken = default)
-		{
-			await target.FlushAsync(cancellationToken);
-			BundleNodeLocator locator = target.GetLocator();
-			await _store.WriteRefAsync(_refName, locator, data, options, cancellationToken);
-		}
 	}
 }
diff --git a/Engine/Source/Programs/Shared/EpicGames.Horde/Storage/Clients/BundleStorageClient.cs b/Engine/Source/Programs/Shared/EpicGames.Horde/Storage/Clients/BundleStorageClient.cs
--- a/Engine/Source/Programs/Shared/EpicGames.Horde/Storage/Clients/BundleStorageClient.cs
+++ b/Engine/Source/Programs/Shared/EpicGames.Horde/Storage/Clients/BundleStorageClient.cs
@@ -1,7 +1,6 @@
 // Copyright Epic Games, Inc. All Rights Reserved.
 
 using System;
-using System.Collections.Generic;
 using System.IO;
 using System.Threading;
 using System.Threading.Tasks;
@@ -72,13 +71,13 @@ namespace EpicGames.Horde.Storage.Clients
 		/// </summary>
 		/// <param name="storageClient">Storage client</param>
 		/// <param name="bundle">Bundle to write</param>
-		/// <param name="prefix">Prefix for the uploaded data</param>
+		/// <param name="basePath">Prefix for the uploaded data</param>
 		/// <param name="cancellationToken">Cancellation token for the operation</param>
 		/// <returns>Locator for reading the bundle back in</returns>
-		public static async Task<BundleLocator> WriteBundleAsync(this IBundleStorageClient storageClient, Bundle bundle, Utf8String prefix = default, CancellationToken cancellationToken = default)
+		public static async Task<BundleLocator> WriteBundleAsync(this IBundleStorageClient storageClient, Bundle bundle, string? basePath = null, CancellationToken cancellationToken = default)
 		{
 			using ReadOnlySequenceStream stream = new ReadOnlySequenceStream(bundle.AsSequence());
-			string path = await storageClient.Backend.WriteAsync(stream, prefix.IsEmpty ? null : prefix.ToString(), cancellationToken);
+			string path = await storageClient.Backend.WriteAsync(stream, basePath, cancellationToken);
 			return new BundleLocator(path);
 		}
 	}
@@ -149,27 +148,10 @@ namespace EpicGames.Horde.Storage.Clients
 		public BundleNodeHandle CreateNodeHandle(BundleNodeLocator locator) => new FlushedNodeHandle(_bundleReader, locator);
 
 		/// <inheritdoc/>
-		public BundleWriter CreateWriter(RefName refName = default, BundleOptions? options = null) => new BundleWriter(this, _bundleReader, refName, options);
+		public BundleWriter CreateWriter(string? basePath = null, BundleOptions? options = null) => new BundleWriter(this, _bundleReader, basePath, options);
 
 		/// <inheritdoc/>
-		IStorageWriter IStorageClient.CreateWriter(RefName refName) => CreateWriter(refName);
-
-		/// <inheritdoc/>
-		public async ValueTask<BlobData> ReadBlobAsync(BlobLocator locator, CancellationToken cancellationToken = default)
-		{
-			BlobHandle handle = CreateBlobHandle(locator);
-			return await handle.ReadAsync(cancellationToken);
-		}
-
-		/// <inheritdoc/>
-		public async ValueTask<BlobHandle> WriteBlobAsync(BlobType type, ReadOnlyMemory<byte> data, IReadOnlyList<BlobHandle> references, CancellationToken cancellationToken = default)
-		{
-			await using (IStorageWriter writer = CreateWriter())
-			{
-				data.CopyTo(writer.GetOutputBuffer(0, data.Length));
-				return await writer.WriteBlobAsync(type, data.Length, references, cancellationToken);
-			}
-		}
+		IStorageWriter IStorageClient.CreateWriter(string? basePath) => CreateWriter(basePath);
 
 		#endregion
 
diff --git a/Engine/Source/Programs/Shared/EpicGames.Horde/Storage/IStorageClient.cs b/Engine/Source/Programs/Shared/EpicGames.Horde/Storage/IStorageClient.cs
--- a/Engine/Source/Programs/Shared/EpicGames.Horde/Storage/IStorageClient.cs
+++ b/Engine/Source/Programs/Shared/EpicGames.Horde/Storage/IStorageClient.cs
@@ -92,25 +92,9 @@ namespace EpicGames.Horde.Storage
 		/// <summary>
 		/// Creates a new writer for storage blobs
 		/// </summary>
-		/// <param name="refName">Name of the ref being written.</param>
+		/// <param name="basePath">Base path for any nodes written from the writer.</param>
 		/// <returns>New writer instance. Must be disposed after use.</returns>
-		IStorageWriter CreateWriter(RefName refName = default);
-
-		/// <summary>
-		/// Read a blob from the underlying storage system. Calling <see cref="BlobHandle.ReadAsync(CancellationToken)"/> is more efficient than calling this method repeatedly for small blobs.
-		/// </summary>
-		/// <param name="locator">Locator for the blob</param>
-		/// <param name="cancellationToken">Cancellation token for the operation</param>
-		ValueTask<BlobData> ReadBlobAsync(BlobLocator locator, CancellationToken cancellationToken = default);
-
-		/// <summary>
-		/// Write a blob to the underlying storage system. Using the writer instance returned from <see cref="CreateWriter(RefName)"/> is more efficient than calling this method repeatedly for small blobs.
-		/// </summary>
-		/// <param name="type">Type of the blob</param>
-		/// <param name="data">Data to be stored</param>
-		/// <param name="references">References to other blobs</param>
-		/// <param name="cancellationToken">Cancellation token for the operation</param>
-		ValueTask<BlobHandle> WriteBlobAsync(BlobType type, ReadOnlyMemory<byte> data, IReadOnlyList<BlobHandle> references, CancellationToken cancellationToken = default);
+		IStorageWriter CreateWriter(string? basePath = null);
 
 		#endregion
 
@@ -284,6 +268,60 @@ namespace EpicGames.Horde.Storage
 	/// </summary>
 	public static class StorageClientExtensions
 	{
+		#region Blobs
+
+		/// <summary>
+		/// Creates a writer using a refname as a base path
+		/// </summary>
+		/// <param name="store">The store instance to read from</param>
+		/// <param name="refName">Ref name to use as a base path</param>
+		public static IStorageWriter CreateWriter(this IStorageClient store, RefName refName) => store.CreateWriter(refName.ToString());
+
+		/// <summary>
+		/// Read a blob from the underlying storage system. Calling <see cref="BlobHandle.ReadAsync(CancellationToken)"/> is more efficient than calling this method repeatedly for small blobs.
+		/// </summary>
+		/// <param name="store">The store instance to read from</param>
+		/// <param name="locator">Locator for the blob</param>
+		/// <param name="cancellationToken">Cancellation token for the operation</param>
+		public static async ValueTask<BlobData> ReadBlobAsync(this IStorageClient store, BlobLocator locator, CancellationToken cancellationToken = default)
+		{
+			BlobHandle handle = store.CreateBlobHandle(locator);
+			return await handle.ReadAsync(cancellationToken);
+		}
+
+		/// <summary>
+		/// Write a blob to the underlying storage system. Using the writer instance returned from <see cref="IStorageClient.CreateWriter(String)"/> is more efficient than calling this method repeatedly for small blobs.
+		/// </summary>
+		/// <param name="store">The store instance to read from</param>
+		/// <param name="type">Type of the blob</param>
+		/// <param name="data">Data to be stored</param>
+		/// <param name="references">References to other blobs</param>
+		/// <param name="cancellationToken">Cancellation token for the operation</param>
+		public static ValueTask<BlobHandle> WriteBlobAsync(this IStorageClient store, BlobType type, ReadOnlyMemory<byte> data, IReadOnlyList<BlobHandle> references, CancellationToken cancellationToken = default)
+		{
+			return WriteBlobAsync(store, null, type, data, references, cancellationToken);
+		}
+
+		/// <summary>
+		/// Write a blob to the underlying storage system. Using the writer instance returned from <see cref="IStorageClient.CreateWriter(String)"/> is more efficient than calling this method repeatedly for small blobs.
+		/// </summary>
+		/// <param name="store">The store instance to read from</param>
+		/// <param name="basePath">Base path for writes</param>
+		/// <param name="type">Type of the blob</param>
+		/// <param name="data">Data to be stored</param>
+		/// <param name="references">References to other blobs</param>
+		/// <param name="cancellationToken">Cancellation token for the operation</param>
+		public static async ValueTask<BlobHandle> WriteBlobAsync(this IStorageClient store, string? basePath, BlobType type, ReadOnlyMemory<byte> data, IReadOnlyList<BlobHandle> references, CancellationToken cancellationToken = default)
+		{
+			await using (IStorageWriter writer = store.CreateWriter(basePath))
+			{
+				data.CopyTo(writer.GetOutputBuffer(0, data.Length));
+				return await writer.WriteBlobAsync(type, data.Length, references, cancellationToken);
+			}
+		}
+
+		#endregion
+
 		#region Aliases
 
 		/// <summary>
diff --git a/Engine/Source/Programs/Shared/EpicGames.Horde/Storage/IStorageWriter.cs b/Engine/Source/Programs/Shared/EpicGames.Horde/Storage/IStorageWriter.cs
--- a/Engine/Source/Programs/Shared/EpicGames.Horde/Storage/IStorageWriter.cs
+++ b/Engine/Source/Programs/Shared/EpicGames.Horde/Storage/IStorageWriter.cs
@@ -52,15 +52,6 @@ namespace EpicGames.Horde.Storage
 		/// <param name="cancellationToken">Cancellation token for the operation</param>
 		/// <returns>Handle to the written node</returns>
 		ValueTask<BlobHandle> WriteBlobAsync(BlobType type, int size, IReadOnlyList<BlobHandle> references, IReadOnlyList<AliasInfo> aliases, CancellationToken cancellationToken = default);
-
-		/// <summary>
-		/// Writes the reference using the given target node
-		/// </summary>
-		/// <param name="target">The target node</param>
-		/// <param name="data">Inline data for the node</param>
-		/// <param name="options">Options for the new ref</param>
-		/// <param name="cancellationToken">Cancellation token for the operation</param>
-		ValueTask WriteRefAsync(BlobHandle target, ReadOnlyMemory<byte> data = default, RefOptions? options = null, CancellationToken cancellationToken = default);
 	}
 
 	/// <summary>
@@ -202,9 +193,6 @@ namespace EpicGames.Horde.Storage
 			wrappedHandle._inner = await _inner.WriteBlobAsync(type, size, references.ConvertAll(x => ((WrappedHandle)x)._inner!), aliases, cancellationToken);
 			return wrappedHandle;
 		}
-
-		/// <inheritdoc/>
-		public ValueTask WriteRefAsync(BlobHandle target, ReadOnlyMemory<byte> data, RefOptions? options = null, CancellationToken cancellationToken = default) => _inner.WriteRefAsync(target, data, options, cancellationToken);
 	}
 
 	/// <summary>
diff --git a/Engine/Source/Programs/Shared/EpicGames.Horde/Storage/Node.cs b/Engine/Source/Programs/Shared/EpicGames.Horde/Storage/Node.cs
--- a/Engine/Source/Programs/Shared/EpicGames.Horde/Storage/Node.cs
+++ b/Engine/Source/Programs/Shared/EpicGames.Horde/Storage/Node.cs
@@ -429,7 +429,7 @@ namespace EpicGames.Horde.Storage
 		{
 			await using IStorageWriter writer = store.CreateWriter(name);
 			HashedNodeRef<Node> nodeRef = await writer.WriteHashedNodeAsync(node, cancellationToken);
-			await writer.WriteRefAsync(nodeRef.Handle, ReadOnlyMemory<byte>.Empty, refOptions, cancellationToken);
+			await store.WriteRefTargetAsync(name, nodeRef.Handle, refOptions, cancellationToken);
 			return nodeRef.Handle;
 		}
 
diff --git a/Engine/Source/Programs/UnrealCloudDDC/Jupiter/Controllers/BundlesController.cs b/Engine/Source/Programs/UnrealCloudDDC/Jupiter/Controllers/BundlesController.cs
--- a/Engine/Source/Programs/UnrealCloudDDC/Jupiter/Controllers/BundlesController.cs
+++ b/Engine/Source/Programs/UnrealCloudDDC/Jupiter/Controllers/BundlesController.cs
@@ -87,7 +87,7 @@ namespace Jupiter.Controllers
 				using (Stream stream = file.OpenReadStream())
 				{
 					Bundle bundle = await Bundle.FromStreamAsync(stream, cancellationToken);
-					BundleLocator locator = await client.WriteBundleAsync(bundle, prefix: (prefix == null) ? Utf8String.Empty : new Utf8String(prefix), cancellationToken: cancellationToken);
+					BundleLocator locator = await client.WriteBundleAsync(bundle, prefix, cancellationToken: cancellationToken);
 
 					WriteBlobResponse response = new WriteBlobResponse();
 					response.Blob = locator.ToString();
diff --git a/Engine/Source/Programs/UnrealCloudDDC/Jupiter/Implementation/Bundles/StorageService.cs b/Engine/Source/Programs/UnrealCloudDDC/Jupiter/Implementation/Bundles/StorageService.cs
--- a/Engine/Source/Programs/UnrealCloudDDC/Jupiter/Implementation/Bundles/StorageService.cs
+++ b/Engine/Source/Programs/UnrealCloudDDC/Jupiter/Implementation/Bundles/StorageService.cs
@@ -212,7 +212,7 @@ public class StorageClient : BundleStorageClient
 		}
 	}
 
-	public async Task<BlobHandle> WriteRefAsync(RefName name, Bundle bundle, int exportIdx, Utf8String prefix = default, RefOptions? options = null, CancellationToken cancellationToken = default)
+	public async Task<BlobHandle> WriteRefAsync(RefName name, Bundle bundle, int exportIdx, string? prefix = null, RefOptions? options = null, CancellationToken cancellationToken = default)
 	{
 		BundleLocator locator = await this.WriteBundleAsync(bundle, prefix, cancellationToken);
 		BundleNodeLocator nodeLocator = new BundleNodeLocator(locator, exportIdx);
