Horde: Replace BundleLocator with BlobLocator.

diff --git a/Engine/Source/Programs/Horde/Horde.Server.Tests/BlobStoreTests.cs b/Engine/Source/Programs/Horde/Horde.Server.Tests/BlobStoreTests.cs
--- a/Engine/Source/Programs/Horde/Horde.Server.Tests/BlobStoreTests.cs
+++ b/Engine/Source/Programs/Horde/Horde.Server.Tests/BlobStoreTests.cs
@@ -37,15 +37,15 @@ namespace Horde.Server.Tests
 		class Blob
 		{
 			public ReadOnlyMemory<byte> Data { get; set; } = ReadOnlyMemory<byte>.Empty;
-			public List<BundleLocator> References { get; set; } = new List<BundleLocator>();
+			public List<BlobLocator> References { get; set; } = new List<BlobLocator>();
 		}
 
-		static Bundle CreateTestBundle(ReadOnlyMemory<byte> data, IReadOnlyList<BundleLocator> refs)
+		static Bundle CreateTestBundle(ReadOnlyMemory<byte> data, IReadOnlyList<BlobLocator> refs)
 		{
 			List<BlobType> types = new List<BlobType>();
 			types.Add(new BlobType(Guid.Parse("{AFDF76A7-5333-4DEE-B837-B5F5CA511245}"), 1));
 
-			List<BundleLocator> imports = new List<BundleLocator>(refs);
+			List<BlobLocator> imports = new List<BlobLocator>(refs);
 
 			List<BundleExportRef> exportRefs = new List<BundleExportRef>();
 			for(int idx = 0; idx < refs.Count; idx++)
@@ -63,7 +63,7 @@ namespace Horde.Server.Tests
 			return new Bundle(header, new[] { data });
 		}
 
-		static async Task<Blob> ReadBlobAsync(IServerStorageClient store, BundleLocator locator)
+		static async Task<Blob> ReadBlobAsync(IServerStorageClient store, BlobLocator locator)
 		{
 			Bundle bundle = await store.ReadBundleAsync(locator);
 			return ExtractBlobFromBundle(bundle);
@@ -83,7 +83,7 @@ namespace Horde.Server.Tests
 			using IServerStorageClient store = CreateStorageClient();
 
 			byte[] input = CreateTestData(256, 0);
-			BundleLocator locator = await store.WriteBundleAsync(CreateTestBundle(input, Array.Empty<BundleLocator>()));
+			BlobLocator locator = await store.WriteBundleAsync(CreateTestBundle(input, Array.Empty<BlobLocator>()));
 
 			Bundle outputBundle = await store.ReadBundleAsync(locator);
 			Assert.IsTrue(outputBundle.Packets[0].Span.SequenceEqual(input));
@@ -95,25 +95,25 @@ namespace Horde.Server.Tests
 			using IServerStorageClient store = CreateStorageClient();
 
 			byte[] input1 = CreateTestData(256, 1);
-			Bundle bundle1 = CreateTestBundle(input1, Array.Empty<BundleLocator>());
-			BundleLocator locator1 = await store.WriteBundleAsync(bundle1);
+			Bundle bundle1 = CreateTestBundle(input1, Array.Empty<BlobLocator>());
+			BlobLocator locator1 = await store.WriteBundleAsync(bundle1);
 			Blob blob1 = await ReadBlobAsync(store, locator1);
 			Assert.IsTrue(blob1.Data.Span.SequenceEqual(input1));
-			Assert.IsTrue(blob1.References.SequenceEqual(Array.Empty<BundleLocator>()));
+			Assert.IsTrue(blob1.References.SequenceEqual(Array.Empty<BlobLocator>()));
 
 			byte[] input2 = CreateTestData(256, 2);
-			Bundle bundle2 = CreateTestBundle(input2, new BundleLocator[] { locator1 });
-			BundleLocator locator2 = await store.WriteBundleAsync(bundle2);
+			Bundle bundle2 = CreateTestBundle(input2, new BlobLocator[] { locator1 });
+			BlobLocator locator2 = await store.WriteBundleAsync(bundle2);
 			Blob blob2 = await ReadBlobAsync(store, locator2);
 			Assert.IsTrue(blob2.Data.Span.SequenceEqual(input2));
-			Assert.IsTrue(blob2.References.SequenceEqual(new BundleLocator[] { locator1 }));
+			Assert.IsTrue(blob2.References.SequenceEqual(new BlobLocator[] { locator1 }));
 
 			byte[] input3 = CreateTestData(256, 3);
-			Bundle bundle3 = CreateTestBundle(input3, new BundleLocator[] { locator1, locator2, locator1 });
-			BundleLocator locator3 = await store.WriteBundleAsync(bundle3);
+			Bundle bundle3 = CreateTestBundle(input3, new BlobLocator[] { locator1, locator2, locator1 });
+			BlobLocator locator3 = await store.WriteBundleAsync(bundle3);
 			Blob blob3 = await ReadBlobAsync(store, locator3);
 			Assert.IsTrue(blob3.Data.Span.SequenceEqual(input3));
-			Assert.IsTrue(blob3.References.SequenceEqual(new BundleLocator[] { locator1, locator2, locator1 }));
+			Assert.IsTrue(blob3.References.SequenceEqual(new BlobLocator[] { locator1, locator2, locator1 }));
 
 			for(int idx = 0; idx < 2; idx++)
 			{
@@ -129,8 +129,8 @@ namespace Horde.Server.Tests
 		{
 			using IServerStorageClient store = CreateStorageClient();
 
-			Bundle bundle1 = CreateTestBundle(new byte[] { 1, 2, 3 }, Array.Empty<BundleLocator>());
-			BundleLocator locator1 = await store.WriteBundleAsync(bundle1);
+			Bundle bundle1 = CreateTestBundle(new byte[] { 1, 2, 3 }, Array.Empty<BlobLocator>());
+			BlobLocator locator1 = await store.WriteBundleAsync(bundle1);
 			BundleNodeLocator target = new BundleNodeLocator(locator1, 0);
 
 			await store.WriteRefTargetAsync("test-ref-1", store.CreateBlobHandle(target.ToBlobLocator()));
diff --git a/Engine/Source/Programs/Horde/Horde.Server.Tests/GcServiceTests.cs b/Engine/Source/Programs/Horde/Horde.Server.Tests/GcServiceTests.cs
--- a/Engine/Source/Programs/Horde/Horde.Server.Tests/GcServiceTests.cs
+++ b/Engine/Source/Programs/Horde/Horde.Server.Tests/GcServiceTests.cs
@@ -34,9 +34,9 @@ namespace Horde.Server.Tests
 			using IServerStorageClient store = StorageService.CreateClient(new NamespaceId("default"));
 
 			Random random = new Random(0);
-			BundleLocator[] blobs = await CreateTestDataAsync(store, 30, 50, 30, 5, random);
+			BlobLocator[] blobs = await CreateTestDataAsync(store, 30, 50, 30, 5, random);
 
-			HashSet<BundleLocator> roots = new HashSet<BundleLocator>();
+			HashSet<BlobLocator> roots = new HashSet<BlobLocator>();
 			for (int idx = 0; idx < 10; idx++)
 			{
 				int blobIdx = (int)(random.NextDouble() * blobs.Length);
@@ -48,25 +48,25 @@ namespace Horde.Server.Tests
 				}
 			}
 
-			HashSet<BundleLocator> nodes = await FindNodesAsync(store, roots);
+			HashSet<BlobLocator> nodes = await FindNodesAsync(store, roots);
 
 			await Clock.AdvanceAsync(TimeSpan.FromDays(1.0));
 
-			BundleLocator[] remaining = await store.Backend.EnumerateAsync().Select(x => new BundleLocator(x)).ToArrayAsync();
+			BlobLocator[] remaining = await store.Backend.EnumerateAsync().Select(x => new BlobLocator(x)).ToArrayAsync();
 			Assert.AreEqual(nodes.Count, remaining.Length);
 			Assert.IsTrue(remaining.All(x => nodes.Contains(x)));
 		}
 
-		static async Task<HashSet<BundleLocator>> FindNodesAsync(IBundleStorageClient store, IEnumerable<BundleLocator> roots)
+		static async Task<HashSet<BlobLocator>> FindNodesAsync(IBundleStorageClient store, IEnumerable<BlobLocator> roots)
 		{
-			HashSet<BundleLocator> nodes = new HashSet<BundleLocator>();
+			HashSet<BlobLocator> nodes = new HashSet<BlobLocator>();
 			await FindNodesAsync(store, roots, nodes);
 			return nodes;
 		}
 
-		static async Task FindNodesAsync(IBundleStorageClient store, IEnumerable<BundleLocator> roots, HashSet<BundleLocator> nodes)
+		static async Task FindNodesAsync(IBundleStorageClient store, IEnumerable<BlobLocator> roots, HashSet<BlobLocator> nodes)
 		{
-			foreach (BundleLocator root in roots)
+			foreach (BlobLocator root in roots)
 			{
 				if (nodes.Add(root))
 				{
@@ -76,7 +76,7 @@ namespace Horde.Server.Tests
 			}
 		}
 
-		static async ValueTask<BundleLocator[]> CreateTestDataAsync(IBundleStorageClient store, int numRoots, int numInterior, int numLeaves, int avgChildren, Random random)
+		static async ValueTask<BlobLocator[]> CreateTestDataAsync(IBundleStorageClient store, int numRoots, int numInterior, int numLeaves, int avgChildren, Random random)
 		{
 			int firstRoot = 0;
 			int firstInterior = firstRoot + numRoots;
@@ -100,11 +100,11 @@ namespace Horde.Server.Tests
 				}
 			}
 
-			BundleLocator[] locators = new BundleLocator[children.Length];
+			BlobLocator[] locators = new BlobLocator[children.Length];
 			for (int idx = numNodes - 1; idx >= 0; idx--)
 			{
 				List<BlobType> types = new List<BlobType> { new BlobType(Guid.Parse("{AFDF76A7-5333-4DEE-B837-B5F5CA511245}"), 0) };
-				List<BundleLocator> imports = children[idx].ConvertAll(x => locators[x]);
+				List<BlobLocator> imports = children[idx].ConvertAll(x => locators[x]);
 				BundleHeader header = new BundleHeader(types.ToArray(), imports.ToArray(), Array.Empty<BundleExport>(), Array.Empty<BundlePacket>());
 				Bundle bundle = new Bundle(header, Array.Empty<ReadOnlyMemory<byte>>());
 				locators[idx] = await store.WriteBundleAsync(bundle, basePath: "gctest");
diff --git a/Engine/Source/Programs/Horde/Horde.Server.Tests/StorageServiceTests.cs b/Engine/Source/Programs/Horde/Horde.Server.Tests/StorageServiceTests.cs
--- a/Engine/Source/Programs/Horde/Horde.Server.Tests/StorageServiceTests.cs
+++ b/Engine/Source/Programs/Horde/Horde.Server.Tests/StorageServiceTests.cs
@@ -38,9 +38,9 @@ namespace Horde.Server.Tests
 			exports.Add(new BundleExport(0, 0, 0, data1.Length, Array.Empty<BundleExportRef>()));
 			exports.Add(new BundleExport(0, 0, data1.Length, data2.Length, Array.Empty<BundleExportRef>()));
 
-			BundleHeader header = new BundleHeader(types.ToArray(), Array.Empty<BundleLocator>(), exports.ToArray(), new BundlePacket[1]);
+			BundleHeader header = new BundleHeader(types.ToArray(), Array.Empty<BlobLocator>(), exports.ToArray(), new BundlePacket[1]);
 			Bundle bundle = new Bundle(header, Array.Empty<ReadOnlyMemory<byte>>());
-			BundleLocator locator = await client.WriteBundleAsync(bundle);
+			BlobLocator locator = await client.WriteBundleAsync(bundle);
 
 			await client.AddAliasAsync("foo", client.CreateBlobHandle(new BundleNodeLocator(locator, 0).ToBlobLocator()));
 			await client.AddAliasAsync("foo", client.CreateBlobHandle(new BundleNodeLocator(locator, 1).ToBlobLocator()));
diff --git a/Engine/Source/Programs/Horde/Horde.Server/Artifacts/ArtifactsController.cs b/Engine/Source/Programs/Horde/Horde.Server/Artifacts/ArtifactsController.cs
--- a/Engine/Source/Programs/Horde/Horde.Server/Artifacts/ArtifactsController.cs
+++ b/Engine/Source/Programs/Horde/Horde.Server/Artifacts/ArtifactsController.cs
@@ -9,7 +9,6 @@ using System.Threading.Tasks;
 using EpicGames.Core;
 using EpicGames.Horde.Artifacts;
 using EpicGames.Horde.Storage;
-using EpicGames.Horde.Storage.Bundles;
 using EpicGames.Horde.Storage.Nodes;
 using Horde.Server.Server;
 using Horde.Server.Storage;
@@ -77,7 +76,7 @@ namespace Horde.Server.Artifacts
 		[HttpGet]
 		[Route("/api/v2/artifacts/{id}/blobs/{*locator}")]
 		[Route("/api/v2/artifacts/{id}/bundles/{*locator}")]
-		public async Task<ActionResult> ReadArtifactBlobAsync(ArtifactId id, BundleLocator locator, CancellationToken cancellationToken = default)
+		public async Task<ActionResult> ReadArtifactBlobAsync(ArtifactId id, BlobLocator locator, CancellationToken cancellationToken = default)
 		{
 			IArtifact? artifact = await _artifactCollection.GetAsync(id, cancellationToken);
 			if (artifact == null)
diff --git a/Engine/Source/Programs/Horde/Horde.Server/Startup.cs b/Engine/Source/Programs/Horde/Horde.Server/Startup.cs
--- a/Engine/Source/Programs/Horde/Horde.Server/Startup.cs
+++ b/Engine/Source/Programs/Horde/Horde.Server/Startup.cs
@@ -106,7 +106,6 @@ using Polly;
 using Polly.Extensions.Http;
 using Horde.Server.Jobs.Bisect;
 using EpicGames.Horde;
-using EpicGames.Horde.Storage.Bundles;
 using Horde.Server.Ddc;
 using System.Net.Mime;
 using Microsoft.Extensions.Logging.Abstractions;
@@ -967,16 +966,16 @@ namespace Horde.Server
 				.WaitAndRetryAsync(new[] { TimeSpan.FromSeconds(1), TimeSpan.FromSeconds(5), TimeSpan.FromSeconds(10), TimeSpan.FromSeconds(10) });
 		}
 		
-		public sealed class BundleLocatorBsonSerializer : SerializerBase<BundleLocator>
+		public sealed class BlobLocatorBsonSerializer : SerializerBase<BlobLocator>
 		{
 			/// <inheritdoc/>
-			public override BundleLocator Deserialize(BsonDeserializationContext context, BsonDeserializationArgs args)
+			public override BlobLocator Deserialize(BsonDeserializationContext context, BsonDeserializationArgs args)
 			{
-				return new BundleLocator(context.Reader.ReadString());
+				return new BlobLocator(context.Reader.ReadString());
 			}
 
 			/// <inheritdoc/>
-			public override void Serialize(BsonSerializationContext context, BsonSerializationArgs args, BundleLocator value)
+			public override void Serialize(BsonSerializationContext context, BsonSerializationArgs args, BlobLocator value)
 			{
 				context.Writer.WriteString(value.ToString());
 			}
@@ -1107,7 +1106,7 @@ namespace Horde.Server
 				ConventionRegistry.Register("Horde", conventionPack, type => true);
 
 				// Register the custom serializers
-				BsonSerializer.RegisterSerializer(new BundleLocatorBsonSerializer());
+				BsonSerializer.RegisterSerializer(new BlobLocatorBsonSerializer());
 				BsonSerializer.RegisterSerializer(new RefNameBsonSerializer());
 				BsonSerializer.RegisterSerializer(new IoHashBsonSerializer());
 				BsonSerializer.RegisterSerializer(new NamespaceIdBsonSerializer());
diff --git a/Engine/Source/Programs/Horde/Horde.Server/Storage/StorageController.cs b/Engine/Source/Programs/Horde/Horde.Server/Storage/StorageController.cs
--- a/Engine/Source/Programs/Horde/Horde.Server/Storage/StorageController.cs
+++ b/Engine/Source/Programs/Horde/Horde.Server/Storage/StorageController.cs
@@ -104,7 +104,7 @@ namespace Horde.Server.Storage
 					bundle = await Bundle.FromStreamAsync(stream, cancellationToken);
 				}
 
-				BundleLocator locator = await storageClient.WriteBundleAsync(bundle, prefix, cancellationToken: cancellationToken);
+				BlobLocator locator = await storageClient.WriteBundleAsync(bundle, prefix, cancellationToken: cancellationToken);
 				return new WriteBlobResponse { Blob = locator.ToString(), SupportsRedirects = storageClient?.Backend.SupportsRedirects };
 			}
 		}
@@ -118,7 +118,7 @@ namespace Horde.Server.Storage
 		[HttpGet]
 		[Route("/api/v1/storage/{namespaceId}/blobs/{*locator}")]
 		[Route("/api/v1/storage/{namespaceId}/bundles/{*locator}")]
-		public async Task<ActionResult> ReadBlobAsync(NamespaceId namespaceId, BundleLocator locator, CancellationToken cancellationToken = default)
+		public async Task<ActionResult> ReadBlobAsync(NamespaceId namespaceId, BlobLocator locator, CancellationToken cancellationToken = default)
 		{
 			using IServerStorageClient? client = _storageService.TryCreateClient(namespaceId);
 			if (client == null)
@@ -136,7 +136,7 @@ namespace Horde.Server.Storage
 		/// <summary>
 		/// Reads a blob from storage, without performing namespace access checks.
 		/// </summary>
-		internal static async Task<ActionResult> ReadBlobInternalAsync(IBundleStorageClient storageClient, BundleLocator locator, IHeaderDictionary headers, CancellationToken cancellationToken)
+		internal static async Task<ActionResult> ReadBlobInternalAsync(IBundleStorageClient storageClient, BlobLocator locator, IHeaderDictionary headers, CancellationToken cancellationToken)
 		{
 			Uri? redirectUrl = await storageClient.Backend.TryGetReadRedirectAsync(locator.ToString(), cancellationToken);
 			if (redirectUrl != null)
@@ -311,7 +311,7 @@ namespace Horde.Server.Storage
 
 				if (hashIdx != -1)
 				{
-					response.Blob = new BundleLocator(locator.Substring(0, hashIdx));
+					response.Blob = new BlobLocator(locator.Substring(0, hashIdx));
 					response.ExportIdx = Int32.Parse(locator.Substring(hashIdx + 1));
 				}
 			}
diff --git a/Engine/Source/Programs/Horde/Horde.Server/Storage/StorageService.cs b/Engine/Source/Programs/Horde/Horde.Server/Storage/StorageService.cs
--- a/Engine/Source/Programs/Horde/Horde.Server/Storage/StorageService.cs
+++ b/Engine/Source/Programs/Horde/Horde.Server/Storage/StorageService.cs
@@ -127,7 +127,7 @@ namespace Horde.Server.Storage
 			{
 				string path = await _inner.WriteAsync(stream, $"{_prefix}{prefix}", cancellationToken);
 
-				BundleLocator locator = new BundleLocator(path);
+				BlobLocator locator = new BlobLocator(path);
 				await _outer.AddBlobAsync(_namespaceId, locator, null, cancellationToken);
 
 				return path;
@@ -174,7 +174,7 @@ namespace Horde.Server.Storage
 					return null;
 				}
 
-				BundleLocator locator = new BundleLocator(redirect.Value.Path);
+				BlobLocator locator = new BlobLocator(redirect.Value.Path);
 				await _outer.AddBlobAsync(_namespaceId, locator, null, cancellationToken);
 
 				return redirect;
@@ -308,9 +308,9 @@ namespace Horde.Server.Storage
 
 			#region Bundles
 
-			public Task<Stream> OpenAsync(BundleLocator locator, int offset, int? length = null, CancellationToken cancellationToken = default) => _impl.OpenAsync(locator, offset, length, cancellationToken);
+			public Task<Stream> OpenAsync(BlobLocator locator, int offset, int? length = null, CancellationToken cancellationToken = default) => _impl.OpenAsync(locator, offset, length, cancellationToken);
 
-			public Task<BundleHeader> ReadHeaderAsync(BundleLocator locator, CancellationToken cancellationToken) => _impl.ReadHeaderAsync(locator, cancellationToken);
+			public Task<BundleHeader> ReadHeaderAsync(BlobLocator locator, CancellationToken cancellationToken) => _impl.ReadHeaderAsync(locator, cancellationToken);
 
 			#endregion
 
@@ -409,7 +409,7 @@ namespace Horde.Server.Storage
 				Path = String.Empty;
 			}
 
-			public BlobInfo(ObjectId id, NamespaceId namespaceId, BundleLocator locator)
+			public BlobInfo(ObjectId id, NamespaceId namespaceId, BlobLocator locator)
 			{
 				Id = id;
 				NamespaceId = namespaceId;
@@ -683,7 +683,7 @@ namespace Horde.Server.Storage
 		#region Blobs
 
 		/// <inheritdoc/>
-		async Task AddBlobAsync(NamespaceId namespaceId, BundleLocator locator, List<AliasInfo>? exports = null, CancellationToken cancellationToken = default)
+		async Task AddBlobAsync(NamespaceId namespaceId, BlobLocator locator, List<AliasInfo>? exports = null, CancellationToken cancellationToken = default)
 		{
 			ObjectId id = ObjectId.GenerateNewId(_clock.UtcNow);
 			BlobInfo blobInfo = new BlobInfo(id, namespaceId, locator);
@@ -746,11 +746,11 @@ namespace Horde.Server.Storage
 
 							if (client != null)
 							{
-								BundleHeader? header = await client.ReadHeaderAsync(new BundleLocator(blobInfo.Locator.Path), cancellationToken);
+								BundleHeader? header = await client.ReadHeaderAsync(blobInfo.Locator, cancellationToken);
 								if (header != null)
 								{
 									List<ObjectId> importInfoIds = new List<ObjectId>();
-									foreach (BundleLocator import in header.Imports)
+									foreach (BlobLocator import in header.Imports)
 									{
 										FilterDefinition<BlobInfo> filter = Builders<BlobInfo>.Filter.Expr(x => x.NamespaceId == blobInfo.NamespaceId && x.Path == import.Path.ToString());
 										UpdateDefinition<BlobInfo> update = Builders<BlobInfo>.Update.SetOnInsert(x => x.Imports, null);
diff --git a/Engine/Source/Programs/Horde/Horde.Server/Tools/ToolsController.cs b/Engine/Source/Programs/Horde/Horde.Server/Tools/ToolsController.cs
--- a/Engine/Source/Programs/Horde/Horde.Server/Tools/ToolsController.cs
+++ b/Engine/Source/Programs/Horde/Horde.Server/Tools/ToolsController.cs
@@ -2,7 +2,6 @@
 
 using EpicGames.Core;
 using EpicGames.Horde.Storage;
-using EpicGames.Horde.Storage.Bundles;
 using EpicGames.Horde.Storage.Clients;
 using EpicGames.Horde.Storage.Nodes;
 using Horde.Server.Server;
@@ -535,7 +534,7 @@ namespace Horde.Server.Tools
 		/// <returns>Information about all the artifacts</returns>
 		[HttpGet]
 		[Route("/api/v1/tools/{id}/blobs/{*locator}")]
-		public async Task<ActionResult<object>> ReadToolBlobAsync(ToolId id, BundleLocator locator, CancellationToken cancellationToken = default)
+		public async Task<ActionResult<object>> ReadToolBlobAsync(ToolId id, BlobLocator locator, CancellationToken cancellationToken = default)
 		{
 			ITool? tool = await _toolCollection.GetAsync(id, _globalConfig.Value);
 			if (tool == null)
diff --git a/Engine/Source/Programs/Shared/EpicGames.Horde.Tests/BundleStoreTests.cs b/Engine/Source/Programs/Shared/EpicGames.Horde.Tests/BundleStoreTests.cs
--- a/Engine/Source/Programs/Shared/EpicGames.Horde.Tests/BundleStoreTests.cs
+++ b/Engine/Source/Programs/Shared/EpicGames.Horde.Tests/BundleStoreTests.cs
@@ -83,7 +83,7 @@ namespace EpicGames.Horde.Tests
 			List<BundlePacket> packets = new List<BundlePacket>();
 			packets.Add(new BundlePacket(BundleCompressionFormat.None, 0, payload.Length, payload.Length));
 
-			BundleHeader header = new BundleHeader(types.ToArray(), Array.Empty<BundleLocator>(), exports.ToArray(), packets.ToArray());
+			BundleHeader header = new BundleHeader(types.ToArray(), Array.Empty<BlobLocator>(), exports.ToArray(), packets.ToArray());
 			return new Bundle(header, new List<ReadOnlyMemory<byte>> { payload });
 		}
 
diff --git a/Engine/Source/Programs/Shared/EpicGames.Horde.Tests/BundleTests.cs b/Engine/Source/Programs/Shared/EpicGames.Horde.Tests/BundleTests.cs
--- a/Engine/Source/Programs/Shared/EpicGames.Horde.Tests/BundleTests.cs
+++ b/Engine/Source/Programs/Shared/EpicGames.Horde.Tests/BundleTests.cs
@@ -138,9 +138,9 @@ namespace EpicGames.Horde.Tests
 				List<BlobType> types = new List<BlobType>();
 				types.Add(new BlobType(Guid.NewGuid(), 0));
 
-				List<BundleLocator> imports = new List<BundleLocator>();
-				imports.Add(new BundleLocator("import1"));
-				imports.Add(new BundleLocator("import2"));
+				List<BlobLocator> imports = new List<BlobLocator>();
+				imports.Add(new BlobLocator("import1"));
+				imports.Add(new BlobLocator("import2"));
 
 				List<BundleExport> exports = new List<BundleExport>();
 				exports.Add(new BundleExport(0, 0, 0, 2, new BundleExportRef[] { new BundleExportRef(0, 5), new BundleExportRef(0, 6) }));
diff --git a/Engine/Source/Programs/Shared/EpicGames.Horde/Compute/AgentMessage.cs b/Engine/Source/Programs/Shared/EpicGames.Horde/Compute/AgentMessage.cs
--- a/Engine/Source/Programs/Shared/EpicGames.Horde/Compute/AgentMessage.cs
+++ b/Engine/Source/Programs/Shared/EpicGames.Horde/Compute/AgentMessage.cs
@@ -243,7 +243,7 @@ namespace EpicGames.Horde.Compute
 	/// <summary>
 	/// Creates a blob read request
 	/// </summary>
-	public record struct ReadBlobMessage(BundleLocator Locator, int Offset, int Length);
+	public record struct ReadBlobMessage(BlobLocator Locator, int Offset, int Length);
 
 	/// <summary>
 	/// Message for running an XOR command
@@ -540,7 +540,7 @@ namespace EpicGames.Horde.Compute
 		/// <returns></returns>
 		public static ReadBlobMessage ParseReadBlobRequest(this AgentMessage message)
 		{
-			BundleLocator locator = message.ReadBlobLocator();
+			BlobLocator locator = new BlobLocator(message.ReadUtf8String());
 			int offset = (int)message.ReadUnsignedVarInt();
 			int length = (int)message.ReadUnsignedVarInt();
 			return new ReadBlobMessage(locator, offset, length);
@@ -644,7 +644,7 @@ namespace EpicGames.Horde.Compute
 		/// <param name="length">Length of the data</param>
 		/// <param name="storage">Storage client to retrieve the blob from</param>
 		/// <param name="cancellationToken">Cancellation token for the operation</param>
-		public static async Task SendBlobDataAsync(this AgentMessageChannel channel, BundleLocator locator, int offset, int length, BundleStorageClient storage, CancellationToken cancellationToken = default)
+		public static async Task SendBlobDataAsync(this AgentMessageChannel channel, BlobLocator locator, int offset, int length, BundleStorageClient storage, CancellationToken cancellationToken = default)
 		{
 			using Stream stream = await storage.OpenAsync(locator, offset, length, cancellationToken);
 
diff --git a/Engine/Source/Programs/Shared/EpicGames.Horde/Storage/BlobLocator.cs b/Engine/Source/Programs/Shared/EpicGames.Horde/Storage/BlobLocator.cs
--- a/Engine/Source/Programs/Shared/EpicGames.Horde/Storage/BlobLocator.cs
+++ b/Engine/Source/Programs/Shared/EpicGames.Horde/Storage/BlobLocator.cs
@@ -61,6 +61,11 @@ namespace EpicGames.Horde.Storage
 			_path = new Utf8String(buffer);
 		}
 
+		/// <summary>
+		/// Whether the blob locator is valid
+		/// </summary>
+		public bool IsValid() => !_path.IsEmpty;
+
 		/// <summary>
 		/// The outermost blob locator
 		/// </summary>
@@ -109,6 +114,17 @@ namespace EpicGames.Horde.Storage
 			}
 		}
 
+		/// <summary>
+		/// Checks whether this blob is within the given folder
+		/// </summary>
+		/// <param name="folderName">Name of the folder</param>
+		/// <returns>True if the the blob id is within the given folder</returns>
+		public bool WithinFolder(Utf8String folderName)
+		{
+			Utf8String path = Path;
+			return path.Length > folderName.Length && path.StartsWith(folderName) && path[folderName.Length] == '/';
+		}
+
 		/// <inheritdoc/>
 		public override bool Equals(object? obj) => obj is BlobLocator blobId && Equals(blobId);
 
diff --git a/Engine/Source/Programs/Shared/EpicGames.Horde/Storage/Bundles/Bundle.cs b/Engine/Source/Programs/Shared/EpicGames.Horde/Storage/Bundles/Bundle.cs
--- a/Engine/Source/Programs/Shared/EpicGames.Horde/Storage/Bundles/Bundle.cs
+++ b/Engine/Source/Programs/Shared/EpicGames.Horde/Storage/Bundles/Bundle.cs
@@ -270,7 +270,7 @@ namespace EpicGames.Horde.Storage.Bundles
 		/// <summary>
 		/// Constructor
 		/// </summary>
-		public BundleHeader(BlobType[] types, BundleLocator[] imports, BundleExport[] exports, BundlePacket[] packets)
+		public BundleHeader(BlobType[] types, BlobLocator[] imports, BundleExport[] exports, BundlePacket[] packets)
 		{
 			Types = new BundleTypeCollection(types);
 			Imports = new BundleImportCollection(imports);
@@ -484,12 +484,12 @@ namespace EpicGames.Horde.Storage.Bundles
 
 			// Read the imports
 			int numImports = (int)reader.ReadUnsignedVarInt();
-			List<BundleLocator> imports = new List<BundleLocator>(numImports);
+			List<BlobLocator> imports = new List<BlobLocator>(numImports);
 			List<BundleExportRef> allExportReferences = new List<BundleExportRef>();
 
 			for (int importIdx = 0; importIdx < numImports; importIdx++)
 			{
-				BundleLocator locator = reader.ReadBlobLocator();
+				BlobLocator locator = new BlobLocator(reader.ReadUtf8String());
 				imports.Add(locator);
 
 				int[] exportIndexes = reader.ReadVariableLengthArray(() => (int)reader.ReadUnsignedVarInt());
@@ -676,9 +676,9 @@ namespace EpicGames.Horde.Storage.Bundles
 	/// <summary>
 	/// Collection of imported node references
 	/// </summary>
-	public struct BundleImportCollection : IReadOnlyList<BundleLocator>
+	public struct BundleImportCollection : IReadOnlyList<BlobLocator>
 	{
-		readonly BundleLocator[] _imports;
+		readonly BlobLocator[] _imports;
 
 		/// <inheritdoc/>
 		public int Count => _imports.Length;
@@ -686,7 +686,7 @@ namespace EpicGames.Horde.Storage.Bundles
 		/// <summary>
 		/// Constructor
 		/// </summary>
-		public BundleImportCollection(BundleLocator[] imports) => _imports = imports;
+		public BundleImportCollection(BlobLocator[] imports) => _imports = imports;
 
 		/// <summary>
 		/// Reads a collection from memory
@@ -696,12 +696,12 @@ namespace EpicGames.Horde.Storage.Bundles
 			ReadOnlySpan<byte> span = data.Span;
 			int count = (span.Length == 0) ? 0 : BinaryPrimitives.ReadInt32LittleEndian(span) / sizeof(int);
 
-			BundleLocator[] imports = new BundleLocator[count];
+			BlobLocator[] imports = new BlobLocator[count];
 			for (int idx = 0; idx < count; idx++)
 		{
 				int offset = BinaryPrimitives.ReadInt32LittleEndian(span.Slice(idx * sizeof(int)));
 				int length = span.Slice(offset).IndexOf((byte)0);
-				imports[idx] = new BundleLocator(new Utf8String(data.Slice(offset, length)));
+				imports[idx] = new BlobLocator(new Utf8String(data.Slice(offset, length)));
 			}
 
 			return new BundleImportCollection(imports);
@@ -718,10 +718,10 @@ namespace EpicGames.Horde.Storage.Bundles
 		}
 
 		/// <inheritdoc/>
-		public BundleLocator this[int index] => _imports[index];
+		public BlobLocator this[int index] => _imports[index];
 
 		/// <inheritdoc/>
-		public IEnumerator<BundleLocator> GetEnumerator() => _imports.AsEnumerable().GetEnumerator();
+		public IEnumerator<BlobLocator> GetEnumerator() => _imports.AsEnumerable().GetEnumerator();
 
 		/// <inheritdoc/>
 		IEnumerator IEnumerable.GetEnumerator() => GetEnumerator();
@@ -739,12 +739,12 @@ namespace EpicGames.Horde.Storage.Bundles
 		public void Write(IMemoryWriter writer)
 		{
 			int offset = _imports.Length * sizeof(int);
-			foreach (BundleLocator import in _imports)
+			foreach (BlobLocator import in _imports)
 			{
 				writer.WriteInt32(offset);
 				offset += import.Path.Length + 1;
 			}
-			foreach (BundleLocator import in _imports)
+			foreach (BlobLocator import in _imports)
 			{
 				writer.WriteNullTerminatedUtf8String(import.Path);
 			}
diff --git a/Engine/Source/Programs/Shared/EpicGames.Horde/Storage/Bundles/BundleLocator.cs b/Engine/Source/Programs/Shared/EpicGames.Horde/Storage/Bundles/BundleLocator.cs
deleted file mode 100644
diff --git a/Engine/Source/Programs/Shared/EpicGames.Horde/Storage/Bundles/BundleNodeLocator.cs b/Engine/Source/Programs/Shared/EpicGames.Horde/Storage/Bundles/BundleNodeLocator.cs
--- a/Engine/Source/Programs/Shared/EpicGames.Horde/Storage/Bundles/BundleNodeLocator.cs
+++ b/Engine/Source/Programs/Shared/EpicGames.Horde/Storage/Bundles/BundleNodeLocator.cs
@@ -22,7 +22,7 @@ namespace EpicGames.Horde.Storage.Bundles
 		/// <summary>
 		/// Location of the blob containing this node
 		/// </summary>
-		public BundleLocator Blob { get; }
+		public BlobLocator Blob { get; }
 
 		/// <summary>
 		/// Index of the export within the blob
@@ -32,7 +32,7 @@ namespace EpicGames.Horde.Storage.Bundles
 		/// <summary>
 		/// Constructor
 		/// </summary>
-		public BundleNodeLocator(BundleLocator blob, int exportIdx)
+		public BundleNodeLocator(BlobLocator blob, int exportIdx)
 		{
 			Blob = blob;
 			ExportIdx = exportIdx;
@@ -41,7 +41,7 @@ namespace EpicGames.Horde.Storage.Bundles
 		/// <summary>
 		/// Determines if this locator points to a valid entry
 		/// </summary>
-		public bool IsValid() => Blob.IsValid();
+		public bool IsValid() => !Blob.Path.IsEmpty;
 
 		/// <summary>
 		/// Parse a string as a node locator
@@ -63,7 +63,7 @@ namespace EpicGames.Horde.Storage.Bundles
 			}
 
 			int exportIdx = Int32.Parse(text.Slice(hashIdx + 1), NumberStyles.None, CultureInfo.InvariantCulture);
-			BundleLocator blobLocator = new BundleLocator(new Utf8String(text.Slice(0, hashIdx)));
+			BlobLocator blobLocator = new BlobLocator(new Utf8String(text.Slice(0, hashIdx)));
 			return new BundleNodeLocator(blobLocator, exportIdx);
 		}
 
diff --git a/Engine/Source/Programs/Shared/EpicGames.Horde/Storage/Bundles/BundleReader.cs b/Engine/Source/Programs/Shared/EpicGames.Horde/Storage/Bundles/BundleReader.cs
--- a/Engine/Source/Programs/Shared/EpicGames.Horde/Storage/Bundles/BundleReader.cs
+++ b/Engine/Source/Programs/Shared/EpicGames.Horde/Storage/Bundles/BundleReader.cs
@@ -22,7 +22,7 @@ namespace EpicGames.Horde.Storage.Bundles
 		/// <summary>
 		/// Locator for the bundle
 		/// </summary>
-		public BundleLocator Locator { get; }
+		public BlobLocator Locator { get; }
 
 		/// <summary>
 		/// Bundle header
@@ -37,7 +37,7 @@ namespace EpicGames.Horde.Storage.Bundles
 		/// <summary>
 		/// Constructor
 		/// </summary>
-		public BundleInfo(BundleLocator locator, BundleHeader header, int headerLength)
+		public BundleInfo(BlobLocator locator, BundleHeader header, int headerLength)
 		{
 			Locator = locator;
 			Header = header;
@@ -56,7 +56,7 @@ namespace EpicGames.Horde.Storage.Bundles
 		[DebuggerDisplay("{Locator}")]
 		class QueuedBundle
 		{
-			public BundleLocator Locator { get; }
+			public BlobLocator Locator { get; }
 			public TaskCompletionSource<BundleInfo> BundleInfo { get; } = new TaskCompletionSource<BundleInfo>();
 			public List<QueuedPacket> QueuedPackets { get; } = new List<QueuedPacket>(); // Sorted by index
 			public int InfoRefCount { get; set; } // Ref count for reading the header
@@ -65,7 +65,7 @@ namespace EpicGames.Horde.Storage.Bundles
 			public CancellationTokenSource CancellationSource { get; set; } = new CancellationTokenSource();
 			public bool Complete { get; set; }
 
-			public QueuedBundle(BundleLocator locator)
+			public QueuedBundle(BlobLocator locator)
 			{
 				Locator = locator;
 			}
@@ -166,7 +166,7 @@ namespace EpicGames.Horde.Storage.Bundles
 		/// <param name="locator">Locator for the bundle</param>
 		/// <param name="cancellationToken">Cancellation token for the operation</param>
 		/// <returns>Information about the bundle</returns>
-		public async Task<BundleHeader> ReadHeaderAsync(BundleLocator locator, CancellationToken cancellationToken)
+		public async Task<BundleHeader> ReadHeaderAsync(BlobLocator locator, CancellationToken cancellationToken)
 		{
 			// Check for a cached value first
 			if (_cache.TryGetCachedHeader(locator, out BundleInfo? cachedBundleInfo))
@@ -194,7 +194,7 @@ namespace EpicGames.Horde.Storage.Bundles
 			}
 		}
 
-		static string GetDecodeTaskKey(BundleLocator locator, int packetIdx) => $"{locator}#{packetIdx}";
+		static string GetDecodeTaskKey(BlobLocator locator, int packetIdx) => $"{locator}#{packetIdx}";
 
 		/// <summary>
 		/// Reads decoded packet data from a bundle
@@ -203,7 +203,7 @@ namespace EpicGames.Horde.Storage.Bundles
 		/// <param name="packetIdx">Index of the bundle packet</param>
 		/// <param name="cancellationToken">Cancellation token for the operation</param>
 		/// <returns>Data for the packet</returns>
-		public async Task<ReadOnlyMemory<byte>> ReadPacketAsync(BundleLocator locator, int packetIdx, CancellationToken cancellationToken)
+		public async Task<ReadOnlyMemory<byte>> ReadPacketAsync(BlobLocator locator, int packetIdx, CancellationToken cancellationToken)
 		{
 			ReadOnlyMemory<byte> cachedDecodedPacket;
 			if (_cache.TryGetCachedDecodedPacket(locator, packetIdx, out cachedDecodedPacket))
@@ -231,7 +231,7 @@ namespace EpicGames.Horde.Storage.Bundles
 			return await decodeTask.WaitAsync(cancellationToken);
 		}
 
-		async Task<ReadOnlyMemory<byte>> ReadAndDecodePacketAsync(BundleLocator locator, int packetIdx, CancellationToken cancellationToken)
+		async Task<ReadOnlyMemory<byte>> ReadAndDecodePacketAsync(BlobLocator locator, int packetIdx, CancellationToken cancellationToken)
 		{
 			BundleHeader header = await ReadHeaderAsync(locator, cancellationToken);
 			ReadOnlyMemory<byte> encodedPacket = await ReadEncodedPacketAsync(locator, packetIdx, cancellationToken);
@@ -260,7 +260,7 @@ namespace EpicGames.Horde.Storage.Bundles
 		/// <param name="packetIdx">Index of the bundle packet</param>
 		/// <param name="cancellationToken">Cancellation token for the operation</param>
 		/// <returns>Dtaa for the packet</returns>
-		public async Task<ReadOnlyMemory<byte>> ReadEncodedPacketAsync(BundleLocator locator, int packetIdx, CancellationToken cancellationToken)
+		public async Task<ReadOnlyMemory<byte>> ReadEncodedPacketAsync(BlobLocator locator, int packetIdx, CancellationToken cancellationToken)
 		{
 			// Register a read for the packet
 			QueuedBundle? bundle = null;
@@ -290,7 +290,7 @@ namespace EpicGames.Horde.Storage.Bundles
 			}
 		}
 
-		QueuedBundle FindOrAddBundle(BundleLocator locator)
+		QueuedBundle FindOrAddBundle(BlobLocator locator)
 		{
 			QueuedBundle? bundle = _queuedBundles.FirstOrDefault(x => x.Locator == locator);
 			if (bundle == null)
@@ -547,7 +547,7 @@ namespace EpicGames.Horde.Storage.Bundles
 			List<BlobHandle> refs = new List<BlobHandle>(export.References.Count);
 			foreach (BundleExportRef reference in export.References)
 			{
-				BundleLocator importBlob;
+				BlobLocator importBlob;
 				if (reference.ImportIdx == -1)
 				{
 					importBlob = locator.Blob;
diff --git a/Engine/Source/Programs/Shared/EpicGames.Horde/Storage/Bundles/BundleReaderCache.cs b/Engine/Source/Programs/Shared/EpicGames.Horde/Storage/Bundles/BundleReaderCache.cs
--- a/Engine/Source/Programs/Shared/EpicGames.Horde/Storage/Bundles/BundleReaderCache.cs
+++ b/Engine/Source/Programs/Shared/EpicGames.Horde/Storage/Bundles/BundleReaderCache.cs
@@ -104,14 +104,14 @@ namespace EpicGames.Horde.Storage
 			return cache.TryGetValue(cacheKey, out value);
 		}
 
-		static string GetBundleInfoCacheKey(BundleLocator locator) => $"bundle:{locator}";
-		static string GetEncodedPacketCacheKey(BundleLocator locator, int packetIdx) => $"encoded-packet:{locator}#{packetIdx}";
-		static string GetDecodedPacketCacheKey(BundleLocator locator, int packetIdx) => $"decoded-packet:{locator}#{packetIdx}";
+		static string GetBundleInfoCacheKey(BlobLocator locator) => $"bundle:{locator}";
+		static string GetEncodedPacketCacheKey(BlobLocator locator, int packetIdx) => $"encoded-packet:{locator}#{packetIdx}";
+		static string GetDecodedPacketCacheKey(BlobLocator locator, int packetIdx) => $"decoded-packet:{locator}#{packetIdx}";
 
 		/// <summary>
 		/// Adds a bundle info object to the cache
 		/// </summary>
-		public void AddCachedHeader(BundleLocator locator, BundleInfo bundleInfo)
+		public void AddCachedHeader(BlobLocator locator, BundleInfo bundleInfo)
 		{
 			AddCachedValue(_headerCache, GetBundleInfoCacheKey(locator), bundleInfo, bundleInfo.HeaderLength);
 		}
@@ -119,7 +119,7 @@ namespace EpicGames.Horde.Storage
 		/// <summary>
 		/// Try to read a bundle info object from the cache
 		/// </summary>
-		public bool TryGetCachedHeader(BundleLocator locator, [NotNullWhen(true)] out BundleInfo? bundleInfo)
+		public bool TryGetCachedHeader(BlobLocator locator, [NotNullWhen(true)] out BundleInfo? bundleInfo)
 		{
 			return TryGetCachedValue(_headerCache, GetBundleInfoCacheKey(locator), out bundleInfo);
 		}
@@ -127,7 +127,7 @@ namespace EpicGames.Horde.Storage
 		/// <summary>
 		/// Adds an encoded bundle packet to the cache
 		/// </summary>
-		public void AddCachedEncodedPacket(BundleLocator locator, int packetIdx, ReadOnlyMemory<byte> data)
+		public void AddCachedEncodedPacket(BlobLocator locator, int packetIdx, ReadOnlyMemory<byte> data)
 		{
 			AddCachedValue(_packetCache, GetEncodedPacketCacheKey(locator, packetIdx), data, data.Length);
 		}
@@ -135,7 +135,7 @@ namespace EpicGames.Horde.Storage
 		/// <summary>
 		/// Try to read an encoded bundle packet from the cache
 		/// </summary>
-		public bool TryGetCachedEncodedPacket(BundleLocator locator, int packetIdx, out ReadOnlyMemory<byte> data)
+		public bool TryGetCachedEncodedPacket(BlobLocator locator, int packetIdx, out ReadOnlyMemory<byte> data)
 		{
 			return TryGetCachedValue(_packetCache, GetEncodedPacketCacheKey(locator, packetIdx), out data);
 		}
@@ -143,7 +143,7 @@ namespace EpicGames.Horde.Storage
 		/// <summary>
 		/// Adds a decoded bundle packet to the cache
 		/// </summary>
-		public void AddCachedDecodedPacket(BundleLocator locator, int packetIdx, ReadOnlyMemory<byte> data)
+		public void AddCachedDecodedPacket(BlobLocator locator, int packetIdx, ReadOnlyMemory<byte> data)
 		{
 			AddCachedValue(_packetCache, GetDecodedPacketCacheKey(locator, packetIdx), data, data.Length);
 		}
@@ -151,7 +151,7 @@ namespace EpicGames.Horde.Storage
 		/// <summary>
 		/// Try to read a decoded bundle packet from the cache
 		/// </summary>
-		public bool TryGetCachedDecodedPacket(BundleLocator locator, int packetIdx, out ReadOnlyMemory<byte> data)
+		public bool TryGetCachedDecodedPacket(BlobLocator locator, int packetIdx, out ReadOnlyMemory<byte> data)
 		{
 			return TryGetCachedValue(_packetCache, GetDecodedPacketCacheKey(locator, packetIdx), out data);
 		}
diff --git a/Engine/Source/Programs/Shared/EpicGames.Horde/Storage/Bundles/BundleWriter.cs b/Engine/Source/Programs/Shared/EpicGames.Horde/Storage/Bundles/BundleWriter.cs
--- a/Engine/Source/Programs/Shared/EpicGames.Horde/Storage/Bundles/BundleWriter.cs
+++ b/Engine/Source/Programs/Shared/EpicGames.Horde/Storage/Bundles/BundleWriter.cs
@@ -474,7 +474,7 @@ namespace EpicGames.Horde.Storage.Bundles
 				try
 				{
 					Bundle bundle = CreateBundle();
-					BundleLocator locator = await store.WriteBundleAsync(bundle, basePath);
+					BlobLocator locator = await store.WriteBundleAsync(bundle, basePath);
 					traceLogger?.LogInformation("Written bundle {BundleId} as {Locator}", BundleId, locator);
 
 					for (int idx = 0; idx < _queue.Count; idx++)
@@ -521,8 +521,8 @@ namespace EpicGames.Horde.Storage.Bundles
 			Bundle CreateBundleInternal()
 			{
 				// List of imported blobs
-				List<BundleLocator> imports = new List<BundleLocator>();
-				Dictionary<BundleLocator, int> importToIndex = new Dictionary<BundleLocator, int>();
+				List<BlobLocator> imports = new List<BlobLocator>();
+				Dictionary<BlobLocator, int> importToIndex = new Dictionary<BlobLocator, int>();
 
 				// List of types in the bundle
 				List<BlobType> types = new List<BlobType>();
diff --git a/Engine/Source/Programs/Shared/EpicGames.Horde/Storage/Clients/BundleStorageClient.cs b/Engine/Source/Programs/Shared/EpicGames.Horde/Storage/Clients/BundleStorageClient.cs
--- a/Engine/Source/Programs/Shared/EpicGames.Horde/Storage/Clients/BundleStorageClient.cs
+++ b/Engine/Source/Programs/Shared/EpicGames.Horde/Storage/Clients/BundleStorageClient.cs
@@ -23,7 +23,7 @@ namespace EpicGames.Horde.Storage.Clients
 		#region Blobs
 
 		/// <inheritdoc/>
-		Task<Stream> OpenAsync(BundleLocator locator, int offset, int? length = null, CancellationToken cancellationToken = default);
+		Task<Stream> OpenAsync(BlobLocator locator, int offset, int? length = null, CancellationToken cancellationToken = default);
 
 		#endregion
 
@@ -34,7 +34,7 @@ namespace EpicGames.Horde.Storage.Clients
 		/// </summary>
 		/// <param name="locator">Locator for the bundle</param>
 		/// <param name="cancellationToken">Cancellation token for the operation</param>
-		Task<BundleHeader> ReadHeaderAsync(BundleLocator locator, CancellationToken cancellationToken);
+		Task<BundleHeader> ReadHeaderAsync(BlobLocator locator, CancellationToken cancellationToken);
 
 		#endregion
 	}
@@ -51,7 +51,7 @@ namespace EpicGames.Horde.Storage.Clients
 		/// <param name="locator">Locator for the bundle</param>
 		/// <param name="cancellationToken">Cancellation token for the operation</param>
 		/// <returns>Stream for reading from the bundle</returns>
-		public static Task<Stream> OpenAsync(this IBundleStorageClient storageClient, BundleLocator locator, CancellationToken cancellationToken = default) => storageClient.OpenAsync(locator, 0, null, cancellationToken);
+		public static Task<Stream> OpenAsync(this IBundleStorageClient storageClient, BlobLocator locator, CancellationToken cancellationToken = default) => storageClient.OpenAsync(locator, 0, null, cancellationToken);
 
 		/// <summary>
 		/// Reads an entire bundle into memory
@@ -60,7 +60,7 @@ namespace EpicGames.Horde.Storage.Clients
 		/// <param name="locator">Locator for the bundle</param>
 		/// <param name="cancellationToken">Cancellation token for the operation</param>
 		/// <returns>Bundle that was read</returns>
-		public static async Task<Bundle> ReadBundleAsync(this IBundleStorageClient storageClient, BundleLocator locator, CancellationToken cancellationToken = default)
+		public static async Task<Bundle> ReadBundleAsync(this IBundleStorageClient storageClient, BlobLocator locator, CancellationToken cancellationToken = default)
 		{
 			using Stream stream = await storageClient.OpenAsync(locator, cancellationToken);
 			return await Bundle.FromStreamAsync(stream, cancellationToken);
@@ -74,11 +74,11 @@ namespace EpicGames.Horde.Storage.Clients
 		/// <param name="basePath">Prefix for the uploaded data</param>
 		/// <param name="cancellationToken">Cancellation token for the operation</param>
 		/// <returns>Locator for reading the bundle back in</returns>
-		public static async Task<BundleLocator> WriteBundleAsync(this IBundleStorageClient storageClient, Bundle bundle, string? basePath = null, CancellationToken cancellationToken = default)
+		public static async Task<BlobLocator> WriteBundleAsync(this IBundleStorageClient storageClient, Bundle bundle, string? basePath = null, CancellationToken cancellationToken = default)
 		{
 			using ReadOnlySequenceStream stream = new ReadOnlySequenceStream(bundle.AsSequence());
 			string path = await storageClient.Backend.WriteAsync(stream, basePath, cancellationToken);
-			return new BundleLocator(path);
+			return new BlobLocator(path);
 		}
 	}
 
@@ -123,14 +123,14 @@ namespace EpicGames.Horde.Storage.Clients
 		#region Blobs
 
 		/// <inheritdoc/>
-		public async Task<Stream> OpenAsync(BundleLocator locator, int offset, int? length = null, CancellationToken cancellationToken = default) => await _backend.OpenAsync(locator.Path.ToString(), offset, length, cancellationToken);
+		public async Task<Stream> OpenAsync(BlobLocator locator, int offset, int? length = null, CancellationToken cancellationToken = default) => await _backend.OpenAsync(locator.Path.ToString(), offset, length, cancellationToken);
 
 		#endregion
 
 		#region Bundles
 
 		/// <inheritdoc/>
-		public Task<BundleHeader> ReadHeaderAsync(BundleLocator locator, CancellationToken cancellationToken) => _bundleReader.ReadHeaderAsync(locator, cancellationToken);
+		public Task<BundleHeader> ReadHeaderAsync(BlobLocator locator, CancellationToken cancellationToken) => _bundleReader.ReadHeaderAsync(locator, cancellationToken);
 
 		/// <inheritdoc/>
 		public async Task<BlobData> ReadNodeDataAsync(BundleNodeLocator locator, CancellationToken cancellationToken) => await _bundleReader.ReadNodeDataAsync(locator, cancellationToken);
diff --git a/Engine/Source/Programs/Shared/EpicGames.Horde/Storage/StorageMessages.cs b/Engine/Source/Programs/Shared/EpicGames.Horde/Storage/StorageMessages.cs
--- a/Engine/Source/Programs/Shared/EpicGames.Horde/Storage/StorageMessages.cs
+++ b/Engine/Source/Programs/Shared/EpicGames.Horde/Storage/StorageMessages.cs
@@ -2,7 +2,6 @@
 
 using System;
 using System.Collections.Generic;
-using EpicGames.Horde.Storage.Bundles;
 
 #pragma warning disable CA2227 // Collection properties should be read only
 
@@ -93,7 +92,7 @@ namespace EpicGames.Horde.Storage
 		/// Locator for the target blob
 		/// </summary>
 		[Obsolete("Use Target instead")]
-		public BundleLocator? Blob { get; set; }
+		public BlobLocator? Blob { get; set; }
 
 		/// <summary>
 		/// Export index for the ref
@@ -126,7 +125,7 @@ namespace EpicGames.Horde.Storage
 		/// Locator for the target blob
 		/// </summary>
 		[Obsolete("Use Target instead")]
-		public BundleLocator? Blob { get; set; }
+		public BlobLocator? Blob { get; set; }
 
 		/// <summary>
 		/// Export index for the ref
diff --git a/Engine/Source/Programs/UnrealCloudDDC/Jupiter.Common/Implementation/BlobId.cs b/Engine/Source/Programs/UnrealCloudDDC/Jupiter.Common/Implementation/BlobId.cs
--- a/Engine/Source/Programs/UnrealCloudDDC/Jupiter.Common/Implementation/BlobId.cs
+++ b/Engine/Source/Programs/UnrealCloudDDC/Jupiter.Common/Implementation/BlobId.cs
@@ -4,13 +4,11 @@ using System;
 using System.ComponentModel;
 using System.Globalization;
 using System.IO;
-using System.Text;
 using System.Text.Json;
 using System.Text.Json.Serialization;
 using System.Threading.Tasks;
 using Blake3;
 using EpicGames.Core;
-using EpicGames.Horde.Storage.Bundles;
 using EpicGames.Serialization;
 
 namespace Jupiter.Implementation
@@ -153,16 +151,6 @@ namespace Jupiter.Implementation
 		{
 			return new IoHash(HashData);
 		}
-
-		public static BlobId FromBlobLocator(BundleLocator locator)
-		{
-			return new BlobId(Encoding.UTF8.GetBytes(locator.ToString()));
-		}
-
-		public BundleLocator AsBlobLocator()
-		{
-			return new BundleLocator(Encoding.UTF8.GetString(HashData));
-		}
 	}
 
 	public class BlobIdentifierTypeConverter : TypeConverter
