Horde: Remove extension methods designed specifically for writing bundles to the storage layer.

diff --git a/Engine/Source/Programs/Shared/EpicGames.Horde.Tests/BundleStoreTests.cs b/Engine/Source/Programs/Shared/EpicGames.Horde.Tests/BundleStoreTests.cs
--- a/Engine/Source/Programs/Shared/EpicGames.Horde.Tests/BundleStoreTests.cs
+++ b/Engine/Source/Programs/Shared/EpicGames.Horde.Tests/BundleStoreTests.cs
@@ -65,7 +65,10 @@ namespace EpicGames.Horde.Tests
 			TextNode node = new TextNode("Hello world");
 			BundleNodeHandle handle = await writer.FlushAsync(node, CancellationToken.None);
 
-			return await store.ReadBundleAsync(handle.GetLocator().Blob);
+			BlobHandle bundleHandle = store.CreateBlobHandle(handle.GetLocator().Blob);
+			using BlobData blobData = await bundleHandle.ReadAsync();
+
+			return Bundle.FromMemory(blobData.Data.ToArray());
 		}
 
 		static Bundle CreateBundleManually()
diff --git a/Engine/Source/Programs/Shared/EpicGames.Horde/Storage/Bundles/BundleWriter.cs b/Engine/Source/Programs/Shared/EpicGames.Horde/Storage/Bundles/BundleWriter.cs
--- a/Engine/Source/Programs/Shared/EpicGames.Horde/Storage/Bundles/BundleWriter.cs
+++ b/Engine/Source/Programs/Shared/EpicGames.Horde/Storage/Bundles/BundleWriter.cs
@@ -474,7 +474,21 @@ namespace EpicGames.Horde.Storage.Bundles
 				try
 				{
 					Bundle bundle = CreateBundle();
-					BlobLocator locator = await store.WriteBundleAsync(bundle, basePath).GetLocatorAsync();
+
+					// Write the bundle to storage
+					BlobHandle[] imports = new BlobHandle[bundle.Header.Imports.Count];
+					for (int idx = 0; idx < bundle.Header.Imports.Count; idx++)
+					{
+						imports[idx] = store.CreateBlobHandle(new BlobLocator(bundle.Header.Imports[idx].Path));
+					}
+
+					BlobHandle handle;
+					using (ReadOnlySequenceStream stream = new ReadOnlySequenceStream(bundle.AsSequence()))
+					{
+						handle = await store.WriteBlobAsync(BundleStorageClient.BundleBlobType, stream, imports, basePath);
+					}
+
+					BlobLocator locator = handle.GetLocator();
 					traceLogger?.LogInformation("Written bundle {BundleId} as {Locator}", BundleId, locator);
 
 					for (int idx = 0; idx < _queue.Count; idx++)
diff --git a/Engine/Source/Programs/Shared/EpicGames.Horde/Storage/Clients/BundleStorageClient.cs b/Engine/Source/Programs/Shared/EpicGames.Horde/Storage/Clients/BundleStorageClient.cs
--- a/Engine/Source/Programs/Shared/EpicGames.Horde/Storage/Clients/BundleStorageClient.cs
+++ b/Engine/Source/Programs/Shared/EpicGames.Horde/Storage/Clients/BundleStorageClient.cs
@@ -13,58 +13,6 @@ using Microsoft.Extensions.Logging;
 
 namespace EpicGames.Horde.Storage.Clients
 {
-	/// <summary>
-	/// Extension methods for reading bundles
-	/// </summary>
-	public static class BundleStorageClientExtensions
-	{
-		/// <summary>
-		/// Reads an entire bundle into memory
-		/// </summary>
-		/// <param name="handle">Handle to read from</param>
-		/// <param name="cancellationToken">Cancellation token for the operation</param>
-		/// <returns>Bundle that was read</returns>
-		public static async Task<Bundle> ReadBundleAsync(this BlobHandle handle, CancellationToken cancellationToken = default)
-		{
-			using Stream stream = await handle.OpenAsync(cancellationToken: cancellationToken);
-			return await Bundle.FromStreamAsync(stream, cancellationToken);
-		}
-
-		/// <summary>
-		/// Reads an entire bundle into memory
-		/// </summary>
-		/// <param name="storageClient">Storage client</param>
-		/// <param name="locator">Locator for the bundle</param>
-		/// <param name="cancellationToken">Cancellation token for the operation</param>
-		/// <returns>Bundle that was read</returns>
-		public static async Task<Bundle> ReadBundleAsync(this IStorageClient storageClient, BlobLocator locator, CancellationToken cancellationToken = default)
-		{
-			BlobHandle handle = storageClient.CreateBlobHandle(locator);
-			using Stream stream = await handle.OpenAsync(cancellationToken: cancellationToken);
-			return await Bundle.FromStreamAsync(stream, cancellationToken);
-		}
-
-		/// <summary>
-		/// Writes an entire bundle to a storage client
-		/// </summary>
-		/// <param name="storageClient">Storage client</param>
-		/// <param name="bundle">Bundle to write</param>
-		/// <param name="basePath">Prefix for the uploaded data</param>
-		/// <param name="cancellationToken">Cancellation token for the operation</param>
-		/// <returns>Locator for reading the bundle back in</returns>
-		public static async Task<BlobHandle> WriteBundleAsync(this IStorageClient storageClient, Bundle bundle, string? basePath = null, CancellationToken cancellationToken = default)
-		{
-			BlobHandle[] imports = new BlobHandle[bundle.Header.Imports.Count];
-			for (int idx = 0; idx < bundle.Header.Imports.Count; idx++)
-			{
-				imports[idx] = storageClient.CreateBlobHandle(new BlobLocator(bundle.Header.Imports[idx].Path));
-			}
-
-			using ReadOnlySequenceStream stream = new ReadOnlySequenceStream(bundle.AsSequence());
-			return await storageClient.WriteBlobAsync(BundleStorageClient.BundleBlobType, stream, imports, basePath, cancellationToken);
-		}
-	}
-
 	/// <summary>
 	/// Base class for an implementation of <see cref="IStorageClient"/>, providing implementations for some common functionality using bundles.
 	/// </summary>
diff --git a/Engine/Source/Programs/Shared/EpicGames.Horde/Storage/IStorageWriter.cs b/Engine/Source/Programs/Shared/EpicGames.Horde/Storage/IStorageWriter.cs
--- a/Engine/Source/Programs/Shared/EpicGames.Horde/Storage/IStorageWriter.cs
+++ b/Engine/Source/Programs/Shared/EpicGames.Horde/Storage/IStorageWriter.cs
@@ -50,7 +50,7 @@ namespace EpicGames.Horde.Storage
 		Memory<byte> GetOutputBuffer(int usedSize, int desiredSize);
 
 		/// <summary>
-		/// Finish writing a node.
+		/// Finish writing a blob that has been written into the output buffer.
 		/// </summary>
 		/// <param name="type">Type of the node that was written</param>
 		/// <param name="size">Used size of the buffer</param>
