Added the limited ability to add C++ attribute specifiers at the start of a UFUNCTION declaration.

diff --git a/Engine/Source/Programs/Shared/EpicGames.UHT/Parsers/UhtFunctionParser.cs b/Engine/Source/Programs/Shared/EpicGames.UHT/Parsers/UhtFunctionParser.cs
--- a/Engine/Source/Programs/Shared/EpicGames.UHT/Parsers/UhtFunctionParser.cs
+++ b/Engine/Source/Programs/Shared/EpicGames.UHT/Parsers/UhtFunctionParser.cs
@@ -342,6 +342,8 @@ namespace EpicGames.UHT.Parsers
 							break;
 					}
 
+					topScope.TokenReader.OptionalAttributes(false);
+
 					if (topScope.TokenReader.TryOptional("static"))
 					{
 						function.FunctionFlags |= EFunctionFlags.Static;
