Horde: Fix writing non-bundle blob types through BundleStorageClient.

diff --git a/Engine/Source/Programs/Shared/EpicGames.Horde/Compute/AgentMessage.cs b/Engine/Source/Programs/Shared/EpicGames.Horde/Compute/AgentMessage.cs
--- a/Engine/Source/Programs/Shared/EpicGames.Horde/Compute/AgentMessage.cs
+++ b/Engine/Source/Programs/Shared/EpicGames.Horde/Compute/AgentMessage.cs
@@ -646,7 +646,7 @@ namespace EpicGames.Horde.Compute
 		/// <param name="cancellationToken">Cancellation token for the operation</param>
 		public static async Task SendBlobDataAsync(this AgentMessageChannel channel, BlobLocator locator, int offset, int length, BundleStorageClient storage, CancellationToken cancellationToken = default)
 		{
-			using Stream stream = await storage.OpenAsync(locator, offset, length, cancellationToken);
+			using Stream stream = await storage.OpenBlobAsync(locator, offset, length, cancellationToken);
 
 			const int MaxChunkSize = 512 * 1024;
 			for (int chunkOffset = 0; chunkOffset < stream.Length;)
diff --git a/Engine/Source/Programs/Shared/EpicGames.Horde/Storage/Bundles/BundleReader.cs b/Engine/Source/Programs/Shared/EpicGames.Horde/Storage/Bundles/BundleReader.cs
--- a/Engine/Source/Programs/Shared/EpicGames.Horde/Storage/Bundles/BundleReader.cs
+++ b/Engine/Source/Programs/Shared/EpicGames.Horde/Storage/Bundles/BundleReader.cs
@@ -444,7 +444,7 @@ namespace EpicGames.Horde.Storage.Bundles
 			int prefetchSize = _cache != null ? DefaultFetchSize : DefaultUncachedFetchSize;
 			for (; ; )
 			{
-				await using (Stream stream = await _store.OpenAsync(bundle.Locator, 0, prefetchSize, cancellationToken))
+				await using (Stream stream = await _store.OpenBlobAsync(bundle.Locator, 0, prefetchSize, cancellationToken))
 				{
 					// Read the header data
 					byte[] prelude = new byte[BundleHeader.PreludeLength];
@@ -499,7 +499,7 @@ namespace EpicGames.Horde.Storage.Bundles
 			Interlocked.Increment(ref _numPacketReads);
 			Interlocked.Add(ref _numBytesRead, maxOffset - minOffset);
 
-			await using (Stream stream = await _store.OpenAsync(bundleInfo.Locator, bundleInfo.HeaderLength + minOffset, maxOffset - minOffset, cancellationToken))
+			await using (Stream stream = await _store.OpenBlobAsync(bundleInfo.Locator, bundleInfo.HeaderLength + minOffset, maxOffset - minOffset, cancellationToken))
 			{
 				// Copy all the packets that have been read into separate buffers, so we can cache them indidually.
 				for (int packetIdx = minPacketIdx; packetIdx <= maxPacketIdx; packetIdx++)
diff --git a/Engine/Source/Programs/Shared/EpicGames.Horde/Storage/Clients/BundleStorageClient.cs b/Engine/Source/Programs/Shared/EpicGames.Horde/Storage/Clients/BundleStorageClient.cs
--- a/Engine/Source/Programs/Shared/EpicGames.Horde/Storage/Clients/BundleStorageClient.cs
+++ b/Engine/Source/Programs/Shared/EpicGames.Horde/Storage/Clients/BundleStorageClient.cs
@@ -2,7 +2,6 @@
 
 using System;
 using System.Collections.Generic;
-using System.Diagnostics;
 using System.Diagnostics.CodeAnalysis;
 using System.IO;
 using System.Threading;
@@ -56,12 +55,12 @@ namespace EpicGames.Horde.Storage.Clients
 
 			public override Task<Stream> OpenAsync(int offset = 0, int? length = null, CancellationToken cancellationToken = default)
 			{
-				return _storageClient.OpenAsync(new BlobLocator(_path), offset, length, cancellationToken);
+				return _storageClient.OpenBlobAsync(new BlobLocator(_path), offset, length, cancellationToken);
 			}
 
 			public override async ValueTask<BlobData> ReadAsync(CancellationToken cancellationToken = default)
 			{
-				using (Stream stream = await _storageClient.OpenAsync(new BlobLocator(_path), 0, cancellationToken: cancellationToken))
+				using (Stream stream = await _storageClient.OpenBlobAsync(new BlobLocator(_path), 0, cancellationToken: cancellationToken))
 				{
 					byte[] data = await stream.ReadAllBytesAsync(cancellationToken);
 					return new BlobData(BundleBlobType, data, await GetRefsAsync(cancellationToken));
@@ -114,14 +113,34 @@ namespace EpicGames.Horde.Storage.Clients
 		#region Blobs
 
 		/// <inheritdoc/>
-		public async Task<Stream> OpenAsync(BlobLocator locator, int offset = 0, int? length = null, CancellationToken cancellationToken = default) => await _backend.OpenAsync(locator.Path.ToString(), offset, length, cancellationToken);
+		public async Task<Stream> OpenBlobAsync(BlobLocator locator, int offset = 0, int? length = null, CancellationToken cancellationToken = default) 
+			=> await _backend.OpenAsync(locator.Path.ToString(), offset, length, cancellationToken);
 
 		/// <inheritdoc/>
 		public async ValueTask<BlobHandle> WriteBlobAsync(BlobType type, Stream stream, IReadOnlyList<BlobHandle> references, string? basePath = null, CancellationToken cancellationToken = default)
 		{
-			Debug.Assert(type == BundleBlobType);
-			string path = await _backend.WriteAsync(stream, basePath, cancellationToken);
-			return new BundleHandle(this, path);
+			if (type == BundleBlobType)
+			{
+				string path = await _backend.WriteAsync(stream, basePath, cancellationToken);
+				return new BundleHandle(this, path);
+			}
+			else
+			{
+				await using IStorageWriter writer = CreateWriter(basePath);
+
+				int length = 0;
+				for (; ; )
+				{
+					int readLength = await stream.ReadAsync(writer.GetOutputBuffer(length, length + 1), cancellationToken);
+					if (readLength == 0)
+					{
+						break;
+					}
+					length += readLength;
+				}
+
+				return await writer.WriteBlobAsync(type, length, references, Array.Empty<AliasInfo>(), cancellationToken);
+			}
 		}
 
 		/// <inheritdoc/>
