Horde: Missing file.

diff --git a/Engine/Source/Programs/Shared/EpicGames.Horde.Tests/BundleTests.cs b/Engine/Source/Programs/Shared/EpicGames.Horde.Tests/BundleTests.cs
--- a/Engine/Source/Programs/Shared/EpicGames.Horde.Tests/BundleTests.cs
+++ b/Engine/Source/Programs/Shared/EpicGames.Horde.Tests/BundleTests.cs
@@ -215,7 +215,10 @@ namespace EpicGames.Horde.Tests
 
 			// Check the ref
 			BlobHandle refTarget =  await store.ReadRefTargetAsync(refName);
-			Bundle bundle = await store.ReadBundleAsync(((BundleNodeHandle)refTarget).GetLocator().Blob);
+			BlobHandle bundleTarget = store.CreateBlobHandle(refTarget.GetLocator().Outermost);
+			using BlobData bundleData = await bundleTarget.ReadAsync();
+			Bundle bundle = Bundle.FromMemory(bundleData.Data);
+
 			Assert.AreEqual(0, bundle.Header.Imports.Count);
 			Assert.AreEqual(3, bundle.Header.Exports.Count);
 
