Horde: Fix static analysis warnings.

diff --git a/Engine/Source/Programs/Horde/Horde.Server/Storage/StorageService.cs b/Engine/Source/Programs/Horde/Horde.Server/Storage/StorageService.cs
--- a/Engine/Source/Programs/Horde/Horde.Server/Storage/StorageService.cs
+++ b/Engine/Source/Programs/Horde/Horde.Server/Storage/StorageService.cs
@@ -4,11 +4,9 @@ using System;
 using System.Buffers;
 using System.Collections.Generic;
 using System.ComponentModel;
-using System.Diagnostics;
 using System.Diagnostics.CodeAnalysis;
 using System.IO;
 using System.Linq;
-using System.Security.Claims;
 using System.Threading;
 using System.Threading.Tasks;
 using EpicGames.Core;
@@ -18,7 +16,6 @@ using EpicGames.Horde.Storage.Clients;
 using EpicGames.Horde.Storage.Nodes;
 using EpicGames.Redis;
 using EpicGames.Redis.Utility;
-using Horde.Server.Acls;
 using Horde.Server.Server;
 using Horde.Server.Utilities;
 using HordeCommon;
diff --git a/Engine/Source/Programs/Shared/EpicGames.Horde/Storage/Clients/BundleStorageClient.cs b/Engine/Source/Programs/Shared/EpicGames.Horde/Storage/Clients/BundleStorageClient.cs
--- a/Engine/Source/Programs/Shared/EpicGames.Horde/Storage/Clients/BundleStorageClient.cs
+++ b/Engine/Source/Programs/Shared/EpicGames.Horde/Storage/Clients/BundleStorageClient.cs
@@ -8,7 +8,6 @@ using System.Threading;
 using System.Threading.Tasks;
 using EpicGames.Core;
 using EpicGames.Horde.Storage.Bundles;
-using Microsoft.CodeAnalysis;
 using Microsoft.Extensions.Logging;
 
 namespace EpicGames.Horde.Storage.Clients
@@ -361,7 +360,8 @@ namespace EpicGames.Horde.Storage.Clients
 			RefValue? refValue = await _inner.TryReadRefAsync(name, cacheTime, cancellationToken);
 			if (refValue != null)
 			{
-				refValue = new RefValue(CreateBlobHandle(refValue.Target.GetLocator()), refValue.Data);
+				BlobHandle target = CreateBlobHandle(refValue.Target.GetLocator());
+				refValue = new RefValue(target, refValue.Data);
 			}
 			return refValue;
 		}
