Horde: Remove overloads to specify inline data on refs. This was meant for encoding fragment data, which is now handled explicitly.

diff --git a/Engine/Source/Programs/Horde/Horde.Server/Commands/Bundles/BundleWriteTest.cs b/Engine/Source/Programs/Horde/Horde.Server/Commands/Bundles/BundleWriteTest.cs
--- a/Engine/Source/Programs/Horde/Horde.Server/Commands/Bundles/BundleWriteTest.cs
+++ b/Engine/Source/Programs/Horde/Horde.Server/Commands/Bundles/BundleWriteTest.cs
@@ -37,8 +37,8 @@ namespace Horde.Server.Commands.Bundles
 			public Task AddAliasAsync(string name, BlobHandle handle, int rank = 0, ReadOnlyMemory<byte> data = default, CancellationToken cancellationToken = default) => throw new NotImplementedException();
 			public Task RemoveAliasAsync(string name, BlobHandle handle, CancellationToken cancellationToken = default) => throw new NotImplementedException();
 			public Task<BlobAlias[]> FindAliasesAsync(string alias, int? maxResults = null, CancellationToken cancellationToken = default) => throw new NotImplementedException();
-			public Task<RefValue?> TryReadRefAsync(RefName name, RefCacheTime cacheTime = default, CancellationToken cancellationToken = default) => throw new NotImplementedException();
-			public Task WriteRefAsync(RefName name, BlobHandle target, ReadOnlyMemory<byte> data = default, RefOptions? options = null, CancellationToken cancellationToken = default) => Task.CompletedTask;
+			public Task<BlobHandle?> TryReadRefAsync(RefName name, RefCacheTime cacheTime = default, CancellationToken cancellationToken = default) => throw new NotImplementedException();
+			public Task WriteRefAsync(RefName name, BlobHandle target, RefOptions? options = null, CancellationToken cancellationToken = default) => Task.CompletedTask;
 
 			public BundleWriter CreateWriter(string? basePath = null, BundleOptions? options = null)
 			{
diff --git a/Engine/Source/Programs/Horde/Horde.Server/Storage/StorageController.cs b/Engine/Source/Programs/Horde/Horde.Server/Storage/StorageController.cs
--- a/Engine/Source/Programs/Horde/Horde.Server/Storage/StorageController.cs
+++ b/Engine/Source/Programs/Horde/Horde.Server/Storage/StorageController.cs
@@ -246,7 +246,7 @@ namespace Horde.Server.Storage
 #pragma warning restore CS0618 // Type or member is obsolete
 
 			BlobHandle target = client.CreateBlobHandle(request.Target);
-			await client.WriteRefAsync(refName, target, request.Data, request.Options, cancellationToken);
+			await client.WriteRefAsync(refName, target, request.Options, cancellationToken);
 
 			return Ok();
 		}
@@ -290,19 +290,19 @@ namespace Horde.Server.Storage
 				}
 			}
 
-			RefValue? target = await client.TryReadRefAsync(refName, cacheTime, cancellationToken: cancellationToken);
+			BlobHandle? target = await client.TryReadRefAsync(refName, cacheTime, cancellationToken: cancellationToken);
 			if (target == null)
 			{
 				return new NotFoundResult();
 			}
 
-			string link = $"/api/v1/storage/{namespaceId}/nodes/{target.Target.GetLocator()}";
-			ReadRefResponse response = new ReadRefResponse { Target = target.Target.GetLocator(), Data = target.Data.ToArray(), Link = link };
+			string link = $"/api/v1/storage/{namespaceId}/nodes/{target.GetLocator()}";
+			ReadRefResponse response = new ReadRefResponse { Target = target.GetLocator(), Link = link };
 
 #pragma warning disable CS0618 // Type or member is obsolete
 			try
 			{
-				string locator = target.Target.ToString();
+				string locator = target.ToString();
 				int hashIdx = locator.LastIndexOf('#');
 
 				if (hashIdx != -1)
diff --git a/Engine/Source/Programs/Horde/Horde.Server/Storage/StorageService.cs b/Engine/Source/Programs/Horde/Horde.Server/Storage/StorageService.cs
--- a/Engine/Source/Programs/Horde/Horde.Server/Storage/StorageService.cs
+++ b/Engine/Source/Programs/Horde/Horde.Server/Storage/StorageService.cs
@@ -248,7 +248,7 @@ namespace Horde.Server.Storage
 			#region Refs
 
 			/// <inheritdoc/>
-			public async Task<RefValue?> TryReadRefAsync(RefName name, RefCacheTime cacheTime = default, CancellationToken cancellationToken = default)
+			public async Task<BlobHandle?> TryReadRefAsync(RefName name, RefCacheTime cacheTime = default, CancellationToken cancellationToken = default)
 			{
 				RefInfo? result = await _outer.TryReadRefAsync(NamespaceId, name, cacheTime, cancellationToken);
 				if (result == null)
@@ -263,11 +263,11 @@ namespace Horde.Server.Storage
 				}
 
 				BlobLocator locator = packet.Refs[0];
-				return new RefValue(CreateBlobHandle(locator), ReadOnlyMemory<byte>.Empty);
+				return CreateBlobHandle(locator);
 			}
 
 			/// <inheritdoc/>
-			public async Task WriteRefAsync(RefName name, BlobHandle target, ReadOnlyMemory<byte> data = default, RefOptions? options = null, CancellationToken cancellationToken = default)
+			public async Task WriteRefAsync(RefName name, BlobHandle target, RefOptions? options = null, CancellationToken cancellationToken = default)
 			{
 				BlobType blobType = Node.GetNodeType<RedirectNode>();
 				using BlobData blobData = new BlobData(blobType, ReadOnlyMemory<byte>.Empty, new List<BlobHandle> { target });
diff --git a/Engine/Source/Programs/Shared/EpicGames.Horde/Compute/AgentStorageClient.cs b/Engine/Source/Programs/Shared/EpicGames.Horde/Compute/AgentStorageClient.cs
--- a/Engine/Source/Programs/Shared/EpicGames.Horde/Compute/AgentStorageClient.cs
+++ b/Engine/Source/Programs/Shared/EpicGames.Horde/Compute/AgentStorageClient.cs
@@ -45,10 +45,10 @@ namespace EpicGames.Horde.Compute
 		public override Task<bool> DeleteRefAsync(RefName name, CancellationToken cancellationToken = default) => throw new NotSupportedException();
 
 		/// <inheritdoc/>
-		public override Task<RefValue?> TryReadRefAsync(RefName name, RefCacheTime cacheTime = default, CancellationToken cancellationToken = default) => throw new NotSupportedException();
+		public override Task<BlobHandle?> TryReadRefAsync(RefName name, RefCacheTime cacheTime = default, CancellationToken cancellationToken = default) => throw new NotSupportedException();
 
 		/// <inheritdoc/>
-		public override Task WriteRefAsync(RefName name, BlobHandle target, ReadOnlyMemory<byte> data = default, RefOptions? options = null, CancellationToken cancellationToken = default) => throw new NotSupportedException();
+		public override Task WriteRefAsync(RefName name, BlobHandle target, RefOptions? options = null, CancellationToken cancellationToken = default) => throw new NotSupportedException();
 
 		#endregion
 	}
diff --git a/Engine/Source/Programs/Shared/EpicGames.Horde/Storage/Clients/BundleStorageClient.cs b/Engine/Source/Programs/Shared/EpicGames.Horde/Storage/Clients/BundleStorageClient.cs
--- a/Engine/Source/Programs/Shared/EpicGames.Horde/Storage/Clients/BundleStorageClient.cs
+++ b/Engine/Source/Programs/Shared/EpicGames.Horde/Storage/Clients/BundleStorageClient.cs
@@ -215,10 +215,10 @@ namespace EpicGames.Horde.Storage.Clients
 		public abstract Task<bool> DeleteRefAsync(RefName name, CancellationToken cancellationToken = default);
 
 		/// <inheritdoc/>
-		public abstract Task<RefValue?> TryReadRefAsync(RefName name, RefCacheTime cacheTime = default, CancellationToken cancellationToken = default);
+		public abstract Task<BlobHandle?> TryReadRefAsync(RefName name, RefCacheTime cacheTime = default, CancellationToken cancellationToken = default);
 
 		/// <inheritdoc/>
-		public abstract Task WriteRefAsync(RefName name, BlobHandle target, ReadOnlyMemory<byte> data, RefOptions? options, CancellationToken cancellationToken);
+		public abstract Task WriteRefAsync(RefName name, BlobHandle target, RefOptions? options, CancellationToken cancellationToken);
 
 		#endregion
 
@@ -355,20 +355,19 @@ namespace EpicGames.Horde.Storage.Clients
 			=> _inner.DeleteRefAsync(name, cancellationToken);
 
 		/// <inheritdoc/>
-		public override async Task<RefValue?> TryReadRefAsync(RefName name, RefCacheTime cacheTime = default, CancellationToken cancellationToken = default)
+		public override async Task<BlobHandle?> TryReadRefAsync(RefName name, RefCacheTime cacheTime = default, CancellationToken cancellationToken = default)
 		{
-			RefValue? refValue = await _inner.TryReadRefAsync(name, cacheTime, cancellationToken);
-			if (refValue != null)
+			BlobHandle? target = await _inner.TryReadRefAsync(name, cacheTime, cancellationToken);
+			if (target != null)
 			{
-				BlobHandle target = CreateBlobHandle(refValue.Target.GetLocator());
-				refValue = new RefValue(target, refValue.Data);
+				target = CreateBlobHandle(target.GetLocator());
 			}
-			return refValue;
+			return target;
 		}
 
 		/// <inheritdoc/>
-		public override Task WriteRefAsync(RefName name, BlobHandle target, ReadOnlyMemory<byte> data, RefOptions? options, CancellationToken cancellationToken)
-			=> _inner.WriteRefAsync(name, target, data, options, cancellationToken);
+		public override Task WriteRefAsync(RefName name, BlobHandle target, RefOptions? options, CancellationToken cancellationToken)
+			=> _inner.WriteRefAsync(name, target, options, cancellationToken);
 
 		#endregion
 	}
diff --git a/Engine/Source/Programs/Shared/EpicGames.Horde/Storage/Clients/FileStorageClient.cs b/Engine/Source/Programs/Shared/EpicGames.Horde/Storage/Clients/FileStorageClient.cs
--- a/Engine/Source/Programs/Shared/EpicGames.Horde/Storage/Clients/FileStorageClient.cs
+++ b/Engine/Source/Programs/Shared/EpicGames.Horde/Storage/Clients/FileStorageClient.cs
@@ -81,7 +81,7 @@ namespace EpicGames.Horde.Storage.Clients
 		}
 
 		/// <inheritdoc/>
-		public override async Task<RefValue?> TryReadRefAsync(RefName name, RefCacheTime cacheTime = default, CancellationToken cancellationToken = default)
+		public override async Task<BlobHandle?> TryReadRefAsync(RefName name, RefCacheTime cacheTime = default, CancellationToken cancellationToken = default)
 		{
 			FileReference file = GetRefFile(name);
 			if (!FileReference.Exists(file))
@@ -93,17 +93,11 @@ namespace EpicGames.Horde.Storage.Clients
 			string[] lines = await FileReference.ReadAllLinesAsync(file, cancellationToken);
 
 			BlobHandle handle = CreateBlobHandle(new BlobLocator(lines[0].Trim()));
-			ReadOnlyMemory<byte> data = ReadOnlyMemory<byte>.Empty;
-			if (lines.Length >= 2)
-			{
-				data = Convert.FromBase64String(lines[1].Trim());
-			}
-
-			return new RefValue(handle, data);
+			return handle;
 		}
 
 		/// <inheritdoc/>
-		public override async Task WriteRefAsync(RefName name, BlobHandle target, ReadOnlyMemory<byte> data, RefOptions? options = null, CancellationToken cancellationToken = default)
+		public override async Task WriteRefAsync(RefName name, BlobHandle target, RefOptions? options = null, CancellationToken cancellationToken = default)
 		{
 			await target.FlushAsync(cancellationToken);
 			BlobLocator locator = target.GetLocator();
diff --git a/Engine/Source/Programs/Shared/EpicGames.Horde/Storage/Clients/HttpStorageClient.cs b/Engine/Source/Programs/Shared/EpicGames.Horde/Storage/Clients/HttpStorageClient.cs
--- a/Engine/Source/Programs/Shared/EpicGames.Horde/Storage/Clients/HttpStorageClient.cs
+++ b/Engine/Source/Programs/Shared/EpicGames.Horde/Storage/Clients/HttpStorageClient.cs
@@ -113,7 +113,7 @@ namespace EpicGames.Horde.Storage.Clients
 		}
 
 		/// <inheritdoc/>
-		public override async Task<RefValue?> TryReadRefAsync(RefName name, RefCacheTime cacheTime = default, CancellationToken cancellationToken = default)
+		public override async Task<BlobHandle?> TryReadRefAsync(RefName name, RefCacheTime cacheTime = default, CancellationToken cancellationToken = default)
 		{
 			using (HttpClient httpClient = _createClient())
 			{
@@ -142,8 +142,7 @@ namespace EpicGames.Horde.Storage.Clients
 							ReadRefResponse? data = await response.Content.ReadFromJsonAsync<ReadRefResponse>(cancellationToken: cancellationToken);
 							_logger.LogDebug("Read ref {RefName} -> {Blob}", name, data!.Target);
 
-							BlobHandle handle = CreateBlobHandle(data.Target);
-							return new RefValue(handle, data.Data);
+							return CreateBlobHandle(data.Target);
 						}
 					}
 				}
@@ -151,7 +150,7 @@ namespace EpicGames.Horde.Storage.Clients
 		}
 
 		/// <inheritdoc/>
-		public override async Task WriteRefAsync(RefName name, BlobHandle target, ReadOnlyMemory<byte> data, RefOptions? options = null, CancellationToken cancellationToken = default)
+		public override async Task WriteRefAsync(RefName name, BlobHandle target, RefOptions? options = null, CancellationToken cancellationToken = default)
 		{
 			await target.FlushAsync(cancellationToken);
 			BlobLocator locator = target.GetLocator();
@@ -161,7 +160,6 @@ namespace EpicGames.Horde.Storage.Clients
 			{
 				WriteRefRequest request = new WriteRefRequest();
 				request.Target = locator;
-				request.Data = data.ToArray();
 				request.Options = options;
 
 				using (HttpResponseMessage response = await httpClient.PutAsync($"{_basePath}/refs/{name}", request, cancellationToken))
diff --git a/Engine/Source/Programs/Shared/EpicGames.Horde/Storage/Clients/MemoryStorageClient.cs b/Engine/Source/Programs/Shared/EpicGames.Horde/Storage/Clients/MemoryStorageClient.cs
--- a/Engine/Source/Programs/Shared/EpicGames.Horde/Storage/Clients/MemoryStorageClient.cs
+++ b/Engine/Source/Programs/Shared/EpicGames.Horde/Storage/Clients/MemoryStorageClient.cs
@@ -10,11 +10,6 @@ using EpicGames.Horde.Storage.Backends;
 
 namespace EpicGames.Horde.Storage.Clients
 {
-	/// <summary>
-	/// Stores a RefValue in the <see cref="MemoryStorageClient"/>
-	/// </summary>
-	public record class RefData(BlobLocator Locator, ReadOnlyMemory<byte> Data);
-
 	/// <summary>
 	/// Implementation of <see cref="IStorageClient"/> which stores data in memory. Not intended for production use.
 	/// </summary>
@@ -30,7 +25,7 @@ namespace EpicGames.Horde.Storage.Clients
 		/// <summary>
 		/// Map of ref name to ref data
 		/// </summary>
-		readonly ConcurrentDictionary<RefName, RefData> _refs = new ConcurrentDictionary<RefName, RefData>();
+		readonly ConcurrentDictionary<RefName, BlobLocator> _refs = new ConcurrentDictionary<RefName, BlobLocator>();
 
 		/// <summary>
 		/// Content addressed data lookup
@@ -45,7 +40,7 @@ namespace EpicGames.Horde.Storage.Clients
 		/// <summary>
 		/// All refs stored by the client
 		/// </summary>
-		public IReadOnlyDictionary<RefName, RefData> Refs => _refs;
+		public IReadOnlyDictionary<RefName, BlobLocator> Refs => _refs;
 
 		/// <summary>
 		/// Constructor
@@ -103,26 +98,25 @@ namespace EpicGames.Horde.Storage.Clients
 		public override Task<bool> DeleteRefAsync(RefName name, CancellationToken cancellationToken) => Task.FromResult(_refs.TryRemove(name, out _));
 
 		/// <inheritdoc/>
-		public override Task<RefValue?> TryReadRefAsync(RefName name, RefCacheTime cacheTime = default, CancellationToken cancellationToken = default)
+		public override Task<BlobHandle?> TryReadRefAsync(RefName name, RefCacheTime cacheTime = default, CancellationToken cancellationToken = default)
 		{
-			RefData? refData;
+			BlobLocator refData;
 			if (_refs.TryGetValue(name, out refData))
 			{
-				return Task.FromResult<RefValue?>(new RefValue(CreateBlobHandle(refData.Locator), refData.Data)); 
+				return Task.FromResult<BlobHandle?>(CreateBlobHandle(refData)); 
 			}
 			else
 			{
-				return Task.FromResult<RefValue?>(null);
+				return Task.FromResult<BlobHandle?>(null);
 			}
 		}
 
 		/// <inheritdoc/>
-		public override async Task WriteRefAsync(RefName name, BlobHandle target, ReadOnlyMemory<byte> data = default, RefOptions? options = null, CancellationToken cancellationToken = default)
+		public override async Task WriteRefAsync(RefName name, BlobHandle target, RefOptions? options = null, CancellationToken cancellationToken = default)
 		{
 			await target.FlushAsync(cancellationToken);
 
-			BlobLocator locator = target.GetLocator();
-			_refs[name] = new RefData(locator, data);
+			_refs[name] = target.GetLocator();
 		}
 
 		#endregion
diff --git a/Engine/Source/Programs/Shared/EpicGames.Horde/Storage/Clients/SharedStorageClient.cs b/Engine/Source/Programs/Shared/EpicGames.Horde/Storage/Clients/SharedStorageClient.cs
--- a/Engine/Source/Programs/Shared/EpicGames.Horde/Storage/Clients/SharedStorageClient.cs
+++ b/Engine/Source/Programs/Shared/EpicGames.Horde/Storage/Clients/SharedStorageClient.cs
@@ -101,10 +101,10 @@ namespace EpicGames.Horde.Storage.Clients
 		#region Refs
 
 		/// <inheritdoc/>
-		public Task<RefValue?> TryReadRefAsync(RefName name, RefCacheTime cacheTime = default, CancellationToken cancellationToken = default) => _inner.TryReadRefAsync(name, cacheTime, cancellationToken);
+		public Task<BlobHandle?> TryReadRefAsync(RefName name, RefCacheTime cacheTime = default, CancellationToken cancellationToken = default) => _inner.TryReadRefAsync(name, cacheTime, cancellationToken);
 
 		/// <inheritdoc/>
-		public Task WriteRefAsync(RefName name, BlobHandle handle, ReadOnlyMemory<byte> data = default, RefOptions? options = null, CancellationToken cancellationToken = default) => _inner.WriteRefAsync(name, handle, data, options, cancellationToken);
+		public Task WriteRefAsync(RefName name, BlobHandle handle, RefOptions? options = null, CancellationToken cancellationToken = default) => _inner.WriteRefAsync(name, handle, options, cancellationToken);
 
 		/// <inheritdoc/>
 		public Task<bool> DeleteRefAsync(RefName name, CancellationToken cancellationToken = default) => _inner.DeleteRefAsync(name, cancellationToken);
diff --git a/Engine/Source/Programs/Shared/EpicGames.Horde/Storage/IStorageClient.cs b/Engine/Source/Programs/Shared/EpicGames.Horde/Storage/IStorageClient.cs
--- a/Engine/Source/Programs/Shared/EpicGames.Horde/Storage/IStorageClient.cs
+++ b/Engine/Source/Programs/Shared/EpicGames.Horde/Storage/IStorageClient.cs
@@ -170,18 +170,17 @@ namespace EpicGames.Horde.Storage
 		/// <param name="cacheTime">Minimum coherency for any cached value to be returned</param>
 		/// <param name="cancellationToken">Cancellation token for the operation</param>
 		/// <returns>Blob pointed to by the ref</returns>
-		Task<RefValue?> TryReadRefAsync(RefName name, RefCacheTime cacheTime = default, CancellationToken cancellationToken = default);
+		Task<BlobHandle?> TryReadRefAsync(RefName name, RefCacheTime cacheTime = default, CancellationToken cancellationToken = default);
 
 		/// <summary>
 		/// Writes a new ref to the store
 		/// </summary>
 		/// <param name="name">Ref to write</param>
 		/// <param name="handle">Handle to the target blob</param>
-		/// <param name="data">Inline data to store with the ref</param>
 		/// <param name="options">Options for the new ref</param>
 		/// <param name="cancellationToken">Cancellation token for the operation</param>
 		/// <returns>Unique identifier for the blob</returns>
-		Task WriteRefAsync(RefName name, BlobHandle handle, ReadOnlyMemory<byte> data = default, RefOptions? options = null, CancellationToken cancellationToken = default);
+		Task WriteRefAsync(RefName name, BlobHandle handle, RefOptions? options = null, CancellationToken cancellationToken = default);
 
 		/// <summary>
 		/// Reads data for a ref from the store
@@ -402,7 +401,7 @@ namespace EpicGames.Horde.Storage
 		/// <returns>True if the ref exists, false if it did not exist</returns>
 		public static async Task<bool> HasRefAsync(this IStorageClient store, RefName name, RefCacheTime cacheTime = default, CancellationToken cancellationToken = default)
 		{
-			RefValue? target = await store.TryReadRefAsync(name, cacheTime, cancellationToken);
+			BlobHandle? target = await store.TryReadRefAsync(name, cacheTime, cancellationToken);
 			return target != null;
 		}
 
@@ -416,8 +415,7 @@ namespace EpicGames.Horde.Storage
 		/// <returns>Blob pointed to by the ref</returns>
 		public static async Task<BlobHandle?> TryReadRefTargetAsync(this IStorageClient store, RefName name, RefCacheTime cacheTime = default, CancellationToken cancellationToken = default)
 		{
-			RefValue? refValue = await store.TryReadRefAsync(name, cacheTime, cancellationToken);
-			return refValue?.Target;
+			return await store.TryReadRefAsync(name, cacheTime, cancellationToken);
 		}
 
 		/// <summary>
