Horde: Add additional logging around HTTP artifact uploads, and log request method and URI when retries occur.

diff --git a/Engine/Source/Programs/Horde/Horde.Agent/Execution/JobExecutor.cs b/Engine/Source/Programs/Horde/Horde.Agent/Execution/JobExecutor.cs
--- a/Engine/Source/Programs/Horde/Horde.Agent/Execution/JobExecutor.cs
+++ b/Engine/Source/Programs/Horde/Horde.Agent/Execution/JobExecutor.cs
@@ -848,10 +848,20 @@ namespace Horde.Agent.Execution
 				using IStorageClient storage = StorageFactory.CreateClient(new NamespaceId(artifact.NamespaceId), artifact.Token);
 				await using IStorageWriter writer = storage.CreateWriter(new RefName(artifact.RefName));
 
-				DirectoryNode dir = new DirectoryNode();
-				await dir.AddFilesAsync(baseDir, files, new ChunkingOptions(), writer, new CopyStatsLogger(logger), cancellationToken);
+				try
+				{
+					DirectoryNode dir = new DirectoryNode();
+					await dir.AddFilesAsync(baseDir, files, new ChunkingOptions(), writer, new CopyStatsLogger(logger), cancellationToken);
+
+					await storage.WriteRefAsync(new RefName(artifact.RefName), dir, cancellationToken: cancellationToken);
+				}
+				catch (Exception ex)
+				{
+					Logger.LogInformation(ex, "Error uploading files for artifact {ArtifactId}", artifact.Id);
+					throw;
+				}
 
-				await storage.WriteRefAsync(new RefName(artifact.RefName), dir, cancellationToken: cancellationToken);
+				Logger.LogInformation("Uploaded artifact {ArtifactId}", artifact.Id);
 			}
 			catch (Exception ex)
 			{
diff --git a/Engine/Source/Programs/Shared/EpicGames.Core/KnownLogEvents.cs b/Engine/Source/Programs/Shared/EpicGames.Core/KnownLogEvents.cs
--- a/Engine/Source/Programs/Shared/EpicGames.Core/KnownLogEvents.cs
+++ b/Engine/Source/Programs/Shared/EpicGames.Core/KnownLogEvents.cs
@@ -250,6 +250,11 @@ namespace EpicGames.Core
 		/// </summary>
 		public static EventId Systemic_Horde_Compute { get; } = new EventId(732);
 
+		/// <summary>
+		/// HTTP error
+		/// </summary>
+		public static EventId Systemic_Horde_Http { get; } = new EventId(733);
+
 		/// <summary>
 		/// Harmless pdbutil error
 		/// </summary>
diff --git a/Engine/Source/Programs/Shared/EpicGames.Horde/HordeHttpClient.cs b/Engine/Source/Programs/Shared/EpicGames.Horde/HordeHttpClient.cs
--- a/Engine/Source/Programs/Shared/EpicGames.Horde/HordeHttpClient.cs
+++ b/Engine/Source/Programs/Shared/EpicGames.Horde/HordeHttpClient.cs
@@ -416,8 +416,8 @@ namespace EpicGames.Horde
 			}
 
 			IHttpClientBuilder builder = services.AddHttpClient<HordeHttpClient>(HordeHttpClient.HttpClientName, ConfigureClientFromEnvironment)
-				.AddPolicyHandler((serviceProvider, request) => CreateDefaultTimeoutRetryPolicy(serviceProvider.GetRequiredService<ILogger<HttpStorageClient>>()))
-				.AddPolicyHandler((serviceProvider, request) => CreateDefaultTransientErrorPolicy(serviceProvider.GetRequiredService<ILogger<HttpStorageClient>>()));
+				.AddPolicyHandler((serviceProvider, request) => CreateDefaultTimeoutRetryPolicy(request, serviceProvider.GetRequiredService<ILogger<HttpStorageClient>>()))
+				.AddPolicyHandler((serviceProvider, request) => CreateDefaultTransientErrorPolicy(request, serviceProvider.GetRequiredService<ILogger<HttpStorageClient>>()));
 
 			if (useAuthChallenge)
 			{
@@ -434,12 +434,12 @@ namespace EpicGames.Horde
 		/// <summary>
 		/// Create a default timeout retry policy
 		/// </summary>
-		public static IAsyncPolicy<HttpResponseMessage> CreateDefaultTimeoutRetryPolicy(ILogger logger)
+		public static IAsyncPolicy<HttpResponseMessage> CreateDefaultTimeoutRetryPolicy(HttpRequestMessage request, ILogger logger)
 		{
 			// Wait 30 seconds for operations to timeout
 			Task OnTimeoutAsync(Context context, TimeSpan timespan, Task timeoutTask)
 			{
-				logger.LogWarning("Http request timed out after {Time}s.", (int)timespan.TotalSeconds);
+				logger.LogWarning(KnownLogEvents.Systemic_Horde_Http, "{Method} {Url} timed out after {Time}s.", request.Method, request.RequestUri, (int)timespan.TotalSeconds);
 				return Task.CompletedTask;
 			}
 
@@ -448,7 +448,7 @@ namespace EpicGames.Horde
 			// Retry twice after a timeout
 			void OnRetry(Exception ex, TimeSpan timespan)
 			{
-				logger.LogWarning("Retrying http call after {Time}s.", timespan.TotalSeconds);
+				logger.LogWarning(KnownLogEvents.Systemic_Horde_Http, ex, "{Method} {Url} retrying after {Time}s.", request.Method, request.RequestUri, timespan.TotalSeconds);
 			}
 
 			TimeSpan[] retryTimes = new[] { TimeSpan.FromSeconds(5.0), TimeSpan.FromSeconds(10.0) };
@@ -459,11 +459,11 @@ namespace EpicGames.Horde
 		/// <summary>
 		/// Create a default timeout retry policy
 		/// </summary>
-		public static IAsyncPolicy<HttpResponseMessage> CreateDefaultTransientErrorPolicy(ILogger logger)
+		public static IAsyncPolicy<HttpResponseMessage> CreateDefaultTransientErrorPolicy(HttpRequestMessage request, ILogger logger)
 		{
 			Task OnTimeoutAsync(DelegateResult<HttpResponseMessage> outcome, TimeSpan timespan, int retryAttempt, Context context)
 			{
-				logger.LogWarning("Http request failed. Delaying for {DelayMs}ms (attempt #{RetryNum}).", timespan.TotalMilliseconds, retryAttempt);
+				logger.LogWarning(KnownLogEvents.Systemic_Horde_Http, "{Method} {Url} failed ({Result}). Delaying for {DelayMs}ms (attempt #{RetryNum}).", request.Method, request.RequestUri, outcome.Result?.StatusCode, timespan.TotalMilliseconds, retryAttempt);
 				return Task.CompletedTask;
 			}
 
