EpicGames.Perforce: Add triggers in the fixture server

Makes for easier testing prior to applying to any real P4 server.

diff --git a/Engine/Source/Programs/Shared/EpicGames.Perforce.Fixture/Dockerfile b/Engine/Source/Programs/Shared/EpicGames.Perforce.Fixture/Dockerfile
--- a/Engine/Source/Programs/Shared/EpicGames.Perforce.Fixture/Dockerfile
+++ b/Engine/Source/Programs/Shared/EpicGames.Perforce.Fixture/Dockerfile
@@ -20,4 +20,5 @@ RUN dos2unix setup.sh && bash setup.sh
 FROM docker.io/library/ubuntu:22.04
 COPY --from=builder /usr/local/bin/p4 /usr/local/bin/p4d /usr/local/bin/
 COPY --from=builder /app/data /app/data
+COPY --from=builder /trigger-log.sh /trigger-log.sh
 ENTRYPOINT ["p4d", "-r", "/app/data", "-p", "1666"]
diff --git a/Engine/Source/Programs/Shared/EpicGames.Perforce.Fixture/setup.sh b/Engine/Source/Programs/Shared/EpicGames.Perforce.Fixture/setup.sh
--- a/Engine/Source/Programs/Shared/EpicGames.Perforce.Fixture/setup.sh
+++ b/Engine/Source/Programs/Shared/EpicGames.Perforce.Fixture/setup.sh
@@ -26,6 +26,21 @@ FullName:	Test User
 AuthMethod:	perforce
 EOF
 
+# Set up triggers
+cat > /trigger-log.sh <<'EOF'
+#!/bin/bash
+echo "Trigger: $@" >> /tmp/trigger.log
+EOF
+chmod +x /trigger-log.sh
+
+cat > /tmp/triggers.txt <<EOF
+Triggers:
+        FormSaveAll form-save change "/trigger-log.sh form-save<SEP>%client%<SEP>%clienthost%<SEP>%clientip%<SEP>%clientprog%<SEP>%clientversion%<SEP>%peerhost%<SEP>%peerip%<SEP>%serverhost%<SEP>%serverid%<SEP>%serverip%<SEP>%servername%<SEP>%serverport%<SEP>%submitserverid%<SEP>%user%<SEP>%formname%<SEP>%formtype%<SEP>%change%<SEP>%changeroot%"
+        ChangeCommitAll change-commit //... "/trigger-log.sh change-commit<SEP>%client%<SEP>%clienthost%<SEP>%clientip%<SEP>%clientprog%<SEP>%clientversion%<SEP>%peerhost%<SEP>%peerip%<SEP>%serverhost%<SEP>%serverid%<SEP>%serverip%<SEP>%servername%<SEP>%serverport%<SEP>%submitserverid%<SEP>%user%<SEP>%formname%<SEP>%formtype%<SEP>%change%<SEP>%changeroot%"
+        ChangeCommitAll shelve-commit //... "/trigger-log.sh shelve-commit<SEP>%client%<SEP>%clienthost%<SEP>%clientip%<SEP>%clientprog%<SEP>%clientversion%<SEP>%peerhost%<SEP>%peerip%<SEP>%serverhost%<SEP>%serverid%<SEP>%serverip%<SEP>%servername%<SEP>%serverport%<SEP>%submitserverid%<SEP>%user%<SEP>%formname%<SEP>%formtype%<SEP>%change%<SEP>%changeroot%"
+EOF
+p4 triggers -i < /tmp/triggers.txt
+
 # Create a stream depot
 p4 depot -i << EOF
 Depot:          Foo
