Horde: Fix invalid message posted to slack when choosing between fixers for an issue.

diff --git a/Engine/Source/Programs/Horde/Horde.Server/Notifications/Sinks/SlackNotificationSink.cs b/Engine/Source/Programs/Horde/Horde.Server/Notifications/Sinks/SlackNotificationSink.cs
--- a/Engine/Source/Programs/Horde/Horde.Server/Notifications/Sinks/SlackNotificationSink.cs
+++ b/Engine/Source/Programs/Horde/Horde.Server/Notifications/Sinks/SlackNotificationSink.cs
@@ -3071,7 +3071,7 @@ namespace Horde.Server.Notifications.Sinks
 						options.Add(new SlackOption($"Current owner ({user.Name})", owner.Id.ToString()));
 
 						RadioButtonGroupElement ownership = new RadioButtonGroupElement("fixed_by_action", options);
-						ownership.InitialOptions.Add(options[0]);
+						ownership.InitialOption = options[0];
 						view.AddInput("Fixed By:", ownership).BlockId = "fixed_by";
 					}
 				}
diff --git a/Engine/Source/Programs/Shared/EpicGames.Slack/Elements/RadioButtonGroupElement.cs b/Engine/Source/Programs/Shared/EpicGames.Slack/Elements/RadioButtonGroupElement.cs
--- a/Engine/Source/Programs/Shared/EpicGames.Slack/Elements/RadioButtonGroupElement.cs
+++ b/Engine/Source/Programs/Shared/EpicGames.Slack/Elements/RadioButtonGroupElement.cs
@@ -26,8 +26,8 @@ namespace EpicGames.Slack.Elements
 		/// <summary>
 		/// An array of option objects that exactly matches one or more of the options within options. These options will be selected when the checkbox group initially loads.
 		/// </summary>
-		[JsonPropertyName("initial_options")]
-		public List<SlackOption> InitialOptions { get; } = new List<SlackOption>();
+		[JsonPropertyName("initial_option")]
+		public SlackOption? InitialOption { get; set; }
 
 		/// <summary>
 		/// A confirm object that defines an optional confirmation dialog that appears after clicking one of the checkboxes in this element.
@@ -39,7 +39,7 @@ namespace EpicGames.Slack.Elements
 		/// Indicates whether the element will be set to auto focus within the view object. Only one element can be set to true. Defaults to false.
 		/// </summary>
 		[JsonPropertyName("focus_on_load")]
-		public bool FocusOnLoad { get; set; }
+		public bool? FocusOnLoad { get; set; }
 
 		/// <summary>
 		/// Construct a new Button action element.
