Horde: Fix parsing of values from radio buttons in Slack views.

diff --git a/Engine/Source/Programs/Shared/EpicGames.Slack/SlackView.cs b/Engine/Source/Programs/Shared/EpicGames.Slack/SlackView.cs
--- a/Engine/Source/Programs/Shared/EpicGames.Slack/SlackView.cs
+++ b/Engine/Source/Programs/Shared/EpicGames.Slack/SlackView.cs
@@ -111,16 +111,30 @@ namespace EpicGames.Slack
 		/// <summary>
 		/// Gets a value from the state
 		/// </summary>
-		/// <param name="blockId"></param>
-		/// <param name="actionId"></param>
-		/// <param name="value"></param>
-		/// <returns></returns>
 		public bool TryGetValue(string blockId, string actionId, [NotNullWhen(true)] out string? value)
+		{
+			SlackViewValue? viewValue;
+			if (TryGetValue(blockId, actionId, out viewValue))
+			{
+				value = viewValue.Value ?? viewValue.SelectedOption?.Value;
+				return value != null;
+			}
+			else
+			{
+				value = null;
+				return false;
+			}
+		}
+
+		/// <summary>
+		/// Gets a value from the state
+		/// </summary>
+		public bool TryGetValue(string blockId, string actionId, [NotNullWhen(true)] out SlackViewValue? value)
 		{
 			Dictionary<string, SlackViewValue>? actions;
 			if (Values.TryGetValue(blockId, out actions) && actions.TryGetValue(actionId, out SlackViewValue? viewValue))
 			{
-				value = viewValue.Value;
+				value = viewValue;
 				return true;
 			}
 			else
@@ -143,9 +157,15 @@ namespace EpicGames.Slack
 		public string Type { get; set; } = String.Empty;
 
 		/// <summary>
-		/// Value for the widget
+		/// Value for the widget (used for text fields)
 		/// </summary>
 		[JsonPropertyName("value")]
-		public string Value { get; set; } = String.Empty;
+		public string? Value { get; set; }
+
+		/// <summary>
+		/// Current selected option (used for radio buttons)
+		/// </summary>
+		[JsonPropertyName("selected_option")]
+		public SlackOption? SelectedOption { get; set; }
 	}
 }
