Horde: Fix deserialization of SlackOption objects due to missing default constructor.

diff --git a/Engine/Source/Programs/Shared/EpicGames.Slack/SlackOption.cs b/Engine/Source/Programs/Shared/EpicGames.Slack/SlackOption.cs
--- a/Engine/Source/Programs/Shared/EpicGames.Slack/SlackOption.cs
+++ b/Engine/Source/Programs/Shared/EpicGames.Slack/SlackOption.cs
@@ -37,6 +37,14 @@ namespace EpicGames.Slack
 		[JsonPropertyName("url")]
 		public Uri? Url { get; set; }
 
+		/// <summary>
+		/// Constructor
+		/// </summary>
+		public SlackOption()
+			: this("", "")
+		{
+		}
+
 		/// <summary>
 		/// Constructor
 		/// </summary>
