Horde: Add a IReadOnlyMemory<T> type to EpicGames.Core for controlling ownership of a shared read-only block of memory.

diff --git a/Engine/Source/Programs/Horde/Horde.Server/Storage/Backends/AwsStorageBackend.cs b/Engine/Source/Programs/Horde/Horde.Server/Storage/Backends/AwsStorageBackend.cs
--- a/Engine/Source/Programs/Horde/Horde.Server/Storage/Backends/AwsStorageBackend.cs
+++ b/Engine/Source/Programs/Horde/Horde.Server/Storage/Backends/AwsStorageBackend.cs
@@ -104,13 +104,6 @@ namespace Horde.Server.Storage.Backends
 	/// </summary>
 	public sealed class AwsStorageBackend : IStorageBackend, IDisposable
 	{
-		class StorageObject : IStorageObject
-		{
-			public ReadOnlyMemory<byte> Data { get; }
-			public StorageObject(ReadOnlyMemory<byte> data) => Data = data;
-			public void Dispose() { }
-		}
-
 		/// <summary>
 		/// S3 Client
 		/// </summary>
@@ -348,10 +341,10 @@ namespace Horde.Server.Storage.Backends
 		}
 
 		/// <inheritdoc/>
-		public async Task<IStorageObject> ReadAsync(string path, int offset, int? length, CancellationToken cancellationToken)
+		public async Task<IReadOnlyMemoryOwner<byte>> ReadAsync(string path, int offset, int? length, CancellationToken cancellationToken)
 		{
 			using Stream stream = await OpenAsync(path, offset, length, cancellationToken);
-			return new StorageObject(await stream.ReadAllBytesAsync(cancellationToken));
+			return ReadOnlyMemoryOwner.Create(await stream.ReadAllBytesAsync(cancellationToken));
 		}
 
 		/// <inheritdoc/>
diff --git a/Engine/Source/Programs/Horde/Horde.Server/Storage/Backends/AzureStorageBackend.cs b/Engine/Source/Programs/Horde/Horde.Server/Storage/Backends/AzureStorageBackend.cs
--- a/Engine/Source/Programs/Horde/Horde.Server/Storage/Backends/AzureStorageBackend.cs
+++ b/Engine/Source/Programs/Horde/Horde.Server/Storage/Backends/AzureStorageBackend.cs
@@ -54,13 +54,6 @@ namespace Horde.Server.Storage.Backends
 	/// </summary>
 	public sealed class AzureStorageBackend : IStorageBackend
 	{
-		class StorageObject : IStorageObject
-		{
-			public ReadOnlyMemory<byte> Data { get; }
-			public StorageObject(ReadOnlyMemory<byte> data) => Data = data;
-			public void Dispose() { }
-		}
-
 		private readonly ILogger _logger;
 		private readonly Tracer _tracer;
 		private readonly BlobContainerClient _blobContainer;
@@ -113,10 +106,10 @@ namespace Horde.Server.Storage.Backends
 		}
 
 		/// <inheritdoc/>
-		public async Task<IStorageObject> ReadAsync(string path, int offset, int? length, CancellationToken cancellationToken = default)
+		public async Task<IReadOnlyMemoryOwner<byte>> ReadAsync(string path, int offset, int? length, CancellationToken cancellationToken = default)
 		{
 			using Stream stream = await OpenAsync(path, offset, length, cancellationToken);
-			return new StorageObject(await stream.ReadAllBytesAsync(cancellationToken));
+			return ReadOnlyMemoryOwner.Create(await stream.ReadAllBytesAsync(cancellationToken));
 		}
 
 		/// <inheritdoc/>
diff --git a/Engine/Source/Programs/Horde/Horde.Server/Storage/StorageBackendProvider.cs b/Engine/Source/Programs/Horde/Horde.Server/Storage/StorageBackendProvider.cs
--- a/Engine/Source/Programs/Horde/Horde.Server/Storage/StorageBackendProvider.cs
+++ b/Engine/Source/Programs/Horde/Horde.Server/Storage/StorageBackendProvider.cs
@@ -85,7 +85,7 @@ namespace Horde.Server.Storage
 			}
 
 			/// <inheritdoc/>
-			public Task<IStorageObject> ReadAsync(string path, int offset, int? length, CancellationToken cancellationToken = default)
+			public Task<IReadOnlyMemoryOwner<byte>> ReadAsync(string path, int offset, int? length, CancellationToken cancellationToken = default)
 			{
 				return _backend?.ReadAsync(path, offset, length, cancellationToken) ?? throw new InvalidOperationException("Backend has already been disposed");
 			}
diff --git a/Engine/Source/Programs/Horde/Horde.Server/Storage/StorageService.cs b/Engine/Source/Programs/Horde/Horde.Server/Storage/StorageService.cs
--- a/Engine/Source/Programs/Horde/Horde.Server/Storage/StorageService.cs
+++ b/Engine/Source/Programs/Horde/Horde.Server/Storage/StorageService.cs
@@ -58,10 +58,10 @@ namespace Horde.Server.Storage
 
 		sealed class LeafBlobData : BlobData
 		{
-			readonly IStorageObject _obj;
+			readonly IReadOnlyMemoryOwner<byte> _obj;
 
-			public LeafBlobData(IStorageObject obj)
-				: base(s_leafBlobType, obj.Data, Array.Empty<BlobHandle>())
+			public LeafBlobData(IReadOnlyMemoryOwner<byte> obj)
+				: base(s_leafBlobType, obj.Memory, Array.Empty<BlobHandle>())
 			{
 				_obj = obj;
 			}
@@ -111,7 +111,7 @@ namespace Horde.Server.Storage
 			/// <inheritdoc/>
 			public override async ValueTask<BlobData> ReadAsync(CancellationToken cancellationToken = default)
 			{
-				IStorageObject obj = await _backend.ReadAsync(_path, cancellationToken);
+				IReadOnlyMemoryOwner<byte> obj = await _backend.ReadAsync(_path, cancellationToken);
 				return new LeafBlobData(obj);
 			}
 
diff --git a/Engine/Source/Programs/Horde/Horde.Server/Storage/TypedStorageBackend.cs b/Engine/Source/Programs/Horde/Horde.Server/Storage/TypedStorageBackend.cs
--- a/Engine/Source/Programs/Horde/Horde.Server/Storage/TypedStorageBackend.cs
+++ b/Engine/Source/Programs/Horde/Horde.Server/Storage/TypedStorageBackend.cs
@@ -5,6 +5,7 @@ using System.Collections.Generic;
 using System.IO;
 using System.Threading;
 using System.Threading.Tasks;
+using EpicGames.Core;
 using EpicGames.Horde.Storage;
 
 namespace Horde.Server.Storage
@@ -46,7 +47,7 @@ namespace Horde.Server.Storage
 			public Task<Stream> OpenAsync(string path, int offset, int? length, CancellationToken cancellationToken) => _inner.OpenAsync(path, offset, length, cancellationToken);
 
 			/// <inheritdoc/>
-			public Task<IStorageObject> ReadAsync(string path, int offset, int? length, CancellationToken cancellationToken) => _inner.ReadAsync(path, offset, length, cancellationToken);
+			public Task<IReadOnlyMemoryOwner<byte>> ReadAsync(string path, int offset, int? length, CancellationToken cancellationToken) => _inner.ReadAsync(path, offset, length, cancellationToken);
 
 			/// <inheritdoc/>
 			public Task<string> WriteAsync(Stream stream, string? prefix, CancellationToken cancellationToken) => _inner.WriteAsync(stream, prefix, cancellationToken);
diff --git a/Engine/Source/Programs/Shared/EpicGames.Core/IReadOnlyMemoryOwner.cs b/Engine/Source/Programs/Shared/EpicGames.Core/IReadOnlyMemoryOwner.cs
new file mode 100644
--- /dev/null
+++ b/Engine/Source/Programs/Shared/EpicGames.Core/IReadOnlyMemoryOwner.cs
@@ -0,0 +1,111 @@
+// Copyright Epic Games, Inc. All Rights Reserved.
+
+using System;
+using System.IO;
+
+namespace EpicGames.Core
+{
+	/// <summary>
+	/// Handle to a shared block of memory with a concrete lifetime
+	/// </summary>
+	public interface IReadOnlyMemoryOwner<T> : IDisposable
+	{
+		/// <summary>
+		/// The underlying memory object
+		/// </summary>
+		public ReadOnlyMemory<T> Memory { get; }
+	}
+
+	/// <summary>
+	/// Extension methods for <see cref="IReadOnlyMemoryOwner{T}"/>
+	/// </summary>
+	public static class ReadOnlyMemoryOwner
+	{
+		class DefaultReadOnlyMemoryOwner<T> : IReadOnlyMemoryOwner<T>
+		{
+			public ReadOnlyMemory<T> Memory { get; }
+
+			public DefaultReadOnlyMemoryOwner(ReadOnlyMemory<T> memory) => Memory = memory;
+			public void Dispose() { }
+		}
+
+		/// <summary>
+		/// Wrap an array as a <see cref="IReadOnlyMemoryOwner{T}"/>
+		/// </summary>
+		public static IReadOnlyMemoryOwner<T> Create<T>(T[] memory)
+		{
+			return new DefaultReadOnlyMemoryOwner<T>(memory);
+		}
+
+		/// <summary>
+		/// Wrap a <see cref="ReadOnlyMemory{T}"/> as a <see cref="IReadOnlyMemoryOwner{T}"/>
+		/// </summary>
+		public static IReadOnlyMemoryOwner<T> Create<T>(ReadOnlyMemory<T> memory)
+		{
+			return new DefaultReadOnlyMemoryOwner<T>(memory);
+		}
+
+		// Wraps a ReadOnlyMemoryOwner in a stream
+		class ReadOnlyMemoryOwnerStream : ReadOnlyMemoryStream
+		{
+			readonly IReadOnlyMemoryOwner<byte> _owner;
+
+			public ReadOnlyMemoryOwnerStream(IReadOnlyMemoryOwner<byte> owner)
+				: base(owner.Memory) => _owner = owner;
+
+			protected override void Dispose(bool disposing)
+			{
+				base.Dispose(disposing);
+
+				if (disposing)
+				{
+					_owner.Dispose();
+				}
+			}
+		}
+
+		/// <summary>
+		/// Create a stream to wrap a storage object. The object will be disposed when the stream is closed.
+		/// </summary>
+		/// <param name="owner">Storage object to wrap</param>
+		/// <returns>Stream for the given storage object</returns>
+		public static Stream AsStream(this IReadOnlyMemoryOwner<byte> owner) => new ReadOnlyMemoryOwnerStream(owner);
+
+		// Creates a slice of a storage object
+		class ReadOnlyMemoryOwnerSlice<T> : IReadOnlyMemoryOwner<T>
+		{
+			readonly IReadOnlyMemoryOwner<T> _inner;
+
+			public ReadOnlyMemory<T> Memory { get; }
+
+			public ReadOnlyMemoryOwnerSlice(IReadOnlyMemoryOwner<T> inner, ReadOnlyMemory<T> memory)
+			{
+				_inner = inner;
+				Memory = memory;
+			}
+
+			public void Dispose() => _inner.Dispose();
+		}
+
+		/// <summary>
+		/// Creates a slice of an existing memory object.
+		/// </summary>
+		/// <param name="owner">Storage object to wrap</param>
+		/// <param name="offset">Offset of the slice</param>
+		/// <param name="length">Length to take for the slice</param>
+		/// <returns>Stream for the given storage object</returns>
+		public static IReadOnlyMemoryOwner<T> Slice<T>(this IReadOnlyMemoryOwner<T> owner, int offset, int? length)
+		{
+			int maxLength = owner.Memory.Length - offset;
+			int actualLength = length.HasValue ? Math.Min(length.Value, maxLength) : maxLength;
+			if (offset == 0 && actualLength == owner.Memory.Length)
+			{
+				return owner;
+			}
+			else
+			{
+				return new ReadOnlyMemoryOwnerSlice<T>(owner, owner.Memory.Slice(offset, actualLength));
+			}
+		}
+	}
+}
diff --git a/Engine/Source/Programs/Shared/EpicGames.Horde.Tests/BackendTests.cs b/Engine/Source/Programs/Shared/EpicGames.Horde.Tests/BackendTests.cs
--- a/Engine/Source/Programs/Shared/EpicGames.Horde.Tests/BackendTests.cs
+++ b/Engine/Source/Programs/Shared/EpicGames.Horde.Tests/BackendTests.cs
@@ -61,17 +61,17 @@ namespace EpicGames.Horde.Tests
 				using FileStorageBackend backend = new FileStorageBackend(tempDir.Location);
 				string path = await backend.WriteBytesAsync(Encoding.UTF8.GetBytes("hello world"));
 
-				using (IStorageObject handle = backend.Read(path, 0, null))
+				using (IReadOnlyMemoryOwner<byte> handle = backend.Read(path, 0, null))
 				{
-					Assert.AreEqual("hello world", Encoding.UTF8.GetString(handle.Data.Span));
+					Assert.AreEqual("hello world", Encoding.UTF8.GetString(handle.Memory.Span));
 				}
-				using (IStorageObject handle = backend.Read(path, 0, 5))
+				using (IReadOnlyMemoryOwner<byte> handle = backend.Read(path, 0, 5))
 				{
-					Assert.AreEqual("hello", Encoding.UTF8.GetString(handle.Data.Span));
+					Assert.AreEqual("hello", Encoding.UTF8.GetString(handle.Memory.Span));
 				}
-				using (IStorageObject handle = backend.Read(path, 4, 3))
+				using (IReadOnlyMemoryOwner<byte> handle = backend.Read(path, 4, 3))
 				{
-					Assert.AreEqual("o w", Encoding.UTF8.GetString(handle.Data.Span));
+					Assert.AreEqual("o w", Encoding.UTF8.GetString(handle.Memory.Span));
 				}
 			}
 		}
diff --git a/Engine/Source/Programs/Shared/EpicGames.Horde/Compute/AgentStorageClient.cs b/Engine/Source/Programs/Shared/EpicGames.Horde/Compute/AgentStorageClient.cs
--- a/Engine/Source/Programs/Shared/EpicGames.Horde/Compute/AgentStorageClient.cs
+++ b/Engine/Source/Programs/Shared/EpicGames.Horde/Compute/AgentStorageClient.cs
@@ -55,13 +55,6 @@ namespace EpicGames.Horde.Compute
 
 	class AgentStorageBackend : IStorageBackend
 	{
-		class StorageObject : IStorageObject
-		{
-			public ReadOnlyMemory<byte> Data { get; }
-			public StorageObject(ReadOnlyMemory<byte> data) => Data = data;
-			public void Dispose() { }
-		}
-
 		readonly AgentMessageChannel _channel;
 		readonly SemaphoreSlim _semaphore;
 
@@ -108,12 +101,12 @@ namespace EpicGames.Horde.Compute
 		}
 
 		/// <inheritdoc/>
-		public async Task<IStorageObject> ReadAsync(string path, int offset, int? length, CancellationToken cancellationToken = default)
+		public async Task<IReadOnlyMemoryOwner<byte>> ReadAsync(string path, int offset, int? length, CancellationToken cancellationToken = default)
 		{
 			using (Stream stream = await OpenAsync(path, offset, length, cancellationToken))
 			{
 				byte[] data = await stream.ReadAllBytesAsync(cancellationToken);
-				return new StorageObject(data);
+				return ReadOnlyMemoryOwner.Create(data);
 			}
 		}
 
diff --git a/Engine/Source/Programs/Shared/EpicGames.Horde/Storage/Backends/FileStorageBackend.cs b/Engine/Source/Programs/Shared/EpicGames.Horde/Storage/Backends/FileStorageBackend.cs
--- a/Engine/Source/Programs/Shared/EpicGames.Horde/Storage/Backends/FileStorageBackend.cs
+++ b/Engine/Source/Programs/Shared/EpicGames.Horde/Storage/Backends/FileStorageBackend.cs
@@ -122,12 +122,12 @@ namespace EpicGames.Horde.Storage.Backends
 		}
 
 		// Handle to a file in memory
-		class MappedFileHandle : IStorageObject
+		class MappedFileHandle : IReadOnlyMemoryOwner<byte>
 		{
 			MappedFile? _mappedFile;
 			ReadOnlyMemory<byte> _data;
 
-			public ReadOnlyMemory<byte> Data => _data;
+			public ReadOnlyMemory<byte> Memory => _data;
 
 			public MappedFileHandle(MappedFile? mappedFile, ReadOnlyMemory<byte> data)
 			{
@@ -243,8 +243,10 @@ namespace EpicGames.Horde.Storage.Backends
 		/// <inheritdoc/>
 		public Task<Stream> OpenAsync(string path, int offset, int? length, CancellationToken cancellationToken)
 		{
-			IStorageObject storageObject = Read(path, offset, length);
-			return Task.FromResult<Stream>(storageObject.CreateStream());
+#pragma warning disable CA2000 // Dispose objects before losing scope
+			IReadOnlyMemoryOwner<byte> storageObject = Read(path, offset, length);
+#pragma warning restore CA2000 // Dispose objects before losing scope
+			return Task.FromResult<Stream>(storageObject.AsStream());
 		}
 
 		/// <summary>
@@ -254,7 +256,7 @@ namespace EpicGames.Horde.Storage.Backends
 		/// <param name="offset">Offset of the data to retrieve</param>
 		/// <param name="length">Length of the data</param>
 		/// <returns>Handle to the data. Must be disposed by the caller.</returns>
-		public IStorageObject Read(string path, int offset, int? length)
+		public IReadOnlyMemoryOwner<byte> Read(string path, int offset, int? length)
 		{
 			lock (_lockObject)
 			{
@@ -264,7 +266,7 @@ namespace EpicGames.Horde.Storage.Backends
 		}
 
 		/// <inheritdoc/>
-		public Task<IStorageObject> ReadAsync(string path, int offset, int? length, CancellationToken cancellationToken = default)
+		public Task<IReadOnlyMemoryOwner<byte>> ReadAsync(string path, int offset, int? length, CancellationToken cancellationToken = default)
 		{
 			return Task.FromResult(Read(path, offset, length));
 		}
diff --git a/Engine/Source/Programs/Shared/EpicGames.Horde/Storage/Backends/HttpStorageBackend.cs b/Engine/Source/Programs/Shared/EpicGames.Horde/Storage/Backends/HttpStorageBackend.cs
--- a/Engine/Source/Programs/Shared/EpicGames.Horde/Storage/Backends/HttpStorageBackend.cs
+++ b/Engine/Source/Programs/Shared/EpicGames.Horde/Storage/Backends/HttpStorageBackend.cs
@@ -18,13 +18,6 @@ namespace EpicGames.Horde.Storage.Backends
 	/// </summary>
 	public sealed class HttpStorageBackend : IStorageBackend
 	{
-		class StorageObject : IStorageObject
-		{
-			public ReadOnlyMemory<byte> Data { get; }
-			public StorageObject(ReadOnlyMemory<byte> data) => Data = data;
-			public void Dispose() { }
-		}
-
 		class WriteBlobResponse
 		{
 			public string Blob { get; set; } = String.Empty;
@@ -95,12 +88,12 @@ namespace EpicGames.Horde.Storage.Backends
 		}
 
 		/// <inheritdoc/>
-		public async Task<IStorageObject> ReadAsync(string path, int offset, int? length, CancellationToken cancellationToken = default)
+		public async Task<IReadOnlyMemoryOwner<byte>> ReadAsync(string path, int offset, int? length, CancellationToken cancellationToken = default)
 		{
 			using (Stream stream = await OpenAsync(path, offset, length, cancellationToken))
 			{
 				byte[] data = await stream.ReadAllBytesAsync(cancellationToken);
-				return new StorageObject(data);
+				return ReadOnlyMemoryOwner.Create(data);
 			}
 		}
 
diff --git a/Engine/Source/Programs/Shared/EpicGames.Horde/Storage/Backends/MemoryStorageBackend.cs b/Engine/Source/Programs/Shared/EpicGames.Horde/Storage/Backends/MemoryStorageBackend.cs
--- a/Engine/Source/Programs/Shared/EpicGames.Horde/Storage/Backends/MemoryStorageBackend.cs
+++ b/Engine/Source/Programs/Shared/EpicGames.Horde/Storage/Backends/MemoryStorageBackend.cs
@@ -16,13 +16,6 @@ namespace EpicGames.Horde.Storage.Backends
 	/// </summary>
 	public sealed class MemoryStorageBackend : IStorageBackend
 	{
-		class StorageObject : IStorageObject
-		{
-			public ReadOnlyMemory<byte> Data { get; }
-			public StorageObject(ReadOnlyMemory<byte> data) => Data = data;
-			public void Dispose() { }
-		}
-
 		/// <summary>
 		/// Data storage
 		/// </summary>
@@ -48,9 +41,9 @@ namespace EpicGames.Horde.Storage.Backends
 		}
 
 		/// <inheritdoc/>
-		public Task<IStorageObject> ReadAsync(string path, int offset, int? length, CancellationToken cancellationToken)
+		public Task<IReadOnlyMemoryOwner<byte>> ReadAsync(string path, int offset, int? length, CancellationToken cancellationToken)
 		{
-			return Task.FromResult<IStorageObject>(new StorageObject(GetData(path, offset, length)));
+			return Task.FromResult(ReadOnlyMemoryOwner.Create(GetData(path, offset, length)));
 		}
 
 		ReadOnlyMemory<byte> GetData(string path, int offset, int? length)
diff --git a/Engine/Source/Programs/Shared/EpicGames.Horde/Storage/Backends/PrefixStorageBackend.cs b/Engine/Source/Programs/Shared/EpicGames.Horde/Storage/Backends/PrefixStorageBackend.cs
--- a/Engine/Source/Programs/Shared/EpicGames.Horde/Storage/Backends/PrefixStorageBackend.cs
+++ b/Engine/Source/Programs/Shared/EpicGames.Horde/Storage/Backends/PrefixStorageBackend.cs
@@ -6,6 +6,7 @@ using System.IO;
 using System.Runtime.CompilerServices;
 using System.Threading;
 using System.Threading.Tasks;
+using EpicGames.Core;
 
 namespace EpicGames.Horde.Storage.Backends
 {
@@ -59,7 +60,7 @@ namespace EpicGames.Horde.Storage.Backends
 		public Task<Stream> OpenAsync(string path, int offset, int? length, CancellationToken cancellationToken = default) => _inner.OpenAsync($"{_prefix}{path}", offset, length, cancellationToken);
 
 		/// <inheritdoc/>
-		public Task<IStorageObject> ReadAsync(string path, int offset, int? length, CancellationToken cancellationToken = default) => _inner.ReadAsync($"{_prefix}{path}", offset, length, cancellationToken);
+		public Task<IReadOnlyMemoryOwner<byte>> ReadAsync(string path, int offset, int? length, CancellationToken cancellationToken = default) => _inner.ReadAsync($"{_prefix}{path}", offset, length, cancellationToken);
 
 		/// <inheritdoc/>
 		public ValueTask<Uri?> TryGetReadRedirectAsync(string path, CancellationToken cancellationToken = default) => _inner.TryGetReadRedirectAsync($"{_prefix}{path}", cancellationToken);
diff --git a/Engine/Source/Programs/Shared/EpicGames.Horde/Storage/BlobHandle.cs b/Engine/Source/Programs/Shared/EpicGames.Horde/Storage/BlobHandle.cs
--- a/Engine/Source/Programs/Shared/EpicGames.Horde/Storage/BlobHandle.cs
+++ b/Engine/Source/Programs/Shared/EpicGames.Horde/Storage/BlobHandle.cs
@@ -15,6 +15,21 @@ namespace EpicGames.Horde.Storage
 	/// </summary>
 	public abstract class BlobHandle
 	{
+		class BlobDataFragment : IReadOnlyMemoryOwner<byte>
+		{
+			readonly BlobData _data;
+
+			public ReadOnlyMemory<byte> Memory { get; }
+
+			public BlobDataFragment(BlobData data, int offset, int length)
+			{
+				_data = data;
+				Memory = data.Data.Slice(offset, length);
+			}
+
+			public void Dispose() => _data.Dispose();
+		}
+
 		class BlobDataStream : ReadOnlyMemoryStream
 		{
 			readonly BlobData _blobData;
@@ -94,6 +109,19 @@ namespace EpicGames.Horde.Storage
 		/// <param name="cancellationToken">Cancellation token for the operation</param>
 		public abstract ValueTask<BlobData> ReadAsync(CancellationToken cancellationToken = default);
 
+		/// <summary>
+		/// Reads part of the blob, and returns a handle that can be used to access the data.
+		/// </summary>
+		/// <param name="offset">Offset of the data within the blob</param>
+		/// <param name="length">Length of the data to read</param>
+		/// <param name="cancellationToken">Cancellation token for the operation</param>
+		/// <returns></returns>
+		public virtual async ValueTask<IReadOnlyMemoryOwner<byte>> ReadPartialAsync(int offset, int length, CancellationToken cancellationToken = default)
+		{
+			BlobData data = await ReadAsync(cancellationToken);
+			return new BlobDataFragment(data, offset, length);
+		}
+
 		/// <summary>
 		/// Creates a reader for this node's data
 		/// </summary>
diff --git a/Engine/Source/Programs/Shared/EpicGames.Horde/Storage/IStorageBackend.cs b/Engine/Source/Programs/Shared/EpicGames.Horde/Storage/IStorageBackend.cs
--- a/Engine/Source/Programs/Shared/EpicGames.Horde/Storage/IStorageBackend.cs
+++ b/Engine/Source/Programs/Shared/EpicGames.Horde/Storage/IStorageBackend.cs
@@ -10,17 +10,6 @@ using EpicGames.Core;
 
 namespace EpicGames.Horde.Storage
 {
-	/// <summary>
-	/// Handle to a raw object read from storage.
-	/// </summary>
-	public interface IStorageObject : IDisposable
-	{
-		/// <summary>
-		/// Data for the item. Consumers of this interface must not cache this value beyond the lifetime of the object.
-		/// </summary>
-		public ReadOnlyMemory<byte> Data { get; }
-	}
-
 	/// <summary>
 	/// Interface for a low-level storage backend.
 	/// </summary>
@@ -49,7 +38,7 @@ namespace EpicGames.Horde.Storage
 		/// <param name="length">Length of the data</param>
 		/// <param name="cancellationToken">Cancellation token for the operation</param>
 		/// <returns>Handle to the data. Must be disposed by the caller.</returns>
-		Task<IStorageObject> ReadAsync(string path, int offset, int? length, CancellationToken cancellationToken = default);
+		Task<IReadOnlyMemoryOwner<byte>> ReadAsync(string path, int offset, int? length, CancellationToken cancellationToken = default);
 
 		/// <summary>
 		/// Writes a stream to the storage backend. If the stream throws an exception during read, the write will be aborted.
@@ -153,118 +142,6 @@ namespace EpicGames.Horde.Storage
 	/// </summary>
 	public static class StorageBackendExtensions
 	{
-		// Wraps an IStorageObject in a stream
-		class StorageObjectStream : Stream
-		{
-			readonly IStorageObject _storageObject;
-			int _offset;
-
-			public override bool CanRead => true;
-			public override bool CanSeek => true;
-			public override bool CanWrite => false;
-
-			public override long Length => _storageObject.Data.Length;
-
-			public override long Position
-			{
-				get => _offset;
-				set => _offset = (int)Math.Clamp(value, 0, _storageObject.Data.Length);
-			}
-
-			public StorageObjectStream(IStorageObject storageObject) => _storageObject = storageObject;
-
-			protected override void Dispose(bool disposing)
-			{
-				base.Dispose(disposing);
-
-				if (disposing)
-				{
-					_storageObject.Dispose();
-				}
-			}
-
-			public override void Flush() { }
-
-			public override int Read(Span<byte> buffer)
-			{
-				int length = (int)Math.Min(Length - _offset, buffer.Length);
-				_storageObject.Data.Span.Slice(_offset, length).CopyTo(buffer);
-				_offset += length;
-				return length;
-			}
-
-			public override int Read(byte[] buffer, int offset, int count) => Read(buffer.AsSpan(offset, count));
-
-			public override long Seek(long offset, SeekOrigin origin)
-			{
-				return _offset = origin switch
-				{
-					SeekOrigin.Begin => (int)Math.Clamp(offset, 0, Length),
-					SeekOrigin.Current => (int)Math.Clamp(_offset + offset, 0, Length),
-					SeekOrigin.End => (int)Math.Clamp(Length + offset, 0, Length),
-					_ => throw new InvalidOperationException()
-				};
-			}
-
-			public override void SetLength(long value) => throw new InvalidOperationException();
-			public override void Write(byte[] buffer, int offset, int count) => throw new InvalidOperationException();
-		}
-
-		/// <summary>
-		/// Create a stream to wrap a storage object. The object will be disposed when the stream is closed.
-		/// </summary>
-		/// <param name="storageObject">Storage object to wrap</param>
-		/// <returns>Stream for the given storage object</returns>
-		public static Stream CreateStream(this IStorageObject storageObject) => new StorageObjectStream(storageObject);
-
-		// Creates a slice of a storage object
-		class StorageObjectSlice : IStorageObject
-		{
-			readonly IStorageObject _inner;
-			readonly int _offset;
-			readonly int _length;
-
-			public ReadOnlyMemory<byte> Data => _inner.Data.Slice(_offset, _length);
-
-			public StorageObjectSlice(IStorageObject inner, int offset, int? length)
-			{
-				_inner = inner;
-				_offset = offset;
-
-				int maxLength = inner.Data.Length - offset;
-				if (length == null)
-				{
-					_length = maxLength;
-				}
-				else
-				{
-					_length = Math.Min(length.Value, maxLength);
-				}
-			}
-
-			public void Dispose() => _inner.Dispose();
-		}
-
-		/// <summary>
-		/// Create a stream to wrap a storage object. The object will be disposed when the stream is closed.
-		/// </summary>
-		/// <param name="storageObject">Storage object to wrap</param>
-		/// <param name="offset">Offset of the slice</param>
-		/// <param name="length">Length to take for the slice</param>
-		/// <returns>Stream for the given storage object</returns>
-		public static IStorageObject CreateSlice(this IStorageObject storageObject, int offset, int? length)
-		{
-			int actualLength = length ?? (storageObject.Data.Length - offset);
-			if (offset == 0 && actualLength == storageObject.Data.Length)
-			{
-				return storageObject;
-			}
-			else
-			{
-				return new StorageObjectSlice(storageObject, offset, length);
-			}
-		}
-
 		/// <summary>
 		/// Attempts to open a read stream for the given path.
 		/// </summary>
@@ -281,7 +158,7 @@ namespace EpicGames.Horde.Storage
 		/// <param name="path">Object name within the store</param>
 		/// <param name="cancellationToken">Cancellation token for the operation</param>
 		/// <returns>Stream for the object</returns>
-		public static Task<IStorageObject> ReadAsync(this IStorageBackend storageBackend, string path, CancellationToken cancellationToken = default) => storageBackend.ReadAsync(path, 0, null, cancellationToken);
+		public static Task<IReadOnlyMemoryOwner<byte>> ReadAsync(this IStorageBackend storageBackend, string path, CancellationToken cancellationToken = default) => storageBackend.ReadAsync(path, 0, null, cancellationToken);
 
 		/// <summary>
 		/// Reads an object as an array of bytes
@@ -292,8 +169,8 @@ namespace EpicGames.Horde.Storage
 		/// <returns>Contents of the object</returns>
 		public static async Task<byte[]> ReadBytesAsync(this IStorageBackend storageBackend, string path, CancellationToken cancellationToken = default)
 		{
-			using IStorageObject storageObject = await storageBackend.ReadAsync(path, cancellationToken);
-			return storageObject.Data.ToArray();
+			using IReadOnlyMemoryOwner<byte> storageObject = await storageBackend.ReadAsync(path, cancellationToken);
+			return storageObject.Memory.ToArray();
 		}
 
 		/// <summary>
diff --git a/Engine/Source/Programs/Shared/EpicGames.Horde/Storage/StorageBackendCache.cs b/Engine/Source/Programs/Shared/EpicGames.Horde/Storage/StorageBackendCache.cs
--- a/Engine/Source/Programs/Shared/EpicGames.Horde/Storage/StorageBackendCache.cs
+++ b/Engine/Source/Programs/Shared/EpicGames.Horde/Storage/StorageBackendCache.cs
@@ -73,14 +73,16 @@ namespace EpicGames.Horde.Storage
 
 			public async Task<Stream> OpenAsync(string path, int offset, int? length, CancellationToken cancellationToken = default)
 			{
-				IStorageObject storageObject = await ReadAsync(path, offset, length, cancellationToken);
-				return storageObject.CreateStream();
+				IReadOnlyMemoryOwner<byte> storageObject = await ReadAsync(path, offset, length, cancellationToken);
+				return storageObject.AsStream();
 			}
 
-			public async Task<IStorageObject> ReadAsync(string path, int offset, int? length, CancellationToken cancellationToken = default)
+			public async Task<IReadOnlyMemoryOwner<byte>> ReadAsync(string path, int offset, int? length, CancellationToken cancellationToken = default)
 			{
-				IStorageObject storageObject = await _cacheStorage.ReadAsync($"{_keyPrefix}{path}", ctx => _inner.OpenAsync(path, ctx), cancellationToken);
-				return storageObject.CreateSlice(offset, length);
+#pragma warning disable CA2000 // Dispose objects before losing scope
+				IReadOnlyMemoryOwner<byte> storageObject = await _cacheStorage.ReadAsync($"{_keyPrefix}{path}", ctx => _inner.OpenAsync(path, ctx), cancellationToken);
+#pragma warning restore CA2000 // Dispose objects before losing scope
+				return storageObject.Slice(offset, length);
 			}
 
 			public Task<string> WriteAsync(Stream stream, string? prefix = null, CancellationToken cancellationToken = default) => _inner.WriteAsync(stream, prefix, cancellationToken);
@@ -185,7 +187,7 @@ namespace EpicGames.Horde.Storage
 		}
 
 		/// <inheritdoc/>
-		public async Task<IStorageObject> ReadAsync(string key, Func<CancellationToken, Task<Stream>> createStreamAsync, CancellationToken cancellationToken = default)
+		public async Task<IReadOnlyMemoryOwner<byte>> ReadAsync(string key, Func<CancellationToken, Task<Stream>> createStreamAsync, CancellationToken cancellationToken = default)
 		{
 			for (; ; )
 			{
