Horde: Remove hard references to bundle storage client from compute system.

diff --git a/Engine/Source/Programs/Horde/Horde.Agent/Commands/Compute/ComputeWorkerCommand.cs b/Engine/Source/Programs/Horde/Horde.Agent/Commands/Compute/ComputeWorkerCommand.cs
--- a/Engine/Source/Programs/Horde/Horde.Agent/Commands/Compute/ComputeWorkerCommand.cs
+++ b/Engine/Source/Programs/Horde/Horde.Agent/Commands/Compute/ComputeWorkerCommand.cs
@@ -5,7 +5,6 @@ using System.Net.Sockets;
 using EpicGames.Core;
 using EpicGames.Horde.Compute;
 using EpicGames.Horde.Compute.Transports;
-using EpicGames.Horde.Storage;
 using Microsoft.Extensions.Logging;
 
 namespace Horde.Agent.Commands.Compute
@@ -16,16 +15,9 @@ namespace Horde.Agent.Commands.Compute
 	[Command("computeworker", "Runs the agent as a local compute host, accepting incoming connections on the loopback adapter with a given port")]
 	class ComputeWorkerCommand : Command
 	{
-		readonly BundleReaderCache _bundleReaderCache;
-
 		[CommandLine("-Port=")]
 		int Port { get; set; } = 2000;
 
-		public ComputeWorkerCommand(BundleReaderCache bundleReaderCache)
-		{
-			_bundleReaderCache = bundleReaderCache;
-		}
-
 		public override async Task<int> ExecuteAsync(ILogger logger)
 		{
 			logger.LogInformation("** WORKER **");
@@ -36,7 +28,7 @@ namespace Horde.Agent.Commands.Compute
 			await using (RemoteComputeSocket socket = new RemoteComputeSocket(new TcpTransport(tcpSocket), logger))
 			{
 				logger.LogInformation("Running worker...");
-				await RunWorkerAsync(socket, _bundleReaderCache, logger, CancellationToken.None);
+				await RunWorkerAsync(socket, logger, CancellationToken.None);
 				logger.LogInformation("Worker complete");
 				await socket.CloseAsync(CancellationToken.None);
 			}
@@ -45,11 +37,11 @@ namespace Horde.Agent.Commands.Compute
 			return 0;
 		}
 
-		public static async Task RunWorkerAsync(ComputeSocket socket, BundleReaderCache bundleReaderCache, ILogger logger, CancellationToken cancellationToken)
+		public static async Task RunWorkerAsync(ComputeSocket socket, ILogger logger, CancellationToken cancellationToken)
 		{
 			DirectoryReference sandboxDir = DirectoryReference.Combine(AgentApp.DataDir, "Sandbox");
 
-			AgentMessageHandler worker = new AgentMessageHandler(sandboxDir, bundleReaderCache, null, false, null, logger);
+			AgentMessageHandler worker = new AgentMessageHandler(sandboxDir, null, false, null, logger);
 			await worker.RunAsync(socket, cancellationToken);
 		}
 	}
diff --git a/Engine/Source/Programs/Horde/Horde.Agent/Leases/Handlers/ComputeHandler.cs b/Engine/Source/Programs/Horde/Horde.Agent/Leases/Handlers/ComputeHandler.cs
--- a/Engine/Source/Programs/Horde/Horde.Agent/Leases/Handlers/ComputeHandler.cs
+++ b/Engine/Source/Programs/Horde/Horde.Agent/Leases/Handlers/ComputeHandler.cs
@@ -6,7 +6,6 @@ using System.Net.Sockets;
 using EpicGames.Core;
 using EpicGames.Horde.Compute;
 using EpicGames.Horde.Compute.Transports;
-using EpicGames.Horde.Storage;
 using Horde.Agent.Services;
 using Horde.Agent.Utility;
 using HordeCommon.Rpc.Tasks;
@@ -76,7 +75,6 @@ namespace Horde.Agent.Leases.Handlers
 		static TimeSpan NoDataTimeout { get; } = TimeSpan.FromSeconds(20);
 
 		readonly ComputeListenerService _listenerService;
-		readonly BundleReaderCache _bundleReaderCache;
 		readonly IServerLoggerFactory _serverLoggerFactory;
 		readonly AgentSettings _settings;
 		readonly ILogger _logger;
@@ -84,10 +82,9 @@ namespace Horde.Agent.Leases.Handlers
 		/// <summary>
 		/// Constructor
 		/// </summary>
-		public ComputeHandler(ComputeListenerService listenerService, BundleReaderCache bundleReaderCache, IServerLoggerFactory serverLoggerFactory, IOptions<AgentSettings> settings, ILogger<ComputeHandler> logger)
+		public ComputeHandler(ComputeListenerService listenerService, IServerLoggerFactory serverLoggerFactory, IOptions<AgentSettings> settings, ILogger<ComputeHandler> logger)
 		{
 			_listenerService = listenerService;
-			_bundleReaderCache = bundleReaderCache;
 			_serverLoggerFactory = serverLoggerFactory;
 			_settings = settings.Value;
 			_logger = logger;
@@ -146,7 +143,7 @@ namespace Horde.Agent.Leases.Handlers
 								newEnvVars["UE_HORDE_SHARED_DIR"] = sharedDir.FullName;
 								newEnvVars["UE_HORDE_TERMINATION_SIGNAL_FILE"] = _settings.GetTerminationSignalFile().FullName;
 
-								AgentMessageHandler worker = new AgentMessageHandler(sandboxDir, _bundleReaderCache, newEnvVars, false, _settings.WineExecutablePath, serverLogger ?? _logger);
+								AgentMessageHandler worker = new AgentMessageHandler(sandboxDir, newEnvVars, false, _settings.WineExecutablePath, serverLogger ?? _logger);
 								await worker.RunAsync(socket, cts.Token);
 								await socket.CloseAsync(cts.Token);
 								return LeaseResult.Success;
diff --git a/Engine/Source/Programs/Horde/Horde/Commands/Compute/ComputeWorker.cs b/Engine/Source/Programs/Horde/Horde/Commands/Compute/ComputeWorker.cs
--- a/Engine/Source/Programs/Horde/Horde/Commands/Compute/ComputeWorker.cs
+++ b/Engine/Source/Programs/Horde/Horde/Commands/Compute/ComputeWorker.cs
@@ -5,7 +5,6 @@ using System.Net.Sockets;
 using EpicGames.Core;
 using EpicGames.Horde.Compute;
 using EpicGames.Horde.Compute.Transports;
-using EpicGames.Horde.Storage;
 using Microsoft.Extensions.Logging;
 
 namespace Horde.Commands.Compute
@@ -26,11 +25,10 @@ namespace Horde.Commands.Compute
 			using Socket tcpSocket = new Socket(SocketType.Stream, ProtocolType.IP);
 			await tcpSocket.ConnectAsync(IPAddress.Loopback, Port);
 
-			using BundleReaderCache bundleReaderCache = new BundleReaderCache();
 			await using (RemoteComputeSocket socket = new RemoteComputeSocket(new TcpTransport(tcpSocket), logger))
 			{
 				logger.LogInformation("Running worker...");
-				await RunWorkerAsync(socket, bundleReaderCache, logger, CancellationToken.None);
+				await RunWorkerAsync(socket, logger, CancellationToken.None);
 				logger.LogInformation("Worker complete");
 				await socket.CloseAsync(CancellationToken.None);
 			}
@@ -39,11 +37,11 @@ namespace Horde.Commands.Compute
 			return 0;
 		}
 
-		public static async Task RunWorkerAsync(ComputeSocket socket, BundleReaderCache bundleReaderCache, ILogger logger, CancellationToken cancellationToken)
+		public static async Task RunWorkerAsync(ComputeSocket socket, ILogger logger, CancellationToken cancellationToken)
 		{
 			DirectoryReference sandboxDir = DirectoryReference.Combine(DirectoryReference.GetSpecialFolder(Environment.SpecialFolder.LocalApplicationData)!, "Horde", "Sandbox");
 
-			AgentMessageHandler worker = new AgentMessageHandler(sandboxDir, bundleReaderCache, null, false, null, logger);
+			AgentMessageHandler worker = new AgentMessageHandler(sandboxDir, null, false, null, logger);
 			await worker.RunAsync(socket, cancellationToken);
 		}
 	}
diff --git a/Engine/Source/Programs/Shared/EpicGames.Horde.Tests/ComputeSocketTests.cs b/Engine/Source/Programs/Shared/EpicGames.Horde.Tests/ComputeSocketTests.cs
--- a/Engine/Source/Programs/Shared/EpicGames.Horde.Tests/ComputeSocketTests.cs
+++ b/Engine/Source/Programs/Shared/EpicGames.Horde.Tests/ComputeSocketTests.cs
@@ -140,7 +140,8 @@ namespace EpicGames.Horde.Tests
 						await channel2.CloseAsync();
 					}
 
-					using MemoryStorageClient storage = new MemoryStorageClient();
+					using MemoryStorageClient memoryStorage = new MemoryStorageClient();
+					using BundleStorageClientWrapper storage = new BundleStorageClientWrapper(memoryStorage, BundleReaderCache.None, NullLogger.Instance);
 					await using (BundleWriter treeWriter = storage.CreateWriter())
 					{
 						FileReference file = FileReference.Combine(tempDir, "subdir/hello.txt");
@@ -190,8 +191,7 @@ namespace EpicGames.Horde.Tests
 
 		static async Task RunAgentAsync(ComputeSocket socket, DirectoryReference tempDir, CancellationToken cancellationToken)
 		{
-			using BundleReaderCache bundleReaderCache = new BundleReaderCache();
-			AgentMessageHandler handler = new AgentMessageHandler(tempDir, bundleReaderCache, null, true, null, NullLogger.Instance);
+			AgentMessageHandler handler = new AgentMessageHandler(tempDir, null, true, null, NullLogger.Instance);
 			await handler.RunAsync(socket, cancellationToken);
 		}
 	}
diff --git a/Engine/Source/Programs/Shared/EpicGames.Horde/Compute/AgentMessage.cs b/Engine/Source/Programs/Shared/EpicGames.Horde/Compute/AgentMessage.cs
--- a/Engine/Source/Programs/Shared/EpicGames.Horde/Compute/AgentMessage.cs
+++ b/Engine/Source/Programs/Shared/EpicGames.Horde/Compute/AgentMessage.cs
@@ -9,8 +9,6 @@ using System.Threading;
 using System.Threading.Tasks;
 using EpicGames.Core;
 using EpicGames.Horde.Storage;
-using EpicGames.Horde.Storage.Bundles;
-using EpicGames.Horde.Storage.Clients;
 using EpicGames.Horde.Storage.Nodes;
 
 namespace EpicGames.Horde.Compute
@@ -216,7 +214,7 @@ namespace EpicGames.Horde.Compute
 	/// </summary>
 	/// <param name="Name">Path to extract the files to</param>
 	/// <param name="Locator">Locator for the tree to extract</param>
-	public record struct UploadFilesMessage(string Name, BundleNodeLocator Locator);
+	public record struct UploadFilesMessage(string Name, BlobLocator Locator);
 
 	/// <summary>
 	/// Deletes files or directories in the remote
@@ -348,7 +346,7 @@ namespace EpicGames.Horde.Compute
 
 		#region Process
 
-		static async Task<AgentMessage> RunStorageServerAsync(this AgentMessageChannel channel, BundleStorageClient storage, CancellationToken cancellationToken = default)
+		static async Task<AgentMessage> RunStorageServerAsync(this AgentMessageChannel channel, IStorageClient storage, CancellationToken cancellationToken = default)
 		{
 			for (; ; )
 			{
@@ -378,7 +376,7 @@ namespace EpicGames.Horde.Compute
 		/// <param name="locator">Location of a <see cref="DirectoryNode"/> describing contents of the sandbox</param>
 		/// <param name="storage">Storage for the sandbox data</param>
 		/// <param name="cancellationToken">Cancellation token for the operation</param>
-		public static async Task UploadFilesAsync(this AgentMessageChannel channel, string path, BlobLocator locator, BundleStorageClient storage, CancellationToken cancellationToken = default)
+		public static async Task UploadFilesAsync(this AgentMessageChannel channel, string path, BlobLocator locator, IStorageClient storage, CancellationToken cancellationToken = default)
 		{
 			using (IAgentMessageBuilder request = await channel.CreateMessageAsync(AgentMessageType.WriteFiles, cancellationToken))
 			{
@@ -400,7 +398,7 @@ namespace EpicGames.Horde.Compute
 		public static UploadFilesMessage ParseUploadFilesMessage(this AgentMessage message)
 		{
 			string name = message.ReadString();
-			BundleNodeLocator locator = message.ReadNodeLocator();
+			BlobLocator locator = new BlobLocator(message.ReadString());
 			return new UploadFilesMessage(name, locator);
 		}
 
@@ -630,7 +628,7 @@ namespace EpicGames.Horde.Compute
 		/// <param name="message">The read request</param>
 		/// <param name="storage">Storage client to retrieve the blob from</param>
 		/// <param name="cancellationToken">Cancellation token for the operation</param>
-		public static Task SendBlobDataAsync(this AgentMessageChannel channel, ReadBlobMessage message, BundleStorageClient storage, CancellationToken cancellationToken = default)
+		public static Task SendBlobDataAsync(this AgentMessageChannel channel, ReadBlobMessage message, IStorageClient storage, CancellationToken cancellationToken = default)
 		{
 			return SendBlobDataAsync(channel, message.Locator, message.Offset, message.Length, storage, cancellationToken);
 		}
@@ -644,7 +642,7 @@ namespace EpicGames.Horde.Compute
 		/// <param name="length">Length of the data</param>
 		/// <param name="storage">Storage client to retrieve the blob from</param>
 		/// <param name="cancellationToken">Cancellation token for the operation</param>
-		public static async Task SendBlobDataAsync(this AgentMessageChannel channel, BlobLocator locator, int offset, int length, BundleStorageClient storage, CancellationToken cancellationToken = default)
+		public static async Task SendBlobDataAsync(this AgentMessageChannel channel, BlobLocator locator, int offset, int length, IStorageClient storage, CancellationToken cancellationToken = default)
 		{
 			BlobHandle handle = storage.CreateBlobHandle(locator);
 			using Stream stream = await handle.OpenAsync(offset, length, cancellationToken);
diff --git a/Engine/Source/Programs/Shared/EpicGames.Horde/Compute/AgentMessageHandler.cs b/Engine/Source/Programs/Shared/EpicGames.Horde/Compute/AgentMessageHandler.cs
--- a/Engine/Source/Programs/Shared/EpicGames.Horde/Compute/AgentMessageHandler.cs
+++ b/Engine/Source/Programs/Shared/EpicGames.Horde/Compute/AgentMessageHandler.cs
@@ -13,7 +13,6 @@ using Microsoft.Extensions.Logging;
 using EpicGames.Horde.Storage;
 using System.IO;
 using System.Linq;
-using EpicGames.Horde.Storage.Bundles;
 
 namespace EpicGames.Horde.Compute
 {
@@ -23,7 +22,6 @@ namespace EpicGames.Horde.Compute
 	public class AgentMessageHandler
 	{
 		readonly DirectoryReference _sandboxDir;
-		readonly BundleReaderCache _cache;
 		readonly Dictionary<string, string?> _envVars;
 		readonly bool _executeInProcess;
 		readonly string? _wineExecutablePath;
@@ -33,15 +31,13 @@ namespace EpicGames.Horde.Compute
 		/// Constructor
 		/// </summary>
 		/// <param name="sandboxDir">Directory to use for reading/writing files</param>
-		/// <param name="cache">Cache for nodes read from storage</param>
 		/// <param name="envVars">Environment variables to set for any child processes</param>
 		/// <param name="executeInProcess">Whether to execute any external assemblies in the current process</param>
 		/// <param name="wineExecutablePath">Path to Wine executable. If null, execution under Wine is disabled</param>
 		/// <param name="logger">Logger for diagnostics</param>
-		public AgentMessageHandler(DirectoryReference sandboxDir, BundleReaderCache cache, Dictionary<string, string?>? envVars, bool executeInProcess, string? wineExecutablePath, ILogger logger)
+		public AgentMessageHandler(DirectoryReference sandboxDir, Dictionary<string, string?>? envVars, bool executeInProcess, string? wineExecutablePath, ILogger logger)
 		{
 			_sandboxDir = sandboxDir;
-			_cache = cache;
 			_envVars = envVars ?? new Dictionary<string, string?>();
 			_executeInProcess = executeInProcess;
 			_wineExecutablePath = wineExecutablePath;
@@ -136,12 +132,12 @@ namespace EpicGames.Horde.Compute
 			}
 		}
 
-		async Task WriteFilesAsync(AgentMessageChannel channel, string path, BundleNodeLocator locator, CancellationToken cancellationToken)
+		async Task WriteFilesAsync(AgentMessageChannel channel, string path, BlobLocator locator, CancellationToken cancellationToken)
 		{
 			using AgentStorageClient store = new AgentStorageClient(channel);
-			BundleReader reader = new BundleReader(store, _cache, _logger);
 
-			DirectoryNode directoryNode = await reader.ReadNodeAsync<DirectoryNode>(locator, cancellationToken);
+			BlobHandle handle = store.CreateBlobHandle(locator);
+			DirectoryNode directoryNode = await handle.ReadNodeAsync<DirectoryNode>(cancellationToken);
 
 			DirectoryReference outputDir = DirectoryReference.Combine(_sandboxDir, path);
 			if (!outputDir.IsUnderDirectory(_sandboxDir))
diff --git a/Engine/Source/Programs/Shared/EpicGames.Horde/Compute/Clients/LocalComputeClient.cs b/Engine/Source/Programs/Shared/EpicGames.Horde/Compute/Clients/LocalComputeClient.cs
--- a/Engine/Source/Programs/Shared/EpicGames.Horde/Compute/Clients/LocalComputeClient.cs
+++ b/Engine/Source/Programs/Shared/EpicGames.Horde/Compute/Clients/LocalComputeClient.cs
@@ -7,7 +7,6 @@ using System.Threading;
 using System.Threading.Tasks;
 using EpicGames.Core;
 using EpicGames.Horde.Compute.Transports;
-using EpicGames.Horde.Storage;
 using Microsoft.Extensions.Logging;
 
 namespace EpicGames.Horde.Compute.Clients
@@ -75,11 +74,9 @@ namespace EpicGames.Horde.Compute.Clients
 		{
 			using Socket tcpSocket = await listener.AcceptAsync(cancellationToken);
 
-			using BundleReaderCache bundleReaderCache = new BundleReaderCache();
-
 			await using (RemoteComputeSocket socket = new RemoteComputeSocket(new TcpTransport(tcpSocket), logger))
 			{
-				AgentMessageHandler worker = new AgentMessageHandler(sandboxDir, bundleReaderCache, null, executeInProcess, null, logger);
+				AgentMessageHandler worker = new AgentMessageHandler(sandboxDir, null, executeInProcess, null, logger);
 				await worker.RunAsync(socket, cancellationToken);
 				await socket.CloseAsync(cancellationToken);
 			}
