Horde: Change MemoryStorageClient to be a self-contained storage client implementation, and wrap it in a BundleStorageClientWrapper for existing tests.

diff --git a/Engine/Source/Programs/Horde/Horde/Commands/Compute/ComputeRun.cs b/Engine/Source/Programs/Horde/Horde/Commands/Compute/ComputeRun.cs
--- a/Engine/Source/Programs/Horde/Horde/Commands/Compute/ComputeRun.cs
+++ b/Engine/Source/Programs/Horde/Horde/Commands/Compute/ComputeRun.cs
@@ -7,7 +7,6 @@ using EpicGames.Horde.Common;
 using EpicGames.Horde.Compute;
 using EpicGames.Horde.Compute.Clients;
 using EpicGames.Horde.Storage;
-using EpicGames.Horde.Storage.Bundles;
 using EpicGames.Horde.Storage.Clients;
 using EpicGames.Horde.Storage.Nodes;
 using Microsoft.Extensions.Logging;
@@ -129,9 +128,9 @@ namespace Horde.Commands.Compute
 			return true;
 		}
 
-		static async Task<BlobLocator> CreateSandboxAsync(FileReference taskFile, BundleStorageClient storage, CancellationToken cancellationToken)
+		static async Task<BlobLocator> CreateSandboxAsync(FileReference taskFile, IStorageClient storage, CancellationToken cancellationToken)
 		{
-			await using BundleWriter writer = storage.CreateWriter();
+			await using IStorageWriter writer = storage.CreateWriter();
 
 			DirectoryNode sandbox = new DirectoryNode();
 			await sandbox.CopyFromDirectoryAsync(taskFile.Directory.ToDirectoryInfo(), new ChunkingOptions(), writer, null, cancellationToken);
diff --git a/Engine/Source/Programs/Horde/Samples/RemoteClient/RemoteClientApp.cs b/Engine/Source/Programs/Horde/Samples/RemoteClient/RemoteClientApp.cs
--- a/Engine/Source/Programs/Horde/Samples/RemoteClient/RemoteClientApp.cs
+++ b/Engine/Source/Programs/Horde/Samples/RemoteClient/RemoteClientApp.cs
@@ -111,7 +111,7 @@ namespace RemoteClient
 
 				// Upload the sandbox to the primary channel.
 				using MemoryStorageClient storage = new MemoryStorageClient();
-				await using (BundleWriter treeWriter = storage.CreateWriter())
+				await using (IStorageWriter treeWriter = storage.CreateWriter())
 				{
 					DirectoryNode sandbox = new DirectoryNode();
 					await sandbox.CopyFromDirectoryAsync(uploadDir.ToDirectoryInfo(), new ChunkingOptions(), treeWriter, null);
diff --git a/Engine/Source/Programs/Shared/EpicGames.Horde.Tests/BundleStoreTests.cs b/Engine/Source/Programs/Shared/EpicGames.Horde.Tests/BundleStoreTests.cs
--- a/Engine/Source/Programs/Shared/EpicGames.Horde.Tests/BundleStoreTests.cs
+++ b/Engine/Source/Programs/Shared/EpicGames.Horde.Tests/BundleStoreTests.cs
@@ -59,7 +59,8 @@ namespace EpicGames.Horde.Tests
 
 		static async Task<Bundle> CreateBundleNormalAsync()
 		{
-			using MemoryStorageClient store = new MemoryStorageClient();
+			using MemoryStorageClient memoryStore = new MemoryStorageClient();
+			using BundleStorageClientWrapper store = new BundleStorageClientWrapper(memoryStore, BundleReaderCache.None, NullLogger.Instance);
 			await using BundleWriter writer = store.CreateWriter(options: new BundleOptions { CompressionFormat = BundleCompressionFormat.None });
 
 			TextNode node = new TextNode("Hello world");
@@ -96,7 +97,9 @@ namespace EpicGames.Horde.Tests
 			using IMemoryCache cache = new MemoryCache(new MemoryCacheOptions());
 
 			using MemoryStorageClient blobStore = new MemoryStorageClient();
-			await TestTreeAsync(blobStore, new BundleOptions { MaxBlobSize = 1024 * 1024 });
+			using BundleStorageClientWrapper bundleStore = new BundleStorageClientWrapper(blobStore, BundleReaderCache.None, NullLogger.Instance);
+
+			await TestTreeAsync(bundleStore, new BundleOptions { MaxBlobSize = 1024 * 1024 });
 
 			Assert.AreEqual(1, blobStore.Blobs.Count);
 			Assert.AreEqual(1, blobStore.Refs.Count);
@@ -108,7 +111,9 @@ namespace EpicGames.Horde.Tests
 			using IMemoryCache cache = new MemoryCache(new MemoryCacheOptions());
 
 			using MemoryStorageClient blobStore = new MemoryStorageClient();
-			await TestTreeAsync(blobStore, new BundleOptions { MaxBlobSize = 1 });
+			using BundleStorageClientWrapper bundleStore = new BundleStorageClientWrapper(blobStore, BundleReaderCache.None, NullLogger.Instance);
+
+			await TestTreeAsync(bundleStore, new BundleOptions { MaxBlobSize = 1 });
 
 			Assert.AreEqual(5, blobStore.Blobs.Count);
 			Assert.AreEqual(1, blobStore.Refs.Count);
@@ -139,7 +144,7 @@ namespace EpicGames.Horde.Tests
 			}
 		}
 
-		static async Task TestTreeAsync(MemoryStorageClient store, BundleOptions options)
+		static async Task TestTreeAsync(BundleStorageClientBase store, BundleOptions options)
 		{
 			// Generate a tree
 			{
@@ -309,7 +314,8 @@ namespace EpicGames.Horde.Tests
 		[TestMethod]
 		public async Task StreamTestAsync()
 		{
-			using MemoryStorageClient store = new MemoryStorageClient();
+			using MemoryStorageClient memoryStore = new MemoryStorageClient();
+			using BundleStorageClientWrapper store = new BundleStorageClientWrapper(memoryStore, BundleReaderCache.None, NullLogger.Instance);
 
 			const int Length = 4096;
 
@@ -349,7 +355,8 @@ namespace EpicGames.Horde.Tests
 		public async Task LargeFileTestAsync()
 		{
 			using IMemoryCache cache = new MemoryCache(new MemoryCacheOptions());
-			using MemoryStorageClient store = new MemoryStorageClient();
+			using MemoryStorageClient memoryStore = new MemoryStorageClient();
+			using BundleStorageClientWrapper store = new BundleStorageClientWrapper(memoryStore, BundleReaderCache.None, NullLogger.Instance);
 
 			const int Length = 1024;
 			const int Copies = 4096;
@@ -391,7 +398,7 @@ namespace EpicGames.Horde.Tests
 
 				HashedNodeRef<ChunkedDataNode> file = root.GetFileEntry("test");
 
-				long uniqueSize = store.Blobs.Values.Select(x => Bundle.FromMemory(x)).SelectMany(x => x.Header.Packets).Sum(x => x.DecodedLength);
+				long uniqueSize = memoryStore.Blobs.Values.Select(x => Bundle.FromMemory(x.Data)).SelectMany(x => x.Header.Packets).Sum(x => x.DecodedLength);
 				Assert.IsTrue(uniqueSize < data.Length / 3); // random fraction meaning "lots of dedupe happened"
 			}
 		}
diff --git a/Engine/Source/Programs/Shared/EpicGames.Horde.Tests/BundleTests.cs b/Engine/Source/Programs/Shared/EpicGames.Horde.Tests/BundleTests.cs
--- a/Engine/Source/Programs/Shared/EpicGames.Horde.Tests/BundleTests.cs
+++ b/Engine/Source/Programs/Shared/EpicGames.Horde.Tests/BundleTests.cs
@@ -6,6 +6,7 @@ using EpicGames.Horde.Storage.Bundles;
 using EpicGames.Horde.Storage.Clients;
 using EpicGames.Horde.Storage.Nodes;
 using Microsoft.Extensions.Caching.Memory;
+using Microsoft.Extensions.Logging.Abstractions;
 using Microsoft.VisualStudio.TestTools.UnitTesting;
 using System;
 using System.Collections.Generic;
@@ -20,17 +21,20 @@ namespace EpicGames.Horde.Tests
 	public sealed class BundleTests : IDisposable
 	{
 		readonly IMemoryCache _cache;
-		readonly MemoryStorageClient _storage;
+		readonly MemoryStorageClient _memoryStore;
+		readonly BundleStorageClientBase _storage;
 
 		public BundleTests()
 		{
 			_cache = new MemoryCache(new MemoryCacheOptions());
-			_storage = new MemoryStorageClient();
+			_memoryStore = new MemoryStorageClient();
+			_storage = new BundleStorageClientWrapper(_memoryStore, BundleReaderCache.None, NullLogger.Instance);
 		}
 
 		public void Dispose()
 		{
 			_storage.Dispose();
+			_memoryStore.Dispose();
 			_cache.Dispose();
 		}
 
@@ -187,7 +191,7 @@ namespace EpicGames.Horde.Tests
 		[TestMethod]
 		public async Task BasicTestDirectoryAsync()
 		{
-			MemoryStorageClient store = _storage;
+			IStorageClient store = _storage;
 
 			HashedNodeRef<DirectoryNode> rootRef;
 			await using (IStorageWriter writer = store.CreateWriter())
@@ -210,8 +214,8 @@ namespace EpicGames.Horde.Tests
 			await store.WriteRefTargetAsync(refName, rootRef);
 
 			// Should be stored inline
-			Assert.AreEqual(1, store.Refs.Count);
-			Assert.AreEqual(1, store.Blobs.Count);
+			Assert.AreEqual(1, _memoryStore.Refs.Count);
+			Assert.AreEqual(1, _memoryStore.Blobs.Count);
 
 			// Check the ref
 			BlobHandle refTarget =  await store.ReadRefTargetAsync(refName);
@@ -242,8 +246,8 @@ namespace EpicGames.Horde.Tests
 		[TestMethod]
 		public async Task DedupTestsAsync()
 		{
-			Assert.AreEqual(0, _storage.Refs.Count);
-			Assert.AreEqual(0, _storage.Blobs.Count);
+			Assert.AreEqual(0, _memoryStore.Refs.Count);
+			Assert.AreEqual(0, _memoryStore.Blobs.Count);
 
 			BundleOptions options = new BundleOptions();
 			options.MaxBlobSize = 1;
@@ -263,8 +267,8 @@ namespace EpicGames.Horde.Tests
 				await _storage.WriteRefTargetAsync(refName, rootRef);
 			}
 
-			Assert.AreEqual(1, _storage.Refs.Count);
-			Assert.AreEqual(2, _storage.Blobs.Count);
+			Assert.AreEqual(1, _memoryStore.Refs.Count);
+			Assert.AreEqual(2, _memoryStore.Blobs.Count);
 		}
 
 		[TestMethod]
@@ -288,8 +292,8 @@ namespace EpicGames.Horde.Tests
 					await _storage.WriteRefTargetAsync(refName, rootRef);
 				}
 
-				Assert.AreEqual(1, _storage.Refs.Count);
-				Assert.AreEqual(5, _storage.Blobs.Count);
+				Assert.AreEqual(1, _memoryStore.Refs.Count);
+				Assert.AreEqual(5, _memoryStore.Blobs.Count);
 			}
 
 			{
diff --git a/Engine/Source/Programs/Shared/EpicGames.Horde/Storage/BlobHandle.cs b/Engine/Source/Programs/Shared/EpicGames.Horde/Storage/BlobHandle.cs
--- a/Engine/Source/Programs/Shared/EpicGames.Horde/Storage/BlobHandle.cs
+++ b/Engine/Source/Programs/Shared/EpicGames.Horde/Storage/BlobHandle.cs
@@ -99,7 +99,8 @@ namespace EpicGames.Horde.Storage
 		public virtual async Task<Stream> OpenAsync(int offset = 0, int? length = null, CancellationToken cancellationToken = default)
 		{
 			BlobData blobData = await ReadAsync(cancellationToken);
-			ReadOnlyMemory<byte> memory = blobData.Data.Slice(offset, length ?? blobData.Data.Length - offset);
+			int maxLength = blobData.Data.Length - offset;
+			ReadOnlyMemory<byte> memory = blobData.Data.Slice(offset, length.HasValue ? Math.Min(length.Value, maxLength) : maxLength);
 			return new BlobDataStream(blobData, memory);
 		}
 
diff --git a/Engine/Source/Programs/Shared/EpicGames.Horde/Storage/Bundles/BundleReader.cs b/Engine/Source/Programs/Shared/EpicGames.Horde/Storage/Bundles/BundleReader.cs
--- a/Engine/Source/Programs/Shared/EpicGames.Horde/Storage/Bundles/BundleReader.cs
+++ b/Engine/Source/Programs/Shared/EpicGames.Horde/Storage/Bundles/BundleReader.cs
@@ -168,10 +168,10 @@ namespace EpicGames.Horde.Storage.Bundles
 		public async Task<BundleHeader> ReadHeaderAsync(BlobLocator locator, CancellationToken cancellationToken)
 		{
 			// Check for a cached value first
-			if (_cache.TryGetCachedHeader(locator, out BundleInfo? cachedBundleInfo))
-			{
-				return cachedBundleInfo.Header;
-			}
+//			if (_cache.TryGetCachedHeader(locator, out BundleInfo? cachedBundleInfo))
+//			{
+//				return cachedBundleInfo.Header;
+//			}
 
 			// Find a registered bundle info with the given locator
 			QueuedBundle? bundle = null;
diff --git a/Engine/Source/Programs/Shared/EpicGames.Horde/Storage/Bundles/BundleReaderCache.cs b/Engine/Source/Programs/Shared/EpicGames.Horde/Storage/Bundles/BundleReaderCache.cs
--- a/Engine/Source/Programs/Shared/EpicGames.Horde/Storage/Bundles/BundleReaderCache.cs
+++ b/Engine/Source/Programs/Shared/EpicGames.Horde/Storage/Bundles/BundleReaderCache.cs
@@ -34,7 +34,7 @@ namespace EpicGames.Horde.Storage
 		/// <summary>
 		/// Instance of an empty cache
 		/// </summary>
-		public static BundleReaderCache None { get; } = new BundleReaderCache();
+		public static BundleReaderCache None { get; } = new BundleReaderCache(new BundleReaderCacheOptions { HeaderCacheSize = 0, PacketCacheSize = 0 });
 
 		/// <summary>
 		/// Size of the configured header cache
diff --git a/Engine/Source/Programs/Shared/EpicGames.Horde/Storage/Clients/MemoryStorageClient.cs b/Engine/Source/Programs/Shared/EpicGames.Horde/Storage/Clients/MemoryStorageClient.cs
--- a/Engine/Source/Programs/Shared/EpicGames.Horde/Storage/Clients/MemoryStorageClient.cs
+++ b/Engine/Source/Programs/Shared/EpicGames.Horde/Storage/Clients/MemoryStorageClient.cs
@@ -3,83 +3,173 @@
 using System;
 using System.Collections.Concurrent;
 using System.Collections.Generic;
+using System.Diagnostics.CodeAnalysis;
+using System.IO;
 using System.Threading;
 using System.Threading.Tasks;
-using Microsoft.Extensions.Logging.Abstractions;
-using EpicGames.Horde.Storage.Backends;
+using EpicGames.Core;
 
 namespace EpicGames.Horde.Storage.Clients
 {
 	/// <summary>
 	/// Implementation of <see cref="IStorageClient"/> which stores data in memory. Not intended for production use.
 	/// </summary>
-	public class MemoryStorageClient : BundleStorageClient
+	public sealed class MemoryStorageClient : IStorageClient
 	{
-		record class ExportEntry(BlobLocator Locator, int Rank, ReadOnlyMemory<byte> Data,  ExportEntry? Next);
+		class Handle : BlobHandle
+		{
+			readonly MemoryStorageClient _owner;
+			readonly BlobLocator _locator;
 
-		/// <summary>
-		/// Backend instance
-		/// </summary>
-		readonly MemoryStorageBackend _backend;
+			public Handle(MemoryStorageClient owner, BlobLocator locator)
+			{
+				_owner = owner;
+				_locator = locator;
+			}
 
-		/// <summary>
-		/// Map of ref name to ref data
-		/// </summary>
+			/// <inheritdoc/>
+			public override ValueTask<BlobData> ReadAsync(CancellationToken cancellationToken = default) => new ValueTask<BlobData>(_owner._blobs[_locator]);
+
+			/// <inheritdoc/>
+			public override bool TryGetLocator([NotNullWhen(true)] out BlobLocator locator)
+			{
+				locator = _locator;
+				return true;
+			}
+		}
+
+		record class AliasListNode(BlobLocator Locator, int Rank, ReadOnlyMemory<byte> Data, AliasListNode? Next);
+
+		readonly ConcurrentDictionary<BlobLocator, BlobData> _blobs = new ConcurrentDictionary<BlobLocator, BlobData>();
 		readonly ConcurrentDictionary<RefName, BlobLocator> _refs = new ConcurrentDictionary<RefName, BlobLocator>();
+		readonly ConcurrentDictionary<string, AliasListNode?> _aliases = new ConcurrentDictionary<string, AliasListNode?>(StringComparer.Ordinal);
 
-		/// <summary>
-		/// Content addressed data lookup
-		/// </summary>
-		readonly ConcurrentDictionary<string, ExportEntry> _aliases = new ConcurrentDictionary<string, ExportEntry>(StringComparer.Ordinal);
+		int _nextId;
 
 		/// <summary>
 		/// All data stored by the client
 		/// </summary>
-		public IReadOnlyDictionary<string, byte[]> Blobs => _backend.Blobs;
+		public IReadOnlyDictionary<BlobLocator, BlobData> Blobs => _blobs;
 
 		/// <summary>
-		/// All refs stored by the client
+		/// Accessor for all refs stored by the client
 		/// </summary>
 		public IReadOnlyDictionary<RefName, BlobLocator> Refs => _refs;
 
-		/// <summary>
-		/// Constructor
-		/// </summary>
-		public MemoryStorageClient() 
-			: this(new MemoryStorageBackend())
+		/// <inheritdoc/>
+		public bool SupportsRedirects => false;
+
+		/// <inheritdoc/>
+		public void Dispose()
 		{
 		}
 
-		/// <summary>
-		/// Constructor
-		/// </summary>
-		private MemoryStorageClient(MemoryStorageBackend backend)
-			: base(backend, BundleReaderCache.None, NullLogger.Instance)
+		#region Blobs
+
+		/// <inheritdoc/>
+		public BlobHandle CreateBlobHandle(BlobLocator locator)
+		{
+			if (locator.TryUnwrapFull(out BlobLocator outer, out Utf8String fragment))
+			{
+				return new BlobFragmentHandle(new Handle(this, outer), fragment);
+			}
+			else
+			{
+				return new Handle(this, locator);
+			}
+		}
+
+		/// <inheritdoc/>
+		public IStorageWriter CreateWriter(string? basePath = null) => new DefaultStorageWriter(this, basePath);
+
+		/// <inheritdoc/>
+		public async ValueTask<BlobHandle> WriteBlobAsync(BlobType type, Stream stream, IReadOnlyList<BlobHandle> references, string? basePath = null, CancellationToken cancellationToken = default)
 		{
-			_backend = backend;
+			BlobLocator locator = new BlobLocator($"mem-{Interlocked.Increment(ref _nextId)}");
+			_blobs[locator] = new BlobData(type, await stream.ReadAllBytesAsync(cancellationToken), references);
+			return new Handle(this, locator);
 		}
 
+		/// <inheritdoc/>
+		public ValueTask<Uri?> TryGetReadRedirectAsync(BlobLocator locator, CancellationToken cancellationToken = default)
+			=> default;
+
+		/// <inheritdoc/>
+		public ValueTask<(BlobLocator, Uri)?> TryGetWriteRedirectAsync(string? prefix = null, CancellationToken cancellationToken = default)
+			=> default;
+
+		#endregion
+
 		#region Aliases
 
 		/// <inheritdoc/>
-		public override async Task AddAliasAsync(string name, BlobHandle handle, int rank = 0, ReadOnlyMemory<byte> data = default, CancellationToken cancellationToken = default)
+		public Task AddAliasAsync(string name, BlobHandle handle, int rank = 0, ReadOnlyMemory<byte> data = default, CancellationToken cancellationToken = default)
 		{
-			await handle.FlushAsync(cancellationToken);
 			BlobLocator locator = handle.GetLocator();
-			_aliases.AddOrUpdate(name, _ => new ExportEntry(locator, rank, data, null), (_, entry) => new ExportEntry(locator, rank, data, entry));
+			_aliases.AddOrUpdate(name, _ => new AliasListNode(locator, rank, data, null), (_, entry) => new AliasListNode(locator, rank, data, entry));
+			return Task.CompletedTask;
 		}
 
 		/// <inheritdoc/>
-		public override Task RemoveAliasAsync(string name, BlobHandle handle, CancellationToken cancellationToken = default)
+		public Task RemoveAliasAsync(string name, BlobHandle handle, CancellationToken cancellationToken = default)
+		{
+			BlobLocator locator = handle.GetLocator();
+			for (; ; )
+			{
+				AliasListNode? entry;
+				if (!_aliases.TryGetValue(name, out entry))
+				{
+					break;
+				}
+
+				AliasListNode? newEntry = RemoveAliasFromList(entry, locator);
+				if (entry == newEntry)
+				{
+					break;
+				}
+
+				if (newEntry == null)
+				{
+					if (_aliases.TryRemove(new KeyValuePair<string, AliasListNode?>(name, entry)))
+					{
+						break;
+					}
+				}
+				else
+				{
+					if (_aliases.TryUpdate(name, newEntry, entry))
+					{
+						break;
+					}
+				}
+			}
+			return Task.CompletedTask;
+		}
+
+		static AliasListNode? RemoveAliasFromList(AliasListNode? entry, BlobLocator locator)
 		{
-			throw new NotSupportedException();
+			if (entry == null)
+			{
+				return null;
+			}
+			if (entry.Locator == locator)
+			{
+				return entry.Next;
+			}
+
+			AliasListNode? nextEntry = RemoveAliasFromList(entry.Next, locator);
+			if (nextEntry != entry.Next)
+			{
+				entry = new AliasListNode(entry.Locator, entry.Rank, entry.Data, nextEntry);
+			}
+			return entry;
 		}
 
 		/// <inheritdoc/>
-		public override Task<BlobAlias[]> FindAliasesAsync(string alias, int? maxResults = null, CancellationToken cancellationToken = default)
+		public Task<BlobAlias[]> FindAliasesAsync(string alias, int? maxResults = null, CancellationToken cancellationToken = default)
 		{
 			List<BlobAlias> aliases = new List<BlobAlias>();
-			if (_aliases.TryGetValue(alias, out ExportEntry? entry))
+			if (_aliases.TryGetValue(alias, out AliasListNode? entry))
 			{
 				for (; entry != null; entry = entry.Next)
 				{
@@ -95,15 +185,15 @@ namespace EpicGames.Horde.Storage.Clients
 		#region Refs
 
 		/// <inheritdoc/>
-		public override Task<bool> DeleteRefAsync(RefName name, CancellationToken cancellationToken) => Task.FromResult(_refs.TryRemove(name, out _));
+		public Task<bool> DeleteRefAsync(RefName name, CancellationToken cancellationToken) => Task.FromResult(_refs.TryRemove(name, out _));
 
 		/// <inheritdoc/>
-		public override Task<BlobHandle?> TryReadRefAsync(RefName name, RefCacheTime cacheTime = default, CancellationToken cancellationToken = default)
+		public Task<BlobHandle?> TryReadRefAsync(RefName name, RefCacheTime cacheTime = default, CancellationToken cancellationToken = default)
 		{
-			BlobLocator refData;
-			if (_refs.TryGetValue(name, out refData))
+			BlobLocator locator;
+			if (_refs.TryGetValue(name, out locator))
 			{
-				return Task.FromResult<BlobHandle?>(CreateBlobHandle(refData)); 
+				return Task.FromResult<BlobHandle?>(CreateBlobHandle(locator));
 			}
 			else
 			{
@@ -112,13 +202,17 @@ namespace EpicGames.Horde.Storage.Clients
 		}
 
 		/// <inheritdoc/>
-		public override async Task WriteRefAsync(RefName name, BlobHandle target, RefOptions? options = null, CancellationToken cancellationToken = default)
+		public async Task WriteRefAsync(RefName name, BlobHandle target, RefOptions? options = null, CancellationToken cancellationToken = default)
 		{
 			await target.FlushAsync(cancellationToken);
-
 			_refs[name] = target.GetLocator();
 		}
 
 		#endregion
+
+		/// <inheritdoc/>
+		public void GetStats(StorageStats stats)
+		{
+		}
 	}
 }
diff --git a/Engine/Source/Programs/Shared/EpicGames.Horde/Storage/IStorageWriter.cs b/Engine/Source/Programs/Shared/EpicGames.Horde/Storage/IStorageWriter.cs
--- a/Engine/Source/Programs/Shared/EpicGames.Horde/Storage/IStorageWriter.cs
+++ b/Engine/Source/Programs/Shared/EpicGames.Horde/Storage/IStorageWriter.cs
@@ -113,7 +113,9 @@ namespace EpicGames.Horde.Storage
 			ReadOnlyMemory<byte> data = _data.AsMemory(_offset, size);
 			_offset += size;
 
-			BlobHandle handle = await _outer.WriteBlobAsync(type, data, references, cancellationToken);
+			using ReadOnlyMemoryStream stream = new ReadOnlyMemoryStream(data);
+
+			BlobHandle handle = await _outer.WriteBlobAsync(type, stream, references, _basePath, cancellationToken);
 			foreach (AliasInfo aliasInfo in aliases)
 			{
 				await _outer.AddAliasAsync(aliasInfo.Name, handle, aliasInfo.Rank, aliasInfo.Data, cancellationToken);
