Improve diagnostics for race conditions involving dispose calls on ThreadPoolWorkQueue.

* Event is now set to null during dispose in an atomic operation. The event being null is indication that the object has been disposed.
* Event is now always read with an atomic operation, and throws an exception if null (indicating that the object has already been disposed).
* Reads of _numOutstandingItems are also now atomic to ensure correct sequencing with increment/decrement operations.

diff --git a/Engine/Source/Programs/Shared/EpicGames.Core/ThreadPoolWorkQueue.cs b/Engine/Source/Programs/Shared/EpicGames.Core/ThreadPoolWorkQueue.cs
--- a/Engine/Source/Programs/Shared/EpicGames.Core/ThreadPoolWorkQueue.cs
+++ b/Engine/Source/Programs/Shared/EpicGames.Core/ThreadPoolWorkQueue.cs
@@ -25,7 +25,9 @@ namespace EpicGames.Core
 		/// <summary>
 		/// Event which indicates whether the queue is empty.
 		/// </summary>
+#pragma warning disable CA2213 // Disposable fields should be disposed
 		ManualResetEvent? _emptyEvent = new(true);
+#pragma warning restore CA2213 // Disposable fields should be disposed
 
 		/// <summary>
 		/// Exceptions which occurred while executing tasks
@@ -44,12 +46,12 @@ namespace EpicGames.Core
 		/// </summary>
 		public void Dispose()
 		{
-			if (_emptyEvent != null)
+			ManualResetEvent? emptyEvent = Interlocked.CompareExchange(ref _emptyEvent, null, null);
+			if (emptyEvent != null)
 			{
-				Wait();
-
-				_emptyEvent?.Dispose();
-				_emptyEvent = null;
+				emptyEvent.WaitOne();
+				Interlocked.CompareExchange(ref _emptyEvent, null, emptyEvent)?.Dispose();
+				RethrowExceptions();
 			}
 		}
 
@@ -99,12 +101,25 @@ namespace EpicGames.Core
 			}
 		}
 
+		/// <summary>
+		/// Atomically read the contents of <see cref="_emptyEvent"/>, throwing an exception if it's already null (indicating that the object has been disposed)
+		/// </summary>
+		ManualResetEvent GetEmptyEvent()
+		{
+			ManualResetEvent? emptyEvent = Interlocked.CompareExchange(ref _emptyEvent, null, null);
+			if (emptyEvent == null)
+			{
+				throw new ObjectDisposedException(typeof(ThreadPoolWorkQueue).Name);
+			}
+			return emptyEvent;
+		}
+
 		/// <summary>
 		/// Waits for all queued tasks to finish
 		/// </summary>
 		public void Wait()
 		{
-			_emptyEvent?.WaitOne();
+			GetEmptyEvent().WaitOne();
 			RethrowExceptions();
 		}
 
@@ -125,7 +140,7 @@ namespace EpicGames.Core
 		/// <returns>True if the queue completed, false if the timeout elapsed</returns>
 		public bool Wait(TimeSpan timeout)
 		{
-			bool bResult = _emptyEvent?.WaitOne(timeout) ?? false;
+			bool bResult = GetEmptyEvent().WaitOne(timeout);
 			if (bResult)
 			{
 				RethrowExceptions();
@@ -153,13 +168,14 @@ namespace EpicGames.Core
 		{
 			lock (_lockObject)
 			{
-				if (_numOutstandingJobs > 0)
+				ManualResetEvent emptyEvent = GetEmptyEvent();
+				if (Interlocked.CompareExchange(ref _numOutstandingJobs, 0, 0) > 0)
 				{
-					_emptyEvent?.Reset();
+					emptyEvent.Reset();
 				}
 				else
 				{
-					_emptyEvent?.Set();
+					emptyEvent.Set();
 				}
 			}
 		}
