Horde: Fix cache being ignored for header reads.

diff --git a/Engine/Source/Programs/Shared/EpicGames.Horde/Storage/Bundles/BundleReader.cs b/Engine/Source/Programs/Shared/EpicGames.Horde/Storage/Bundles/BundleReader.cs
--- a/Engine/Source/Programs/Shared/EpicGames.Horde/Storage/Bundles/BundleReader.cs
+++ b/Engine/Source/Programs/Shared/EpicGames.Horde/Storage/Bundles/BundleReader.cs
@@ -168,10 +168,10 @@ namespace EpicGames.Horde.Storage.Bundles
 		public async Task<BundleHeader> ReadHeaderAsync(BlobLocator locator, CancellationToken cancellationToken)
 		{
 			// Check for a cached value first
-//			if (_cache.TryGetCachedHeader(locator, out BundleInfo? cachedBundleInfo))
-//			{
-//				return cachedBundleInfo.Header;
-//			}
+			if (_cache.TryGetCachedHeader(locator, out BundleInfo? cachedBundleInfo))
+			{
+				return cachedBundleInfo.Header;
+			}
 
 			// Find a registered bundle info with the given locator
 			QueuedBundle? bundle = null;
