Horde: Add a base class for storage clients that wrap a key/value storage model, without any of the splitting/merging that the bundle storage client does.

diff --git a/Engine/Source/Programs/Horde/Horde.Server/Commands/Bundles/BundleWriteTest.cs b/Engine/Source/Programs/Horde/Horde.Server/Commands/Bundles/BundleWriteTest.cs
--- a/Engine/Source/Programs/Horde/Horde.Server/Commands/Bundles/BundleWriteTest.cs
+++ b/Engine/Source/Programs/Horde/Horde.Server/Commands/Bundles/BundleWriteTest.cs
@@ -30,6 +30,7 @@ namespace Horde.Server.Commands.Bundles
 
 			public BlobHandle CreateBlobHandle(BlobLocator locator) => throw new NotImplementedException();
 			public Task<bool> DeleteRefAsync(RefName name, CancellationToken cancellationToken = default) => Task.FromResult(true);
+			public ValueTask<BlobData> ReadBlobAsync(BlobLocator locator, CancellationToken cancellationToken = default) => throw new NotImplementedException();
 			public ValueTask<BlobHandle> WriteBlobAsync(BlobType type, Stream stream, IReadOnlyList<BlobHandle> references, string? basePath = null, CancellationToken cancellationToken = default) => throw new NotImplementedException();
 			public ValueTask<Uri?> TryGetReadRedirectAsync(BlobLocator locator, CancellationToken cancellationToken = default) => new ValueTask<Uri?>();
 			public ValueTask<(BlobLocator, Uri)?> TryGetWriteRedirectAsync(string? prefix = null, CancellationToken cancellationToken = default) => new ValueTask<(BlobLocator, Uri)?>();
diff --git a/Engine/Source/Programs/Horde/Horde.Server/Storage/StorageService.cs b/Engine/Source/Programs/Horde/Horde.Server/Storage/StorageService.cs
--- a/Engine/Source/Programs/Horde/Horde.Server/Storage/StorageService.cs
+++ b/Engine/Source/Programs/Horde/Horde.Server/Storage/StorageService.cs
@@ -184,6 +184,13 @@ namespace Horde.Server.Storage
 			/// <inheritdoc/>
 			public void GetStats(StorageStats stats) => _backend.GetStats(stats);
 
+			/// <inheritdoc/>
+			public async ValueTask<BlobData> ReadBlobAsync(BlobLocator locator, CancellationToken cancellationToken = default)
+			{
+				BlobHandle handle = CreateBlobHandle(locator);
+				return await handle.ReadAsync(cancellationToken);
+			}
+
 			/// <inheritdoc/>
 			public async ValueTask<BlobHandle> WriteBlobAsync(BlobType type, Stream stream, IReadOnlyList<BlobHandle> references, string? basePath = null, CancellationToken cancellationToken = default)
 			{
diff --git a/Engine/Source/Programs/Shared/EpicGames.Horde/Storage/Clients/BundleStorageClient.cs b/Engine/Source/Programs/Shared/EpicGames.Horde/Storage/Clients/BundleStorageClient.cs
--- a/Engine/Source/Programs/Shared/EpicGames.Horde/Storage/Clients/BundleStorageClient.cs
+++ b/Engine/Source/Programs/Shared/EpicGames.Horde/Storage/Clients/BundleStorageClient.cs
@@ -119,6 +119,13 @@ namespace EpicGames.Horde.Storage.Clients
 		public async Task<Stream> OpenBlobAsync(BlobLocator locator, int offset = 0, int? length = null, CancellationToken cancellationToken = default) 
 			=> await OpenBundleAsync(locator, offset, length, cancellationToken);
 
+		/// <inheritdoc/>
+		public async ValueTask<BlobData> ReadBlobAsync(BlobLocator locator, CancellationToken cancellationToken = default)
+		{
+			BlobHandle handle = CreateBlobHandle(locator);
+			return await handle.ReadAsync(cancellationToken);
+		}
+
 		/// <inheritdoc/>
 		public async ValueTask<BlobHandle> WriteBlobAsync(BlobType type, Stream stream, IReadOnlyList<BlobHandle> references, string? basePath = null, CancellationToken cancellationToken = default)
 		{
diff --git a/Engine/Source/Programs/Shared/EpicGames.Horde/Storage/Clients/KeyValueStorageClient.cs b/Engine/Source/Programs/Shared/EpicGames.Horde/Storage/Clients/KeyValueStorageClient.cs
new file mode 100644
--- /dev/null
+++ b/Engine/Source/Programs/Shared/EpicGames.Horde/Storage/Clients/KeyValueStorageClient.cs
@@ -0,0 +1,101 @@
+// Copyright Epic Games, Inc. All Rights Reserved.
+
+using System;
+using System.Collections.Generic;
+using System.Diagnostics.CodeAnalysis;
+using System.IO;
+using System.Threading;
+using System.Threading.Tasks;
+
+namespace EpicGames.Horde.Storage.Clients
+{
+	/// <summary>
+	/// Base class for storage clients that wrap a diirect key/value type store without any merging/splitting.
+	/// </summary>
+	abstract class KeyValueStorageClient : IStorageClient
+	{
+		class Handle : BlobHandle
+		{
+			readonly KeyValueStorageClient _outer;
+			readonly BlobLocator _locator;
+
+			public Handle(KeyValueStorageClient outer, BlobLocator locator)
+			{
+				_outer = outer;
+				_locator = locator;
+			}
+
+			public override ValueTask<BlobData> ReadAsync(CancellationToken cancellationToken = default) => _outer.ReadBlobAsync(_locator, cancellationToken);
+
+			public override bool TryGetLocator([NotNullWhen(true)] out BlobLocator locator)
+			{
+				locator = _locator;
+				return true;
+			}
+		}
+
+		/// <inheritdoc/>
+		public abstract bool SupportsRedirects { get; }
+
+		/// <inheritdoc/>
+		public void Dispose()
+		{
+			Dispose(true);
+			GC.SuppressFinalize(this);
+		}
+
+		protected virtual void Dispose(bool disposing)
+		{
+		}
+
+		#region Blobs
+
+		/// <inheritdoc/>
+		public virtual BlobHandle CreateBlobHandle(BlobLocator locator) => new Handle(this, locator);
+
+		/// <inheritdoc/>
+		public virtual IStorageWriter CreateWriter(string? basePath = null) => new DefaultStorageWriter(this, basePath);
+
+		/// <inheritdoc/>
+		public abstract ValueTask<BlobData> ReadBlobAsync(BlobLocator locator, CancellationToken cancellationToken = default);
+
+		/// <inheritdoc/>
+		public abstract ValueTask<BlobHandle> WriteBlobAsync(BlobType type, Stream stream, IReadOnlyList<BlobHandle> references, string? basePath = null, CancellationToken cancellationToken = default);
+
+		/// <inheritdoc/>
+		public abstract ValueTask<Uri?> TryGetReadRedirectAsync(BlobLocator locator, CancellationToken cancellationToken = default);
+
+		/// <inheritdoc/>
+		public abstract ValueTask<(BlobLocator, Uri)?> TryGetWriteRedirectAsync(string? prefix = null, CancellationToken cancellationToken = default);
+
+		#endregion
+
+		#region Aliases
+
+		/// <inheritdoc/>
+		public abstract Task AddAliasAsync(string name, BlobHandle handle, int rank = 0, ReadOnlyMemory<byte> data = default, CancellationToken cancellationToken = default);
+
+		/// <inheritdoc/>
+		public abstract Task RemoveAliasAsync(string name, BlobHandle handle, CancellationToken cancellationToken = default);
+
+		/// <inheritdoc/>
+		public abstract Task<BlobAlias[]> FindAliasesAsync(string name, int? maxResults = null, CancellationToken cancellationToken = default);
+
+		#endregion
+
+		#region Refs
+
+		/// <inheritdoc/>
+		public abstract Task<BlobHandle?> TryReadRefAsync(RefName name, RefCacheTime cacheTime = default, CancellationToken cancellationToken = default);
+
+		/// <inheritdoc/>
+		public abstract Task WriteRefAsync(RefName name, BlobHandle handle, RefOptions? options = null, CancellationToken cancellationToken = default);
+
+		/// <inheritdoc/>
+		public abstract Task<bool> DeleteRefAsync(RefName name, CancellationToken cancellationToken = default);
+
+		#endregion
+
+		public abstract void GetStats(StorageStats stats);
+	}
+}
diff --git a/Engine/Source/Programs/Shared/EpicGames.Horde/Storage/Clients/SharedStorageClient.cs b/Engine/Source/Programs/Shared/EpicGames.Horde/Storage/Clients/SharedStorageClient.cs
--- a/Engine/Source/Programs/Shared/EpicGames.Horde/Storage/Clients/SharedStorageClient.cs
+++ b/Engine/Source/Programs/Shared/EpicGames.Horde/Storage/Clients/SharedStorageClient.cs
@@ -76,6 +76,9 @@ namespace EpicGames.Horde.Storage.Clients
 		/// <inheritdoc/>
 		public IStorageWriter CreateWriter(string? basePath = null) => _inner.CreateWriter(basePath);
 
+		/// <inheritdoc/>
+		public ValueTask<BlobData> ReadBlobAsync(BlobLocator locator, CancellationToken cancellationToken = default) => _inner.ReadBlobAsync(locator, cancellationToken);
+
 		/// <inheritdoc/>
 		public ValueTask<BlobHandle> WriteBlobAsync(BlobType type, Stream stream, IReadOnlyList<BlobHandle> references, string? basePath = null, CancellationToken cancellationToken = default) => _inner.WriteBlobAsync(type, stream, references, basePath, cancellationToken);
 
diff --git a/Engine/Source/Programs/Shared/EpicGames.Horde/Storage/IStorageClient.cs b/Engine/Source/Programs/Shared/EpicGames.Horde/Storage/IStorageClient.cs
--- a/Engine/Source/Programs/Shared/EpicGames.Horde/Storage/IStorageClient.cs
+++ b/Engine/Source/Programs/Shared/EpicGames.Horde/Storage/IStorageClient.cs
@@ -6,6 +6,7 @@ using System.IO;
 using System.Linq;
 using System.Threading;
 using System.Threading.Tasks;
+using EpicGames.Core;
 using Microsoft.Extensions.Logging;
 
 namespace EpicGames.Horde.Storage
@@ -363,11 +364,8 @@ namespace EpicGames.Horde.Storage
 		/// <param name="cancellationToken">Cancellation token for the operation</param>
 		public static async ValueTask<BlobHandle> WriteBlobAsync(this IStorageClient store, string? basePath, BlobType type, ReadOnlyMemory<byte> data, IReadOnlyList<BlobHandle> references, CancellationToken cancellationToken = default)
 		{
-			await using (IStorageWriter writer = store.CreateWriter(basePath))
-			{
-				data.CopyTo(writer.GetOutputBuffer(0, data.Length));
-				return await writer.WriteBlobAsync(type, data.Length, references, cancellationToken);
-			}
+			using ReadOnlyMemoryStream stream = new ReadOnlyMemoryStream(data);
+			return await store.WriteBlobAsync(type, stream, references, basePath, cancellationToken);
 		}
 
 		#endregion
