Horde: Derive FileStorageClient from IStorageClient rather than BundleStorageClientBase, allowing it to be used with a different packing algorithm.

diff --git a/Engine/Source/Programs/Horde/Horde.Server/Commands/Bundles/CreateCommand.cs b/Engine/Source/Programs/Horde/Horde.Server/Commands/Bundles/CreateCommand.cs
--- a/Engine/Source/Programs/Horde/Horde.Server/Commands/Bundles/CreateCommand.cs
+++ b/Engine/Source/Programs/Horde/Horde.Server/Commands/Bundles/CreateCommand.cs
@@ -47,7 +47,7 @@ namespace Horde.Server.Commands.Bundles
 			}
 			else
 			{
-				using FileStorageClient storageClient = new FileStorageClient(OutputDir, serviceProvider.GetRequiredService<BundleReaderCache>(), logger);
+				using IStorageClient storageClient = BundleStorageClientWrapper.CreateFromDirectory(OutputDir, serviceProvider.GetRequiredService<BundleReaderCache>(), logger);
 				await ExecuteInternalAsync(storageClient, logger);
 			}
 
diff --git a/Engine/Source/Programs/Horde/Horde.Server/Commands/Install/SetupCommand.cs b/Engine/Source/Programs/Horde/Horde.Server/Commands/Install/SetupCommand.cs
--- a/Engine/Source/Programs/Horde/Horde.Server/Commands/Install/SetupCommand.cs
+++ b/Engine/Source/Programs/Horde/Horde.Server/Commands/Install/SetupCommand.cs
@@ -60,7 +60,7 @@ namespace Horde.Server.Commands.Install
 
 			using BundleReaderCache bundleReaderCache = new BundleReaderCache();
 
-			using FileStorageClient client = new FileStorageClient(bundleDir, bundleReaderCache, logger);
+			using IStorageClient client = BundleStorageClientWrapper.CreateFromDirectory(bundleDir, bundleReaderCache, logger);
 			await using (IStorageWriter writer = client.CreateWriter(refName))
 			{
 				DirectoryNode dirNode = new DirectoryNode();
diff --git a/Engine/Source/Programs/Horde/Horde.Server/Tools/ToolCollection.cs b/Engine/Source/Programs/Horde/Horde.Server/Tools/ToolCollection.cs
--- a/Engine/Source/Programs/Horde/Horde.Server/Tools/ToolCollection.cs
+++ b/Engine/Source/Programs/Horde/Horde.Server/Tools/ToolCollection.cs
@@ -425,7 +425,7 @@ namespace Horde.Server.Tools
 		{
 			if (tool.Config is BundledToolConfig bundledConfig)
 			{
-				return new FileStorageClient(DirectoryReference.Combine(ServerApp.AppDir, bundledConfig.DataDir ?? $"tools/{tool.Id}"), _cache, _logger);
+				return BundleStorageClientWrapper.CreateFromDirectory(DirectoryReference.Combine(ServerApp.AppDir, bundledConfig.DataDir ?? $"tools/{tool.Id}"), _cache, _logger);
 			}
 			else
 			{
diff --git a/Engine/Source/Programs/Horde/Horde/Commands/Bundles/BundleCreate.cs b/Engine/Source/Programs/Horde/Horde/Commands/Bundles/BundleCreate.cs
--- a/Engine/Source/Programs/Horde/Horde/Commands/Bundles/BundleCreate.cs
+++ b/Engine/Source/Programs/Horde/Horde/Commands/Bundles/BundleCreate.cs
@@ -34,7 +34,7 @@ namespace Horde.Commands.Bundles
 		{
 			if (File != null)
 			{
-				using FileStorageClient store = new FileStorageClient(File.Directory, BundleReaderCache, logger);
+				using IStorageClient store = BundleStorageClientWrapper.CreateFromDirectory(File.Directory, BundleReaderCache, logger);
 				return await ExecuteInternalAsync(store, logger);
 			}
 			else if (!String.IsNullOrEmpty(Ref))
diff --git a/Engine/Source/Programs/Horde/Horde/Commands/Bundles/BundleExtract.cs b/Engine/Source/Programs/Horde/Horde/Commands/Bundles/BundleExtract.cs
--- a/Engine/Source/Programs/Horde/Horde/Commands/Bundles/BundleExtract.cs
+++ b/Engine/Source/Programs/Horde/Horde/Commands/Bundles/BundleExtract.cs
@@ -49,7 +49,7 @@ namespace Horde.Commands.Bundles
 			}
 			else if (Node != null)
 			{
-				using BundleStorageClient store = (BundleStorageClient)CreateStorageClient();
+				using IStorageClient store = CreateStorageClient();
 				BlobHandle handle = store.CreateBlobHandle(new BlobLocator(Node));
 				await ExecuteInternalAsync(store, handle, logger);
 			}
diff --git a/Engine/Source/Programs/Horde/Horde/Commands/Workspace/WorkspaceSync.cs b/Engine/Source/Programs/Horde/Horde/Commands/Workspace/WorkspaceSync.cs
--- a/Engine/Source/Programs/Horde/Horde/Commands/Workspace/WorkspaceSync.cs
+++ b/Engine/Source/Programs/Horde/Horde/Commands/Workspace/WorkspaceSync.cs
@@ -42,8 +42,8 @@ namespace Horde.Commands.Workspace
 		{
 			if (File != null)
 			{
-				using FileStorageClient store = new FileStorageClient(File.Directory, BundleReaderCache, logger);
-				BlobHandle handle = await store.ReadRefAsync(File);
+				using IStorageClient store = BundleStorageClientWrapper.CreateFromDirectory(File.Directory, BundleReaderCache, logger);
+				BlobHandle handle = store.CreateBlobHandle(await FileStorageClient.ReadRefAsync(File));
 				return await ExecuteInternalAsync(store, handle, logger);
 			}
 			else if (Ref != null)
diff --git a/Engine/Source/Programs/Shared/EpicGames.Horde/Storage/BlobType.cs b/Engine/Source/Programs/Shared/EpicGames.Horde/Storage/BlobType.cs
--- a/Engine/Source/Programs/Shared/EpicGames.Horde/Storage/BlobType.cs
+++ b/Engine/Source/Programs/Shared/EpicGames.Horde/Storage/BlobType.cs
@@ -17,6 +17,11 @@ namespace EpicGames.Horde.Storage
 		/// </summary>
 		public const int NumBytes = 20;
 
+		/// <summary>
+		/// Blob type used to indicate a leaf node. Data is an opaque blob and has no references.
+		/// </summary>
+		public static BlobType Leaf { get; } = new BlobType(Guid.Parse("{1080B643-8015-4A4D-ADBB-C2FB586C8D81}"), 1);
+
 		/// <summary>
 		/// Deserialize a type from a byte span
 		/// </summary>
diff --git a/Engine/Source/Programs/Shared/EpicGames.Horde/Storage/Clients/BundleStorageClient.cs b/Engine/Source/Programs/Shared/EpicGames.Horde/Storage/Clients/BundleStorageClient.cs
--- a/Engine/Source/Programs/Shared/EpicGames.Horde/Storage/Clients/BundleStorageClient.cs
+++ b/Engine/Source/Programs/Shared/EpicGames.Horde/Storage/Clients/BundleStorageClient.cs
@@ -317,6 +317,24 @@ namespace EpicGames.Horde.Storage.Clients
 			base.Dispose(disposing);
 		}
 
+		/// <summary>
+		/// Creates a bundle storage client around a memory client backend
+		/// </summary>
+		public static BundleStorageClientWrapper CreateFromMemory(ILogger logger)
+		{
+			MemoryStorageClient blobStore = new MemoryStorageClient();
+			return new BundleStorageClientWrapper(blobStore, BundleReaderCache.None, logger);
+		}
+
+		/// <summary>
+		/// Creates a bundle storage client around a directory on the filesystem
+		/// </summary>
+		public static BundleStorageClientWrapper CreateFromDirectory(DirectoryReference rootDir, BundleReaderCache cache, ILogger logger)
+		{
+			FileStorageClient fileStorageClient = new FileStorageClient(rootDir, logger);
+			return new BundleStorageClientWrapper(fileStorageClient, cache, logger);
+		}
+
 		/// <inheritdoc/>
 		protected override async Task<Stream> OpenBundleAsync(BlobLocator locator, int offset, int? length, CancellationToken cancellationToken)
 			=> await _inner.CreateBlobHandle(locator).OpenAsync(offset, length, cancellationToken);
diff --git a/Engine/Source/Programs/Shared/EpicGames.Horde/Storage/Clients/FileStorageClient.cs b/Engine/Source/Programs/Shared/EpicGames.Horde/Storage/Clients/FileStorageClient.cs
--- a/Engine/Source/Programs/Shared/EpicGames.Horde/Storage/Clients/FileStorageClient.cs
+++ b/Engine/Source/Programs/Shared/EpicGames.Horde/Storage/Clients/FileStorageClient.cs
@@ -1,6 +1,7 @@
 // Copyright Epic Games, Inc. All Rights Reserved.
 
 using System;
+using System.Collections.Generic;
 using System.IO;
 using System.Threading;
 using System.Threading.Tasks;
@@ -13,37 +14,95 @@ namespace EpicGames.Horde.Storage.Clients
 	/// <summary>
 	/// Implementation of <see cref="IStorageClient"/> which writes data to files on disk.
 	/// </summary>
-	public class FileStorageClient : BundleStorageClient
+	public class FileStorageClient : KeyValueStorageClient
 	{
+		class LeafBlobData : BlobData
+		{
+			readonly IReadOnlyMemoryOwner<byte> _owner;
+
+			public LeafBlobData(IReadOnlyMemoryOwner<byte> owner)
+				: base(BlobType.Leaf, owner.Memory, Array.Empty<BlobHandle>())
+			{
+				_owner = owner;
+			}
+
+			protected override void Dispose(bool disposing)
+			{
+				base.Dispose(disposing);
+
+				if (disposing)
+				{
+					_owner.Dispose();
+				}
+			}
+		}
+
 		readonly DirectoryReference _rootDir;
+		readonly FileStorageBackend _backend;
 		readonly ILogger _logger;
 
+		/// <inheritdoc/>
+		public override bool SupportsRedirects => false;
+
 		/// <summary>
 		/// Constructor
 		/// </summary>
 		/// <param name="rootDir">Root directory for storing blobs</param>
-		/// <param name="cache">Memory cache for read data</param>
 		/// <param name="logger">Logger interface</param>
-		public FileStorageClient(DirectoryReference rootDir, BundleReaderCache cache, ILogger logger)
-			: base(new FileStorageBackend(rootDir), cache, logger)
+		public FileStorageClient(DirectoryReference rootDir, ILogger logger)
 		{
 			_rootDir = rootDir;
+			_backend = new FileStorageBackend(rootDir);
 			_logger = logger;
 
 			DirectoryReference.CreateDirectory(_rootDir);
 		}
 
+		/// <inheritdoc/>
+		protected override void Dispose(bool disposing)
+		{
+			base.Dispose(disposing);
+
+			if (disposing)
+			{
+				_backend.Dispose();
+			}
+		}
+
 		/// <summary>
 		/// Reads a ref from a file on disk
 		/// </summary>
-		public async ValueTask<BlobHandle> ReadRefAsync(FileReference file)
+		public static async ValueTask<BlobLocator> ReadRefAsync(FileReference file)
 		{
 			string text = await FileReference.ReadAllTextAsync(file);
-			return CreateBlobHandle(new BlobLocator(text));
+			return new BlobLocator(text);
 		}
 
 		FileReference GetRefFile(RefName name) => FileReference.Combine(_rootDir, name.ToString() + ".ref");
 
+		#region Blobs
+		/// <inheritdoc/>
+		public override async ValueTask<BlobData> ReadBlobAsync(BlobLocator locator, CancellationToken cancellationToken = default)
+		{
+			IReadOnlyMemoryOwner<byte> owner = await _backend.ReadAsync(locator.ToString(), cancellationToken);
+			return new LeafBlobData(owner);
+		}
+
+		/// <inheritdoc/>
+		public override async ValueTask<BlobHandle> WriteBlobAsync(BlobType type, Stream stream, IReadOnlyList<BlobHandle> references, string? basePath = null, CancellationToken cancellationToken = default)
+		{
+			string path = await _backend.WriteAsync(stream, basePath, cancellationToken);
+			return CreateBlobHandle(new BlobLocator(path));
+		}
+
+		/// <inheritdoc/>
+		public override ValueTask<Uri?> TryGetReadRedirectAsync(BlobLocator locator, CancellationToken cancellationToken = default) => default;
+
+		/// <inheritdoc/>
+		public override ValueTask<(BlobLocator, Uri)?> TryGetWriteRedirectAsync(string? prefix = null, CancellationToken cancellationToken = default) => default;
+
+		#endregion
+
 		#region Aliases
 
 		/// <inheritdoc/>
@@ -121,6 +180,11 @@ namespace EpicGames.Horde.Storage.Clients
 			}
 		}
 
+		/// <inheritdoc/>
+		public override void GetStats(StorageStats stats)
+		{
+		}
+
 		#endregion
 	}
 }
diff --git a/Engine/Source/Programs/UnrealBuildTool/Artifacts/HordeStorageArtifactCache.cs b/Engine/Source/Programs/UnrealBuildTool/Artifacts/HordeStorageArtifactCache.cs
--- a/Engine/Source/Programs/UnrealBuildTool/Artifacts/HordeStorageArtifactCache.cs
+++ b/Engine/Source/Programs/UnrealBuildTool/Artifacts/HordeStorageArtifactCache.cs
@@ -526,7 +526,7 @@ namespace UnrealBuildTool.Artifacts
 				}
 				Directory.CreateDirectory(directory.FullName);
 
-				_store = new FileStorageClient(directory, BundleReaderCache.None, logger);
+				_store = BundleStorageClientWrapper.CreateFromDirectory(directory, BundleReaderCache.None, logger);
 
 				State = ArtifactCacheState.Available;
 				return State;
