Horde: Remove extension method for reading blobs from IStorageClient. This is already declared directly in all implementations.

diff --git a/Engine/Source/Programs/Shared/EpicGames.Horde/Storage/IStorageClient.cs b/Engine/Source/Programs/Shared/EpicGames.Horde/Storage/IStorageClient.cs
--- a/Engine/Source/Programs/Shared/EpicGames.Horde/Storage/IStorageClient.cs
+++ b/Engine/Source/Programs/Shared/EpicGames.Horde/Storage/IStorageClient.cs
@@ -103,6 +103,13 @@ namespace EpicGames.Horde.Storage
 		/// <returns>New writer instance. Must be disposed after use.</returns>
 		IStorageWriter CreateWriter(string? basePath = null);
 
+		/// <summary>
+		/// Read a blob from the underlying storage system. Calling <see cref="BlobHandle.ReadAsync(CancellationToken)"/> is more efficient than calling this method repeatedly for small blobs.
+		/// </summary>
+		/// <param name="locator">Locator for the blob</param>
+		/// <param name="cancellationToken">Cancellation token for the operation</param>
+		ValueTask<BlobData> ReadBlobAsync(BlobLocator locator, CancellationToken cancellationToken = default);
+
 		/// <summary>
 		/// Write a blob to the underlying storage system. Using the writer instance returned from <see cref="IStorageClient.CreateWriter(String)"/> is more efficient than calling this method repeatedly for small blobs.
 		/// </summary>
@@ -328,18 +335,6 @@ namespace EpicGames.Horde.Storage
 		/// <param name="refName">Ref name to use as a base path</param>
 		public static IStorageWriter CreateWriter(this IStorageClient store, RefName refName) => store.CreateWriter(refName.ToString());
 
-		/// <summary>
-		/// Read a blob from the underlying storage system. Calling <see cref="BlobHandle.ReadAsync(CancellationToken)"/> is more efficient than calling this method repeatedly for small blobs.
-		/// </summary>
-		/// <param name="store">The store instance to read from</param>
-		/// <param name="locator">Locator for the blob</param>
-		/// <param name="cancellationToken">Cancellation token for the operation</param>
-		public static async ValueTask<BlobData> ReadBlobAsync(this IStorageClient store, BlobLocator locator, CancellationToken cancellationToken = default)
-		{
-			BlobHandle handle = store.CreateBlobHandle(locator);
-			return await handle.ReadAsync(cancellationToken);
-		}
-
 		/// <summary>
 		/// Write a blob to the underlying storage system. Using the writer instance returned from <see cref="IStorageClient.CreateWriter(String)"/> is more efficient than calling this method repeatedly for small blobs.
 		/// </summary>
