Horde: Implement AgentStorageClient directly as an IStorageClient, rather than deriving from BundleStorageClientBase.

diff --git a/Engine/Source/Programs/Shared/EpicGames.Horde/Compute/AgentMessageHandler.cs b/Engine/Source/Programs/Shared/EpicGames.Horde/Compute/AgentMessageHandler.cs
--- a/Engine/Source/Programs/Shared/EpicGames.Horde/Compute/AgentMessageHandler.cs
+++ b/Engine/Source/Programs/Shared/EpicGames.Horde/Compute/AgentMessageHandler.cs
@@ -13,6 +13,7 @@ using Microsoft.Extensions.Logging;
 using EpicGames.Horde.Storage;
 using System.IO;
 using System.Linq;
+using EpicGames.Horde.Storage.Clients;
 
 namespace EpicGames.Horde.Compute
 {
@@ -134,7 +135,8 @@ namespace EpicGames.Horde.Compute
 
 		async Task WriteFilesAsync(AgentMessageChannel channel, string path, BlobLocator locator, CancellationToken cancellationToken)
 		{
-			using AgentStorageClient store = new AgentStorageClient(channel);
+			using AgentStorageClient innerStore = new AgentStorageClient(channel);
+			using BundleStorageClientWrapper store = new BundleStorageClientWrapper(innerStore, BundleReaderCache.None, _logger);
 
 			BlobHandle handle = store.CreateBlobHandle(locator);
 			DirectoryNode directoryNode = await handle.ReadNodeAsync<DirectoryNode>(cancellationToken);
diff --git a/Engine/Source/Programs/Shared/EpicGames.Horde/Compute/AgentStorageClient.cs b/Engine/Source/Programs/Shared/EpicGames.Horde/Compute/AgentStorageClient.cs
--- a/Engine/Source/Programs/Shared/EpicGames.Horde/Compute/AgentStorageClient.cs
+++ b/Engine/Source/Programs/Shared/EpicGames.Horde/Compute/AgentStorageClient.cs
@@ -5,9 +5,7 @@ using System.Collections.Generic;
 using System.IO;
 using System.Threading;
 using System.Threading.Tasks;
-using EpicGames.Core;
 using EpicGames.Horde.Storage;
-using Microsoft.Extensions.Logging.Abstractions;
 using EpicGames.Horde.Storage.Clients;
 
 namespace EpicGames.Horde.Compute
@@ -15,84 +13,45 @@ namespace EpicGames.Horde.Compute
 	/// <summary>
 	/// Storage client which can read bundles over a compute channel
 	/// </summary>
-	public sealed class AgentStorageClient : BundleStorageClient
-	{
-		/// <summary>
-		/// Constructor
-		/// </summary>
-		/// <param name="channel"></param>
-		public AgentStorageClient(AgentMessageChannel channel)
-			: base(new AgentStorageBackend(channel), BundleReaderCache.None, NullLogger.Instance)
-		{
-		}
-
-		#region Nodes
-
-		/// <inheritdoc/>
-		public override Task AddAliasAsync(string name, BlobHandle target, int rank, ReadOnlyMemory<byte> data, CancellationToken cancellationToken = default) => throw new NotSupportedException();
-
-		/// <inheritdoc/>
-		public override Task RemoveAliasAsync(string name, BlobHandle target, CancellationToken cancellationToken = default) => throw new NotSupportedException();
-
-		/// <inheritdoc/>
-		public override Task<BlobAlias[]> FindAliasesAsync(string name, int? maxResults, CancellationToken cancellationToken = default) => throw new NotSupportedException();
-
-		#endregion
-
-		#region Refs
-
-		/// <inheritdoc/>
-		public override Task<bool> DeleteRefAsync(RefName name, CancellationToken cancellationToken = default) => throw new NotSupportedException();
-
-		/// <inheritdoc/>
-		public override Task<BlobHandle?> TryReadRefAsync(RefName name, RefCacheTime cacheTime = default, CancellationToken cancellationToken = default) => throw new NotSupportedException();
-
-		/// <inheritdoc/>
-		public override Task WriteRefAsync(RefName name, BlobHandle target, RefOptions? options = null, CancellationToken cancellationToken = default) => throw new NotSupportedException();
-
-		#endregion
-	}
-
-	class AgentStorageBackend : IStorageBackend
+	public sealed class AgentStorageClient : KeyValueStorageClient
 	{
 		readonly AgentMessageChannel _channel;
 		readonly SemaphoreSlim _semaphore;
 
+		/// <inheritdoc/>
+		public override bool SupportsRedirects => false;
+
 		/// <summary>
 		/// Constructor
 		/// </summary>
 		/// <param name="channel"></param>
-		public AgentStorageBackend(AgentMessageChannel channel)
+		public AgentStorageClient(AgentMessageChannel channel)
 		{
 			_channel = channel;
 			_semaphore = new SemaphoreSlim(1);
 		}
 
 		/// <inheritdoc/>
-		public void Dispose()
+		protected override void Dispose(bool disposing)
 		{
-			_semaphore.Dispose();
+			base.Dispose(disposing);
+
+			if (disposing)
+			{
+				_semaphore.Dispose();
+			}
 		}
 
-		/// <inheritdoc/>
-		public bool SupportsRedirects => throw new NotImplementedException();
+		#region Blobs
 
 		/// <inheritdoc/>
-		public async Task<Stream> OpenAsync(string path, int offset, int? length, CancellationToken cancellationToken = default)
+		public override async ValueTask<BlobData> ReadBlobAsync(BlobLocator locator, CancellationToken cancellationToken = default)
 		{
 			await _semaphore.WaitAsync(cancellationToken);
 			try
 			{
-				ReadOnlyMemory<byte> data;
-				if (length.HasValue && length.Value == 0)
-				{
-					data = ReadOnlyMemory<byte>.Empty;
-				}
-				else
-				{
-					data = await _channel.ReadBlobAsync(path, offset, length ?? 0, cancellationToken);
-				}
-				return new ReadOnlyMemoryStream(data);
+				ReadOnlyMemory<byte> data = await _channel.ReadBlobAsync(locator.ToString(), 0, 0, cancellationToken);
+				return new BlobData(BlobType.Leaf, data, Array.Empty<BlobHandle>());
 			}
 			finally
 			{
@@ -101,37 +60,45 @@ namespace EpicGames.Horde.Compute
 		}
 
 		/// <inheritdoc/>
-		public async Task<IReadOnlyMemoryOwner<byte>> ReadAsync(string path, int offset, int? length, CancellationToken cancellationToken = default)
-		{
-			using (Stream stream = await OpenAsync(path, offset, length, cancellationToken))
-			{
-				byte[] data = await stream.ReadAllBytesAsync(cancellationToken);
-				return ReadOnlyMemoryOwner.Create(data);
-			}
-		}
+		public override ValueTask<BlobHandle> WriteBlobAsync(BlobType type, Stream stream, IReadOnlyList<BlobHandle> references, string? basePath = null, CancellationToken cancellationToken = default) => throw new NotSupportedException();
+
+		/// <inheritdoc/>
+		public override ValueTask<Uri?> TryGetReadRedirectAsync(BlobLocator locator, CancellationToken cancellationToken = default) => default;
 
 		/// <inheritdoc/>
-		public Task<string> WriteAsync(Stream stream, string? prefix = null, CancellationToken cancellationToken = default) => throw new NotSupportedException();
+		public override ValueTask<(BlobLocator, Uri)?> TryGetWriteRedirectAsync(string? prefix = null, CancellationToken cancellationToken = default) => default;
+
+		#endregion
+
+		#region Aliases
 
 		/// <inheritdoc/>
-		public Task WriteExplicitPathAsync(string path, Stream stream, CancellationToken cancellationToken = default) => throw new NotSupportedException();
+		public override Task AddAliasAsync(string name, BlobHandle target, int rank, ReadOnlyMemory<byte> data, CancellationToken cancellationToken = default) => throw new NotSupportedException();
 
 		/// <inheritdoc/>
-		public Task DeleteAsync(string path, CancellationToken cancellationToken = default) => throw new NotSupportedException();
+		public override Task RemoveAliasAsync(string name, BlobHandle target, CancellationToken cancellationToken = default) => throw new NotSupportedException();
 
 		/// <inheritdoc/>
-		public IAsyncEnumerable<string> EnumerateAsync(CancellationToken cancellationToken = default) => throw new NotSupportedException();
+		public override Task<BlobAlias[]> FindAliasesAsync(string name, int? maxResults, CancellationToken cancellationToken = default) => throw new NotSupportedException();
+
+		#endregion
+
+		#region Refs
 
 		/// <inheritdoc/>
-		public Task<bool> ExistsAsync(string path, CancellationToken cancellationToken = default) => throw new NotSupportedException();
+		public override Task<bool> DeleteRefAsync(RefName name, CancellationToken cancellationToken = default) => throw new NotSupportedException();
 
 		/// <inheritdoc/>
-		public ValueTask<Uri?> TryGetReadRedirectAsync(string path, CancellationToken cancellationToken = default) => throw new NotSupportedException();
+		public override Task<BlobHandle?> TryReadRefAsync(RefName name, RefCacheTime cacheTime = default, CancellationToken cancellationToken = default) => throw new NotSupportedException();
 
 		/// <inheritdoc/>
-		public ValueTask<(string, Uri)?> TryGetWriteRedirectAsync(string? prefix = null, CancellationToken cancellationToken = default) => throw new NotSupportedException();
+		public override Task WriteRefAsync(RefName name, BlobHandle target, RefOptions? options = null, CancellationToken cancellationToken = default) => throw new NotSupportedException();
+
+		#endregion
 
 		/// <inheritdoc/>
-		public void GetStats(StorageStats stats) { }
+		public override void GetStats(StorageStats stats)
+		{
+		}
 	}
 }
