Horde: Make HttpStorageClient a direct implementation of IStorageClient rather than deriving from BundleStorageClient.

diff --git a/Engine/Source/Programs/AutomationTool/BuildGraph/Tasks/DeployToolTask.cs b/Engine/Source/Programs/AutomationTool/BuildGraph/Tasks/DeployToolTask.cs
--- a/Engine/Source/Programs/AutomationTool/BuildGraph/Tasks/DeployToolTask.cs
+++ b/Engine/Source/Programs/AutomationTool/BuildGraph/Tasks/DeployToolTask.cs
@@ -153,8 +153,9 @@ namespace AutomationTool.Tasks
 			BundleNodeHandle handle;
 
 			string basePath = $"api/v1/tools/{Parameters.Id}";
-			HttpStorageBackend storageBackend = new HttpStorageBackend(basePath, CreateHttpClient, Logger);
-			HttpStorageClient storageClient = new HttpStorageClient(basePath, CreateHttpClient, storageBackend, BundleReaderCache.None, Logger);
+			using HttpStorageBackend httpStorageBackend = new HttpStorageBackend(basePath, CreateHttpClient, Logger);
+			using HttpStorageClient httpStorageClient = new HttpStorageClient(basePath, CreateHttpClient, httpStorageBackend, Logger);
+			using BundleStorageClientWrapper storageClient = new BundleStorageClientWrapper(httpStorageClient, BundleReaderCache.None, Logger);
 
 			await using (BundleWriter treeWriter = storageClient.CreateWriter())
 			{
diff --git a/Engine/Source/Programs/Horde/Horde.Server/Storage/StorageService.cs b/Engine/Source/Programs/Horde/Horde.Server/Storage/StorageService.cs
--- a/Engine/Source/Programs/Horde/Horde.Server/Storage/StorageService.cs
+++ b/Engine/Source/Programs/Horde/Horde.Server/Storage/StorageService.cs
@@ -51,28 +51,6 @@ namespace Horde.Server.Storage
 	/// </summary>
 	public sealed class StorageService : IHostedService, IStorageClientFactory, IAsyncDisposable
 	{
-		/// <summary>
-		/// Type used for blobs that only contain leaf data.
-		/// </summary>
-		static readonly BlobType s_leafBlobType = new BlobType(Guid.Parse("{9E88B4E3-E4C5-445E-A14C-9698A30BF8C5}"), 1);
-
-		sealed class LeafBlobData : BlobData
-		{
-			readonly IReadOnlyMemoryOwner<byte> _obj;
-
-			public LeafBlobData(IReadOnlyMemoryOwner<byte> obj)
-				: base(s_leafBlobType, obj.Memory, Array.Empty<BlobHandle>())
-			{
-				_obj = obj;
-			}
-
-			protected override void Dispose(bool disposing)
-			{
-				_obj.Dispose();
-				base.Dispose(disposing);
-			}
-		}
-
 		sealed class LeafBlobHandle : BlobHandle
 		{
 			readonly IStorageBackend _backend;
@@ -87,7 +65,7 @@ namespace Horde.Server.Storage
 			}
 
 			/// <inheritdoc/>
-			public override ValueTask<BlobType> GetTypeAsync(CancellationToken cancellationToken = default) => new ValueTask<BlobType>(s_leafBlobType);
+			public override ValueTask<BlobType> GetTypeAsync(CancellationToken cancellationToken = default) => new ValueTask<BlobType>(BlobType.Leaf);
 
 			/// <inheritdoc/>
 			public override async Task<Stream> OpenAsync(int offset = 0, int? length = null, CancellationToken cancellationToken = default)
@@ -112,7 +90,7 @@ namespace Horde.Server.Storage
 			public override async ValueTask<BlobData> ReadAsync(CancellationToken cancellationToken = default)
 			{
 				IReadOnlyMemoryOwner<byte> obj = await _backend.ReadAsync(_path, cancellationToken);
-				return new LeafBlobData(obj);
+				return new ReadOnlyMemoryOwnerBlobData(BlobType.Leaf, obj, Array.Empty<BlobHandle>());
 			}
 
 			/// <inheritdoc/>
diff --git a/Engine/Source/Programs/Shared/EpicGames.Horde/Storage/BlobData.cs b/Engine/Source/Programs/Shared/EpicGames.Horde/Storage/BlobData.cs
--- a/Engine/Source/Programs/Shared/EpicGames.Horde/Storage/BlobData.cs
+++ b/Engine/Source/Programs/Shared/EpicGames.Horde/Storage/BlobData.cs
@@ -2,6 +2,7 @@
 
 using System;
 using System.Collections.Generic;
+using EpicGames.Core;
 
 namespace EpicGames.Horde.Storage
 {
@@ -50,4 +51,32 @@ namespace EpicGames.Horde.Storage
 		{
 		}
 	}
+
+	/// <summary>
+	/// Implementation of <see cref="BlobData"/> for <see cref="IReadOnlyMemoryOwner{Byte}"/> instances.
+	/// </summary>
+	public class ReadOnlyMemoryOwnerBlobData : BlobData
+	{
+		readonly IReadOnlyMemoryOwner<byte> _owner;
+
+		/// <summary>
+		/// Constructor
+		/// </summary>
+		public ReadOnlyMemoryOwnerBlobData(BlobType type, IReadOnlyMemoryOwner<byte> owner, IReadOnlyList<BlobHandle> refs)
+			: base(type, owner.Memory, refs)
+		{
+			_owner = owner;
+		}
+
+		/// <inheritdoc/>
+		protected override void Dispose(bool disposing)
+		{
+			base.Dispose(disposing);
+
+			if (disposing)
+			{
+				_owner.Dispose();
+			}
+		}
+	}
 }
diff --git a/Engine/Source/Programs/Shared/EpicGames.Horde/Storage/Clients/FileStorageClient.cs b/Engine/Source/Programs/Shared/EpicGames.Horde/Storage/Clients/FileStorageClient.cs
--- a/Engine/Source/Programs/Shared/EpicGames.Horde/Storage/Clients/FileStorageClient.cs
+++ b/Engine/Source/Programs/Shared/EpicGames.Horde/Storage/Clients/FileStorageClient.cs
@@ -16,27 +16,6 @@ namespace EpicGames.Horde.Storage.Clients
 	/// </summary>
 	public class FileStorageClient : KeyValueStorageClient
 	{
-		class LeafBlobData : BlobData
-		{
-			readonly IReadOnlyMemoryOwner<byte> _owner;
-
-			public LeafBlobData(IReadOnlyMemoryOwner<byte> owner)
-				: base(BlobType.Leaf, owner.Memory, Array.Empty<BlobHandle>())
-			{
-				_owner = owner;
-			}
-
-			protected override void Dispose(bool disposing)
-			{
-				base.Dispose(disposing);
-
-				if (disposing)
-				{
-					_owner.Dispose();
-				}
-			}
-		}
-
 		readonly DirectoryReference _rootDir;
 		readonly FileStorageBackend _backend;
 		readonly ILogger _logger;
@@ -81,11 +60,12 @@ namespace EpicGames.Horde.Storage.Clients
 		FileReference GetRefFile(RefName name) => FileReference.Combine(_rootDir, name.ToString() + ".ref");
 
 		#region Blobs
+
 		/// <inheritdoc/>
 		public override async ValueTask<BlobData> ReadBlobAsync(BlobLocator locator, CancellationToken cancellationToken = default)
 		{
 			IReadOnlyMemoryOwner<byte> owner = await _backend.ReadAsync(locator.ToString(), cancellationToken);
-			return new LeafBlobData(owner);
+			return new ReadOnlyMemoryOwnerBlobData(BlobType.Leaf, owner, Array.Empty<BlobHandle>());
 		}
 
 		/// <inheritdoc/>
diff --git a/Engine/Source/Programs/Shared/EpicGames.Horde/Storage/Clients/HttpStorageClient.cs b/Engine/Source/Programs/Shared/EpicGames.Horde/Storage/Clients/HttpStorageClient.cs
--- a/Engine/Source/Programs/Shared/EpicGames.Horde/Storage/Clients/HttpStorageClient.cs
+++ b/Engine/Source/Programs/Shared/EpicGames.Horde/Storage/Clients/HttpStorageClient.cs
@@ -11,45 +11,132 @@ using System.Web;
 using EpicGames.Core;
 using Microsoft.Extensions.Logging;
 using EpicGames.Horde.Storage.Backends;
+using System.IO;
+using System.Collections.Generic;
+using System.Diagnostics.CodeAnalysis;
 
 namespace EpicGames.Horde.Storage.Clients
 {
 	/// <summary>
 	/// Implementation of <see cref="IStorageClient"/> which communicates with an upstream Horde instance via HTTP.
 	/// </summary>
-	public class HttpStorageClient : BundleStorageClient
+	public sealed class HttpStorageClient : IStorageClient
 	{
+		class Handle : BlobHandle
+		{
+			readonly HttpStorageClient _outer;
+			readonly BlobLocator _locator;
+
+			public Handle(HttpStorageClient outer, BlobLocator locator)
+			{
+				_outer = outer;
+				_locator = locator;
+			}
+
+			public override Task<Stream> OpenAsync(int offset = 0, int? length = null, CancellationToken cancellationToken = default)
+				=> _outer._backend.OpenAsync(_locator.ToString(), offset, length, cancellationToken);
+
+			public override ValueTask<BlobData> ReadAsync(CancellationToken cancellationToken = default)
+				=> _outer.ReadBlobAsync(_locator, cancellationToken);
+
+			public override bool TryGetLocator([NotNullWhen(true)] out BlobLocator locator)
+			{
+				locator = _locator;
+				return true;
+			}
+		}
+
 		readonly string _basePath;
 		readonly Func<HttpClient> _createClient;
+		readonly IStorageBackend _backend;
 		readonly ILogger _logger;
 
+		/// <inheritdoc/>
+		public bool SupportsRedirects => _backend.SupportsRedirects;
+
 		/// <summary>
 		/// Constructor
 		/// </summary>
-		public HttpStorageClient(string basePath, Func<HttpClient> createClient, IStorageBackend backend, BundleReaderCache cache, ILogger logger) 
-			: base(backend, cache, logger)
+		public HttpStorageClient(string basePath, Func<HttpClient> createClient, IStorageBackend backend, ILogger logger) 
 		{
 			_basePath = basePath.TrimEnd('/');
 			_createClient = createClient;
+			_backend = backend;
 			_logger = logger;
 		}
 
-		#region Nodes
+		/// <inheritdoc/>
+		public void Dispose()
+		{
+			_backend.Dispose();
+		}
+
+		#region Blobs
 
 		/// <inheritdoc/>
-		public override Task AddAliasAsync(string name, BlobHandle target, int rank = 0, ReadOnlyMemory<byte> data = default, CancellationToken cancellationToken = default)
+		public BlobHandle CreateBlobHandle(BlobLocator locator)
+		{
+			if (locator.TryUnwrapFull(out BlobLocator outer, out Utf8String fragment))
+			{
+				return new BlobFragmentHandle(new Handle(this, outer), fragment);
+			}
+			else
+			{
+				return new Handle(this, locator);
+			}
+		}
+
+		/// <inheritdoc/>
+		public IStorageWriter CreateWriter(string? basePath = null)
+			=> new DefaultStorageWriter(this, basePath);
+
+		/// <inheritdoc/>
+		public async ValueTask<BlobData> ReadBlobAsync(BlobLocator locator, CancellationToken cancellationToken = default)
+		{
+			IReadOnlyMemoryOwner<byte> owner = await _backend.ReadAsync(locator.ToString(), cancellationToken);
+			return new ReadOnlyMemoryOwnerBlobData(BlobType.Leaf, owner, Array.Empty<BlobHandle>());
+		}
+
+		/// <inheritdoc/>
+		public async ValueTask<BlobHandle> WriteBlobAsync(BlobType type, Stream stream, IReadOnlyList<BlobHandle> references, string? basePath = null, CancellationToken cancellationToken = default)
+		{
+			string path = await _backend.WriteAsync(stream, basePath, cancellationToken);
+			return new Handle(this, new BlobLocator(path));
+		}
+
+		/// <inheritdoc/>
+		public ValueTask<Uri?> TryGetReadRedirectAsync(BlobLocator locator, CancellationToken cancellationToken = default)
+			=> _backend.TryGetReadRedirectAsync(locator.ToString(), cancellationToken);
+
+		/// <inheritdoc/>
+		public async ValueTask<(BlobLocator, Uri)?> TryGetWriteRedirectAsync(string? prefix = null, CancellationToken cancellationToken = default)
+		{
+			(string Path, Uri Url)? result = await _backend.TryGetWriteRedirectAsync(prefix, cancellationToken);
+			if (result == null)
+			{
+				return null;
+			}
+			return (new BlobLocator(result.Value.Path), result.Value.Url);
+		}
+
+		#endregion
+
+		#region Aliases
+
+		/// <inheritdoc/>
+		public Task AddAliasAsync(string name, BlobHandle target, int rank = 0, ReadOnlyMemory<byte> data = default, CancellationToken cancellationToken = default)
 		{
 			throw new NotSupportedException("Http storage client does not currently support aliases.");
 		}
 
 		/// <inheritdoc/>
-		public override Task RemoveAliasAsync(string name, BlobHandle target, CancellationToken cancellationToken = default)
+		public Task RemoveAliasAsync(string name, BlobHandle target, CancellationToken cancellationToken = default)
 		{
 			throw new NotSupportedException("Http storage client does not currently support aliases.");
 		}
 
 		/// <inheritdoc/>
-		public override async Task<BlobAlias[]> FindAliasesAsync(string alias, int? maxResults = null, CancellationToken cancellationToken = default)
+		public async Task<BlobAlias[]> FindAliasesAsync(string alias, int? maxResults = null, CancellationToken cancellationToken = default)
 		{
 			_logger.LogDebug("Finding nodes with alias {Alias}", alias);
 			using (HttpClient httpClient = _createClient())
@@ -87,7 +174,7 @@ namespace EpicGames.Horde.Storage.Clients
 		#region Refs
 
 		/// <inheritdoc/>
-		public override async Task<bool> DeleteRefAsync(RefName name, CancellationToken cancellationToken)
+		public async Task<bool> DeleteRefAsync(RefName name, CancellationToken cancellationToken)
 		{
 			_logger.LogDebug("Deleting ref {RefName}", name);
 			using (HttpClient httpClient = _createClient())
@@ -113,7 +200,7 @@ namespace EpicGames.Horde.Storage.Clients
 		}
 
 		/// <inheritdoc/>
-		public override async Task<BlobHandle?> TryReadRefAsync(RefName name, RefCacheTime cacheTime = default, CancellationToken cancellationToken = default)
+		public async Task<BlobHandle?> TryReadRefAsync(RefName name, RefCacheTime cacheTime = default, CancellationToken cancellationToken = default)
 		{
 			using (HttpClient httpClient = _createClient())
 			{
@@ -150,7 +237,7 @@ namespace EpicGames.Horde.Storage.Clients
 		}
 
 		/// <inheritdoc/>
-		public override async Task WriteRefAsync(RefName name, BlobHandle target, RefOptions? options = null, CancellationToken cancellationToken = default)
+		public async Task WriteRefAsync(RefName name, BlobHandle target, RefOptions? options = null, CancellationToken cancellationToken = default)
 		{
 			await target.FlushAsync(cancellationToken);
 			BlobLocator locator = target.GetLocator();
@@ -170,6 +257,11 @@ namespace EpicGames.Horde.Storage.Clients
 		}
 
 		#endregion
+
+		/// <inheritdoc/>
+		public void GetStats(StorageStats stats)
+		{
+		}
 	}
 
 	/// <summary>
@@ -200,7 +292,7 @@ namespace EpicGames.Horde.Storage.Clients
 		/// </summary>
 		/// <param name="basePath">Base path for all requests</param>
 		/// <param name="accessToken">Custom access token to use for requests</param>
-		public HttpStorageClient CreateClientWithPath(string basePath, string? accessToken = null)
+		public IStorageClient CreateClientWithPath(string basePath, string? accessToken = null)
 		{
 			HttpClient CreateClient()
 			{
@@ -217,7 +309,9 @@ namespace EpicGames.Horde.Storage.Clients
 			{
 				backend = _backendCache.CreateWrapper(basePath, backend);
 			}
-			return new HttpStorageClient(basePath, CreateClient, backend, _readerCache, _clientLogger);
+
+			HttpStorageClient client = new HttpStorageClient(basePath, CreateClient, backend, _clientLogger);
+			return new BundleStorageClientWrapper(client, _readerCache, _clientLogger);
 		}
 
 		/// <summary>
@@ -225,7 +319,7 @@ namespace EpicGames.Horde.Storage.Clients
 		/// </summary>
 		/// <param name="namespaceId">Namespace to create a client for</param>
 		/// <param name="accessToken">Custom access token to use for requests</param>
-		public HttpStorageClient CreateClient(NamespaceId namespaceId, string? accessToken = null) => CreateClientWithPath($"api/v1/storage/{namespaceId}", accessToken);
+		public IStorageClient CreateClient(NamespaceId namespaceId, string? accessToken = null) => CreateClientWithPath($"api/v1/storage/{namespaceId}", accessToken);
 
 		/// <inheritdoc/>
 		IStorageClient? IStorageClientFactory.TryCreateClient(NamespaceId namespaceId) => CreateClient(namespaceId);
