Horde: Unify BundleStorageClient and BundleStorageClientWrapper.

diff --git a/Engine/Source/Programs/AutomationTool/BuildGraph/Tasks/DeployToolTask.cs b/Engine/Source/Programs/AutomationTool/BuildGraph/Tasks/DeployToolTask.cs
--- a/Engine/Source/Programs/AutomationTool/BuildGraph/Tasks/DeployToolTask.cs
+++ b/Engine/Source/Programs/AutomationTool/BuildGraph/Tasks/DeployToolTask.cs
@@ -155,7 +155,7 @@ namespace AutomationTool.Tasks
 			string basePath = $"api/v1/tools/{Parameters.Id}";
 			using HttpStorageBackend httpStorageBackend = new HttpStorageBackend(basePath, CreateHttpClient, Logger);
 			using HttpStorageClient httpStorageClient = new HttpStorageClient(basePath, CreateHttpClient, httpStorageBackend, Logger);
-			using BundleStorageClientWrapper storageClient = new BundleStorageClientWrapper(httpStorageClient, BundleReaderCache.None, Logger);
+			using BundleStorageClient storageClient = new BundleStorageClient(httpStorageClient, BundleReaderCache.None, Logger);
 
 			await using (BundleWriter treeWriter = storageClient.CreateWriter())
 			{
diff --git a/Engine/Source/Programs/Horde/Horde.Server/Commands/Bundles/CreateCommand.cs b/Engine/Source/Programs/Horde/Horde.Server/Commands/Bundles/CreateCommand.cs
--- a/Engine/Source/Programs/Horde/Horde.Server/Commands/Bundles/CreateCommand.cs
+++ b/Engine/Source/Programs/Horde/Horde.Server/Commands/Bundles/CreateCommand.cs
@@ -47,7 +47,7 @@ namespace Horde.Server.Commands.Bundles
 			}
 			else
 			{
-				using IStorageClient storageClient = BundleStorageClientWrapper.CreateFromDirectory(OutputDir, serviceProvider.GetRequiredService<BundleReaderCache>(), logger);
+				using IStorageClient storageClient = BundleStorageClient.CreateFromDirectory(OutputDir, serviceProvider.GetRequiredService<BundleReaderCache>(), logger);
 				await ExecuteInternalAsync(storageClient, logger);
 			}
 
diff --git a/Engine/Source/Programs/Horde/Horde.Server/Commands/Install/SetupCommand.cs b/Engine/Source/Programs/Horde/Horde.Server/Commands/Install/SetupCommand.cs
--- a/Engine/Source/Programs/Horde/Horde.Server/Commands/Install/SetupCommand.cs
+++ b/Engine/Source/Programs/Horde/Horde.Server/Commands/Install/SetupCommand.cs
@@ -60,7 +60,7 @@ namespace Horde.Server.Commands.Install
 
 			using BundleReaderCache bundleReaderCache = new BundleReaderCache();
 
-			using IStorageClient client = BundleStorageClientWrapper.CreateFromDirectory(bundleDir, bundleReaderCache, logger);
+			using IStorageClient client = BundleStorageClient.CreateFromDirectory(bundleDir, bundleReaderCache, logger);
 			await using (IStorageWriter writer = client.CreateWriter(refName))
 			{
 				DirectoryNode dirNode = new DirectoryNode();
diff --git a/Engine/Source/Programs/Horde/Horde.Server/Storage/StorageService.cs b/Engine/Source/Programs/Horde/Horde.Server/Storage/StorageService.cs
--- a/Engine/Source/Programs/Horde/Horde.Server/Storage/StorageService.cs
+++ b/Engine/Source/Programs/Horde/Horde.Server/Storage/StorageService.cs
@@ -609,7 +609,7 @@ namespace Horde.Server.Storage
 
 #pragma warning disable CA2000 // Dispose objects before losing scope
 							client = new StorageClientImpl(this, namespaceConfig, backend, _tracer);
-							client = new BundleStorageClientWrapper(client, _bundleReaderCache, _logger);
+							client = new BundleStorageClient(client, _bundleReaderCache, _logger);
 #pragma warning restore CA2000 // Dispose objects before losing scope
 
 							NamespaceInfo namespaceInfo = new NamespaceInfo(namespaceConfig, backend, new SharedStorageClient(client));
diff --git a/Engine/Source/Programs/Horde/Horde.Server/Tools/ToolCollection.cs b/Engine/Source/Programs/Horde/Horde.Server/Tools/ToolCollection.cs
--- a/Engine/Source/Programs/Horde/Horde.Server/Tools/ToolCollection.cs
+++ b/Engine/Source/Programs/Horde/Horde.Server/Tools/ToolCollection.cs
@@ -425,7 +425,7 @@ namespace Horde.Server.Tools
 		{
 			if (tool.Config is BundledToolConfig bundledConfig)
 			{
-				return BundleStorageClientWrapper.CreateFromDirectory(DirectoryReference.Combine(ServerApp.AppDir, bundledConfig.DataDir ?? $"tools/{tool.Id}"), _cache, _logger);
+				return BundleStorageClient.CreateFromDirectory(DirectoryReference.Combine(ServerApp.AppDir, bundledConfig.DataDir ?? $"tools/{tool.Id}"), _cache, _logger);
 			}
 			else
 			{
diff --git a/Engine/Source/Programs/Horde/Horde/Commands/Bundles/BundleCreate.cs b/Engine/Source/Programs/Horde/Horde/Commands/Bundles/BundleCreate.cs
--- a/Engine/Source/Programs/Horde/Horde/Commands/Bundles/BundleCreate.cs
+++ b/Engine/Source/Programs/Horde/Horde/Commands/Bundles/BundleCreate.cs
@@ -34,7 +34,7 @@ namespace Horde.Commands.Bundles
 		{
 			if (File != null)
 			{
-				using IStorageClient store = BundleStorageClientWrapper.CreateFromDirectory(File.Directory, BundleReaderCache, logger);
+				using IStorageClient store = BundleStorageClient.CreateFromDirectory(File.Directory, BundleReaderCache, logger);
 				return await ExecuteInternalAsync(store, logger);
 			}
 			else if (!String.IsNullOrEmpty(Ref))
diff --git a/Engine/Source/Programs/Horde/Horde/Commands/Bundles/BundleExtract.cs b/Engine/Source/Programs/Horde/Horde/Commands/Bundles/BundleExtract.cs
--- a/Engine/Source/Programs/Horde/Horde/Commands/Bundles/BundleExtract.cs
+++ b/Engine/Source/Programs/Horde/Horde/Commands/Bundles/BundleExtract.cs
@@ -37,7 +37,7 @@ namespace Horde.Commands.Bundles
 		{
 			if (File != null)
 			{
-				using IStorageClient store = BundleStorageClientWrapper.CreateFromDirectory(File.Directory, BundleReaderCache, logger);
+				using IStorageClient store = BundleStorageClient.CreateFromDirectory(File.Directory, BundleReaderCache, logger);
 				BlobHandle handle = store.CreateBlobHandle(await FileStorageClient.ReadRefAsync(File));
 				await ExecuteInternalAsync(store, handle, logger);
 			}
diff --git a/Engine/Source/Programs/Horde/Horde/Commands/Workspace/WorkspaceSync.cs b/Engine/Source/Programs/Horde/Horde/Commands/Workspace/WorkspaceSync.cs
--- a/Engine/Source/Programs/Horde/Horde/Commands/Workspace/WorkspaceSync.cs
+++ b/Engine/Source/Programs/Horde/Horde/Commands/Workspace/WorkspaceSync.cs
@@ -42,7 +42,7 @@ namespace Horde.Commands.Workspace
 		{
 			if (File != null)
 			{
-				using IStorageClient store = BundleStorageClientWrapper.CreateFromDirectory(File.Directory, BundleReaderCache, logger);
+				using IStorageClient store = BundleStorageClient.CreateFromDirectory(File.Directory, BundleReaderCache, logger);
 				BlobHandle handle = store.CreateBlobHandle(await FileStorageClient.ReadRefAsync(File));
 				return await ExecuteInternalAsync(store, handle, logger);
 			}
diff --git a/Engine/Source/Programs/Shared/EpicGames.Horde.Tests/BundleStoreTests.cs b/Engine/Source/Programs/Shared/EpicGames.Horde.Tests/BundleStoreTests.cs
--- a/Engine/Source/Programs/Shared/EpicGames.Horde.Tests/BundleStoreTests.cs
+++ b/Engine/Source/Programs/Shared/EpicGames.Horde.Tests/BundleStoreTests.cs
@@ -60,7 +60,7 @@ namespace EpicGames.Horde.Tests
 		static async Task<Bundle> CreateBundleNormalAsync()
 		{
 			using MemoryStorageClient memoryStore = new MemoryStorageClient();
-			using BundleStorageClientWrapper store = new BundleStorageClientWrapper(memoryStore, BundleReaderCache.None, NullLogger.Instance);
+			using BundleStorageClient store = new BundleStorageClient(memoryStore, BundleReaderCache.None, NullLogger.Instance);
 			await using BundleWriter writer = store.CreateWriter(options: new BundleOptions { CompressionFormat = BundleCompressionFormat.None });
 
 			TextNode node = new TextNode("Hello world");
@@ -97,7 +97,7 @@ namespace EpicGames.Horde.Tests
 			using IMemoryCache cache = new MemoryCache(new MemoryCacheOptions());
 
 			using MemoryStorageClient blobStore = new MemoryStorageClient();
-			using BundleStorageClientWrapper bundleStore = new BundleStorageClientWrapper(blobStore, BundleReaderCache.None, NullLogger.Instance);
+			using BundleStorageClient bundleStore = new BundleStorageClient(blobStore, BundleReaderCache.None, NullLogger.Instance);
 
 			await TestTreeAsync(bundleStore, new BundleOptions { MaxBlobSize = 1024 * 1024 });
 
@@ -111,7 +111,7 @@ namespace EpicGames.Horde.Tests
 			using IMemoryCache cache = new MemoryCache(new MemoryCacheOptions());
 
 			using MemoryStorageClient blobStore = new MemoryStorageClient();
-			using BundleStorageClientWrapper bundleStore = new BundleStorageClientWrapper(blobStore, BundleReaderCache.None, NullLogger.Instance);
+			using BundleStorageClient bundleStore = new BundleStorageClient(blobStore, BundleReaderCache.None, NullLogger.Instance);
 
 			await TestTreeAsync(bundleStore, new BundleOptions { MaxBlobSize = 1 });
 
@@ -144,7 +144,7 @@ namespace EpicGames.Horde.Tests
 			}
 		}
 
-		static async Task TestTreeAsync(BundleStorageClientBase store, BundleOptions options)
+		static async Task TestTreeAsync(BundleStorageClient store, BundleOptions options)
 		{
 			// Generate a tree
 			{
@@ -315,7 +315,7 @@ namespace EpicGames.Horde.Tests
 		public async Task StreamTestAsync()
 		{
 			using MemoryStorageClient memoryStore = new MemoryStorageClient();
-			using BundleStorageClientWrapper store = new BundleStorageClientWrapper(memoryStore, BundleReaderCache.None, NullLogger.Instance);
+			using BundleStorageClient store = new BundleStorageClient(memoryStore, BundleReaderCache.None, NullLogger.Instance);
 
 			const int Length = 4096;
 
@@ -356,7 +356,7 @@ namespace EpicGames.Horde.Tests
 		{
 			using IMemoryCache cache = new MemoryCache(new MemoryCacheOptions());
 			using MemoryStorageClient memoryStore = new MemoryStorageClient();
-			using BundleStorageClientWrapper store = new BundleStorageClientWrapper(memoryStore, BundleReaderCache.None, NullLogger.Instance);
+			using BundleStorageClient store = new BundleStorageClient(memoryStore, BundleReaderCache.None, NullLogger.Instance);
 
 			const int Length = 1024;
 			const int Copies = 4096;
diff --git a/Engine/Source/Programs/Shared/EpicGames.Horde.Tests/BundleTests.cs b/Engine/Source/Programs/Shared/EpicGames.Horde.Tests/BundleTests.cs
--- a/Engine/Source/Programs/Shared/EpicGames.Horde.Tests/BundleTests.cs
+++ b/Engine/Source/Programs/Shared/EpicGames.Horde.Tests/BundleTests.cs
@@ -22,13 +22,13 @@ namespace EpicGames.Horde.Tests
 	{
 		readonly IMemoryCache _cache;
 		readonly MemoryStorageClient _memoryStore;
-		readonly BundleStorageClientBase _storage;
+		readonly BundleStorageClient _storage;
 
 		public BundleTests()
 		{
 			_cache = new MemoryCache(new MemoryCacheOptions());
 			_memoryStore = new MemoryStorageClient();
-			_storage = new BundleStorageClientWrapper(_memoryStore, BundleReaderCache.None, NullLogger.Instance);
+			_storage = new BundleStorageClient(_memoryStore, BundleReaderCache.None, NullLogger.Instance);
 		}
 
 		public void Dispose()
diff --git a/Engine/Source/Programs/Shared/EpicGames.Horde.Tests/ComputeSocketTests.cs b/Engine/Source/Programs/Shared/EpicGames.Horde.Tests/ComputeSocketTests.cs
--- a/Engine/Source/Programs/Shared/EpicGames.Horde.Tests/ComputeSocketTests.cs
+++ b/Engine/Source/Programs/Shared/EpicGames.Horde.Tests/ComputeSocketTests.cs
@@ -141,7 +141,7 @@ namespace EpicGames.Horde.Tests
 					}
 
 					using MemoryStorageClient memoryStorage = new MemoryStorageClient();
-					using BundleStorageClientWrapper storage = new BundleStorageClientWrapper(memoryStorage, BundleReaderCache.None, NullLogger.Instance);
+					using BundleStorageClient storage = new BundleStorageClient(memoryStorage, BundleReaderCache.None, NullLogger.Instance);
 					await using (BundleWriter treeWriter = storage.CreateWriter())
 					{
 						FileReference file = FileReference.Combine(tempDir, "subdir/hello.txt");
diff --git a/Engine/Source/Programs/Shared/EpicGames.Horde/Compute/AgentMessageHandler.cs b/Engine/Source/Programs/Shared/EpicGames.Horde/Compute/AgentMessageHandler.cs
--- a/Engine/Source/Programs/Shared/EpicGames.Horde/Compute/AgentMessageHandler.cs
+++ b/Engine/Source/Programs/Shared/EpicGames.Horde/Compute/AgentMessageHandler.cs
@@ -136,7 +136,7 @@ namespace EpicGames.Horde.Compute
 		async Task WriteFilesAsync(AgentMessageChannel channel, string path, BlobLocator locator, CancellationToken cancellationToken)
 		{
 			using AgentStorageClient innerStore = new AgentStorageClient(channel);
-			using BundleStorageClientWrapper store = new BundleStorageClientWrapper(innerStore, BundleReaderCache.None, _logger);
+			using BundleStorageClient store = new BundleStorageClient(innerStore, BundleReaderCache.None, _logger);
 
 			BlobHandle handle = store.CreateBlobHandle(locator);
 			DirectoryNode directoryNode = await handle.ReadNodeAsync<DirectoryNode>(cancellationToken);
diff --git a/Engine/Source/Programs/Shared/EpicGames.Horde/Storage/Clients/BundleStorageClient.cs b/Engine/Source/Programs/Shared/EpicGames.Horde/Storage/Clients/BundleStorageClient.cs
--- a/Engine/Source/Programs/Shared/EpicGames.Horde/Storage/Clients/BundleStorageClient.cs
+++ b/Engine/Source/Programs/Shared/EpicGames.Horde/Storage/Clients/BundleStorageClient.cs
@@ -15,7 +15,7 @@ namespace EpicGames.Horde.Storage.Clients
 	/// <summary>
 	/// Base class for an implementation of <see cref="IStorageClient"/>, providing implementations for some common functionality using bundles.
 	/// </summary>
-	public abstract class BundleStorageClientBase : IStorageClient
+	public sealed class BundleStorageClient : IStorageClient
 	{
 		/// <summary>
 		/// Blob type for bundles
@@ -25,16 +25,16 @@ namespace EpicGames.Horde.Storage.Clients
 		/// <summary>
 		/// Handle to a bundle object
 		/// </summary>
-		protected class BundleHandle : BlobHandle
+		class BundleHandle : BlobHandle
 		{
-			readonly BundleStorageClientBase _storageClient;
+			readonly BundleStorageClient _storageClient;
 			readonly BlobLocator _locator;
 			List<BlobHandle>? _refs;
 
 			/// <summary>
 			/// Constructor
 			/// </summary>
-			public BundleHandle(BundleStorageClientBase storageClient, BlobLocator locator)
+			public BundleHandle(BundleStorageClient storageClient, BlobLocator locator)
 			{
 				_storageClient = storageClient;
 				_locator = locator;
@@ -85,32 +85,43 @@ namespace EpicGames.Horde.Storage.Clients
 			}
 		}
 
+		readonly IStorageClient _inner;
 		readonly BundleReader _bundleReader;
 
 		/// <inheritdoc/>
-		public virtual bool SupportsRedirects { get; } = false;
+		public bool SupportsRedirects { get; } = false;
 
 		/// <summary>
 		/// Constructor
 		/// </summary>
-		protected BundleStorageClientBase(BundleReaderCache cache, ILogger logger)
+		public BundleStorageClient(IStorageClient inner, BundleReaderCache cache, ILogger logger)
 		{
+			_inner = inner;
 			_bundleReader = new BundleReader(this, cache, logger);
 		}
 
 		/// <inheritdoc/>
 		public void Dispose()
 		{
-			Dispose(true);
-			GC.SuppressFinalize(this);
+			_inner.Dispose();
 		}
 
 		/// <summary>
-		/// Overridable dispose method
+		/// Creates a bundle storage client around a memory client backend
 		/// </summary>
-		/// <param name="disposing"></param>
-		protected virtual void Dispose(bool disposing)
+		public static BundleStorageClient CreateFromMemory(ILogger logger)
 		{
+			MemoryStorageClient blobStore = new MemoryStorageClient();
+			return new BundleStorageClient(blobStore, BundleReaderCache.None, logger);
+		}
+
+		/// <summary>
+		/// Creates a bundle storage client around a directory on the filesystem
+		/// </summary>
+		public static BundleStorageClient CreateFromDirectory(DirectoryReference rootDir, BundleReaderCache cache, ILogger logger)
+		{
+			FileStorageClient fileStorageClient = new FileStorageClient(rootDir, logger);
+			return new BundleStorageClient(fileStorageClient, cache, logger);
 		}
 
 		#region Blobs
@@ -153,10 +164,10 @@ namespace EpicGames.Horde.Storage.Clients
 		}
 
 		/// <inheritdoc/>
-		public virtual ValueTask<Uri?> TryGetReadRedirectAsync(BlobLocator locator, CancellationToken cancellationToken = default) => new ValueTask<Uri?>();
+		public ValueTask<Uri?> TryGetReadRedirectAsync(BlobLocator locator, CancellationToken cancellationToken = default) => new ValueTask<Uri?>();
 
 		/// <inheritdoc/>
-		public virtual ValueTask<(BlobLocator, Uri)?> TryGetWriteRedirectAsync(string? prefix = null, CancellationToken cancellationToken = default) => new ValueTask<(BlobLocator, Uri)?>();
+		public ValueTask<(BlobLocator, Uri)?> TryGetWriteRedirectAsync(string? prefix = null, CancellationToken cancellationToken = default) => new ValueTask<(BlobLocator, Uri)?>();
 
 		#endregion
 
@@ -165,12 +176,14 @@ namespace EpicGames.Horde.Storage.Clients
 		/// <summary>
 		/// Read a bundle from the underlying storage
 		/// </summary>
-		protected abstract Task<Stream> OpenBundleAsync(BlobLocator locator, int offset = 0, int? length = null, CancellationToken cancellationToken = default);
+		async Task<Stream> OpenBundleAsync(BlobLocator locator, int offset, int? length, CancellationToken cancellationToken)
+			=> await _inner.CreateBlobHandle(locator).OpenAsync(offset, length, cancellationToken);
 
 		/// <summary>
 		/// Write a bundle to the underlying storage
 		/// </summary>
-		protected abstract ValueTask<BlobHandle> WriteBundleAsync(Stream stream, IReadOnlyList<BlobHandle> references, string? basePath = null, CancellationToken cancellationToken = default);
+		ValueTask<BlobHandle> WriteBundleAsync(Stream stream, IReadOnlyList<BlobHandle> references, string? basePath = null, CancellationToken cancellationToken = default)
+			=> _inner.WriteBlobAsync(BundleBlobType, stream, references, basePath, cancellationToken);
 
 		/// <inheritdoc/>
 		public Task<BundleHeader> ReadHeaderAsync(BlobLocator locator, CancellationToken cancellationToken) => _bundleReader.ReadHeaderAsync(locator, cancellationToken);
@@ -206,162 +219,15 @@ namespace EpicGames.Horde.Storage.Clients
 		#region Aliases
 
 		/// <inheritdoc/>
-		public abstract Task AddAliasAsync(string name, BlobHandle handle, int rank, ReadOnlyMemory<byte> data, CancellationToken cancellationToken);
-
-		/// <inheritdoc/>
-		public abstract Task RemoveAliasAsync(string name, BlobHandle handle, CancellationToken cancellationToken);
-
-		/// <inheritdoc/>
-		public abstract Task<BlobAlias[]> FindAliasesAsync(string name, int? maxLength = null, CancellationToken cancellationToken = default);
-
-		#endregion
-
-		#region Refs
-
-		/// <inheritdoc/>
-		public abstract Task<bool> DeleteRefAsync(RefName name, CancellationToken cancellationToken = default);
-
-		/// <inheritdoc/>
-		public abstract Task<BlobHandle?> TryReadRefAsync(RefName name, RefCacheTime cacheTime = default, CancellationToken cancellationToken = default);
-
-		/// <inheritdoc/>
-		public abstract Task WriteRefAsync(RefName name, BlobHandle target, RefOptions? options, CancellationToken cancellationToken);
-
-		#endregion
-
-		/// <inheritdoc/>
-		public virtual void GetStats(StorageStats stats) => _bundleReader.GetStats(stats);
-	}
-
-	/// <summary>
-	/// Base class for an implementation of <see cref="IStorageClient"/>, providing implementations for some common functionality using bundles.
-	/// </summary>
-	public abstract class BundleStorageClient : BundleStorageClientBase
-	{
-		readonly IStorageBackend _backend;
-
-		/// <summary>
-		/// Backend for this client
-		/// </summary>
-		public IStorageBackend Backend => _backend;
-
-		/// <summary>
-		/// Constructor
-		/// </summary>
-		protected BundleStorageClient(IStorageBackend backend, BundleReaderCache cache, ILogger logger)
-			: base(cache, logger)
-		{
-			_backend = backend;
-		}
-
-		/// <summary>
-		/// Overridable dispose method
-		/// </summary>
-		/// <param name="disposing"></param>
-		protected override void Dispose(bool disposing)
-		{
-			if (disposing)
-			{
-				_backend.Dispose();
-			}
-
-			base.Dispose(disposing);
-		}
-
-		/// <inheritdoc/>
-		protected override async Task<Stream> OpenBundleAsync(BlobLocator locator, int offset, int? length, CancellationToken cancellationToken)
-			=> await _backend.OpenAsync(locator.ToString(), offset, length, cancellationToken);
-
-		/// <inheritdoc/>
-		protected override async ValueTask<BlobHandle> WriteBundleAsync(Stream stream, IReadOnlyList<BlobHandle> references, string? basePath = null, CancellationToken cancellationToken = default)
-		{
-			string path = await _backend.WriteAsync(stream, basePath, cancellationToken);
-			return new BundleHandle(this, new BlobLocator(path));
-		}
-
-		/// <inheritdoc/>
-		public override void GetStats(StorageStats stats)
-		{
-			_backend.GetStats(stats);
-			base.GetStats(stats);
-		}
-	}
-
-	/// <summary>
-	/// Wraps an underlying storage client with functionality that packs nodes into bundles.
-	/// </summary>
-	public sealed class BundleStorageClientWrapper : BundleStorageClientBase
-	{
-		readonly IStorageClient _inner;
-
-		/// <summary>
-		/// Constructor
-		/// </summary>
-		public BundleStorageClientWrapper(IStorageClient inner, BundleReaderCache cache, ILogger logger)
-			: base(cache, logger)
-		{
-			_inner = inner;
-		}
-
-		/// <summary>
-		/// Overridable dispose method
-		/// </summary>
-		/// <param name="disposing"></param>
-		protected override void Dispose(bool disposing)
-		{
-			if (disposing)
-			{
-				_inner.Dispose();
-			}
-
-			base.Dispose(disposing);
-		}
-
-		/// <summary>
-		/// Creates a bundle storage client around a memory client backend
-		/// </summary>
-		public static BundleStorageClientWrapper CreateFromMemory(ILogger logger)
-		{
-			MemoryStorageClient blobStore = new MemoryStorageClient();
-			return new BundleStorageClientWrapper(blobStore, BundleReaderCache.None, logger);
-		}
-
-		/// <summary>
-		/// Creates a bundle storage client around a directory on the filesystem
-		/// </summary>
-		public static BundleStorageClientWrapper CreateFromDirectory(DirectoryReference rootDir, BundleReaderCache cache, ILogger logger)
-		{
-			FileStorageClient fileStorageClient = new FileStorageClient(rootDir, logger);
-			return new BundleStorageClientWrapper(fileStorageClient, cache, logger);
-		}
-
-		/// <inheritdoc/>
-		protected override async Task<Stream> OpenBundleAsync(BlobLocator locator, int offset, int? length, CancellationToken cancellationToken)
-			=> await _inner.CreateBlobHandle(locator).OpenAsync(offset, length, cancellationToken);
-
-		/// <inheritdoc/>
-		protected override ValueTask<BlobHandle> WriteBundleAsync(Stream stream, IReadOnlyList<BlobHandle> references, string? basePath = null, CancellationToken cancellationToken = default)
-			=> _inner.WriteBlobAsync(BundleBlobType, stream, references, basePath, cancellationToken);
-
-		/// <inheritdoc/>
-		public override void GetStats(StorageStats stats)
-		{
-			_inner.GetStats(stats);
-			base.GetStats(stats);
-		}
-
-		#region Aliases
-
-		/// <inheritdoc/>
-		public override Task AddAliasAsync(string name, BlobHandle handle, int rank, ReadOnlyMemory<byte> data, CancellationToken cancellationToken)
+		public Task AddAliasAsync(string name, BlobHandle handle, int rank, ReadOnlyMemory<byte> data, CancellationToken cancellationToken)
 			=> _inner.AddAliasAsync(name, handle, rank, data, cancellationToken);
 
 		/// <inheritdoc/>
-		public override Task RemoveAliasAsync(string name, BlobHandle handle, CancellationToken cancellationToken)
+		public Task RemoveAliasAsync(string name, BlobHandle handle, CancellationToken cancellationToken)
 			=> _inner.RemoveAliasAsync(name, handle, cancellationToken);
 
 		/// <inheritdoc/>
-		public override async Task<BlobAlias[]> FindAliasesAsync(string name, int? maxLength = null, CancellationToken cancellationToken = default)
+		public async Task<BlobAlias[]> FindAliasesAsync(string name, int? maxLength = null, CancellationToken cancellationToken = default)
 		{
 			BlobAlias[] aliases = await _inner.FindAliasesAsync(name, maxLength, cancellationToken);
 			for (int idx = 0; idx < aliases.Length; idx++)
@@ -373,14 +239,15 @@ namespace EpicGames.Horde.Storage.Clients
 		}
 
 		#endregion
+
 		#region Refs
 
 		/// <inheritdoc/>
-		public override Task<bool> DeleteRefAsync(RefName name, CancellationToken cancellationToken = default)
+		public Task<bool> DeleteRefAsync(RefName name, CancellationToken cancellationToken = default)
 			=> _inner.DeleteRefAsync(name, cancellationToken);
 
 		/// <inheritdoc/>
-		public override async Task<BlobHandle?> TryReadRefAsync(RefName name, RefCacheTime cacheTime = default, CancellationToken cancellationToken = default)
+		public async Task<BlobHandle?> TryReadRefAsync(RefName name, RefCacheTime cacheTime = default, CancellationToken cancellationToken = default)
 		{
 			BlobHandle? target = await _inner.TryReadRefAsync(name, cacheTime, cancellationToken);
 			if (target != null)
@@ -391,9 +258,16 @@ namespace EpicGames.Horde.Storage.Clients
 		}
 
 		/// <inheritdoc/>
-		public override Task WriteRefAsync(RefName name, BlobHandle target, RefOptions? options, CancellationToken cancellationToken)
+		public Task WriteRefAsync(RefName name, BlobHandle target, RefOptions? options, CancellationToken cancellationToken)
 			=> _inner.WriteRefAsync(name, target, options, cancellationToken);
 
 		#endregion
+
+		/// <inheritdoc/>
+		public void GetStats(StorageStats stats)
+		{
+			_inner.GetStats(stats);
+			_bundleReader.GetStats(stats);
+		}
 	}
 }
diff --git a/Engine/Source/Programs/Shared/EpicGames.Horde/Storage/Clients/HttpStorageClient.cs b/Engine/Source/Programs/Shared/EpicGames.Horde/Storage/Clients/HttpStorageClient.cs
--- a/Engine/Source/Programs/Shared/EpicGames.Horde/Storage/Clients/HttpStorageClient.cs
+++ b/Engine/Source/Programs/Shared/EpicGames.Horde/Storage/Clients/HttpStorageClient.cs
@@ -311,7 +311,7 @@ namespace EpicGames.Horde.Storage.Clients
 			}
 
 			HttpStorageClient client = new HttpStorageClient(basePath, CreateClient, backend, _clientLogger);
-			return new BundleStorageClientWrapper(client, _readerCache, _clientLogger);
+			return new BundleStorageClient(client, _readerCache, _clientLogger);
 		}
 
 		/// <summary>
diff --git a/Engine/Source/Programs/UnrealBuildTool/Artifacts/HordeStorageArtifactCache.cs b/Engine/Source/Programs/UnrealBuildTool/Artifacts/HordeStorageArtifactCache.cs
--- a/Engine/Source/Programs/UnrealBuildTool/Artifacts/HordeStorageArtifactCache.cs
+++ b/Engine/Source/Programs/UnrealBuildTool/Artifacts/HordeStorageArtifactCache.cs
@@ -526,7 +526,7 @@ namespace UnrealBuildTool.Artifacts
 				}
 				Directory.CreateDirectory(directory.FullName);
 
-				_store = BundleStorageClientWrapper.CreateFromDirectory(directory, BundleReaderCache.None, logger);
+				_store = BundleStorageClient.CreateFromDirectory(directory, BundleReaderCache.None, logger);
 
 				State = ArtifactCacheState.Available;
 				return State;
