Horde: Remove BundleNodeHandle and BundleNodeLocator in favor of storing parsed data on handle objects.

diff --git a/Engine/Source/Programs/AutomationTool/BuildGraph/Tasks/DeployToolTask.cs b/Engine/Source/Programs/AutomationTool/BuildGraph/Tasks/DeployToolTask.cs
--- a/Engine/Source/Programs/AutomationTool/BuildGraph/Tasks/DeployToolTask.cs
+++ b/Engine/Source/Programs/AutomationTool/BuildGraph/Tasks/DeployToolTask.cs
@@ -150,14 +150,14 @@ namespace AutomationTool.Tasks
 				return httpClient;
 			}
 
-			BundleNodeHandle handle;
+			BlobHandle handle;
 
 			string basePath = $"api/v1/tools/{Parameters.Id}";
 			using HttpStorageBackend httpStorageBackend = new HttpStorageBackend(basePath, CreateHttpClient, Logger);
 			using HttpStorageClient httpStorageClient = new HttpStorageClient(basePath, CreateHttpClient, httpStorageBackend, Logger);
 			using BundleStorageClient storageClient = new BundleStorageClient(httpStorageClient, BundleReaderCache.None, Logger);
 
-			await using (BundleWriter treeWriter = storageClient.CreateWriter())
+			await using (IStorageWriter treeWriter = storageClient.CreateWriter())
 			{
 				DirectoryNode sandbox = new DirectoryNode();
 				if (Parameters.File != null)
diff --git a/Engine/Source/Programs/Horde/Horde.Server/Perforce/PerforceReplicator.cs b/Engine/Source/Programs/Horde/Horde.Server/Perforce/PerforceReplicator.cs
--- a/Engine/Source/Programs/Horde/Horde.Server/Perforce/PerforceReplicator.cs
+++ b/Engine/Source/Programs/Horde/Horde.Server/Perforce/PerforceReplicator.cs
@@ -550,7 +550,7 @@ namespace Horde.Server.Perforce
 			BlobHandle commitHandle = await store.WriteRefAsync(refName, commitNode, refOptions: options.RefOptions, cancellationToken: cancellationToken);
 
 			// Log the snapshot info
-			_logger.LogInformation("Snapshot for {StreamId} CL {Change} is ref {RefName} (commit: {CommitHandle}, root: {RootHandle})", streamConfig.Id, change, refName, ((BundleNodeHandle)commitHandle).GetLocator(), rootRef.Handle.GetLocator());
+			_logger.LogInformation("Snapshot for {StreamId} CL {Change} is ref {RefName} (commit: {CommitHandle}, root: {RootHandle})", streamConfig.Id, change, refName, commitHandle.GetLocator(), rootRef.Handle.GetLocator());
 		}
 
 		static int GetFileOffset(string path)
diff --git a/Engine/Source/Programs/Horde/Horde.Server/Storage/StorageController.cs b/Engine/Source/Programs/Horde/Horde.Server/Storage/StorageController.cs
--- a/Engine/Source/Programs/Horde/Horde.Server/Storage/StorageController.cs
+++ b/Engine/Source/Programs/Horde/Horde.Server/Storage/StorageController.cs
@@ -393,28 +393,28 @@ namespace Horde.Server.Storage
 						List<object> directories = new List<object>();
 						foreach ((string name, DirectoryEntry entry) in directoryNode.NameToDirectory)
 						{
-							directories.Add(new { name = name.ToString(), length = entry.Length, hash = entry.Hash, link = GetNodeLink(linkBase, (BundleNodeHandle)entry.Handle) });
+							directories.Add(new { name = name.ToString(), length = entry.Length, hash = entry.Hash, link = GetNodeLink(linkBase, entry.Handle) });
 						}
 
 						List<object> files = new List<object>();
 						foreach ((string name, FileEntry entry) in directoryNode.NameToFile)
 						{
-							files.Add(new { name = name.ToString(), length = entry.Length, flags = entry.Flags, hash = entry.Hash, link = GetNodeLink(linkBase, (BundleNodeHandle)entry.Handle) });
+							files.Add(new { name = name.ToString(), length = entry.Length, flags = entry.Flags, hash = entry.Hash, link = GetNodeLink(linkBase, entry.Handle) });
 						}
 
 						content = new { directoryNode.Length, directories, files };
 					}
 					break;
 				default:
-					content = new { references = blobData.Refs.Select(x => GetNodeLink(linkBase, (BundleNodeHandle)x)) };
+					content = new { references = blobData.Refs.Select(x => GetNodeLink(linkBase, x)) };
 					break;
 			}
 
 			return new { type = blobData.Type.Guid, @class = node.GetType().Name, content = content };
 		}
 
-		static string GetNodeLink(string linkBase, BundleNodeHandle handle) => GetNodeLink(linkBase, handle.GetLocator());
+		static string GetNodeLink(string linkBase, BlobHandle handle) => GetNodeLink(linkBase, handle.GetLocator());
 		
-		static string GetNodeLink(string linkBase, BundleNodeLocator locator) => $"{linkBase}/nodes/{locator.Blob}?export={locator.ExportIdx}";
+		static string GetNodeLink(string linkBase, BlobLocator locator) => $"{linkBase}/nodes/{locator}";
 	}
 }
diff --git a/Engine/Source/Programs/Horde/Horde/Commands/Bundles/BundleDump.cs b/Engine/Source/Programs/Horde/Horde/Commands/Bundles/BundleDump.cs
--- a/Engine/Source/Programs/Horde/Horde/Commands/Bundles/BundleDump.cs
+++ b/Engine/Source/Programs/Horde/Horde/Commands/Bundles/BundleDump.cs
@@ -67,8 +67,7 @@ namespace Horde.Commands.Bundles
 					for(int referenceIdx = 0; referenceIdx < export.References.Count; referenceIdx++)
 					{
 						BundleExportRef exportRef = export.References[referenceIdx];
-						BundleNodeLocator nodeLocator = new BundleNodeLocator(header.Imports[exportRef.ImportIdx], exportRef.NodeIdx);
-						logger.LogInformation("            REF {RefIdx,-3} -> {Node}", referenceIdx, nodeLocator);
+						logger.LogInformation("            REF {RefIdx,-3} -> {Node}", referenceIdx, $"{header.Imports[exportRef.ImportIdx]}#{exportRef.NodeIdx}");
 					}
 				}
 
diff --git a/Engine/Source/Programs/Shared/EpicGames.Horde.Tests/BundleStoreTests.cs b/Engine/Source/Programs/Shared/EpicGames.Horde.Tests/BundleStoreTests.cs
--- a/Engine/Source/Programs/Shared/EpicGames.Horde.Tests/BundleStoreTests.cs
+++ b/Engine/Source/Programs/Shared/EpicGames.Horde.Tests/BundleStoreTests.cs
@@ -64,9 +64,9 @@ namespace EpicGames.Horde.Tests
 			await using BundleWriter writer = store.CreateWriter(options: new BundleOptions { CompressionFormat = BundleCompressionFormat.None });
 
 			TextNode node = new TextNode("Hello world");
-			BundleNodeHandle handle = await writer.FlushAsync(node, CancellationToken.None);
+			BlobHandle handle = await writer.FlushAsync(node, CancellationToken.None);
 
-			BlobHandle bundleHandle = store.CreateBlobHandle(handle.GetLocator().Blob);
+			BlobHandle bundleHandle = store.CreateBlobHandle(handle.GetLocator().Outermost);
 			using BlobData blobData = await bundleHandle.ReadAsync();
 
 			return Bundle.FromMemory(blobData.Data.ToArray());
diff --git a/Engine/Source/Programs/Shared/EpicGames.Horde/Storage/Bundles/BundleNodeHandle.cs b/Engine/Source/Programs/Shared/EpicGames.Horde/Storage/Bundles/BundleNodeHandle.cs
deleted file mode 100644
diff --git a/Engine/Source/Programs/Shared/EpicGames.Horde/Storage/Bundles/BundleNodeLocator.cs b/Engine/Source/Programs/Shared/EpicGames.Horde/Storage/Bundles/BundleNodeLocator.cs
deleted file mode 100644
diff --git a/Engine/Source/Programs/Shared/EpicGames.Horde/Storage/Bundles/BundleReader.cs b/Engine/Source/Programs/Shared/EpicGames.Horde/Storage/Bundles/BundleReader.cs
--- a/Engine/Source/Programs/Shared/EpicGames.Horde/Storage/Bundles/BundleReader.cs
+++ b/Engine/Source/Programs/Shared/EpicGames.Horde/Storage/Bundles/BundleReader.cs
@@ -537,13 +537,14 @@ namespace EpicGames.Horde.Storage.Bundles
 		/// <summary>
 		/// Reads a node from a bundle
 		/// </summary>
-		/// <param name="locator">Locator for the node</param>
+		/// <param name="bundleLocator">Locator for the bundle</param>
+		/// <param name="exportIdx">Index of the export</param>
 		/// <param name="cancellationToken">Cancellation token for the operation</param>
 		/// <returns>Node data read from the given bundle</returns>
-		public async ValueTask<BlobData> ReadNodeDataAsync(BundleNodeLocator locator, CancellationToken cancellationToken = default)
+		public async ValueTask<BlobData> ReadNodeDataAsync(BlobLocator bundleLocator, int exportIdx, CancellationToken cancellationToken = default)
 		{
-			BundleHeader header = await ReadHeaderAsync(locator.Blob, cancellationToken);
-			BundleExport export = header.Exports[locator.ExportIdx];
+			BundleHeader header = await ReadHeaderAsync(bundleLocator, cancellationToken);
+			BundleExport export = header.Exports[exportIdx];
 
 			List<BlobHandle> refs = new List<BlobHandle>(export.References.Count);
 			foreach (BundleExportRef reference in export.References)
@@ -551,20 +552,20 @@ namespace EpicGames.Horde.Storage.Bundles
 				BlobLocator importBlob;
 				if (reference.ImportIdx == -1)
 				{
-					importBlob = locator.Blob;
+					importBlob = bundleLocator;
 				}
 				else
 				{
 					importBlob = header.Imports[reference.ImportIdx];
 				}
 				Debug.Assert(importBlob.IsValid());
-				refs.Add(new FlushedNodeHandle(this, new BundleNodeLocator(importBlob, reference.NodeIdx)));
+				refs.Add(new FlushedNodeHandle(this, importBlob, reference.NodeIdx));
 			}
 
 			ReadOnlyMemory<byte> nodeData = ReadOnlyMemory<byte>.Empty;
 			if (export.Length > 0)
 			{
-				ReadOnlyMemory<byte> packetData = await ReadPacketAsync(locator.Blob, export.Packet, cancellationToken);
+				ReadOnlyMemory<byte> packetData = await ReadPacketAsync(bundleLocator, export.Packet, cancellationToken);
 				nodeData = packetData.Slice(export.Offset, export.Length);
 			}
 
@@ -572,26 +573,6 @@ namespace EpicGames.Horde.Storage.Bundles
 			return new BlobData(nodeType, nodeData, refs);
 		}
 
-		/// <summary>
-		/// Reads a node from a bundle
-		/// </summary>
-		/// <param name="locator">Locator for the node</param>
-		/// <param name="cancellationToken">Cancellation token for the operation</param>
-		/// <returns>Node data read from the given bundle</returns>
-		public async ValueTask<Node> ReadNodeAsync(BundleNodeLocator locator, CancellationToken cancellationToken = default)
-		{
-			using BlobData blobData = await ReadNodeDataAsync(locator, cancellationToken);
-			return Node.Deserialize(blobData);
-		}
-
-		/// <summary>
-		/// Reads a node from a bundle
-		/// </summary>
-		/// <param name="locator">Locator for the node</param>
-		/// <param name="cancellationToken">Cancellation token for the operation</param>
-		/// <returns>Node data read from the given bundle</returns>
-		public async ValueTask<TNode> ReadNodeAsync<TNode>(BundleNodeLocator locator, CancellationToken cancellationToken = default) where TNode : Node => (TNode)await ReadNodeAsync(locator, cancellationToken);
-
 		/// <summary>
 		/// Gets stats for the reader
 		/// </summary>
diff --git a/Engine/Source/Programs/Shared/EpicGames.Horde/Storage/Bundles/BundleWriter.cs b/Engine/Source/Programs/Shared/EpicGames.Horde/Storage/Bundles/BundleWriter.cs
--- a/Engine/Source/Programs/Shared/EpicGames.Horde/Storage/Bundles/BundleWriter.cs
+++ b/Engine/Source/Programs/Shared/EpicGames.Horde/Storage/Bundles/BundleWriter.cs
@@ -2,6 +2,7 @@
 
 using System;
 using System.Buffers;
+using System.Buffers.Text;
 using System.Collections.Generic;
 using System.Diagnostics;
 using System.Diagnostics.CodeAnalysis;
@@ -48,38 +49,44 @@ namespace EpicGames.Horde.Storage.Bundles
 	/// <summary>
 	/// Implementation of <see cref="BlobHandle"/> for nodes which can be read from storage
 	/// </summary>
-	sealed class FlushedNodeHandle : BundleNodeHandle
+	sealed class FlushedNodeHandle : BlobHandle
 	{
 		readonly BundleReader _reader;
-		readonly BundleNodeLocator _locator;
+
+		public BlobLocator BundleLocator { get; }
+		public int ExportIdx { get; }
 
 		/// <summary>
 		/// Constructor
 		/// </summary>
-		public FlushedNodeHandle(BundleReader reader, BundleNodeLocator locator)
+		public FlushedNodeHandle(BundleReader reader, BlobLocator bundleLocator, int exportIdx)
 		{
 			_reader = reader;
-			_locator = locator;
+			BundleLocator = bundleLocator;
+			ExportIdx = exportIdx;
 		}
 
-		/// <inheritdoc/>
-		public override bool TryGetLocator([NotNullWhen(true)] out BlobLocator blobId)
+		public static FlushedNodeHandle FromBlobLocator(BundleReader reader, BlobLocator locator)
 		{
-			blobId = _locator.ToBlobLocator();
-			return true;
+			if (locator.TryUnwrap(out BlobLocator bundleLocator, out Utf8String fragment) && Utf8Parser.TryParse(fragment, out int exportIdx, out int bytesConsumed) && bytesConsumed == fragment.Length)
+			{
+				return new FlushedNodeHandle(reader, bundleLocator, exportIdx);
+			}
+			else
+			{
+				throw new ArgumentException($"Locator {locator} is not a valid bundle node");
+			}
 		}
 
 		/// <inheritdoc/>
-		public override bool HasLocator() => true;
-
-		/// <inheritdoc/>
-		public override BundleNodeLocator GetLocator() => _locator;
-
-		/// <inheritdoc/>
-		public override void AddWriteCallback(BlobWriteCallback callback) => callback.OnWrite();
+		public override bool TryGetLocator([NotNullWhen(true)] out BlobLocator locator)
+		{
+			locator = new BlobLocator(BundleLocator, ExportIdx.ToString());
+			return true;
+		}
 
 		/// <inheritdoc/>
-		public override ValueTask<BlobData> ReadAsync(CancellationToken cancellationToken = default) => _reader.ReadNodeDataAsync(_locator, cancellationToken);
+		public override ValueTask<BlobData> ReadAsync(CancellationToken cancellationToken = default) => _reader.ReadNodeDataAsync(BundleLocator, ExportIdx, cancellationToken);
 
 		/// <inheritdoc/>
 		public override ValueTask FlushAsync(CancellationToken cancellationToken = default) => new ValueTask();
@@ -92,13 +99,13 @@ namespace EpicGames.Horde.Storage.Bundles
 	public sealed class BundleWriter : IStorageWriter
 	{
 		// Information about a unique output node. Note that multiple node refs may de-duplicate to the same output node.
-		internal class PendingNode : BundleNodeHandle
+		internal class PendingNode : BlobHandle
 		{
 			readonly BundleReader _reader;
 
 			object LockObject => _reader;
 
-			BundleNodeLocator _locator;
+			FlushedNodeHandle? _flushedHandle;
 			PendingBundle? _pendingBundle;
 
 			public readonly BlobType BlobType;
@@ -109,6 +116,7 @@ namespace EpicGames.Horde.Storage.Bundles
 			public readonly AliasInfo[] Aliases;
 
 			public PendingBundle? PendingBundle => _pendingBundle;
+			public FlushedNodeHandle? FlushedNodeHandle => _flushedHandle;
 
 			public PendingNode(BundleReader reader, BlobType blobType, int packet, int offset, int length, IReadOnlyList<BlobHandle> refs, IReadOnlyList<AliasInfo> aliases, PendingBundle pendingBundle)
 			{
@@ -125,79 +133,54 @@ namespace EpicGames.Horde.Storage.Bundles
 			}
 
 			/// <inheritdoc/>
-			public override bool TryGetLocator([NotNullWhen(true)] out BlobLocator blobId)
+			public override bool TryGetLocator([NotNullWhen(true)] out BlobLocator locator)
 			{
-				if (_locator.IsValid())
+				if (_flushedHandle != null && _flushedHandle.TryGetLocator(out BlobLocator flushedLocator))
 				{
-					blobId = _locator.ToBlobLocator();
+					locator = flushedLocator;
 					return true;
 				}
 				else
 				{
-					blobId = default;
-					return false;
-				}
-			}
+					Console.WriteLine($"Unable to get locator for node {BlobType}:{Packet}:{Offset}:{Length}. Handle is {FlushedNodeHandle}. Bundle is {PendingBundle}.");
 
-			/// <inheritdoc/>
-			public override bool HasLocator() => _locator.IsValid();
-
-			/// <inheritdoc/>
-			public override BundleNodeLocator GetLocator()
-			{
-				if (!_locator.IsValid())
-				{
-					throw new InvalidOperationException();
+					locator = default;
+					return false;
 				}
-				return _locator;
 			}
 
-			public void MarkAsWritten(BundleNodeLocator locator)
+			public void MarkAsWritten(FlushedNodeHandle flushedHandle)
 			{
 				lock (LockObject)
 				{
-					Debug.Assert(!_locator.IsValid());
-					_locator = locator;
+					Debug.Assert(_flushedHandle == null);
+					_flushedHandle = flushedHandle;
 					_pendingBundle = null;
 				}
 			}
 
-			/// <inheritdoc/>
-			public override void AddWriteCallback(BlobWriteCallback callback)
-			{
-				PendingBundle? pendingBundle = PendingBundle;
-				if (pendingBundle == null || !pendingBundle.TryAddWriteCallback(callback))
-				{
-					Debug.Assert(_locator.IsValid());
-					callback.OnWrite();
-				}
-			}
-
 			/// <inheritdoc/>
 			public override async ValueTask<BlobData> ReadAsync(CancellationToken cancellationToken = default)
 			{
-				if (!_locator.IsValid())
+				if (_flushedHandle == null)
 				{
 					lock (LockObject)
 					{
-						if (!_locator.IsValid() && _pendingBundle != null)
+						if (_flushedHandle == null)
 						{
-							ReadOnlyMemory<byte> data = _pendingBundle.GetNodeData(Packet, Offset, Length);
-							if (data.Length == Length)
-							{
-								return new BlobData(BlobType, data, Refs);
-							}
+							ReadOnlyMemory<byte> data = _pendingBundle!.GetNodeData(Packet, Offset, Length);
+							return new BlobData(BlobType, data, Refs);
 						}
 					}
 				}
 
-				return await _reader.ReadNodeDataAsync(_locator, cancellationToken);
+				return await _flushedHandle!.ReadAsync(cancellationToken);
 			}
 
 			/// <inheritdoc/>
 			public override async ValueTask FlushAsync(CancellationToken cancellationToken = default)
 			{
-				if (_locator.IsValid())
+				if (_flushedHandle != null)
 				{
 					return;
 				}
@@ -218,13 +201,6 @@ namespace EpicGames.Horde.Storage.Bundles
 		// Information about a bundle being built. Metadata operations are synchronous, compression/writes are asynchronous.
 		internal class PendingBundle : IDisposable
 		{
-			class WriteCallbackSentinel : BlobWriteCallback
-			{
-				public override void OnWrite() => throw new NotImplementedException();
-			}
-
-			static readonly WriteCallbackSentinel s_callbackSentinel = new WriteCallbackSentinel();
-
 			static int s_lastBundleId = 0;
 			public int BundleId { get; } = Interlocked.Increment(ref s_lastBundleId);
 
@@ -268,9 +244,6 @@ namespace EpicGames.Horde.Storage.Bundles
 			// Total size of uncompressed data in the current bundle
 			public long UncompressedLength { get; private set; }
 
-			// List of post-write callbacks
-			BlobWriteCallback? _callbacks = null;
-
 			// Task used to compress data in the background
 			Task _compressPacketsTask = Task.CompletedTask;
 
@@ -322,26 +295,6 @@ namespace EpicGames.Horde.Storage.Bundles
 				}
 			}
 
-			// Adds a callback after writing
-			public bool TryAddWriteCallback(BlobWriteCallback callback)
-			{
-				for (; ; )
-				{
-					BlobWriteCallback? tail = _callbacks;
-					if (tail == s_callbackSentinel)
-					{
-						return false;
-					}
-
-					callback._next = tail;
-
-					if (Interlocked.CompareExchange(ref _callbacks, callback, tail) == tail)
-					{
-						return true;
-					}
-				}
-			}
-
 			// Starts writing a node to the buffer
 			public Memory<byte> GetBuffer(int usedSize, int desiredSize)
 			{
@@ -495,23 +448,15 @@ namespace EpicGames.Horde.Storage.Bundles
 					{
 						PendingNode node = _queue[idx];
 
-						BundleNodeLocator nodeLocator = new BundleNodeLocator(locator, idx);
-						_queue[idx].MarkAsWritten(nodeLocator);
+						FlushedNodeHandle flushedHandle = new FlushedNodeHandle(_treeReader, locator, idx);
+						_queue[idx].MarkAsWritten(flushedHandle);
 
 						foreach (AliasInfo alias in node.Aliases)
 						{
-							BlobHandle target = store.CreateBlobHandle(nodeLocator.ToBlobLocator());
-							await store.AddAliasAsync(alias.Name, target, alias.Rank, alias.Data, CancellationToken.None);
+							await store.AddAliasAsync(alias.Name, flushedHandle, alias.Rank, alias.Data, CancellationToken.None);
 						}
 					}
 
-					BlobWriteCallback? callback = Interlocked.Exchange(ref _callbacks, s_callbackSentinel);
-					while (callback != null)
-					{
-						callback.OnWrite();
-						callback = callback._next;
-					}
-
 					_completeEvent.SetResult(true);
 				}
 				catch (Exception ex)
@@ -557,15 +502,23 @@ namespace EpicGames.Horde.Storage.Bundles
 					int typeIdx = FindOrAddItemIndex(nodeInfo.BlobType, types, typeToIndex);
 
 					List<BundleExportRef> exportRefs = new List<BundleExportRef>();
-					foreach (BundleNodeHandle handle in nodeInfo.Refs)
+					foreach (BlobHandle handle in nodeInfo.Refs)
 					{
 						BundleExportRef exportRef;
 						if (!nodeHandleToExportRef.TryGetValue(handle, out exportRef))
 						{
-							BundleNodeLocator locator = handle.GetLocator();
-
-							int importIdx = FindOrAddItemIndex(locator.Blob, imports, importToIndex);
-							exportRef = new BundleExportRef(importIdx, locator.ExportIdx);
+							FlushedNodeHandle? flushedHandle = handle as FlushedNodeHandle;
+							if (flushedHandle == null)
+							{
+								flushedHandle = (handle as PendingNode)?.FlushedNodeHandle;
+								if (flushedHandle == null)
+								{
+									throw new InvalidOperationException($"Node {handle.GetLocator()} is not a bundle node");
+								}
+							}
+							
+							int importIdx = FindOrAddItemIndex(flushedHandle.BundleLocator, imports, importToIndex);
+							exportRef = new BundleExportRef(importIdx, flushedHandle.ExportIdx);
 						}
 						exportRefs.Add(exportRef);
 					}
@@ -858,9 +811,9 @@ namespace EpicGames.Horde.Storage.Bundles
 		/// <param name="root">Root for the tree</param>
 		/// <param name="cancellationToken">Cancellation token for the operation</param>
 		/// <returns></returns>
-		public async Task<BundleNodeHandle> FlushAsync(Node root, CancellationToken cancellationToken = default)
+		public async Task<BlobHandle> FlushAsync(Node root, CancellationToken cancellationToken = default)
 		{
-			return (BundleNodeHandle)await NodeRefExtensions.FlushAsync(this, root, cancellationToken);
+			return await NodeRefExtensions.FlushAsync(this, root, cancellationToken);
 		}
 	}
 }
diff --git a/Engine/Source/Programs/Shared/EpicGames.Horde/Storage/Clients/BundleStorageClient.cs b/Engine/Source/Programs/Shared/EpicGames.Horde/Storage/Clients/BundleStorageClient.cs
--- a/Engine/Source/Programs/Shared/EpicGames.Horde/Storage/Clients/BundleStorageClient.cs
+++ b/Engine/Source/Programs/Shared/EpicGames.Horde/Storage/Clients/BundleStorageClient.cs
@@ -188,9 +188,6 @@ namespace EpicGames.Horde.Storage.Clients
 		/// <inheritdoc/>
 		public Task<BundleHeader> ReadHeaderAsync(BlobLocator locator, CancellationToken cancellationToken) => _bundleReader.ReadHeaderAsync(locator, cancellationToken);
 
-		/// <inheritdoc/>
-		public async Task<BlobData> ReadNodeDataAsync(BundleNodeLocator locator, CancellationToken cancellationToken) => await _bundleReader.ReadNodeDataAsync(locator, cancellationToken);
-
 		#endregion
 
 		#region Nodes
@@ -200,7 +197,7 @@ namespace EpicGames.Horde.Storage.Clients
 		{
 			if (locator.CanUnwrap())
 			{
-				return new FlushedNodeHandle(_bundleReader, BundleNodeLocator.FromBlobLocator(locator));
+				return FlushedNodeHandle.FromBlobLocator(_bundleReader, locator);
 			}
 			else
 			{
