Horde: Fix test failure.

diff --git a/Engine/Source/Programs/Horde/Horde.Server/Storage/StorageController.cs b/Engine/Source/Programs/Horde/Horde.Server/Storage/StorageController.cs
--- a/Engine/Source/Programs/Horde/Horde.Server/Storage/StorageController.cs
+++ b/Engine/Source/Programs/Horde/Horde.Server/Storage/StorageController.cs
@@ -12,7 +12,6 @@ using System.Threading.Tasks;
 using Amazon.EC2.Model;
 using EpicGames.Core;
 using EpicGames.Horde.Storage;
-using EpicGames.Horde.Storage.Bundles;
 using EpicGames.Horde.Storage.Clients;
 using EpicGames.Horde.Storage.Nodes;
 using Horde.Server.Acls;
diff --git a/Engine/Source/Programs/Shared/EpicGames.Horde/Compute/AgentMessage.cs b/Engine/Source/Programs/Shared/EpicGames.Horde/Compute/AgentMessage.cs
--- a/Engine/Source/Programs/Shared/EpicGames.Horde/Compute/AgentMessage.cs
+++ b/Engine/Source/Programs/Shared/EpicGames.Horde/Compute/AgentMessage.cs
@@ -398,7 +398,15 @@ namespace EpicGames.Horde.Compute
 		public static UploadFilesMessage ParseUploadFilesMessage(this AgentMessage message)
 		{
 			string name = message.ReadString();
-			BlobLocator locator = new BlobLocator(message.ReadString());
+			string path = message.ReadString();
+
+			int atIdx = path.IndexOf('@', StringComparison.Ordinal);
+			if (atIdx != -1)
+			{
+				path = path.Substring(atIdx + 1);
+			}
+
+			BlobLocator locator = new BlobLocator(path);
 			return new UploadFilesMessage(name, locator);
 		}
 
@@ -645,7 +653,7 @@ namespace EpicGames.Horde.Compute
 		public static async Task SendBlobDataAsync(this AgentMessageChannel channel, BlobLocator locator, int offset, int length, IStorageClient storage, CancellationToken cancellationToken = default)
 		{
 			BlobHandle handle = storage.CreateBlobHandle(locator);
-			using Stream stream = await handle.OpenAsync(offset, length, cancellationToken);
+			using Stream stream = await handle.OpenAsync(offset, (length == 0)? null : length, cancellationToken);
 
 			const int MaxChunkSize = 512 * 1024;
 			for (int chunkOffset = 0; chunkOffset < stream.Length;)
