Horde: Move code for handling the current bundle schema into a V1 folder.

diff --git a/Engine/Source/Programs/Horde/Horde.Agent.Tests/JsonRpcLoggerTests.cs b/Engine/Source/Programs/Horde/Horde.Agent.Tests/JsonRpcLoggerTests.cs
--- a/Engine/Source/Programs/Horde/Horde.Agent.Tests/JsonRpcLoggerTests.cs
+++ b/Engine/Source/Programs/Horde/Horde.Agent.Tests/JsonRpcLoggerTests.cs
@@ -7,7 +7,6 @@ using System.Threading;
 using System.Threading.Tasks;
 using EpicGames.Core;
 using EpicGames.Horde.Storage;
-using EpicGames.Horde.Storage.Bundles;
 using EpicGames.Horde.Storage.Clients;
 using EpicGames.Horde.Logs;
 using Horde.Agent.Utility;
@@ -17,6 +16,7 @@ using Microsoft.Extensions.Logging;
 using Microsoft.Extensions.Logging.Abstractions;
 using Microsoft.VisualStudio.TestTools.UnitTesting;
 using System.Linq;
+using EpicGames.Horde.Storage.Bundles.V1;
 
 namespace Horde.Agent.Tests
 {
diff --git a/Engine/Source/Programs/Horde/Horde.Server/Commands/Bundles/BundleWriteTest.cs b/Engine/Source/Programs/Horde/Horde.Server/Commands/Bundles/BundleWriteTest.cs
--- a/Engine/Source/Programs/Horde/Horde.Server/Commands/Bundles/BundleWriteTest.cs
+++ b/Engine/Source/Programs/Horde/Horde.Server/Commands/Bundles/BundleWriteTest.cs
@@ -10,6 +10,7 @@ using System.Threading.Tasks;
 using EpicGames.Core;
 using EpicGames.Horde.Storage;
 using EpicGames.Horde.Storage.Bundles;
+using EpicGames.Horde.Storage.Clients;
 using EpicGames.Horde.Storage.Nodes;
 using Microsoft.Extensions.Logging;
 
@@ -18,11 +19,11 @@ namespace Horde.Server.Commands.Bundles
 	[Command("bundle", "writetest", "Synthetic benchmark for bundle write performance")]
 	class BundleWriteTestCommand : Command
 	{
-		class FakeStorageClient : IStorageClient
+		class NullStorageClient : IStorageClient
 		{
 			public bool SupportsRedirects => false;
 
-			public FakeStorageClient()
+			public NullStorageClient()
 			{
 			}
 
@@ -41,20 +42,15 @@ namespace Horde.Server.Commands.Bundles
 			public Task<BlobHandle?> TryReadRefAsync(RefName name, RefCacheTime cacheTime = default, CancellationToken cancellationToken = default) => throw new NotImplementedException();
 			public Task WriteRefAsync(RefName name, BlobHandle target, RefOptions? options = null, CancellationToken cancellationToken = default) => Task.CompletedTask;
 
-			public BundleWriter CreateWriter(string? basePath = null, BundleOptions? options = null)
-			{
-				_ = (this, basePath, options);
-				throw new NotImplementedException();// new BundleWriter(this, _reader, refName, options);
-			}
-
-			IStorageWriter IStorageClient.CreateWriter(string? basePath) => CreateWriter(basePath);
+			public IStorageWriter CreateWriter(string? basePath) => new DefaultStorageWriter(this, basePath);
 
 			public void GetStats(StorageStats stats) { }
 		}
 
 		public override async Task<int> ExecuteAsync(ILogger logger)
 		{
-			using FakeStorageClient store = new FakeStorageClient();
+			using NullStorageClient nullStore = new NullStorageClient();
+			using BundleStorageClient store = new BundleStorageClient(nullStore, BundleReaderCache.None, logger);
 
 			BundleOptions options = new BundleOptions();
 			options.CompressionFormat = BundleCompressionFormat.None;
diff --git a/Engine/Source/Programs/Horde/Horde/Commands/Bundles/BundleDump.cs b/Engine/Source/Programs/Horde/Horde/Commands/Bundles/BundleDump.cs
--- a/Engine/Source/Programs/Horde/Horde/Commands/Bundles/BundleDump.cs
+++ b/Engine/Source/Programs/Horde/Horde/Commands/Bundles/BundleDump.cs
@@ -2,7 +2,7 @@
 
 using EpicGames.Core;
 using EpicGames.Horde.Storage;
-using EpicGames.Horde.Storage.Bundles;
+using EpicGames.Horde.Storage.Bundles.V1;
 using Microsoft.Extensions.Logging;
 
 namespace Horde.Commands.Bundles
@@ -18,14 +18,10 @@ namespace Horde.Commands.Bundles
 
 		public override async Task<int> ExecuteAsync(ILogger logger)
 		{
-			Bundle bundle;
-			using (FileStream stream = FileReference.Open(Input, FileMode.Open, FileAccess.Read, FileShare.Read | FileShare.Delete))
-			{
-				bundle = await Bundle.FromStreamAsync(stream);
-			}
+			byte[] data = await FileReference.ReadAllBytesAsync(Input);
 			logger.LogInformation("Summary for blob {Location}", Input);
 
-			BundleHeader header = bundle.Header;
+			BundleHeader header = BundleHeader.Read(data);
 
 			string[] types = new string[header.Types.Count];
 			for (int idx = 0; idx < header.Types.Count; idx++)
diff --git a/Engine/Source/Programs/Shared/EpicGames.Horde.Tests/BundleStoreTests.cs b/Engine/Source/Programs/Shared/EpicGames.Horde.Tests/BundleStoreTests.cs
--- a/Engine/Source/Programs/Shared/EpicGames.Horde.Tests/BundleStoreTests.cs
+++ b/Engine/Source/Programs/Shared/EpicGames.Horde.Tests/BundleStoreTests.cs
@@ -17,6 +17,7 @@ using Microsoft.Extensions.Caching.Memory;
 using Microsoft.Extensions.Logging.Abstractions;
 using System.Reflection;
 using EpicGames.Horde.Storage.Bundles;
+using EpicGames.Horde.Storage.Bundles.V1;
 
 namespace EpicGames.Horde.Tests
 {
@@ -26,12 +27,9 @@ namespace EpicGames.Horde.Tests
 		[TestMethod]
 		public async Task CreateBundlesManuallyAsync()
 		{
-			Bundle a = CreateBundleManually();
-			Bundle b = await CreateBundleNormalAsync();
-
-			byte[] bytesA = a.AsSequence().ToArray();
-			byte[] bytesB = b.AsSequence().ToArray();
-			Assert.IsTrue(bytesA.SequenceEqual(bytesB));
+			byte[] a = CreateBundleManually();
+			byte[] b = await CreateBundleNormalAsync();
+			Assert.IsTrue(a.SequenceEqual(b));
 		}
 
 		[NodeType("{F63606D4-5DBB-4061-A655-6F444F65229E}")]
@@ -57,11 +55,11 @@ namespace EpicGames.Horde.Tests
 			Node.RegisterTypesFromAssembly(Assembly.GetExecutingAssembly());
 		}
 
-		static async Task<Bundle> CreateBundleNormalAsync()
+		static async Task<byte[]> CreateBundleNormalAsync()
 		{
 			using MemoryStorageClient memoryStore = new MemoryStorageClient();
 			using BundleStorageClient store = new BundleStorageClient(memoryStore, BundleReaderCache.None, NullLogger.Instance);
-			await using BundleWriter writer = store.CreateWriter(options: new BundleOptions { CompressionFormat = BundleCompressionFormat.None });
+			await using IStorageWriter writer = store.CreateWriter(options: new BundleOptions { CompressionFormat = BundleCompressionFormat.None });
 
 			TextNode node = new TextNode("Hello world");
 			BlobHandle handle = await writer.FlushAsync(node, CancellationToken.None);
@@ -69,10 +67,10 @@ namespace EpicGames.Horde.Tests
 			BlobHandle bundleHandle = store.CreateBlobHandle(handle.GetLocator().Outermost);
 			using BlobData blobData = await bundleHandle.ReadAsync();
 
-			return Bundle.FromMemory(blobData.Data.ToArray());
+			return blobData.Data.ToArray();
 		}
 
-		static Bundle CreateBundleManually()
+		static byte[] CreateBundleManually()
 		{
 			ArrayMemoryWriter payloadWriter = new ArrayMemoryWriter(200);
 			payloadWriter.WriteString("Hello world");
@@ -88,7 +86,11 @@ namespace EpicGames.Horde.Tests
 			packets.Add(new BundlePacket(BundleCompressionFormat.None, 0, payload.Length, payload.Length));
 
 			BundleHeader header = new BundleHeader(types.ToArray(), Array.Empty<BlobLocator>(), exports.ToArray(), packets.ToArray());
-			return new Bundle(header, new List<ReadOnlyMemory<byte>> { payload });
+
+			ReadOnlySequenceBuilder<byte> builder = new ReadOnlySequenceBuilder<byte>();
+			header.AppendTo(builder);
+			builder.Append(payload);
+			return builder.Construct().ToArray();
 		}
 
 		[TestMethod]
@@ -398,7 +400,7 @@ namespace EpicGames.Horde.Tests
 
 				HashedNodeRef<ChunkedDataNode> file = root.GetFileEntry("test");
 
-				long uniqueSize = memoryStore.Blobs.Values.Select(x => Bundle.FromMemory(x.Data)).SelectMany(x => x.Header.Packets).Sum(x => x.DecodedLength);
+				long uniqueSize = memoryStore.Blobs.Values.Select(x => BundleHeader.Read(x.Data)).SelectMany(x => x.Packets).Sum(x => x.DecodedLength);
 				Assert.IsTrue(uniqueSize < data.Length / 3); // random fraction meaning "lots of dedupe happened"
 			}
 		}
diff --git a/Engine/Source/Programs/Shared/EpicGames.Horde.Tests/BundleTests.cs b/Engine/Source/Programs/Shared/EpicGames.Horde.Tests/BundleTests.cs
--- a/Engine/Source/Programs/Shared/EpicGames.Horde.Tests/BundleTests.cs
+++ b/Engine/Source/Programs/Shared/EpicGames.Horde.Tests/BundleTests.cs
@@ -3,6 +3,7 @@
 using EpicGames.Core;
 using EpicGames.Horde.Storage;
 using EpicGames.Horde.Storage.Bundles;
+using EpicGames.Horde.Storage.Bundles.V1;
 using EpicGames.Horde.Storage.Clients;
 using EpicGames.Horde.Storage.Nodes;
 using Microsoft.Extensions.Caching.Memory;
@@ -221,10 +222,10 @@ namespace EpicGames.Horde.Tests
 			BlobHandle refTarget =  await store.ReadRefTargetAsync(refName);
 			BlobHandle bundleTarget = store.CreateBlobHandle(refTarget.GetLocator().Outermost);
 			using BlobData bundleData = await bundleTarget.ReadAsync();
-			Bundle bundle = Bundle.FromMemory(bundleData.Data);
 
-			Assert.AreEqual(0, bundle.Header.Imports.Count);
-			Assert.AreEqual(3, bundle.Header.Exports.Count);
+			BundleHeader bundleHeader = BundleHeader.Read(bundleData.Data);
+			Assert.AreEqual(0, bundleHeader.Imports.Count);
+			Assert.AreEqual(3, bundleHeader.Exports.Count);
 
 			// Create a new bundle and read it back in again
 			DirectoryNode newRoot = await store.ReadRefAsync<DirectoryNode>(refName);
diff --git a/Engine/Source/Programs/Shared/EpicGames.Horde.Tests/ComputeSocketTests.cs b/Engine/Source/Programs/Shared/EpicGames.Horde.Tests/ComputeSocketTests.cs
--- a/Engine/Source/Programs/Shared/EpicGames.Horde.Tests/ComputeSocketTests.cs
+++ b/Engine/Source/Programs/Shared/EpicGames.Horde.Tests/ComputeSocketTests.cs
@@ -13,7 +13,6 @@ using EpicGames.Core;
 using EpicGames.Horde.Compute;
 using EpicGames.Horde.Compute.Transports;
 using EpicGames.Horde.Storage;
-using EpicGames.Horde.Storage.Bundles;
 using EpicGames.Horde.Storage.Clients;
 using EpicGames.Horde.Storage.Nodes;
 using Microsoft.Extensions.Logging;
@@ -142,7 +141,7 @@ namespace EpicGames.Horde.Tests
 
 					using MemoryStorageClient memoryStorage = new MemoryStorageClient();
 					using BundleStorageClient storage = new BundleStorageClient(memoryStorage, BundleReaderCache.None, NullLogger.Instance);
-					await using (BundleWriter treeWriter = storage.CreateWriter())
+					await using (IStorageWriter treeWriter = storage.CreateWriter())
 					{
 						FileReference file = FileReference.Combine(tempDir, "subdir/hello.txt");
 						if (FileReference.Exists(file))
diff --git a/Engine/Source/Programs/Shared/EpicGames.Horde/Storage/Bundles/Bundle.cs b/Engine/Source/Programs/Shared/EpicGames.Horde/Storage/Bundles/Bundle.cs
--- a/Engine/Source/Programs/Shared/EpicGames.Horde/Storage/Bundles/Bundle.cs
+++ b/Engine/Source/Programs/Shared/EpicGames.Horde/Storage/Bundles/Bundle.cs
@@ -1,20 +1,8 @@
 // Copyright Epic Games, Inc. All Rights Reserved.
 
-using EpicGames.Compression;
-using EpicGames.Core;
-using K4os.Compression.LZ4;
 using System;
-using System.Buffers;
 using System.Buffers.Binary;
-using System.Collections;
-using System.Collections.Generic;
-using System.Diagnostics;
 using System.IO;
-using System.IO.Compression;
-using System.Linq;
-using System.Text;
-using System.Threading;
-using System.Threading.Tasks;
 
 namespace EpicGames.Horde.Storage.Bundles
 {
@@ -63,1336 +51,63 @@ namespace EpicGames.Horde.Storage.Bundles
 	}
 
 	/// <summary>
-	/// Header for the contents of a bundle. May contain an inlined payload object containing the object data itself.
+	/// Signature for a bundle
 	/// </summary>
-	public class Bundle
-	{
-		/// <summary>
-		/// Maximum size for a section
-		/// </summary>
-		public const int MaxSectionSize = 0x1000000;
-
-		/// <summary>
-		/// Maximum number of exports from a single bundle
-		/// </summary>
-		public const int MaxExports = 65535;
-
-		/// <summary>
-		/// Maximum number of export refs from a single bundle
-		/// </summary>
-		public const int MaxExportRefs = MaxSectionSize / BundleExportRef.NumBytes;
-
-		/// <summary>
-		/// Header for the bundle
-		/// </summary>
-		public BundleHeader Header { get; }
-
-		/// <summary>
-		/// Packet data as described in the header
-		/// </summary>
-		public IReadOnlyList<ReadOnlyMemory<byte>> Packets { get; }
-
-		/// <summary>
-		/// Constructor
-		/// </summary>
-		public Bundle(BundleHeader header, IReadOnlyList<ReadOnlyMemory<byte>> packets)
-		{
-			Header = header;
-			Packets = packets;
-		}
-
-		/// <summary>
-		/// Reads a bundle from the given memory
-		/// </summary>
-		/// <param name="data">Data to read from</param>
-		/// <returns>Bundle that was read</returns>
-		public static Bundle FromMemory(ReadOnlyMemory<byte> data)
-		{
-			using ReadOnlyMemoryStream stream = new ReadOnlyMemoryStream(data);
-			return FromStreamAsync(stream).Result;
-		}
-
-		/// <summary>
-		/// Reads a bundle from the given stream
-		/// </summary>
-		/// <param name="stream">Stream to read from</param>
-		/// <param name="cancellationToken">Cancellation token for the operation</param>
-		/// <returns>Bundle that was read</returns>
-		public static async Task<Bundle> FromStreamAsync(Stream stream, CancellationToken cancellationToken = default)
-		{
-			BundleHeader header = await BundleHeader.ReadAsync(stream, cancellationToken);
-
-			ReadOnlyMemory<byte>[] packets = new ReadOnlyMemory<byte>[header.Packets.Count];
-			for (int idx = 0; idx < header.Packets.Count; idx++)
-			{
-				BundlePacket packet = header.Packets[idx];
-
-				byte[] data = new byte[packet.EncodedLength];
-				await stream.ReadFixedLengthBytesAsync(data, cancellationToken);
-
-				packets[idx] = data;
-			}
-
-			return new Bundle(header, packets);
-		}
-
-		/// <summary>
-		/// Serializes the bundle to a sequence of bytes
-		/// </summary>
-		/// <returns>Sequence for the bundle</returns>
-		public ReadOnlySequence<byte> AsSequence()
-		{
-			ReadOnlySequenceBuilder<byte> sequence = new ReadOnlySequenceBuilder<byte>();
-			Header.AppendTo(sequence);
-
-			foreach (ReadOnlyMemory<byte> packet in Packets)
-			{
-				sequence.Append(packet);
-			}
-
-			return sequence.Construct();
-		}
-	}
-
-	/// <summary>
-	/// Indicates the compression format in the bundle
-	/// </summary>
-	public enum BundleCompressionFormat
-	{
-		/// <summary>
-		/// Packets are uncompressed
-		/// </summary>
-		None = 0,
-
-		/// <summary>
-		/// LZ4 compression
-		/// </summary>
-		LZ4 = 1,
-
-		/// <summary>
-		/// Gzip compression
-		/// </summary>
-		Gzip = 2,
-
-		/// <summary>
-		/// Oodle compression (Kraken)
-		/// </summary>
-		Oodle = 3,
-
-		/// <summary>
-		/// Brotli compression
-		/// </summary>
-		Brotli = 4,
-	}
-
-	/// <summary>
-	/// Identifier for the type of a section in the bundle header
-	/// </summary>
-	public enum BundleSectionType
-	{
-		/// <summary>
-		/// List of custom types
-		/// </summary>
-		Types = 0,
-
-		/// <summary>
-		/// Imports of other bundles
-		/// </summary>
-		Imports = 1,
-
-		/// <summary>
-		/// List of export headers
-		/// </summary>
-		ExportHeaders = 2,
-
-		/// <summary>
-		/// References to exports in other bundles
-		/// </summary>
-		ExportReferences = 3,
-
-		/// <summary>
-		/// Packet headers
-		/// </summary>
-		Packets = 4,
-
-		/// <summary>
-		/// Merged export headers and references
-		/// </summary>
-		Exports = 5,
-	}
+	/// <param name="Version">Version number for the following file data</param>
+	/// <param name="HeaderLength">Length of the initial header</param>
+	public record struct BundleSignature(BundleVersion Version, int HeaderLength);
 
 	/// <summary>
-	/// Header for the contents of a bundle.
+	/// Methods for manipulating bundles
 	/// </summary>
-	public class BundleHeader
+	public static class Bundle
 	{
 		/// <summary>
-		/// Signature bytes
-		/// </summary>
-		public static ReadOnlyMemory<byte> Signature { get; } = Encoding.UTF8.GetBytes("UEBN");
-
-		/// <summary>
-		/// Length of the prelude data
-		/// </summary>
-		public const int PreludeLength = 8;
-
-		/// <summary>
-		/// Types for exports within this bundle
-		/// </summary>
-		public BundleTypeCollection Types { get; }
-
-		/// <summary>
-		/// Bundles that we reference nodes in
-		/// </summary>
-		public BundleImportCollection Imports { get; }
-
-		/// <summary>
-		/// Nodes exported from this bundle
-		/// </summary>
-		public BundleExportCollection Exports { get; }
-
-		/// <summary>
-		/// List of data packets within this bundle
-		/// </summary>
-		public BundlePacketCollection Packets { get; }
-
-		/// <summary>
-		/// Constructor
-		/// </summary>
-		public BundleHeader(BundleTypeCollection types, BundleImportCollection imports, BundleExportCollection exports, BundlePacketCollection packets)
-		{
-			Types = types;
-			Imports = imports;
-			Exports = exports;
-			Packets = packets;
-		}
-
-		/// <summary>
-		/// Constructor
-		/// </summary>
-		public BundleHeader(BlobType[] types, BlobLocator[] imports, BundleExport[] exports, BundlePacket[] packets)
-		{
-			Types = new BundleTypeCollection(types);
-			Imports = new BundleImportCollection(imports);
-			Exports = new BundleExportCollection(exports);
-			Packets = new BundlePacketCollection(packets);
-		}
-
-		/// <summary>
-		/// Writes data for this bundle to a sequence builder
-		/// </summary>
-		public void AppendTo(ReadOnlySequenceBuilder<byte> builder)
-		{
-			// Create a writer for the output data
-			ChunkedArrayMemoryWriter writer = new ChunkedArrayMemoryWriter();
-
-			// Create the prelude
-			Span<byte> prelude = writer.GetSpanAndAdvance(PreludeLength);
-			prelude[0] = (byte)'U';
-			prelude[1] = (byte)'B';
-			prelude[2] = (byte)'N';
-			prelude[3] = (byte)BundleVersion.Latest;
-
-			// Write all the sections
-			WriteSectionHeader(writer, BundleSectionType.Types, Types.Measure());
-			Types.Write(writer);
-
-			WriteSectionHeader(writer, BundleSectionType.Imports, Imports.Measure());
-			Imports.Write(writer);
-
-			WriteSectionHeader(writer, BundleSectionType.ExportHeaders, Exports.MeasureHeaders());
-			Exports.WriteHeaders(writer);
-
-			WriteSectionHeader(writer, BundleSectionType.ExportReferences, Exports.MeasureRefs());
-			Exports.WriteRefs(writer);
-
-			WriteSectionHeader(writer, BundleSectionType.Packets, Packets.Measure());
-			Packets.Write(writer);
-
-			// Fill in the length of the header, and append it to the builder
-			BinaryPrimitives.WriteInt32LittleEndian(prelude[4..], writer.Length);
-			writer.AppendTo(builder);
-		}
-
-		/// <summary>
-		/// Serialize the bundle to a byte array
-		/// </summary>
-		public byte[] ToByteArray()
-		{
-			ReadOnlySequenceBuilder<byte> builder = new ReadOnlySequenceBuilder<byte>();
-			AppendTo(builder);
-			return builder.Construct().ToArray();
-		}
-
-		const int SectionHeaderLength = 4;
-
-		static void WriteSectionHeader(IMemoryWriter writer, BundleSectionType type, int length)
-		{
-			if (length >= 0x1000000)
-			{
-				throw new InvalidDataException($"Section header exceeds maximum allowed size (type: {type}, length: {length:n0}");
-			}
-			writer.WriteInt32((int)type | (length << 8));
-		}
-
-		record class ExportInfo(int TypeIdx, int Length, List<BundleExportRef> References);
-
-		/// <summary>
-		/// Reads a bundle header from memory
-		/// </summary>
-		/// <param name="memory">Memory to read from</param>
-		/// <returns>New header object</returns>
-		public static BundleHeader Read(ReadOnlyMemory<byte> memory)
-		{
-			using ReadOnlyMemoryStream stream = new ReadOnlyMemoryStream(memory);
-			return ReadAsync(stream).Result;
-		}
-
-		/// <summary>
-		/// Reads a bundle header from a stream
-		/// </summary>
-		/// <param name="stream">Stream to deserialize from</param>
-		/// <param name="cancellationToken">Cancellation token for the operation</param>
-		/// <returns>New header object</returns>
-		public static async Task<BundleHeader> ReadAsync(Stream stream, CancellationToken cancellationToken = default)
-		{
-			byte[] prelude = new byte[PreludeLength];
-			await stream.ReadFixedLengthBytesAsync(prelude, cancellationToken);
-
-			return await ReadAsync(prelude, stream, cancellationToken);
-		}
-
-		/// <summary>
-		/// Reads a bundle header from a stream
-		/// </summary>
-		/// <param name="prelude">Prelude data</param>
-		/// <param name="stream">Stream to deserialize from</param>
-		/// <param name="cancellationToken">Cancellation token for the operation</param>
-		/// <returns>New header object</returns>
-		public static async Task<BundleHeader> ReadAsync(ReadOnlyMemory<byte> prelude, Stream stream, CancellationToken cancellationToken = default)
-		{
-			if (prelude.Span.StartsWith(s_latestSignature))
-			{
-				BundleVersion version = (BundleVersion)prelude.Span[3];
-				int headerLength = BinaryPrimitives.ReadInt32LittleEndian(prelude.Span.Slice(4));
-
-				return await ReadLatestAsync(version, stream, headerLength - PreludeLength, cancellationToken);
-			}
-			else if (prelude.Span.StartsWith(s_legacySignature))
-			{
-				int headerLength = BinaryPrimitives.ReadInt32BigEndian(prelude.Span.Slice(4));
-
-				byte[] data = new byte[headerLength - PreludeLength];
-				await stream.ReadFixedLengthBytesAsync(data, cancellationToken);
-
-				return ReadLegacy(data);
-			}
-			else
-			{
-				throw new NotSupportedException();
-			}
-		}
-
-		static readonly byte[] s_latestSignature = new byte[] { (byte)'U', (byte)'B', (byte)'N' };
-
-		/// <summary>
-		/// Construct a header from the given data encoded in the latest format
+		/// Number of bytes in a signature when serialized
 		/// </summary>
-		/// <param name="version">Version to serialize as</param>
-		/// <param name="stream">Stream to read from</param>
-		/// <param name="length">Length of the header</param>
-		/// <param name="cancellationToken">Cancellation token for the operation</param>
-		static async Task<BundleHeader> ReadLatestAsync(BundleVersion version, Stream stream, int length, CancellationToken cancellationToken)
-		{
-			byte[] exportHeaders = Array.Empty<byte>();
-			byte[] exportReferences = Array.Empty<byte>();
-
-			BundleTypeCollection types = new BundleTypeCollection();
-			BundleImportCollection imports = new BundleImportCollection();
-			BundleExportCollection exports = new BundleExportCollection();
-			BundlePacketCollection packets = new BundlePacketCollection();
-
-			byte[] sectionHeader = new byte[SectionHeaderLength];
-			for (int offset = 0; offset < length;)
-			{
-				await stream.ReadFixedLengthBytesAsync(sectionHeader, cancellationToken);
-				offset += SectionHeaderLength;
-
-				BundleSectionType sectionType = (BundleSectionType)(sectionHeader[0] & 255);
-				int sectionLength = (int)(BinaryPrimitives.ReadUInt32LittleEndian(sectionHeader) >> 8);
-
-				switch (sectionType)
-				{
-					case BundleSectionType.Types:
-						types = await BundleTypeCollection.ReadAsync(stream, sectionLength, cancellationToken);
-						break;
-					case BundleSectionType.Imports:
-						imports = await BundleImportCollection.ReadAsync(stream, sectionLength, cancellationToken);
-						break;
-					case BundleSectionType.Exports:
-						exports = await BundleExportCollection.ReadAsync(stream, sectionLength, cancellationToken);
-						break;
-					case BundleSectionType.ExportHeaders:
-						exportHeaders = new byte[sectionLength];
-						await stream.ReadFixedLengthBytesAsync(exportHeaders, cancellationToken);
-						break;
-					case BundleSectionType.ExportReferences:
-						exportReferences = new byte[sectionLength];
-						await stream.ReadFixedLengthBytesAsync(exportReferences, cancellationToken);
-						break;
-					case BundleSectionType.Packets:
-						packets = await BundlePacketCollection.ReadAsync(stream, sectionLength, cancellationToken);
-						break;
-				}
-
-				offset += sectionLength;
-			}
-
-			if (exportHeaders.Length > 0)
-			{
-				exports = BundleExportCollection.Read(exportHeaders, exportReferences, version);
-			}
-
-			return new BundleHeader(types, imports, exports, packets);
-		}
-
-		static readonly byte[] s_legacySignature = new byte[] { (byte)'U', (byte)'E', (byte)'B', (byte)'N' };
-
-		static BundleHeader ReadLegacy(ReadOnlyMemory<byte> memory)
-		{
-			MemoryReader reader = new MemoryReader(memory);
-
-			BundleVersion version = (BundleVersion)reader.ReadUnsignedVarInt();
-			if (version > BundleVersion.RemoveAliases)
-			{
-				throw new InvalidDataException($"Invalid bundle version {(int)version}. Max supported version in legacy format is {(int)BundleVersion.RemoveAliases}.");
-			}
-
-			BundleCompressionFormat compressionFormat = (BundleCompressionFormat)reader.ReadUnsignedVarInt();
-
-			// Read the types
-			int numTypes = (int)reader.ReadUnsignedVarInt();
-			List<BlobType> types = new List<BlobType>(numTypes);
-
-			for (int typeIdx = 0; typeIdx < numTypes; typeIdx++)
-			{
-				Guid guid = reader.ReadGuid();
-				int serializerVersion = (int)reader.ReadUnsignedVarInt();
-
-				types.Add(new BlobType(guid, serializerVersion));
-			}
-
-			// Read the imports
-			int numImports = (int)reader.ReadUnsignedVarInt();
-			List<BlobLocator> imports = new List<BlobLocator>(numImports);
-			List<BundleExportRef> allExportReferences = new List<BundleExportRef>();
-
-			for (int importIdx = 0; importIdx < numImports; importIdx++)
-			{
-				BlobLocator locator = new BlobLocator(reader.ReadUtf8String());
-				imports.Add(locator);
-
-				int[] exportIndexes = reader.ReadVariableLengthArray(() => (int)reader.ReadUnsignedVarInt());
-				for (int exportIdx = 0; exportIdx < exportIndexes.Length; exportIdx++)
-				{
-					BundleExportRef exportReference = new BundleExportRef(imports.Count - 1, exportIndexes[exportIdx]);
-					allExportReferences.Add(exportReference);
-				}
-			}
-
-			// Read the exports
-			int numExports = (int)reader.ReadUnsignedVarInt();
-			List<ExportInfo> exportInfos = new List<ExportInfo>(numExports);
-
-			for (int exportIdx = 0; exportIdx < numExports; exportIdx++)
-			{
-				allExportReferences.Add(new BundleExportRef(-1, exportIdx));
-
-				int typeIdx = (int)reader.ReadUnsignedVarInt();
-				_ = reader.ReadIoHash();
-				int length = (int)reader.ReadUnsignedVarInt();
-
-				List<BundleExportRef> exportRefs = new List<BundleExportRef>();
-
-				int numReferences = (int)reader.ReadUnsignedVarInt();
-				if (numReferences > 0)
-				{
-					for (int idx = 0; idx < numReferences; idx++)
-					{
-						int referenceIdx = (int)reader.ReadUnsignedVarInt();
-						exportRefs.Add(allExportReferences[referenceIdx]);
-					}
-				}
-
-				if (version == BundleVersion.ExportAliases)
-				{
-					_ = reader.ReadUtf8String();
-				}
-
-				exportInfos.Add(new ExportInfo(typeIdx, length, exportRefs));
-			}
-
-			// Read the compression packets
-			List<BundlePacket> packets;
-			if (compressionFormat == BundleCompressionFormat.None)
-			{
-				packets = new List<BundlePacket>(exportInfos.Count);
-
-				int encodedOffset = 0;
-				foreach (ExportInfo export in exportInfos)
-				{
-					packets.Add(new BundlePacket(compressionFormat, encodedOffset, export.Length, export.Length));
-					encodedOffset += export.Length;
-				}
-			}
-			else
-			{
-				int numPackets = (int)reader.ReadUnsignedVarInt();
-				packets = new List<BundlePacket>(numPackets);
-
-				int encodedOffset = 0;
-				for (int packetIdx = 0; packetIdx < numPackets; packetIdx++)
-				{
-					int encodedLength = (int)reader.ReadUnsignedVarInt();
-					int decodedLength = (int)reader.ReadUnsignedVarInt();
-					packets.Add(new BundlePacket(compressionFormat, encodedOffset, encodedLength, decodedLength));
-					encodedOffset += encodedLength;
-				}
-			}
-
-			// Create the final export list
-			List<BundleExport> exports = new List<BundleExport>();
-			{
-				int packetIdx = 0;
-				int packetOffset = 0;
-
-				foreach (ExportInfo exportInfo in exportInfos)
-				{
-					if (packetOffset + exportInfo.Length > packets[packetIdx].DecodedLength)
-					{
-						packetIdx++;
-						packetOffset = 0;
-					}
-
-					exports.Add(new BundleExport(exportInfo.TypeIdx, packetIdx, packetOffset, exportInfo.Length, exportInfo.References));
-					packetOffset += exportInfo.Length;
-				}
-			}
-
-			return new BundleHeader(types.ToArray(), imports.ToArray(), exports.ToArray(), packets.ToArray());
-		}
+		public const int SignatureLength = 8;
 
 		/// <summary>
 		/// Validates that the prelude bytes for a bundle header are correct
 		/// </summary>
-		/// <param name="prelude">The prelude bytes</param>
+		/// <param name="span">The signature bytes</param>
 		/// <returns>Length of the header data, including the prelude</returns>
-		public static int ReadPrelude(ReadOnlySpan<byte> prelude)
+		public static BundleSignature ReadSignature(ReadOnlySpan<byte> span)
 		{
-			if (prelude[0] == 'U' && prelude[1] == 'E' && prelude[2] == 'B' && prelude[3] == 'N')
+			if (span[0] == 'U' && span[1] == 'E' && span[2] == 'B' && span[3] == 'N')
 			{
-				return BinaryPrimitives.ReadInt32BigEndian(prelude.Slice(4));
+				return new BundleSignature(BundleVersion.Initial, BinaryPrimitives.ReadInt32BigEndian(span.Slice(4)));
 			}
-			else if (prelude[0] == 'U' && prelude[1] == 'B' && prelude[2] == 'N')
+			else if (span[0] == 'U' && span[1] == 'B' && span[2] == 'N')
 			{
-				return BinaryPrimitives.ReadInt32LittleEndian(prelude.Slice(4));
+				return new BundleSignature((BundleVersion)span[3], BinaryPrimitives.ReadInt32LittleEndian(span.Slice(4)));
 			}
 			else
 			{
 				throw new InvalidDataException("Invalid signature bytes for bundle. Corrupt data?");
 			}
 		}
-	}
-
-	/// <summary>
-	/// Collection of node types in a bundle
-	/// </summary>
-	public struct BundleTypeCollection : IReadOnlyList<BlobType>
-	{
-		readonly BlobType[] _types;
-
-		/// <inheritdoc/>
-		public int Count => _types.Length;
-
-		/// <summary>
-		/// Constructor
-		/// </summary>
-		public BundleTypeCollection(BlobType[] types) => _types = types;
-
-		/// <summary>
-		/// Reads a type collection from a block of memory
-		/// </summary>
-		public static BundleTypeCollection Read(ReadOnlyMemory<byte> data)
-		{
-			BlobType[] types = new BlobType[data.Length / BlobType.NumBytes];
-			for (int idx = 0; idx < types.Length; idx++)
-		{
-				types[idx] = BlobType.Read(data.Span.Slice(idx * BlobType.NumBytes));
-			}
-
-			return new BundleTypeCollection(types);
-		}
-
-		/// <summary>
-		/// Reads a type collection from a stream
-		/// </summary>
-		public static async Task<BundleTypeCollection> ReadAsync(Stream stream, int length, CancellationToken cancellationToken)
-		{
-			byte[] data = new byte[length];
-			await stream.ReadFixedLengthBytesAsync(data, cancellationToken);
-			return Read(data);
-		}
-
-		/// <summary>
-		/// Gets the size of memory required to serialize a collection of types
-		/// </summary>
-		/// <returns>Number of bytes required to serialize the types</returns>
-		public int Measure() => _types.Length * BlobType.NumBytes;
 
 		/// <summary>
-		/// Serializes a set of types to a fixed block of memory
+		/// Writes a signature to the given memory
 		/// </summary>
-		/// <param name="writer">Writer to serialize to</param>
-		public void Write(IMemoryWriter writer)
+		public static void WriteSignature(Span<byte> span, BundleSignature signature)
 		{
-			Span<byte> span = writer.GetSpanAndAdvance(_types.Length * BlobType.NumBytes);
-			foreach (BlobType type in _types)
+			if (signature.Version == BundleVersion.Initial)
 			{
-				type.Write(span);
-				span = span.Slice(BlobType.NumBytes);
+				span[0] = (byte)'U';
+				span[1] = (byte)'E';
+				span[2] = (byte)'B';
+				span[3] = (byte)'N';
+				BinaryPrimitives.WriteInt32BigEndian(span[4..], signature.HeaderLength);
 			}
-		}
-
-		/// <inheritdoc/>
-		public BlobType this[int index] => _types[index];
-
-		/// <inheritdoc/>
-		public IEnumerator<BlobType> GetEnumerator() => ((IEnumerable<BlobType>)_types).GetEnumerator();
-
-		/// <inheritdoc/>
-		IEnumerator IEnumerable.GetEnumerator() => GetEnumerator();
-	}
-
-	/// <summary>
-	/// Collection of imported node references
-	/// </summary>
-	public struct BundleImportCollection : IReadOnlyList<BlobLocator>
-	{
-		readonly BlobLocator[] _imports;
-
-		/// <inheritdoc/>
-		public int Count => _imports.Length;
-
-		/// <summary>
-		/// Constructor
-		/// </summary>
-		public BundleImportCollection(BlobLocator[] imports) => _imports = imports;
-
-		/// <summary>
-		/// Reads a collection from memory
-		/// </summary>
-		public static BundleImportCollection Read(ReadOnlyMemory<byte> data)
-		{
-			ReadOnlySpan<byte> span = data.Span;
-			int count = (span.Length == 0) ? 0 : BinaryPrimitives.ReadInt32LittleEndian(span) / sizeof(int);
-
-			BlobLocator[] imports = new BlobLocator[count];
-			for (int idx = 0; idx < count; idx++)
-		{
-				int offset = BinaryPrimitives.ReadInt32LittleEndian(span.Slice(idx * sizeof(int)));
-				int length = span.Slice(offset).IndexOf((byte)0);
-				imports[idx] = new BlobLocator(new Utf8String(data.Slice(offset, length)));
-			}
-
-			return new BundleImportCollection(imports);
-		}
-
-		/// <summary>
-		/// Reads a collection from a stream
-		/// </summary>
-		public static async Task<BundleImportCollection> ReadAsync(Stream stream, int length, CancellationToken cancellationToken)
-		{
-			byte[] data = new byte[length];
-			await stream.ReadFixedLengthBytesAsync(data, cancellationToken);
-			return Read(data);
-		}
-
-		/// <inheritdoc/>
-		public BlobLocator this[int index] => _imports[index];
-
-		/// <inheritdoc/>
-		public IEnumerator<BlobLocator> GetEnumerator() => _imports.AsEnumerable().GetEnumerator();
-
-		/// <inheritdoc/>
-		IEnumerator IEnumerable.GetEnumerator() => GetEnumerator();
-
-		/// <summary>
-		/// Measure the size of memory required to store a collection of import locators
-		/// </summary>
-		/// <returns>Size in bytes of the output buffer</returns>
-		public int Measure() => _imports.Sum(x => sizeof(int) + x.Path.Length + 1);
-
-		/// <summary>
-		/// Serialize a collection of locators to memory
-		/// </summary>
-		/// <param name="writer">Writer for output data</param>
-		public void Write(IMemoryWriter writer)
-		{
-			int offset = _imports.Length * sizeof(int);
-			foreach (BlobLocator import in _imports)
-			{
-				writer.WriteInt32(offset);
-				offset += import.Path.Length + 1;
-			}
-			foreach (BlobLocator import in _imports)
-			{
-				writer.WriteNullTerminatedUtf8String(import.Path);
-			}
-		}
-	}
-
-	/// <summary>
-	/// Descriptor for a compression packet
-	/// </summary>
-	public struct BundlePacket
-	{
-		/// <summary>
-		/// Size of this structure when serialized
-		/// </summary>
-		public const int NumBytes = 16;
-
-		/// <summary>
-		/// Compression format for this packet
-		/// </summary>
-		public BundleCompressionFormat CompressionFormat { get; }
-
-		/// <summary>
-		/// Offset of the packet within the payload stream
-		/// </summary>
-		public int EncodedOffset { get; }
-
-		/// <summary>
-		/// Encoded length of the packet
-		/// </summary>
-		public int EncodedLength { get; }
-
-		/// <summary>
-		/// Decoded length of the packet
-		/// </summary>
-		public int DecodedLength { get; }
-
-		/// <summary>
-		/// Constructor
-		/// </summary>
-		/// <param name="compressionFormat">Compression format for the packet</param>
-		/// <param name="encodedOffset">Offset of the data within the payload stream</param>
-		/// <param name="encodedLength">Size of the encoded data</param>
-		/// <param name="decodedLength">Size of the decoded data</param>
-		public BundlePacket(BundleCompressionFormat compressionFormat, int encodedOffset, int encodedLength, int decodedLength)
-		{
-			CompressionFormat = compressionFormat;
-			EncodedOffset = encodedOffset;
-			EncodedLength = encodedLength;
-			DecodedLength = decodedLength;
-		}
-
-		/// <summary>
-		/// Read from a byte array
-		/// </summary>
-		public static BundlePacket Read(ReadOnlySpan<byte> span)
-		{
-			BundleCompressionFormat compressionFormat = (BundleCompressionFormat)BinaryPrimitives.ReadInt32LittleEndian(span);
-			int encodedOffset = BinaryPrimitives.ReadInt32LittleEndian(span.Slice(4));
-			int encodedLength = BinaryPrimitives.ReadInt32LittleEndian(span.Slice(8));
-			int decodedLength = BinaryPrimitives.ReadInt32LittleEndian(span.Slice(12));
-			return new BundlePacket(compressionFormat, encodedOffset, encodedLength, decodedLength);
-		}
-
-		/// <summary>
-		/// Serialize the struct to memory
-		/// </summary>
-		public void Write(Span<byte> span)
-		{
-			BinaryPrimitives.WriteInt32LittleEndian(span, (int)CompressionFormat);
-			BinaryPrimitives.WriteInt32LittleEndian(span.Slice(4), EncodedOffset);
-			BinaryPrimitives.WriteInt32LittleEndian(span.Slice(8), EncodedLength);
-			BinaryPrimitives.WriteInt32LittleEndian(span.Slice(12), DecodedLength);
-		}
-	}
-
-	/// <summary>
-	/// Collection of information about packets in a bundle
-	/// </summary>
-	public struct BundlePacketCollection : IReadOnlyCollection<BundlePacket>
-	{
-		readonly BundlePacket[] _packets;
-
-		/// <inheritdoc/>
-		public int Count => _packets?.Length ?? 0;
-
-		/// <summary>
-		/// Constructor
-		/// </summary>
-		public BundlePacketCollection(BundlePacket[] packets) => _packets = packets;
-
-		/// <summary>
-		/// Reads a collection from memory
-		/// </summary>
-		public static BundlePacketCollection Read(ReadOnlyMemory<byte> data)
-		{
-			int count = data.Length / BundlePacket.NumBytes;
-
-			BundlePacket[] packets = new BundlePacket[count];
-			for (int idx = 0; idx < count; idx++)
-			{
-				packets[idx] = BundlePacket.Read(data.Span.Slice(idx * BundlePacket.NumBytes));
-			}
-
-			return new BundlePacketCollection(packets);
-		}
-
-		/// <summary>
-		/// Reads a collection from a stream
-		/// </summary>
-		public static async Task<BundlePacketCollection> ReadAsync(Stream stream, int length, CancellationToken cancellationToken)
-		{
-			byte[] data = new byte[length];
-			await stream.ReadFixedLengthBytesAsync(data, cancellationToken);
-			return BundlePacketCollection.Read(data);
-		}
-
-		/// <inheritdoc/>
-		public BundlePacket this[int index] => _packets[index];
-
-		/// <inheritdoc/>
-		public IEnumerator<BundlePacket> GetEnumerator() => _packets.AsEnumerable().GetEnumerator();
-
-		/// <inheritdoc/>
-		IEnumerator IEnumerable.GetEnumerator() => GetEnumerator();
-
-		/// <summary>
-		/// Measure the size of memory required to store a collection of packets
-		/// </summary>
-		/// <returns>Size in bytes of the output buffer</returns>
-		public int Measure() => BundlePacket.NumBytes * _packets.Length;
-
-		/// <summary>
-		/// Serialize a collection of packets to memory
-		/// </summary>
-		/// <param name="writer">Writer to serialize to</param>
-		public void Write(IMemoryWriter writer)
-		{
-			Span<byte> span = writer.GetMemoryAndAdvance(_packets.Length * BundlePacket.NumBytes).Span;
-			foreach (BundlePacket packet in _packets)
-			{
-				packet.Write(span);
-				span = span.Slice(BundlePacket.NumBytes);
-			}
-		}
-	}
-
-	/// <summary>
-	/// Entry for a node exported from an object
-	/// </summary>
-	public struct BundleExport
-	{
-		/// <summary>
-		/// Number of bytes in a serialized export object
-		/// </summary>
-		public const int NumHeaderBytes = 32;
-
-		/// <summary>
-		/// Raw data for this export
-		/// </summary>
-		public ReadOnlyMemory<byte> Data { get; }
-
-		/// <summary>
-		/// Raw data for the header
-		/// </summary>
-		public ReadOnlyMemory<byte> HeaderData => Data.Slice(0, NumHeaderBytes);
-
-		/// <summary>
-		/// Hash of the node data
-		/// </summary>
-		public IoHash Hash => new IoHash(Data.Slice(0, 20).Span);
-
-		/// <summary>
-		/// Type id of the node. Can be used to look up the type information from the bundle header.
-		/// </summary>
-		public int TypeIdx => BinaryPrimitives.ReadUInt16LittleEndian(Data.Slice(20).Span);
-
-		/// <summary>
-		/// Packet containing this export's data
-		/// </summary>
-		public int Packet => BinaryPrimitives.ReadUInt16LittleEndian(Data.Slice(22).Span);
-
-		/// <summary>
-		/// Offset within the packet of the node data
-		/// </summary>
-		public int Offset => BinaryPrimitives.ReadInt32LittleEndian(Data.Slice(24).Span);
-
-		/// <summary>
-		/// Length of the node
-		/// </summary>
-		public int Length => BinaryPrimitives.ReadInt32LittleEndian(Data.Slice(28).Span);
-
-		/// <summary>
-		/// References to other nodes
-		/// </summary>
-		public BundleExportRefCollection References => new BundleExportRefCollection(Data.Slice(32));
-
-		/// <summary>
-		/// Constructor
-		/// </summary>
-		public BundleExport(ReadOnlyMemory<byte> data)
-		{
-			Data = data;
-		}
-
-		/// <summary>
-		/// Constructor
-		/// </summary>
-		public BundleExport(int typeIdx, int packetIdx, int offset, int length, IReadOnlyList<BundleExportRef> references)
-		{
-			byte[] data = new byte[NumHeaderBytes + (BundleExportRef.NumBytes * references.Count)];
-			WriteHeader(data, typeIdx, packetIdx, offset, length);
-
-			for (int idx = 0; idx < references.Count; idx++)
-			{
-				references[idx].CopyTo(data.AsSpan(NumHeaderBytes + (BundleExportRef.NumBytes * idx)));
-			}
-
-			Data = data;
-		}
-
-		/// <summary>
-		/// Writes a new export to a block of memory
-		/// </summary>
-		public static void WriteHeader(Span<byte> data, int typeIdx, int packet, int offset, int length)
-		{
-			IoHash.Zero.CopyTo(data); // Previously hash of export
-			BinaryPrimitives.WriteUInt16LittleEndian(data.Slice(20), (ushort)typeIdx);
-			BinaryPrimitives.WriteUInt16LittleEndian(data.Slice(22), (ushort)packet);
-			BinaryPrimitives.WriteInt32LittleEndian(data.Slice(24), offset);
-			BinaryPrimitives.WriteInt32LittleEndian(data.Slice(28), length);
-		}
-	}
-
-	/// <summary>
-	/// Entry for a node exported from an object
-	/// </summary>
-	public struct BundleExportCollection : IReadOnlyList<BundleExport>
-	{
-		readonly BundleExport[] _exports;
-
-		/// <inheritdoc/>
-		public int Count => _exports.Length;
-
-		/// <summary>
-		/// Constructor
-		/// </summary>
-		public BundleExportCollection(BundleExport[] exports) => _exports = exports;
-
-		/// <summary>
-		/// Constructor
-		/// </summary>
-		public static BundleExportCollection Read(ReadOnlyMemory<byte> data, ReadOnlyMemory<byte> refs, BundleVersion version)
-		{
-			int count = data.Length / BundleExport.NumHeaderBytes;
-
-			if (version < BundleVersion.ImportHashes && refs.Length > 0)
-			{
-				const int OldRefNumBytes = 4;
-
-				int headerLength = BinaryPrimitives.ReadInt32LittleEndian(refs.Span);
-				int refCount = (refs.Length - headerLength) / OldRefNumBytes;
-
-				byte[] newRefs = new byte[headerLength + (refCount * BundleExportRef.NumBytes)];
-				for (int headerOffset = 0; headerOffset < headerLength; headerOffset += sizeof(int))
-				{
-					int index = (BinaryPrimitives.ReadInt32LittleEndian(refs.Slice(headerOffset).Span) - headerLength) / OldRefNumBytes;
-					BinaryPrimitives.WriteInt32LittleEndian(newRefs.AsSpan(headerOffset), headerLength + (index * BundleExportRef.NumBytes));
-				}
-				for (int index = 0; index < refCount; index++)
-				{
-					ReadOnlyMemory<byte> source = refs.Slice(headerLength + (index * OldRefNumBytes), OldRefNumBytes);
-					source.CopyTo(newRefs.AsMemory(headerLength + (index * BundleExportRef.NumBytes)));
-				}
-
-				refs = newRefs;
-			}
-
-			BundleExport[] exports = new BundleExport[count];
-			for(int index = 0; index < count; index++)
-			{
-				ReadOnlyMemory<byte> exportData = data.Slice(index * BundleExport.NumHeaderBytes, BundleExport.NumHeaderBytes);
-
-				BundleExportRefCollection exportRefs = new BundleExportRefCollection();
-				if (refs.Length > 0)
-				{
-					int minOffset = BinaryPrimitives.ReadInt32LittleEndian(refs.Span.Slice(index * sizeof(int)));
-					int maxOffset = (index == count - 1) ? refs.Length : BinaryPrimitives.ReadInt32LittleEndian(refs.Span.Slice((index + 1) * sizeof(int)));
-					exportRefs = new BundleExportRefCollection(refs.Slice(minOffset, maxOffset - minOffset));
-				}
-
-				byte[] export = new byte[BundleExport.NumHeaderBytes + exportRefs.Data.Length];
-				exportData.CopyTo(export);
-				exportRefs.Data.CopyTo(export.AsMemory(BundleExport.NumHeaderBytes));
-
-				exports[index] = new BundleExport(export);
-			}
-
-			return new BundleExportCollection(exports);
-		}
-
-		/// <summary>
-		/// Reads an export collection from a stream
-		/// </summary>
-		public static async Task<BundleExportCollection> ReadAsync(Stream stream, int sectionLength, CancellationToken cancellationToken)
-		{
-			BundleExport[] exports = Array.Empty<BundleExport>();
-			if (sectionLength > 0)
-			{
-				const int MaxPageLength = 16384;
-
-				// Read the offsets of each entry
-				byte[] headerBuffer = new byte[sizeof(int)];
-				await stream.ReadFixedLengthBytesAsync(headerBuffer, cancellationToken);
-
-				int headerLength = BinaryPrimitives.ReadInt32LittleEndian(headerBuffer);
-				Array.Resize(ref headerBuffer, headerLength);
-
-				await stream.ReadFixedLengthBytesAsync(headerBuffer.AsMemory(sizeof(int)), cancellationToken);
-
-				// Create the packet array
-				int exportCount = headerLength / sizeof(int);
-				exports = new BundleExport[exportCount];
-
-				// Read the exports in pages
-				int exportIdx = 0;
-				for (int exportOffset = headerLength; exportOffset < sectionLength;)
-				{
-					int pageOffset = exportOffset;
-
-					// Find the exports to read into this page
-					int nextPageOffset = sectionLength;
-					for (int nextExportIdx = exportIdx + 1; nextExportIdx < exportCount; nextExportIdx++)
-					{
-						int nextExportOffset = BinaryPrimitives.ReadInt32LittleEndian(headerBuffer.AsSpan(nextExportIdx * sizeof(int)));
-						if (nextExportOffset > pageOffset + MaxPageLength)
-						{
-							nextPageOffset = nextExportOffset;
-							break;
-						}
-					}
-
-					// Read the page data
-					byte[] pageBuffer = new byte[nextPageOffset - pageOffset];
-					await stream.ReadFixedLengthBytesAsync(pageBuffer, cancellationToken);
-
-					// Parse the individual exports
-					for (; exportOffset < nextPageOffset; exportIdx++)
-					{
-						int startExportOffset = exportOffset;
-						if (exportIdx + 1 < exportCount)
-						{
-							exportOffset = BinaryPrimitives.ReadInt32LittleEndian(headerBuffer.AsSpan((exportIdx + 1) * sizeof(int)));
-						}
-						else
-						{
-							exportOffset = nextPageOffset;
-						}
-						exports[exportIdx] = new BundleExport(pageBuffer.AsMemory(startExportOffset - pageOffset, exportOffset - startExportOffset));
-					}
-				}
-			}
-			return new BundleExportCollection(exports);
-		}
-
-		/// <inheritdoc/>
-		public BundleExport this[int index] => _exports[index];
-
-		/// <inheritdoc/>
-		public IEnumerator<BundleExport> GetEnumerator() => _exports.AsEnumerable().GetEnumerator();
-
-		/// <inheritdoc/>
-		IEnumerator IEnumerable.GetEnumerator() => GetEnumerator();
-
-		/// <summary>
-		/// Serializes this collection as a single section
-		/// </summary>
-		public void Write(IMemoryWriter writer)
-		{
-			int offset = _exports.Length * sizeof(int);
-			Span<byte> index = writer.GetSpanAndAdvance(offset);
-
-			foreach (BundleExport export in _exports)
-			{
-				BinaryPrimitives.WriteInt32LittleEndian(index, offset);
-				index = index.Slice(sizeof(int));
-
-				writer.WriteFixedLengthBytes(export.Data.Span);
-				offset += export.Data.Length;
-			}
-		}
-
-		/// <summary>
-		/// Measure the size of memory required to store a collection of exports
-		/// </summary>
-		/// <returns>Size in bytes of the output buffer</returns>
-		public int MeasureHeaders() => BundleExport.NumHeaderBytes * _exports.Length;
-
-		/// <summary>
-		/// Serialize a collection of exports to memory
-		/// </summary>
-		/// <param name="writer">Writer to serialize to</param>
-		public void WriteHeaders(IMemoryWriter writer)
-		{
-			foreach (BundleExport export in _exports)
-			{
-				writer.WriteFixedLengthBytes(export.HeaderData.Span);
-			}
-		}
-
-		/// <summary>
-		/// Measure the size of memory required to store a collection of export refs
-		/// </summary>
-		/// <returns>Size in bytes of the output buffer</returns>
-		public int MeasureRefs() => (sizeof(int) * _exports.Length) + _exports.Sum(x => x.References.Data.Length);
-
-		/// <summary>
-		/// Serialize a collection of export refs to memory
-		/// </summary>
-		/// <param name="writer">Writer to serialize to</param>
-		public void WriteRefs(IMemoryWriter writer)
-		{
-			int offset = sizeof(int) * _exports.Length;
-			foreach(BundleExport export in _exports)
-			{
-				writer.WriteInt32(offset);
-				offset += export.References.Data.Length;
-			}
-			foreach (BundleExport export in _exports)
-			{
-				writer.WriteFixedLengthBytes(export.References.Data.Span);
-			}
-		}
-	}
-
-	/// <summary>
-	/// Reference to a node in another bundle
-	/// </summary>
-	/// <param name="ImportIdx">Index into the import table of the blob containing the referenced node. Can be -1 for references within the same bundle.</param>
-	/// <param name="NodeIdx">Node imported from the bundle</param>
-	public record struct BundleExportRef(int ImportIdx, int NodeIdx)
-	{
-		/// <summary>
-		/// Number of bytes in the serialized object
-		/// </summary>
-		public const int NumBytes = 4 + IoHash.NumBytes;
-
-		/// <summary>
-		/// Deserialize this object from memory
-		/// </summary>
-		public static BundleExportRef Read(ReadOnlySpan<byte> data)
-		{
-			int importIdx = BinaryPrimitives.ReadInt16LittleEndian(data);
-			int nodeIdx = BinaryPrimitives.ReadUInt16LittleEndian(data.Slice(2));
-			return new BundleExportRef(importIdx, nodeIdx);
-		}
-
-		/// <summary>
-		/// Serialize this object to memory
-		/// </summary>
-		public void CopyTo(Span<byte> data)
-		{
-			if (ImportIdx > 0xffff)
-			{
-				throw new InvalidDataException($"Import index is out of range ({ImportIdx})");
-			}
-			if (NodeIdx > 0xffff)
-			{
-				throw new InvalidDataException($"Node index is out of range ({NodeIdx})");
-			}
-
-			BinaryPrimitives.WriteInt16LittleEndian(data, (short)ImportIdx);
-			BinaryPrimitives.WriteUInt16LittleEndian(data.Slice(2), (ushort)NodeIdx);
-		}
-	}
-
-	/// <summary>
-	/// Collection of information about exported nodes
-	/// </summary>
-	public struct BundleExportRefCollection : IReadOnlyList<BundleExportRef>
-	{
-		/// <summary>
-		/// Data used to store this collection
-		/// </summary>
-		public ReadOnlyMemory<byte> Data { get; }
-
-		/// <summary>
-		/// Constructor
-		/// </summary>
-		public BundleExportRefCollection(ReadOnlyMemory<byte> data) => Data = data;
-
-		/// <summary>
-		/// Constructor
-		/// </summary>
-		public BundleExportRefCollection(IReadOnlyList<BundleExportRef> exportRefs)
-		{
-			byte[] data = new byte[exportRefs.Count * BundleExportRef.NumBytes];
-			for (int idx = 0; idx < exportRefs.Count; idx++)
-			{
-				exportRefs[idx].CopyTo(data.AsSpan(idx * BundleExportRef.NumBytes));
-			}
-			Data = data;
-		}
-
-		/// <inheritdoc/>
-		public int Count => Data.Length / BundleExportRef.NumBytes;
-
-		/// <inheritdoc/>
-		public BundleExportRef this[int index]
-		{
-			get
-			{
-				Debug.Assert(index < Count);
-				return BundleExportRef.Read(Data.Span.Slice(index * BundleExportRef.NumBytes));
-			}
-		}
-
-		/// <inheritdoc/>
-		public IEnumerator<BundleExportRef> GetEnumerator()
-		{
-			ReadOnlyMemory<byte> data = Data;
-			while (data.Length > 0)
-			{
-				yield return BundleExportRef.Read(data.Span);
-				data = data.Slice(BundleExportRef.NumBytes);
-			}
-		}
-
-		/// <inheritdoc/>
-		IEnumerator IEnumerable.GetEnumerator() => GetEnumerator();
-	}
-
-	/// <summary>
-	/// Utility methods for bundles
-	/// </summary>
-	public static class BundleData
-	{
-		/// <summary>
-		/// Compress a data packet
-		/// </summary>
-		/// <param name="format">Format for the compressed data</param>
-		/// <param name="input">The data to compress</param>
-		/// <param name="writer">Writer for output data</param>
-		/// <returns>The compressed data</returns>
-		public static int Compress(BundleCompressionFormat format, ReadOnlyMemory<byte> input, IMemoryWriter writer)
-		{
-			switch (format)
-			{
-				case BundleCompressionFormat.None:
-					{
-						writer.WriteFixedLengthBytes(input.Span);
-						return input.Length;
-					}
-				case BundleCompressionFormat.LZ4:
-					{
-						int maxSize = LZ4Codec.MaximumOutputSize(input.Length);
-
-						Span<byte> buffer = writer.GetSpan(maxSize);
-						int encodedLength = LZ4Codec.Encode(input.Span, buffer);
-
-						writer.Advance(encodedLength);
-						return encodedLength;
-					}
-				case BundleCompressionFormat.Gzip:
-					{
-						using MemoryStream outputStream = new MemoryStream(input.Length);
-
-						using (GZipStream gzipStream = new GZipStream(outputStream, CompressionLevel.Fastest, true))
-						{
-							gzipStream.Write(input.Span);
-						}
-
-						writer.WriteFixedLengthBytes(outputStream.ToArray());
-						return (int)outputStream.Length;
-					}
-				case BundleCompressionFormat.Oodle:
-					{
-						int maxSize = Oodle.MaximumOutputSize(OodleCompressorType.Selkie, input.Length);
-
-						Span<byte> outputSpan = writer.GetSpan(maxSize);
-						int encodedLength = Oodle.Compress(OodleCompressorType.Selkie, input.Span, outputSpan, OodleCompressionLevel.Normal);
-						writer.Advance(encodedLength);
-
-						return encodedLength;
-					}
-				case BundleCompressionFormat.Brotli:
-					{
-						int maxSize = BrotliEncoder.GetMaxCompressedLength(input.Length);
-
-						Span<byte> buffer = writer.GetSpan(maxSize);
-						if (!BrotliEncoder.TryCompress(input.Span, buffer, out int encodedLength))
-						{
-							throw new InvalidOperationException("Unable to compress data using Brotli");
-						}
-
-						writer.Advance(encodedLength);
-						return encodedLength;
-					}
-				default:
-					throw new InvalidDataException($"Invalid compression format '{(int)format}'");
-			}
-		}
-
-		/// <summary>
-		/// Decompress a packet of data
-		/// </summary>
-		/// <param name="format">Format of the compressed data</param>
-		/// <param name="input">Compressed data</param>
-		/// <param name="output">Buffer to receive the decompressed data</param>
-		public static void Decompress(BundleCompressionFormat format, ReadOnlyMemory<byte> input, Memory<byte> output)
-		{
-			switch (format)
+			else
 			{
-				case BundleCompressionFormat.None:
-					input.CopyTo(output);
-					break;
-				case BundleCompressionFormat.LZ4:
-					LZ4Codec.Decode(input.Span, output.Span);
-					break;
-				case BundleCompressionFormat.Gzip:
-					{
-						using ReadOnlyMemoryStream inputStream = new ReadOnlyMemoryStream(input);
-						using GZipStream inflatedStream = new GZipStream(inputStream, CompressionMode.Decompress, true);
-
-						int length = inflatedStream.Read(output.Span);
-						if (length != output.Length)
-						{
-							throw new InvalidDataException($"Decoded data is shorter than expected (expected {output.Length} bytes, got {length} bytes)");
-						}
-						break;
-					}
-				case BundleCompressionFormat.Oodle:
-#if WITH_OODLE
-					Oodle.Decompress(input.Span, output.Span);
-					break;
-#else
-					throw new NotSupportedException("Oodle is not compiled into this build.");
-#endif
-				case BundleCompressionFormat.Brotli:
-					{
-						int bytesWritten;
-						if (!BrotliDecoder.TryDecompress(input.Span, output.Span, out bytesWritten) || bytesWritten != output.Length)
-						{
-							throw new InvalidOperationException("Unable to decompress data using Brotli");
-						}
-						break;
-					}
-				default:
-					throw new InvalidDataException($"Invalid compression format '{(int)format}'");
+				span[0] = (byte)'U';
+				span[1] = (byte)'B';
+				span[2] = (byte)'N';
+				span[3] = (byte)signature.Version;
+				BinaryPrimitives.WriteInt32LittleEndian(span[4..], signature.HeaderLength);
 			}
 		}
 	}
diff --git a/Engine/Source/Programs/Shared/EpicGames.Horde/Storage/Bundles/BundleCompression.cs b/Engine/Source/Programs/Shared/EpicGames.Horde/Storage/Bundles/BundleCompression.cs
new file mode 100644
--- /dev/null
+++ b/Engine/Source/Programs/Shared/EpicGames.Horde/Storage/Bundles/BundleCompression.cs
@@ -0,0 +1,163 @@
+// Copyright Epic Games, Inc. All Rights Reserved.
+
+using System;
+using System.IO.Compression;
+using System.IO;
+using EpicGames.Compression;
+using EpicGames.Core;
+using K4os.Compression.LZ4;
+
+namespace EpicGames.Horde.Storage.Bundles
+{
+	/// <summary>
+	/// Indicates the compression format in the bundle
+	/// </summary>
+	public enum BundleCompressionFormat
+	{
+		/// <summary>
+		/// Packets are uncompressed
+		/// </summary>
+		None = 0,
+
+		/// <summary>
+		/// LZ4 compression
+		/// </summary>
+		LZ4 = 1,
+
+		/// <summary>
+		/// Gzip compression
+		/// </summary>
+		Gzip = 2,
+
+		/// <summary>
+		/// Oodle compression (Kraken)
+		/// </summary>
+		Oodle = 3,
+
+		/// <summary>
+		/// Brotli compression
+		/// </summary>
+		Brotli = 4,
+	}
+
+	/// <summary>
+	/// Utility methods for bundles
+	/// </summary>
+	public static class BundleData
+	{
+		/// <summary>
+		/// Compress a data packet
+		/// </summary>
+		/// <param name="format">Format for the compressed data</param>
+		/// <param name="input">The data to compress</param>
+		/// <param name="writer">Writer for output data</param>
+		/// <returns>The compressed data</returns>
+		public static int Compress(BundleCompressionFormat format, ReadOnlyMemory<byte> input, IMemoryWriter writer)
+		{
+			switch (format)
+			{
+				case BundleCompressionFormat.None:
+					{
+						writer.WriteFixedLengthBytes(input.Span);
+						return input.Length;
+					}
+				case BundleCompressionFormat.LZ4:
+					{
+						int maxSize = LZ4Codec.MaximumOutputSize(input.Length);
+
+						Span<byte> buffer = writer.GetSpan(maxSize);
+						int encodedLength = LZ4Codec.Encode(input.Span, buffer);
+
+						writer.Advance(encodedLength);
+						return encodedLength;
+					}
+				case BundleCompressionFormat.Gzip:
+					{
+						using MemoryStream outputStream = new MemoryStream(input.Length);
+
+						using (GZipStream gzipStream = new GZipStream(outputStream, CompressionLevel.Fastest, true))
+						{
+							gzipStream.Write(input.Span);
+						}
+
+						writer.WriteFixedLengthBytes(outputStream.ToArray());
+						return (int)outputStream.Length;
+					}
+				case BundleCompressionFormat.Oodle:
+					{
+						int maxSize = Oodle.MaximumOutputSize(OodleCompressorType.Selkie, input.Length);
+
+						Span<byte> outputSpan = writer.GetSpan(maxSize);
+						int encodedLength = Oodle.Compress(OodleCompressorType.Selkie, input.Span, outputSpan, OodleCompressionLevel.Normal);
+						writer.Advance(encodedLength);
+
+						return encodedLength;
+					}
+				case BundleCompressionFormat.Brotli:
+					{
+						int maxSize = BrotliEncoder.GetMaxCompressedLength(input.Length);
+
+						Span<byte> buffer = writer.GetSpan(maxSize);
+						if (!BrotliEncoder.TryCompress(input.Span, buffer, out int encodedLength))
+						{
+							throw new InvalidOperationException("Unable to compress data using Brotli");
+						}
+
+						writer.Advance(encodedLength);
+						return encodedLength;
+					}
+				default:
+					throw new InvalidDataException($"Invalid compression format '{(int)format}'");
+			}
+		}
+
+		/// <summary>
+		/// Decompress a packet of data
+		/// </summary>
+		/// <param name="format">Format of the compressed data</param>
+		/// <param name="input">Compressed data</param>
+		/// <param name="output">Buffer to receive the decompressed data</param>
+		public static void Decompress(BundleCompressionFormat format, ReadOnlyMemory<byte> input, Memory<byte> output)
+		{
+			switch (format)
+			{
+				case BundleCompressionFormat.None:
+					input.CopyTo(output);
+					break;
+				case BundleCompressionFormat.LZ4:
+					LZ4Codec.Decode(input.Span, output.Span);
+					break;
+				case BundleCompressionFormat.Gzip:
+					{
+						using ReadOnlyMemoryStream inputStream = new ReadOnlyMemoryStream(input);
+						using GZipStream inflatedStream = new GZipStream(inputStream, CompressionMode.Decompress, true);
+
+						int length = inflatedStream.Read(output.Span);
+						if (length != output.Length)
+						{
+							throw new InvalidDataException($"Decoded data is shorter than expected (expected {output.Length} bytes, got {length} bytes)");
+						}
+						break;
+					}
+				case BundleCompressionFormat.Oodle:
+#if WITH_OODLE
+					Oodle.Decompress(input.Span, output.Span);
+					break;
+#else
+					throw new NotSupportedException("Oodle is not compiled into this build.");
+#endif
+				case BundleCompressionFormat.Brotli:
+					{
+						int bytesWritten;
+						if (!BrotliDecoder.TryDecompress(input.Span, output.Span, out bytesWritten) || bytesWritten != output.Length)
+						{
+							throw new InvalidOperationException("Unable to decompress data using Brotli");
+						}
+						break;
+					}
+				default:
+					throw new InvalidDataException($"Invalid compression format '{(int)format}'");
+			}
+		}
+	}
+}
diff --git a/Engine/Source/Programs/Shared/EpicGames.Horde/Storage/Bundles/BundleOptions.cs b/Engine/Source/Programs/Shared/EpicGames.Horde/Storage/Bundles/BundleOptions.cs
new file mode 100644
--- /dev/null
+++ b/Engine/Source/Programs/Shared/EpicGames.Horde/Storage/Bundles/BundleOptions.cs
@@ -0,0 +1,35 @@
+// Copyright Epic Games, Inc. All Rights Reserved.
+
+namespace EpicGames.Horde.Storage.Bundles
+{
+	/// <summary>
+	/// Options for configuring a bundle serializer
+	/// </summary>
+	public class BundleOptions
+	{
+		/// <summary>
+		/// Maximum payload size fo a blob
+		/// </summary>
+		public int MaxBlobSize { get; set; } = 10 * 1024 * 1024;
+
+		/// <summary>
+		/// Compression format to use
+		/// </summary>
+		public BundleCompressionFormat CompressionFormat { get; set; } = BundleCompressionFormat.LZ4;
+
+		/// <summary>
+		/// Minimum size of a block to be compressed
+		/// </summary>
+		public int MinCompressionPacketSize { get; set; } = 16 * 1024;
+
+		/// <summary>
+		/// Maximum amount of data to store in memory. This includes any background writes as well as bundles being built.
+		/// </summary>
+		public long MaxWriteQueueLength { get; set; } = 256 * 1024 * 1024;
+
+		/// <summary>
+		/// Number of nodes to cache
+		/// </summary>
+		public int NodeCacheSize { get; set; } = 1024;
+	}
+}
diff --git a/Engine/Source/Programs/Shared/EpicGames.Horde/Storage/Bundles/BundleReader.cs b/Engine/Source/Programs/Shared/EpicGames.Horde/Storage/Bundles/BundleReader.cs
deleted file mode 100644
diff --git a/Engine/Source/Programs/Shared/EpicGames.Horde/Storage/Bundles/BundleReaderCache.cs b/Engine/Source/Programs/Shared/EpicGames.Horde/Storage/Bundles/BundleReaderCache.cs
--- a/Engine/Source/Programs/Shared/EpicGames.Horde/Storage/Bundles/BundleReaderCache.cs
+++ b/Engine/Source/Programs/Shared/EpicGames.Horde/Storage/Bundles/BundleReaderCache.cs
@@ -2,7 +2,6 @@
 
 using System;
 using System.Diagnostics.CodeAnalysis;
-using EpicGames.Horde.Storage.Bundles;
 using Microsoft.Extensions.Caching.Memory;
 
 namespace EpicGames.Horde.Storage
@@ -111,7 +110,7 @@ namespace EpicGames.Horde.Storage
 		/// <summary>
 		/// Adds a bundle info object to the cache
 		/// </summary>
-		public void AddCachedHeader(BlobLocator locator, BundleInfo bundleInfo)
+		public void AddCachedHeader(BlobLocator locator, Bundles.V1.BundleInfo bundleInfo)
 		{
 			AddCachedValue(_headerCache, GetBundleInfoCacheKey(locator), bundleInfo, bundleInfo.HeaderLength);
 		}
@@ -119,7 +118,7 @@ namespace EpicGames.Horde.Storage
 		/// <summary>
 		/// Try to read a bundle info object from the cache
 		/// </summary>
-		public bool TryGetCachedHeader(BlobLocator locator, [NotNullWhen(true)] out BundleInfo? bundleInfo)
+		public bool TryGetCachedHeader(BlobLocator locator, [NotNullWhen(true)] out Bundles.V1.BundleInfo? bundleInfo)
 		{
 			return TryGetCachedValue(_headerCache, GetBundleInfoCacheKey(locator), out bundleInfo);
 		}
diff --git a/Engine/Source/Programs/Shared/EpicGames.Horde/Storage/Bundles/BundleWriter.cs b/Engine/Source/Programs/Shared/EpicGames.Horde/Storage/Bundles/BundleWriter.cs
deleted file mode 100644
diff --git a/Engine/Source/Programs/Shared/EpicGames.Horde/Storage/Bundles/V1/BundleHeader.cs b/Engine/Source/Programs/Shared/EpicGames.Horde/Storage/Bundles/V1/BundleHeader.cs
new file mode 100644
--- /dev/null
+++ b/Engine/Source/Programs/Shared/EpicGames.Horde/Storage/Bundles/V1/BundleHeader.cs
@@ -0,0 +1,1082 @@
+// Copyright Epic Games, Inc. All Rights Reserved.
+
+using System;
+using System.Buffers;
+using System.Buffers.Binary;
+using System.Collections;
+using System.Collections.Generic;
+using System.Diagnostics;
+using System.IO;
+using System.Linq;
+using System.Text;
+using System.Threading;
+using System.Threading.Tasks;
+using EpicGames.Core;
+
+namespace EpicGames.Horde.Storage.Bundles.V1
+{
+	/// <summary>
+	/// Identifier for the type of a section in the bundle header
+	/// </summary>
+	public enum BundleSectionType
+	{
+		/// <summary>
+		/// List of custom types
+		/// </summary>
+		Types = 0,
+
+		/// <summary>
+		/// Imports of other bundles
+		/// </summary>
+		Imports = 1,
+
+		/// <summary>
+		/// List of export headers
+		/// </summary>
+		ExportHeaders = 2,
+
+		/// <summary>
+		/// References to exports in other bundles
+		/// </summary>
+		ExportReferences = 3,
+
+		/// <summary>
+		/// Packet headers
+		/// </summary>
+		Packets = 4,
+
+		/// <summary>
+		/// Merged export headers and references
+		/// </summary>
+		Exports = 5,
+	}
+
+	/// <summary>
+	/// Header for the contents of a bundle.
+	/// </summary>
+	public class BundleHeader
+	{
+		/// <summary>
+		/// Maximum size for a section
+		/// </summary>
+		public const int MaxSectionSize = 0x1000000;
+
+		/// <summary>
+		/// Maximum number of exports from a single bundle
+		/// </summary>
+		public const int MaxExports = 65535;
+
+		/// <summary>
+		/// Maximum number of export refs from a single bundle
+		/// </summary>
+		public const int MaxExportRefs = MaxSectionSize / V1.BundleExportRef.NumBytes;
+
+		/// <summary>
+		/// Signature bytes
+		/// </summary>
+		public static ReadOnlyMemory<byte> Signature { get; } = Encoding.UTF8.GetBytes("UEBN");
+
+		/// <summary>
+		/// Types for exports within this bundle
+		/// </summary>
+		public BundleTypeCollection Types { get; }
+
+		/// <summary>
+		/// Bundles that we reference nodes in
+		/// </summary>
+		public BundleImportCollection Imports { get; }
+
+		/// <summary>
+		/// Nodes exported from this bundle
+		/// </summary>
+		public BundleExportCollection Exports { get; }
+
+		/// <summary>
+		/// List of data packets within this bundle
+		/// </summary>
+		public BundlePacketCollection Packets { get; }
+
+		/// <summary>
+		/// Constructor
+		/// </summary>
+		public BundleHeader(BundleTypeCollection types, BundleImportCollection imports, BundleExportCollection exports, BundlePacketCollection packets)
+		{
+			Types = types;
+			Imports = imports;
+			Exports = exports;
+			Packets = packets;
+		}
+
+		/// <summary>
+		/// Constructor
+		/// </summary>
+		public BundleHeader(BlobType[] types, BlobLocator[] imports, BundleExport[] exports, BundlePacket[] packets)
+		{
+			Types = new BundleTypeCollection(types);
+			Imports = new BundleImportCollection(imports);
+			Exports = new BundleExportCollection(exports);
+			Packets = new BundlePacketCollection(packets);
+		}
+
+		/// <summary>
+		/// Writes data for this bundle to a sequence builder
+		/// </summary>
+		public void AppendTo(ReadOnlySequenceBuilder<byte> builder)
+		{
+			// Create a writer for the output data
+			ChunkedArrayMemoryWriter writer = new ChunkedArrayMemoryWriter();
+
+			// Create the prelude
+			Span<byte> signatureData = writer.GetSpanAndAdvance(Bundle.SignatureLength);
+
+			// Write all the sections
+			WriteSectionHeader(writer, BundleSectionType.Types, Types.Measure());
+			Types.Write(writer);
+
+			WriteSectionHeader(writer, BundleSectionType.Imports, Imports.Measure());
+			Imports.Write(writer);
+
+			WriteSectionHeader(writer, BundleSectionType.ExportHeaders, Exports.MeasureHeaders());
+			Exports.WriteHeaders(writer);
+
+			WriteSectionHeader(writer, BundleSectionType.ExportReferences, Exports.MeasureRefs());
+			Exports.WriteRefs(writer);
+
+			WriteSectionHeader(writer, BundleSectionType.Packets, Packets.Measure());
+			Packets.Write(writer);
+
+			// Fill in the length of the header, and append it to the builder
+			BundleSignature signature = new BundleSignature(BundleVersion.Latest, writer.Length);
+			Bundle.WriteSignature(signatureData, signature);
+
+			writer.AppendTo(builder);
+		}
+
+		/// <summary>
+		/// Serialize the bundle to a byte array
+		/// </summary>
+		public byte[] ToByteArray()
+		{
+			ReadOnlySequenceBuilder<byte> builder = new ReadOnlySequenceBuilder<byte>();
+			AppendTo(builder);
+			return builder.Construct().ToArray();
+		}
+
+		const int SectionHeaderLength = 4;
+
+		static void WriteSectionHeader(IMemoryWriter writer, BundleSectionType type, int length)
+		{
+			if (length >= 0x1000000)
+			{
+				throw new InvalidDataException($"Section header exceeds maximum allowed size (type: {type}, length: {length:n0}");
+			}
+			writer.WriteInt32((int)type | (length << 8));
+		}
+
+		record class ExportInfo(int TypeIdx, int Length, List<BundleExportRef> References);
+
+		/// <summary>
+		/// Reads a bundle header from memory
+		/// </summary>
+		/// <param name="memory">Memory to read from</param>
+		/// <returns>New header object</returns>
+		public static BundleHeader Read(ReadOnlyMemory<byte> memory)
+		{
+			using ReadOnlyMemoryStream stream = new ReadOnlyMemoryStream(memory);
+			return ReadAsync(stream).Result;
+		}
+
+		/// <summary>
+		/// Reads a bundle header from a stream
+		/// </summary>
+		/// <param name="stream">Stream to deserialize from</param>
+		/// <param name="cancellationToken">Cancellation token for the operation</param>
+		/// <returns>New header object</returns>
+		public static async Task<BundleHeader> ReadAsync(Stream stream, CancellationToken cancellationToken = default)
+		{
+			byte[] prelude = new byte[Bundle.SignatureLength];
+			await stream.ReadFixedLengthBytesAsync(prelude, cancellationToken);
+
+			return await ReadAsync(Bundle.ReadSignature(prelude), stream, cancellationToken);
+		}
+
+		/// <summary>
+		/// Reads a bundle header from a stream
+		/// </summary>
+		/// <param name="signature">File signature</param>
+		/// <param name="stream">Stream to deserialize from</param>
+		/// <param name="cancellationToken">Cancellation token for the operation</param>
+		/// <returns>New header object</returns>
+		public static async Task<BundleHeader> ReadAsync(BundleSignature signature, Stream stream, CancellationToken cancellationToken = default)
+		{
+			if (signature.Version >= BundleVersion.InPlace)
+			{
+				return await ReadLatestAsync(signature.Version, stream, signature.HeaderLength - Bundle.SignatureLength, cancellationToken);
+			}
+			else
+			{
+				byte[] data = new byte[signature.HeaderLength - Bundle.SignatureLength];
+				await stream.ReadFixedLengthBytesAsync(data, cancellationToken);
+
+				return ReadLegacy(data);
+			}
+		}
+
+		/// <summary>
+		/// Construct a header from the given data encoded in the latest format
+		/// </summary>
+		/// <param name="version">Version to serialize as</param>
+		/// <param name="stream">Stream to read from</param>
+		/// <param name="length">Length of the header</param>
+		/// <param name="cancellationToken">Cancellation token for the operation</param>
+		static async Task<BundleHeader> ReadLatestAsync(BundleVersion version, Stream stream, int length, CancellationToken cancellationToken)
+		{
+			byte[] exportHeaders = Array.Empty<byte>();
+			byte[] exportReferences = Array.Empty<byte>();
+
+			BundleTypeCollection types = new BundleTypeCollection();
+			BundleImportCollection imports = new BundleImportCollection();
+			BundleExportCollection exports = new BundleExportCollection();
+			BundlePacketCollection packets = new BundlePacketCollection();
+
+			byte[] sectionHeader = new byte[SectionHeaderLength];
+			for (int offset = 0; offset < length;)
+			{
+				await stream.ReadFixedLengthBytesAsync(sectionHeader, cancellationToken);
+				offset += SectionHeaderLength;
+
+				BundleSectionType sectionType = (BundleSectionType)(sectionHeader[0] & 255);
+				int sectionLength = (int)(BinaryPrimitives.ReadUInt32LittleEndian(sectionHeader) >> 8);
+
+				switch (sectionType)
+				{
+					case BundleSectionType.Types:
+						types = await BundleTypeCollection.ReadAsync(stream, sectionLength, cancellationToken);
+						break;
+					case BundleSectionType.Imports:
+						imports = await BundleImportCollection.ReadAsync(stream, sectionLength, cancellationToken);
+						break;
+					case BundleSectionType.Exports:
+						exports = await BundleExportCollection.ReadAsync(stream, sectionLength, cancellationToken);
+						break;
+					case BundleSectionType.ExportHeaders:
+						exportHeaders = new byte[sectionLength];
+						await stream.ReadFixedLengthBytesAsync(exportHeaders, cancellationToken);
+						break;
+					case BundleSectionType.ExportReferences:
+						exportReferences = new byte[sectionLength];
+						await stream.ReadFixedLengthBytesAsync(exportReferences, cancellationToken);
+						break;
+					case BundleSectionType.Packets:
+						packets = await BundlePacketCollection.ReadAsync(stream, sectionLength, cancellationToken);
+						break;
+				}
+
+				offset += sectionLength;
+			}
+
+			if (exportHeaders.Length > 0)
+			{
+				exports = BundleExportCollection.Read(exportHeaders, exportReferences, version);
+			}
+
+			return new BundleHeader(types, imports, exports, packets);
+		}
+
+		static BundleHeader ReadLegacy(ReadOnlyMemory<byte> memory)
+		{
+			MemoryReader reader = new MemoryReader(memory);
+
+			BundleVersion version = (BundleVersion)reader.ReadUnsignedVarInt();
+			if (version > BundleVersion.RemoveAliases)
+			{
+				throw new InvalidDataException($"Invalid bundle version {(int)version}. Max supported version in legacy format is {(int)BundleVersion.RemoveAliases}.");
+			}
+
+			BundleCompressionFormat compressionFormat = (BundleCompressionFormat)reader.ReadUnsignedVarInt();
+
+			// Read the types
+			int numTypes = (int)reader.ReadUnsignedVarInt();
+			List<BlobType> types = new List<BlobType>(numTypes);
+
+			for (int typeIdx = 0; typeIdx < numTypes; typeIdx++)
+			{
+				Guid guid = reader.ReadGuid();
+				int serializerVersion = (int)reader.ReadUnsignedVarInt();
+
+				types.Add(new BlobType(guid, serializerVersion));
+			}
+
+			// Read the imports
+			int numImports = (int)reader.ReadUnsignedVarInt();
+			List<BlobLocator> imports = new List<BlobLocator>(numImports);
+			List<BundleExportRef> allExportReferences = new List<BundleExportRef>();
+
+			for (int importIdx = 0; importIdx < numImports; importIdx++)
+			{
+				BlobLocator locator = new BlobLocator(reader.ReadUtf8String());
+				imports.Add(locator);
+
+				int[] exportIndexes = reader.ReadVariableLengthArray(() => (int)reader.ReadUnsignedVarInt());
+				for (int exportIdx = 0; exportIdx < exportIndexes.Length; exportIdx++)
+				{
+					BundleExportRef exportReference = new BundleExportRef(imports.Count - 1, exportIndexes[exportIdx]);
+					allExportReferences.Add(exportReference);
+				}
+			}
+
+			// Read the exports
+			int numExports = (int)reader.ReadUnsignedVarInt();
+			List<ExportInfo> exportInfos = new List<ExportInfo>(numExports);
+
+			for (int exportIdx = 0; exportIdx < numExports; exportIdx++)
+			{
+				allExportReferences.Add(new BundleExportRef(-1, exportIdx));
+
+				int typeIdx = (int)reader.ReadUnsignedVarInt();
+				_ = reader.ReadIoHash();
+				int length = (int)reader.ReadUnsignedVarInt();
+
+				List<BundleExportRef> exportRefs = new List<BundleExportRef>();
+
+				int numReferences = (int)reader.ReadUnsignedVarInt();
+				if (numReferences > 0)
+				{
+					for (int idx = 0; idx < numReferences; idx++)
+					{
+						int referenceIdx = (int)reader.ReadUnsignedVarInt();
+						exportRefs.Add(allExportReferences[referenceIdx]);
+					}
+				}
+
+				if (version == BundleVersion.ExportAliases)
+				{
+					_ = reader.ReadUtf8String();
+				}
+
+				exportInfos.Add(new ExportInfo(typeIdx, length, exportRefs));
+			}
+
+			// Read the compression packets
+			List<BundlePacket> packets;
+			if (compressionFormat == BundleCompressionFormat.None)
+			{
+				packets = new List<BundlePacket>(exportInfos.Count);
+
+				int encodedOffset = 0;
+				foreach (ExportInfo export in exportInfos)
+				{
+					packets.Add(new BundlePacket(compressionFormat, encodedOffset, export.Length, export.Length));
+					encodedOffset += export.Length;
+				}
+			}
+			else
+			{
+				int numPackets = (int)reader.ReadUnsignedVarInt();
+				packets = new List<BundlePacket>(numPackets);
+
+				int encodedOffset = 0;
+				for (int packetIdx = 0; packetIdx < numPackets; packetIdx++)
+				{
+					int encodedLength = (int)reader.ReadUnsignedVarInt();
+					int decodedLength = (int)reader.ReadUnsignedVarInt();
+					packets.Add(new BundlePacket(compressionFormat, encodedOffset, encodedLength, decodedLength));
+					encodedOffset += encodedLength;
+				}
+			}
+
+			// Create the final export list
+			List<BundleExport> exports = new List<BundleExport>();
+			{
+				int packetIdx = 0;
+				int packetOffset = 0;
+
+				foreach (ExportInfo exportInfo in exportInfos)
+				{
+					if (packetOffset + exportInfo.Length > packets[packetIdx].DecodedLength)
+					{
+						packetIdx++;
+						packetOffset = 0;
+					}
+
+					exports.Add(new BundleExport(exportInfo.TypeIdx, packetIdx, packetOffset, exportInfo.Length, exportInfo.References));
+					packetOffset += exportInfo.Length;
+				}
+			}
+
+			return new BundleHeader(types.ToArray(), imports.ToArray(), exports.ToArray(), packets.ToArray());
+		}
+	}
+
+	/// <summary>
+	/// Collection of node types in a bundle
+	/// </summary>
+	public struct BundleTypeCollection : IReadOnlyList<BlobType>
+	{
+		readonly BlobType[] _types;
+
+		/// <inheritdoc/>
+		public int Count => _types.Length;
+
+		/// <summary>
+		/// Constructor
+		/// </summary>
+		public BundleTypeCollection(BlobType[] types) => _types = types;
+
+		/// <summary>
+		/// Reads a type collection from a block of memory
+		/// </summary>
+		public static BundleTypeCollection Read(ReadOnlyMemory<byte> data)
+		{
+			BlobType[] types = new BlobType[data.Length / BlobType.NumBytes];
+			for (int idx = 0; idx < types.Length; idx++)
+			{
+				types[idx] = BlobType.Read(data.Span.Slice(idx * BlobType.NumBytes));
+			}
+
+			return new BundleTypeCollection(types);
+		}
+
+		/// <summary>
+		/// Reads a type collection from a stream
+		/// </summary>
+		public static async Task<BundleTypeCollection> ReadAsync(Stream stream, int length, CancellationToken cancellationToken)
+		{
+			byte[] data = new byte[length];
+			await stream.ReadFixedLengthBytesAsync(data, cancellationToken);
+			return Read(data);
+		}
+
+		/// <summary>
+		/// Gets the size of memory required to serialize a collection of types
+		/// </summary>
+		/// <returns>Number of bytes required to serialize the types</returns>
+		public int Measure() => _types.Length * BlobType.NumBytes;
+
+		/// <summary>
+		/// Serializes a set of types to a fixed block of memory
+		/// </summary>
+		/// <param name="writer">Writer to serialize to</param>
+		public void Write(IMemoryWriter writer)
+		{
+			Span<byte> span = writer.GetSpanAndAdvance(_types.Length * BlobType.NumBytes);
+			foreach (BlobType type in _types)
+			{
+				type.Write(span);
+				span = span.Slice(BlobType.NumBytes);
+			}
+		}
+
+		/// <inheritdoc/>
+		public BlobType this[int index] => _types[index];
+
+		/// <inheritdoc/>
+		public IEnumerator<BlobType> GetEnumerator() => ((IEnumerable<BlobType>)_types).GetEnumerator();
+
+		/// <inheritdoc/>
+		IEnumerator IEnumerable.GetEnumerator() => GetEnumerator();
+	}
+
+	/// <summary>
+	/// Collection of imported node references
+	/// </summary>
+	public struct BundleImportCollection : IReadOnlyList<BlobLocator>
+	{
+		readonly BlobLocator[] _imports;
+
+		/// <inheritdoc/>
+		public int Count => _imports.Length;
+
+		/// <summary>
+		/// Constructor
+		/// </summary>
+		public BundleImportCollection(BlobLocator[] imports) => _imports = imports;
+
+		/// <summary>
+		/// Reads a collection from memory
+		/// </summary>
+		public static BundleImportCollection Read(ReadOnlyMemory<byte> data)
+		{
+			ReadOnlySpan<byte> span = data.Span;
+			int count = (span.Length == 0) ? 0 : BinaryPrimitives.ReadInt32LittleEndian(span) / sizeof(int);
+
+			BlobLocator[] imports = new BlobLocator[count];
+			for (int idx = 0; idx < count; idx++)
+			{
+				int offset = BinaryPrimitives.ReadInt32LittleEndian(span.Slice(idx * sizeof(int)));
+				int length = span.Slice(offset).IndexOf((byte)0);
+				imports[idx] = new BlobLocator(new Utf8String(data.Slice(offset, length)));
+			}
+
+			return new BundleImportCollection(imports);
+		}
+
+		/// <summary>
+		/// Reads a collection from a stream
+		/// </summary>
+		public static async Task<BundleImportCollection> ReadAsync(Stream stream, int length, CancellationToken cancellationToken)
+		{
+			byte[] data = new byte[length];
+			await stream.ReadFixedLengthBytesAsync(data, cancellationToken);
+			return Read(data);
+		}
+
+		/// <inheritdoc/>
+		public BlobLocator this[int index] => _imports[index];
+
+		/// <inheritdoc/>
+		public IEnumerator<BlobLocator> GetEnumerator() => _imports.AsEnumerable().GetEnumerator();
+
+		/// <inheritdoc/>
+		IEnumerator IEnumerable.GetEnumerator() => GetEnumerator();
+
+		/// <summary>
+		/// Measure the size of memory required to store a collection of import locators
+		/// </summary>
+		/// <returns>Size in bytes of the output buffer</returns>
+		public int Measure() => _imports.Sum(x => sizeof(int) + x.Path.Length + 1);
+
+		/// <summary>
+		/// Serialize a collection of locators to memory
+		/// </summary>
+		/// <param name="writer">Writer for output data</param>
+		public void Write(IMemoryWriter writer)
+		{
+			int offset = _imports.Length * sizeof(int);
+			foreach (BlobLocator import in _imports)
+			{
+				writer.WriteInt32(offset);
+				offset += import.Path.Length + 1;
+			}
+			foreach (BlobLocator import in _imports)
+			{
+				writer.WriteNullTerminatedUtf8String(import.Path);
+			}
+		}
+	}
+
+	/// <summary>
+	/// Descriptor for a compression packet
+	/// </summary>
+	public struct BundlePacket
+	{
+		/// <summary>
+		/// Size of this structure when serialized
+		/// </summary>
+		public const int NumBytes = 16;
+
+		/// <summary>
+		/// Compression format for this packet
+		/// </summary>
+		public BundleCompressionFormat CompressionFormat { get; }
+
+		/// <summary>
+		/// Offset of the packet within the payload stream
+		/// </summary>
+		public int EncodedOffset { get; }
+
+		/// <summary>
+		/// Encoded length of the packet
+		/// </summary>
+		public int EncodedLength { get; }
+
+		/// <summary>
+		/// Decoded length of the packet
+		/// </summary>
+		public int DecodedLength { get; }
+
+		/// <summary>
+		/// Constructor
+		/// </summary>
+		/// <param name="compressionFormat">Compression format for the packet</param>
+		/// <param name="encodedOffset">Offset of the data within the payload stream</param>
+		/// <param name="encodedLength">Size of the encoded data</param>
+		/// <param name="decodedLength">Size of the decoded data</param>
+		public BundlePacket(BundleCompressionFormat compressionFormat, int encodedOffset, int encodedLength, int decodedLength)
+		{
+			CompressionFormat = compressionFormat;
+			EncodedOffset = encodedOffset;
+			EncodedLength = encodedLength;
+			DecodedLength = decodedLength;
+		}
+
+		/// <summary>
+		/// Read from a byte array
+		/// </summary>
+		public static BundlePacket Read(ReadOnlySpan<byte> span)
+		{
+			BundleCompressionFormat compressionFormat = (BundleCompressionFormat)BinaryPrimitives.ReadInt32LittleEndian(span);
+			int encodedOffset = BinaryPrimitives.ReadInt32LittleEndian(span.Slice(4));
+			int encodedLength = BinaryPrimitives.ReadInt32LittleEndian(span.Slice(8));
+			int decodedLength = BinaryPrimitives.ReadInt32LittleEndian(span.Slice(12));
+			return new BundlePacket(compressionFormat, encodedOffset, encodedLength, decodedLength);
+		}
+
+		/// <summary>
+		/// Serialize the struct to memory
+		/// </summary>
+		public void Write(Span<byte> span)
+		{
+			BinaryPrimitives.WriteInt32LittleEndian(span, (int)CompressionFormat);
+			BinaryPrimitives.WriteInt32LittleEndian(span.Slice(4), EncodedOffset);
+			BinaryPrimitives.WriteInt32LittleEndian(span.Slice(8), EncodedLength);
+			BinaryPrimitives.WriteInt32LittleEndian(span.Slice(12), DecodedLength);
+		}
+	}
+
+	/// <summary>
+	/// Collection of information about packets in a bundle
+	/// </summary>
+	public struct BundlePacketCollection : IReadOnlyCollection<BundlePacket>
+	{
+		readonly BundlePacket[] _packets;
+
+		/// <inheritdoc/>
+		public int Count => _packets?.Length ?? 0;
+
+		/// <summary>
+		/// Constructor
+		/// </summary>
+		public BundlePacketCollection(BundlePacket[] packets) => _packets = packets;
+
+		/// <summary>
+		/// Reads a collection from memory
+		/// </summary>
+		public static BundlePacketCollection Read(ReadOnlyMemory<byte> data)
+		{
+			int count = data.Length / BundlePacket.NumBytes;
+
+			BundlePacket[] packets = new BundlePacket[count];
+			for (int idx = 0; idx < count; idx++)
+			{
+				packets[idx] = BundlePacket.Read(data.Span.Slice(idx * BundlePacket.NumBytes));
+			}
+
+			return new BundlePacketCollection(packets);
+		}
+
+		/// <summary>
+		/// Reads a collection from a stream
+		/// </summary>
+		public static async Task<BundlePacketCollection> ReadAsync(Stream stream, int length, CancellationToken cancellationToken)
+		{
+			byte[] data = new byte[length];
+			await stream.ReadFixedLengthBytesAsync(data, cancellationToken);
+			return BundlePacketCollection.Read(data);
+		}
+
+		/// <inheritdoc/>
+		public BundlePacket this[int index] => _packets[index];
+
+		/// <inheritdoc/>
+		public IEnumerator<BundlePacket> GetEnumerator() => _packets.AsEnumerable().GetEnumerator();
+
+		/// <inheritdoc/>
+		IEnumerator IEnumerable.GetEnumerator() => GetEnumerator();
+
+		/// <summary>
+		/// Measure the size of memory required to store a collection of packets
+		/// </summary>
+		/// <returns>Size in bytes of the output buffer</returns>
+		public int Measure() => BundlePacket.NumBytes * _packets.Length;
+
+		/// <summary>
+		/// Serialize a collection of packets to memory
+		/// </summary>
+		/// <param name="writer">Writer to serialize to</param>
+		public void Write(IMemoryWriter writer)
+		{
+			Span<byte> span = writer.GetMemoryAndAdvance(_packets.Length * BundlePacket.NumBytes).Span;
+			foreach (BundlePacket packet in _packets)
+			{
+				packet.Write(span);
+				span = span.Slice(BundlePacket.NumBytes);
+			}
+		}
+	}
+
+	/// <summary>
+	/// Entry for a node exported from an object
+	/// </summary>
+	public struct BundleExport
+	{
+		/// <summary>
+		/// Number of bytes in a serialized export object
+		/// </summary>
+		public const int NumHeaderBytes = 32;
+
+		/// <summary>
+		/// Raw data for this export
+		/// </summary>
+		public ReadOnlyMemory<byte> Data { get; }
+
+		/// <summary>
+		/// Raw data for the header
+		/// </summary>
+		public ReadOnlyMemory<byte> HeaderData => Data.Slice(0, NumHeaderBytes);
+
+		/// <summary>
+		/// Hash of the node data
+		/// </summary>
+		public IoHash Hash => new IoHash(Data.Slice(0, 20).Span);
+
+		/// <summary>
+		/// Type id of the node. Can be used to look up the type information from the bundle header.
+		/// </summary>
+		public int TypeIdx => BinaryPrimitives.ReadUInt16LittleEndian(Data.Slice(20).Span);
+
+		/// <summary>
+		/// Packet containing this export's data
+		/// </summary>
+		public int Packet => BinaryPrimitives.ReadUInt16LittleEndian(Data.Slice(22).Span);
+
+		/// <summary>
+		/// Offset within the packet of the node data
+		/// </summary>
+		public int Offset => BinaryPrimitives.ReadInt32LittleEndian(Data.Slice(24).Span);
+
+		/// <summary>
+		/// Length of the node
+		/// </summary>
+		public int Length => BinaryPrimitives.ReadInt32LittleEndian(Data.Slice(28).Span);
+
+		/// <summary>
+		/// References to other nodes
+		/// </summary>
+		public BundleExportRefCollection References => new BundleExportRefCollection(Data.Slice(32));
+
+		/// <summary>
+		/// Constructor
+		/// </summary>
+		public BundleExport(ReadOnlyMemory<byte> data)
+		{
+			Data = data;
+		}
+
+		/// <summary>
+		/// Constructor
+		/// </summary>
+		public BundleExport(int typeIdx, int packetIdx, int offset, int length, IReadOnlyList<BundleExportRef> references)
+		{
+			byte[] data = new byte[NumHeaderBytes + (BundleExportRef.NumBytes * references.Count)];
+			WriteHeader(data, typeIdx, packetIdx, offset, length);
+
+			for (int idx = 0; idx < references.Count; idx++)
+			{
+				references[idx].CopyTo(data.AsSpan(NumHeaderBytes + (BundleExportRef.NumBytes * idx)));
+			}
+
+			Data = data;
+		}
+
+		/// <summary>
+		/// Writes a new export to a block of memory
+		/// </summary>
+		public static void WriteHeader(Span<byte> data, int typeIdx, int packet, int offset, int length)
+		{
+			IoHash.Zero.CopyTo(data); // Previously hash of export
+			BinaryPrimitives.WriteUInt16LittleEndian(data.Slice(20), (ushort)typeIdx);
+			BinaryPrimitives.WriteUInt16LittleEndian(data.Slice(22), (ushort)packet);
+			BinaryPrimitives.WriteInt32LittleEndian(data.Slice(24), offset);
+			BinaryPrimitives.WriteInt32LittleEndian(data.Slice(28), length);
+		}
+	}
+
+	/// <summary>
+	/// Entry for a node exported from an object
+	/// </summary>
+	public struct BundleExportCollection : IReadOnlyList<BundleExport>
+	{
+		readonly BundleExport[] _exports;
+
+		/// <inheritdoc/>
+		public int Count => _exports.Length;
+
+		/// <summary>
+		/// Constructor
+		/// </summary>
+		public BundleExportCollection(BundleExport[] exports) => _exports = exports;
+
+		/// <summary>
+		/// Constructor
+		/// </summary>
+		public static BundleExportCollection Read(ReadOnlyMemory<byte> data, ReadOnlyMemory<byte> refs, BundleVersion version)
+		{
+			int count = data.Length / BundleExport.NumHeaderBytes;
+
+			if (version < BundleVersion.ImportHashes && refs.Length > 0)
+			{
+				const int OldRefNumBytes = 4;
+
+				int headerLength = BinaryPrimitives.ReadInt32LittleEndian(refs.Span);
+				int refCount = (refs.Length - headerLength) / OldRefNumBytes;
+
+				byte[] newRefs = new byte[headerLength + (refCount * BundleExportRef.NumBytes)];
+				for (int headerOffset = 0; headerOffset < headerLength; headerOffset += sizeof(int))
+				{
+					int index = (BinaryPrimitives.ReadInt32LittleEndian(refs.Slice(headerOffset).Span) - headerLength) / OldRefNumBytes;
+					BinaryPrimitives.WriteInt32LittleEndian(newRefs.AsSpan(headerOffset), headerLength + (index * BundleExportRef.NumBytes));
+				}
+				for (int index = 0; index < refCount; index++)
+				{
+					ReadOnlyMemory<byte> source = refs.Slice(headerLength + (index * OldRefNumBytes), OldRefNumBytes);
+					source.CopyTo(newRefs.AsMemory(headerLength + (index * BundleExportRef.NumBytes)));
+				}
+
+				refs = newRefs;
+			}
+
+			BundleExport[] exports = new BundleExport[count];
+			for (int index = 0; index < count; index++)
+			{
+				ReadOnlyMemory<byte> exportData = data.Slice(index * BundleExport.NumHeaderBytes, BundleExport.NumHeaderBytes);
+
+				BundleExportRefCollection exportRefs = new BundleExportRefCollection();
+				if (refs.Length > 0)
+				{
+					int minOffset = BinaryPrimitives.ReadInt32LittleEndian(refs.Span.Slice(index * sizeof(int)));
+					int maxOffset = (index == count - 1) ? refs.Length : BinaryPrimitives.ReadInt32LittleEndian(refs.Span.Slice((index + 1) * sizeof(int)));
+					exportRefs = new BundleExportRefCollection(refs.Slice(minOffset, maxOffset - minOffset));
+				}
+
+				byte[] export = new byte[BundleExport.NumHeaderBytes + exportRefs.Data.Length];
+				exportData.CopyTo(export);
+				exportRefs.Data.CopyTo(export.AsMemory(BundleExport.NumHeaderBytes));
+
+				exports[index] = new BundleExport(export);
+			}
+
+			return new BundleExportCollection(exports);
+		}
+
+		/// <summary>
+		/// Reads an export collection from a stream
+		/// </summary>
+		public static async Task<BundleExportCollection> ReadAsync(Stream stream, int sectionLength, CancellationToken cancellationToken)
+		{
+			BundleExport[] exports = Array.Empty<BundleExport>();
+			if (sectionLength > 0)
+			{
+				const int MaxPageLength = 16384;
+
+				// Read the offsets of each entry
+				byte[] headerBuffer = new byte[sizeof(int)];
+				await stream.ReadFixedLengthBytesAsync(headerBuffer, cancellationToken);
+
+				int headerLength = BinaryPrimitives.ReadInt32LittleEndian(headerBuffer);
+				Array.Resize(ref headerBuffer, headerLength);
+
+				await stream.ReadFixedLengthBytesAsync(headerBuffer.AsMemory(sizeof(int)), cancellationToken);
+
+				// Create the packet array
+				int exportCount = headerLength / sizeof(int);
+				exports = new BundleExport[exportCount];
+
+				// Read the exports in pages
+				int exportIdx = 0;
+				for (int exportOffset = headerLength; exportOffset < sectionLength;)
+				{
+					int pageOffset = exportOffset;
+
+					// Find the exports to read into this page
+					int nextPageOffset = sectionLength;
+					for (int nextExportIdx = exportIdx + 1; nextExportIdx < exportCount; nextExportIdx++)
+					{
+						int nextExportOffset = BinaryPrimitives.ReadInt32LittleEndian(headerBuffer.AsSpan(nextExportIdx * sizeof(int)));
+						if (nextExportOffset > pageOffset + MaxPageLength)
+						{
+							nextPageOffset = nextExportOffset;
+							break;
+						}
+					}
+
+					// Read the page data
+					byte[] pageBuffer = new byte[nextPageOffset - pageOffset];
+					await stream.ReadFixedLengthBytesAsync(pageBuffer, cancellationToken);
+
+					// Parse the individual exports
+					for (; exportOffset < nextPageOffset; exportIdx++)
+					{
+						int startExportOffset = exportOffset;
+						if (exportIdx + 1 < exportCount)
+						{
+							exportOffset = BinaryPrimitives.ReadInt32LittleEndian(headerBuffer.AsSpan((exportIdx + 1) * sizeof(int)));
+						}
+						else
+						{
+							exportOffset = nextPageOffset;
+						}
+						exports[exportIdx] = new BundleExport(pageBuffer.AsMemory(startExportOffset - pageOffset, exportOffset - startExportOffset));
+					}
+				}
+			}
+			return new BundleExportCollection(exports);
+		}
+
+		/// <inheritdoc/>
+		public BundleExport this[int index] => _exports[index];
+
+		/// <inheritdoc/>
+		public IEnumerator<BundleExport> GetEnumerator() => _exports.AsEnumerable().GetEnumerator();
+
+		/// <inheritdoc/>
+		IEnumerator IEnumerable.GetEnumerator() => GetEnumerator();
+
+		/// <summary>
+		/// Serializes this collection as a single section
+		/// </summary>
+		public void Write(IMemoryWriter writer)
+		{
+			int offset = _exports.Length * sizeof(int);
+			Span<byte> index = writer.GetSpanAndAdvance(offset);
+
+			foreach (BundleExport export in _exports)
+			{
+				BinaryPrimitives.WriteInt32LittleEndian(index, offset);
+				index = index.Slice(sizeof(int));
+
+				writer.WriteFixedLengthBytes(export.Data.Span);
+				offset += export.Data.Length;
+			}
+		}
+
+		/// <summary>
+		/// Measure the size of memory required to store a collection of exports
+		/// </summary>
+		/// <returns>Size in bytes of the output buffer</returns>
+		public int MeasureHeaders() => BundleExport.NumHeaderBytes * _exports.Length;
+
+		/// <summary>
+		/// Serialize a collection of exports to memory
+		/// </summary>
+		/// <param name="writer">Writer to serialize to</param>
+		public void WriteHeaders(IMemoryWriter writer)
+		{
+			foreach (BundleExport export in _exports)
+			{
+				writer.WriteFixedLengthBytes(export.HeaderData.Span);
+			}
+		}
+
+		/// <summary>
+		/// Measure the size of memory required to store a collection of export refs
+		/// </summary>
+		/// <returns>Size in bytes of the output buffer</returns>
+		public int MeasureRefs() => (sizeof(int) * _exports.Length) + _exports.Sum(x => x.References.Data.Length);
+
+		/// <summary>
+		/// Serialize a collection of export refs to memory
+		/// </summary>
+		/// <param name="writer">Writer to serialize to</param>
+		public void WriteRefs(IMemoryWriter writer)
+		{
+			int offset = sizeof(int) * _exports.Length;
+			foreach (BundleExport export in _exports)
+			{
+				writer.WriteInt32(offset);
+				offset += export.References.Data.Length;
+			}
+			foreach (BundleExport export in _exports)
+			{
+				writer.WriteFixedLengthBytes(export.References.Data.Span);
+			}
+		}
+	}
+
+	/// <summary>
+	/// Reference to a node in another bundle
+	/// </summary>
+	/// <param name="ImportIdx">Index into the import table of the blob containing the referenced node. Can be -1 for references within the same bundle.</param>
+	/// <param name="NodeIdx">Node imported from the bundle</param>
+	public record struct BundleExportRef(int ImportIdx, int NodeIdx)
+	{
+		/// <summary>
+		/// Number of bytes in the serialized object
+		/// </summary>
+		public const int NumBytes = 4 + IoHash.NumBytes;
+
+		/// <summary>
+		/// Deserialize this object from memory
+		/// </summary>
+		public static BundleExportRef Read(ReadOnlySpan<byte> data)
+		{
+			int importIdx = BinaryPrimitives.ReadInt16LittleEndian(data);
+			int nodeIdx = BinaryPrimitives.ReadUInt16LittleEndian(data.Slice(2));
+			return new BundleExportRef(importIdx, nodeIdx);
+		}
+
+		/// <summary>
+		/// Serialize this object to memory
+		/// </summary>
+		public void CopyTo(Span<byte> data)
+		{
+			if (ImportIdx > 0xffff)
+			{
+				throw new InvalidDataException($"Import index is out of range ({ImportIdx})");
+			}
+			if (NodeIdx > 0xffff)
+			{
+				throw new InvalidDataException($"Node index is out of range ({NodeIdx})");
+			}
+
+			BinaryPrimitives.WriteInt16LittleEndian(data, (short)ImportIdx);
+			BinaryPrimitives.WriteUInt16LittleEndian(data.Slice(2), (ushort)NodeIdx);
+		}
+	}
+
+	/// <summary>
+	/// Collection of information about exported nodes
+	/// </summary>
+	public struct BundleExportRefCollection : IReadOnlyList<BundleExportRef>
+	{
+		/// <summary>
+		/// Data used to store this collection
+		/// </summary>
+		public ReadOnlyMemory<byte> Data { get; }
+
+		/// <summary>
+		/// Constructor
+		/// </summary>
+		public BundleExportRefCollection(ReadOnlyMemory<byte> data) => Data = data;
+
+		/// <summary>
+		/// Constructor
+		/// </summary>
+		public BundleExportRefCollection(IReadOnlyList<BundleExportRef> exportRefs)
+		{
+			byte[] data = new byte[exportRefs.Count * BundleExportRef.NumBytes];
+			for (int idx = 0; idx < exportRefs.Count; idx++)
+			{
+				exportRefs[idx].CopyTo(data.AsSpan(idx * BundleExportRef.NumBytes));
+			}
+			Data = data;
+		}
+
+		/// <inheritdoc/>
+		public int Count => Data.Length / BundleExportRef.NumBytes;
+
+		/// <inheritdoc/>
+		public BundleExportRef this[int index]
+		{
+			get
+			{
+				Debug.Assert(index < Count);
+				return BundleExportRef.Read(Data.Span.Slice(index * BundleExportRef.NumBytes));
+			}
+		}
+
+		/// <inheritdoc/>
+		public IEnumerator<BundleExportRef> GetEnumerator()
+		{
+			ReadOnlyMemory<byte> data = Data;
+			while (data.Length > 0)
+			{
+				yield return BundleExportRef.Read(data.Span);
+				data = data.Slice(BundleExportRef.NumBytes);
+			}
+		}
+
+		/// <inheritdoc/>
+		IEnumerator IEnumerable.GetEnumerator() => GetEnumerator();
+	}
+}
diff --git a/Engine/Source/Programs/Shared/EpicGames.Horde/Storage/Bundles/V1/BundleReader.cs b/Engine/Source/Programs/Shared/EpicGames.Horde/Storage/Bundles/V1/BundleReader.cs
new file mode 100644
--- /dev/null
+++ b/Engine/Source/Programs/Shared/EpicGames.Horde/Storage/Bundles/V1/BundleReader.cs
@@ -0,0 +1,587 @@
+// Copyright Epic Games, Inc. All Rights Reserved.
+
+using System;
+using System.Collections.Generic;
+using System.Linq;
+using System.Threading;
+using System.Threading.Tasks;
+using EpicGames.Core;
+using Microsoft.Extensions.Logging;
+using Microsoft.CodeAnalysis;
+using System.Diagnostics;
+using System.IO;
+
+namespace EpicGames.Horde.Storage.Bundles.V1
+{
+	/// <summary>
+	/// Computed information about a bundle
+	/// </summary>
+	public class BundleInfo
+	{
+		/// <summary>
+		/// Locator for the bundle
+		/// </summary>
+		public BlobLocator Locator { get; }
+
+		/// <summary>
+		/// Bundle header
+		/// </summary>
+		public BundleHeader Header { get; }
+
+		/// <summary>
+		/// Length of the header. Required to offset packets from the start of the bundle
+		/// </summary>
+		public int HeaderLength { get; }
+
+		/// <summary>
+		/// Constructor
+		/// </summary>
+		public BundleInfo(BlobLocator locator, BundleHeader header, int headerLength)
+		{
+			Locator = locator;
+			Header = header;
+			HeaderLength = headerLength;
+		}
+	}
+
+	/// <summary>
+	/// Writes nodes from bundles in an <see cref="IStorageClient"/> instance.
+	/// </summary>
+	public class BundleReader
+	{
+		/// <summary>
+		/// Queued set of requests from a particular bundle
+		/// </summary>
+		[DebuggerDisplay("{Locator}")]
+		class QueuedBundle
+		{
+			public BlobLocator Locator { get; }
+			public TaskCompletionSource<BundleInfo> BundleInfo { get; } = new TaskCompletionSource<BundleInfo>();
+			public List<QueuedPacket> QueuedPackets { get; } = new List<QueuedPacket>(); // Sorted by index
+			public int InfoRefCount { get; set; } // Ref count for reading the header
+			public PendingRead? PendingRead { get; set; }
+			public Task? WorkerTask { get; set; }
+			public CancellationTokenSource CancellationSource { get; set; } = new CancellationTokenSource();
+			public bool Complete { get; set; }
+
+			public QueuedBundle(BlobLocator locator)
+			{
+				Locator = locator;
+			}
+
+			public QueuedPacket AddPacket(int packetIdx)
+			{
+				QueuedPacket packet = new QueuedPacket(packetIdx);
+
+				int insertIdx = QueuedPackets.BinarySearch(packet);
+				if (insertIdx < 0)
+				{
+					insertIdx = ~insertIdx;
+				}
+				QueuedPackets.Insert(insertIdx, packet);
+
+				return packet;
+			}
+
+			public void RemovePacket(QueuedPacket packet) => QueuedPackets.Remove(packet);
+		}
+
+		/// <summary>
+		/// Encoded bundle packet queued to be read
+		/// </summary>
+		class QueuedPacket : IComparable<QueuedPacket>
+		{
+			public readonly int PacketIdx;
+			public readonly TaskCompletionSource<ReadOnlyMemory<byte>> CompletionSource = new TaskCompletionSource<ReadOnlyMemory<byte>>(TaskCreationOptions.RunContinuationsAsynchronously);
+
+			public QueuedPacket(int packetIdx) => PacketIdx = packetIdx;
+			public int CompareTo(QueuedPacket? other) => PacketIdx - other?.PacketIdx ?? 0;
+		}
+
+		/// <summary>
+		/// Information about a pending read
+		/// </summary>
+		sealed class PendingRead : IDisposable
+		{
+			public int MinPacketIdx { get; }
+			public int MaxPacketIdx { get; }
+			public CancellationTokenSource CancellationSource { get; }
+
+			public PendingRead(int minPacketIdx, int maxPacketIdx, CancellationToken cancellationToken)
+			{
+				MinPacketIdx = minPacketIdx;
+				MaxPacketIdx = maxPacketIdx;
+				CancellationSource = CancellationTokenSource.CreateLinkedTokenSource(cancellationToken);
+			}
+
+			public void Dispose()
+			{
+				CancellationSource.Dispose();
+			}
+		}
+
+		// Size of data to fetch by default. This is larger than the minimum request size to reduce number of reads.
+		const int DefaultFetchSize = 15 * 1024 * 1024;
+
+		// When reader is uncached, use a smaller default fetch size
+		const int DefaultUncachedFetchSize = 1 * 1024 * 1024;
+
+		readonly IStorageClient _store;
+		readonly BundleReaderCache _cache;
+		readonly ILogger _logger;
+
+		readonly object _queueLock = new object();
+		readonly List<QueuedBundle> _queuedBundles = new List<QueuedBundle>();
+		readonly Dictionary<string, Task<ReadOnlyMemory<byte>>> _decodeTasks = new Dictionary<string, Task<ReadOnlyMemory<byte>>>(StringComparer.Ordinal);
+
+		int _numHeaderReads;
+		int _numPacketReads;
+		long _numBytesRead;
+		long _decodeTimeTicks;
+
+		/// <summary>
+		/// Accessor for the cache
+		/// </summary>
+		public BundleReaderCache Cache => _cache;
+
+		/// <summary>
+		/// Constructor
+		/// </summary>
+		/// <param name="store"></param>
+		/// <param name="cache">Cache for data</param>
+		/// <param name="logger">Logger for output</param>
+		public BundleReader(IStorageClient store, BundleReaderCache cache, ILogger logger)
+		{
+			_store = store;
+			_cache = cache;
+			_logger = logger;
+		}
+
+		#region Bundles
+
+		/// <summary>
+		/// Reads a bundle header
+		/// </summary>
+		/// <param name="locator">Locator for the bundle</param>
+		/// <param name="cancellationToken">Cancellation token for the operation</param>
+		/// <returns>Information about the bundle</returns>
+		public async Task<BundleHeader> ReadHeaderAsync(BlobLocator locator, CancellationToken cancellationToken)
+		{
+			// Check for a cached value first
+			if (_cache.TryGetCachedHeader(locator, out BundleInfo? cachedBundleInfo))
+			{
+				return cachedBundleInfo.Header;
+			}
+
+			// Find a registered bundle info with the given locator
+			QueuedBundle? bundle = null;
+			lock (_queueLock)
+			{
+				bundle = FindOrAddBundle(locator);
+				bundle.InfoRefCount++;
+			}
+
+			// Wait for the read to complete
+			try
+			{
+				BundleInfo bundleInfo = await bundle.BundleInfo.Task.WaitAsync(cancellationToken);
+				return bundleInfo.Header;
+			}
+			finally
+			{
+				await CancelBundleRequestAsync(bundle, () => bundle.InfoRefCount--);
+			}
+		}
+
+		static string GetDecodeTaskKey(BlobLocator locator, int packetIdx) => $"{locator}#{packetIdx}";
+
+		/// <summary>
+		/// Reads decoded packet data from a bundle
+		/// </summary>
+		/// <param name="locator">Locator for the bundle to read</param>
+		/// <param name="packetIdx">Index of the bundle packet</param>
+		/// <param name="cancellationToken">Cancellation token for the operation</param>
+		/// <returns>Data for the packet</returns>
+		public async Task<ReadOnlyMemory<byte>> ReadPacketAsync(BlobLocator locator, int packetIdx, CancellationToken cancellationToken)
+		{
+			ReadOnlyMemory<byte> cachedDecodedPacket;
+			if (_cache.TryGetCachedDecodedPacket(locator, packetIdx, out cachedDecodedPacket))
+			{
+				return cachedDecodedPacket;
+			}
+
+			// Create an async task to read the data
+			Task<ReadOnlyMemory<byte>>? decodeTask;
+			lock (_queueLock)
+			{
+				if (_cache.TryGetCachedDecodedPacket(locator, packetIdx, out cachedDecodedPacket))
+				{
+					return cachedDecodedPacket;
+				}
+
+				string decodedCacheKey = GetDecodeTaskKey(locator, packetIdx);
+				if (!_decodeTasks.TryGetValue(decodedCacheKey, out decodeTask))
+				{
+					decodeTask = Task.Run(() => ReadAndDecodePacketAsync(locator, packetIdx, CancellationToken.None), CancellationToken.None);
+					_decodeTasks.Add(decodedCacheKey, decodeTask);
+				}
+			}
+
+			return await decodeTask.WaitAsync(cancellationToken);
+		}
+
+		async Task<ReadOnlyMemory<byte>> ReadAndDecodePacketAsync(BlobLocator locator, int packetIdx, CancellationToken cancellationToken)
+		{
+			BundleHeader header = await ReadHeaderAsync(locator, cancellationToken);
+			ReadOnlyMemory<byte> encodedPacket = await ReadEncodedPacketAsync(locator, packetIdx, cancellationToken);
+
+			BundlePacket packet = header.Packets[packetIdx];
+
+			byte[] decodedPacket = new byte[packet.DecodedLength];
+			Stopwatch decodeTimer = Stopwatch.StartNew();
+			BundleData.Decompress(packet.CompressionFormat, encodedPacket, decodedPacket);
+			Interlocked.Add(ref _decodeTimeTicks, decodeTimer.ElapsedTicks);
+			_cache.AddCachedDecodedPacket(locator, packetIdx, decodedPacket);
+
+			lock (_queueLock)
+			{
+				string decodedCacheKey = GetDecodeTaskKey(locator, packetIdx);
+				_decodeTasks.Remove(decodedCacheKey);
+			}
+
+			return decodedPacket;
+		}
+
+		/// <summary>
+		/// Reads encoded packet data from a bundle
+		/// </summary>
+		/// <param name="locator">Locator for the bundle to read</param>
+		/// <param name="packetIdx">Index of the bundle packet</param>
+		/// <param name="cancellationToken">Cancellation token for the operation</param>
+		/// <returns>Dtaa for the packet</returns>
+		public async Task<ReadOnlyMemory<byte>> ReadEncodedPacketAsync(BlobLocator locator, int packetIdx, CancellationToken cancellationToken)
+		{
+			// Register a read for the packet
+			QueuedBundle? bundle = null;
+			QueuedPacket? packet = null;
+			lock (_queueLock)
+			{
+				if (_cache.TryGetCachedEncodedPacket(locator, packetIdx, out ReadOnlyMemory<byte> cachedEncodedPacket))
+				{
+					return cachedEncodedPacket;
+				}
+
+				bundle = FindOrAddBundle(locator);
+				packet = bundle.AddPacket(packetIdx);
+			}
+
+			// Wait for the read to complete
+			using (CancellationTokenRegistration registration = cancellationToken.Register(() => packet.CompletionSource.TrySetCanceled()))
+			{
+				try
+				{
+					return await packet.CompletionSource.Task;
+				}
+				finally
+				{
+					await CancelBundleRequestAsync(bundle, () => bundle.QueuedPackets.Remove(packet));
+				}
+			}
+		}
+
+		QueuedBundle FindOrAddBundle(BlobLocator locator)
+		{
+			QueuedBundle? bundle = _queuedBundles.FirstOrDefault(x => x.Locator == locator);
+			if (bundle == null)
+			{
+				bundle = new QueuedBundle(locator);
+				bundle.WorkerTask = Task.Run(() => HandleBundleRequestsAsync(bundle), bundle.CancellationSource.Token);
+				_queuedBundles.Add(bundle);
+			}
+			return bundle;
+		}
+
+		async Task HandleBundleRequestsAsync(QueuedBundle bundle)
+		{
+			// Read the bundle header
+			BundleInfo? bundleInfo;
+			if (!_cache.TryGetCachedHeader(bundle.Locator, out bundleInfo))
+			{
+				try
+				{
+					bundleInfo = await ReadBundleInfoAsync(bundle, bundle.CancellationSource.Token);
+					_cache.AddCachedHeader(bundle.Locator, bundleInfo);
+				}
+				catch (Exception ex)
+				{
+					bundle.BundleInfo.SetException(ex);
+					return;
+				}
+			}
+			bundle.BundleInfo.SetResult(bundleInfo);
+
+			// Serve any packet read requests
+			for (; ; )
+			{
+				// Create the next read
+				PendingRead? pendingRead;
+				lock (_queueLock)
+				{
+					// Dispose of the previous pending read
+					if (bundle.PendingRead != null)
+					{
+						bundle.PendingRead.Dispose();
+						bundle.PendingRead = null;
+					}
+
+					// If there's nothing left to read, dispose of it
+					if (bundle.QueuedPackets.Count == 0)
+					{
+						bundle.Complete = true;
+						_queuedBundles.Remove(bundle);
+						break;
+					}
+
+					// Figure out the range of packets to read
+					int minPacketIdx = bundle.QueuedPackets[0].PacketIdx;
+					int maxPacketIdx = minPacketIdx;
+
+					long length = 0;
+					for (int packetIdx = minPacketIdx; packetIdx < bundleInfo.Header.Packets.Count; packetIdx++)
+					{
+						length += bundleInfo.Header.Packets[packetIdx].EncodedLength;
+						if (length > DefaultFetchSize)
+						{
+							break;
+						}
+						maxPacketIdx = packetIdx;
+					}
+
+					// Create the new read
+					pendingRead = new PendingRead(minPacketIdx, maxPacketIdx, bundle.CancellationSource.Token);
+					bundle.PendingRead = pendingRead;
+				}
+
+				// Execute the read
+				try
+				{
+					await ReadEncodedPacketsAsync(bundle, bundle.PendingRead.MinPacketIdx, bundle.PendingRead.MaxPacketIdx, pendingRead.CancellationSource.Token);
+				}
+				catch (OperationCanceledException ex)
+				{
+					_logger.LogTrace(ex, "Read from bundle was cancelled");
+				}
+				catch (Exception ex)
+				{
+					_logger.LogError(ex, "Error reading from bundle: {Message}", ex.Message);
+
+					lock (_queueLock)
+					{
+						List<QueuedPacket> packets = bundle.QueuedPackets;
+						for (int idx = 0; idx < packets.Count && packets[idx].PacketIdx <= pendingRead.MaxPacketIdx; idx++)
+						{
+							if (packets[idx].PacketIdx >= pendingRead.MinPacketIdx)
+							{
+								packets[idx].CompletionSource.TrySetException(ex);
+								packets.RemoveAt(idx--);
+							}
+						}
+					}
+				}
+			}
+
+			// Dispose the cancellation source used for this bundle
+			lock (_queueLock)
+			{
+				bundle.CancellationSource.Dispose();
+				bundle.CancellationSource = null!;
+			}
+		}
+
+		async Task CancelBundleRequestAsync(QueuedBundle bundle, Action updateAction)
+		{
+			// Remove this read request
+			Task? workerTask = null;
+			lock (_queueLock)
+			{
+				updateAction();
+
+				// If there's nothing required any more, remove the bundle from the queue
+				if (bundle.InfoRefCount == 0 && bundle.QueuedPackets.Count == 0)
+				{
+					workerTask = bundle.WorkerTask;
+					_queuedBundles.Remove(bundle);
+				}
+			}
+
+			// If this was the last thing using the queued bundle, wait for the worker task to finish
+			if (workerTask != null)
+			{
+				await workerTask;
+
+				try
+				{
+					await bundle.BundleInfo.Task; // Avoid unobserved cancellation exceptions
+				}
+				catch { }
+			}
+		}
+
+		#endregion
+
+		#region Reading
+
+		/// <summary>
+		/// Reads a bundle header from the queue
+		/// </summary>
+		/// <param name="bundle">Bundle header to read</param>
+		/// <param name="cancellationToken">Cancellation token for the operation</param>
+		async Task<BundleInfo> ReadBundleInfoAsync(QueuedBundle bundle, CancellationToken cancellationToken)
+		{
+			Interlocked.Increment(ref _numHeaderReads);
+
+			int prefetchSize = _cache != null ? DefaultFetchSize : DefaultUncachedFetchSize;
+			for (; ; )
+			{
+				BlobHandle handle = _store.CreateBlobHandle(bundle.Locator);
+				await using (Stream stream = await handle.OpenAsync(0, prefetchSize, cancellationToken))
+				{
+					// Read the header data
+					byte[] prelude = new byte[Bundle.SignatureLength];
+					await stream.ReadFixedLengthBytesAsync(prelude, cancellationToken);
+
+					// Make sure we've read enough to hold the header
+					BundleSignature signature = Bundle.ReadSignature(prelude);
+					if (signature.HeaderLength > prefetchSize)
+					{
+						prefetchSize = signature.HeaderLength;
+						continue;
+					}
+
+					// Parse the header and construct the bundle info from it
+					BundleHeader header = await BundleHeader.ReadAsync(signature, stream, cancellationToken);
+
+					// Construct the bundle info
+					BundleInfo bundleInfo = new BundleInfo(bundle.Locator, header, signature.HeaderLength);
+
+					// Also add any encoded packets we prefetched
+					int packetOffset = signature.HeaderLength;
+					for (int packetIdx = 0; packetIdx < header.Packets.Count; packetIdx++)
+					{
+						BundlePacket packet = header.Packets[packetIdx];
+
+						int packetLength = packet.EncodedLength;
+						if (packetOffset + packetLength > prefetchSize)
+						{
+							break;
+						}
+						packetOffset += packetLength;
+
+						await ReadEncodedPacketFromStreamAsync(bundle, packetIdx, packet, stream, cancellationToken);
+					}
+
+					Interlocked.Add(ref _numBytesRead, packetOffset);
+					return bundleInfo;
+				}
+			}
+		}
+
+		async Task ReadEncodedPacketsAsync(QueuedBundle bundle, int minPacketIdx, int maxPacketIdx, CancellationToken cancellationToken)
+		{
+			BundleInfo bundleInfo = await bundle.BundleInfo.Task;
+
+			BundlePacket minPacket = bundleInfo.Header.Packets[minPacketIdx];
+			int minOffset = minPacket.EncodedOffset;
+
+			BundlePacket maxPacket = bundleInfo.Header.Packets[maxPacketIdx];
+			int maxOffset = maxPacket.EncodedOffset + maxPacket.EncodedLength;
+
+			Interlocked.Increment(ref _numPacketReads);
+			Interlocked.Add(ref _numBytesRead, maxOffset - minOffset);
+
+			BlobHandle handle = _store.CreateBlobHandle(bundleInfo.Locator);
+			await using (Stream stream = await handle.OpenAsync(bundleInfo.HeaderLength + minOffset, maxOffset - minOffset, cancellationToken))
+			{
+				// Copy all the packets that have been read into separate buffers, so we can cache them indidually.
+				for (int packetIdx = minPacketIdx; packetIdx <= maxPacketIdx; packetIdx++)
+				{
+					BundlePacket packet = bundleInfo.Header.Packets[packetIdx];
+					await ReadEncodedPacketFromStreamAsync(bundle, packetIdx, packet, stream, cancellationToken);
+				}
+			}
+		}
+
+		async Task ReadEncodedPacketFromStreamAsync(QueuedBundle bundle, int packetIdx, BundlePacket packet, Stream stream, CancellationToken cancellationToken)
+		{
+			byte[] data = new byte[packet.EncodedLength];
+			await stream.ReadFixedLengthBytesAsync(data, cancellationToken);
+
+			lock (_queueLock)
+			{
+				_cache.AddCachedEncodedPacket(bundle.Locator, packetIdx, data); // Note: do this while holding the lock, to avoid races with new encodes being added
+
+				List<QueuedPacket> packets = bundle.QueuedPackets;
+				for (int idx = 0; idx < packets.Count && packets[idx].PacketIdx <= packetIdx; idx++)
+				{
+					if (packets[idx].PacketIdx == packetIdx)
+					{
+						packets[idx].CompletionSource.TrySetResult(data);
+						packets.RemoveAt(idx--);
+					}
+				}
+			}
+		}
+
+		#endregion
+
+		/// <summary>
+		/// Reads a node from a bundle
+		/// </summary>
+		/// <param name="bundleLocator">Locator for the bundle</param>
+		/// <param name="exportIdx">Index of the export</param>
+		/// <param name="cancellationToken">Cancellation token for the operation</param>
+		/// <returns>Node data read from the given bundle</returns>
+		public async ValueTask<BlobData> ReadNodeDataAsync(BlobLocator bundleLocator, int exportIdx, CancellationToken cancellationToken = default)
+		{
+			BundleHeader header = await ReadHeaderAsync(bundleLocator, cancellationToken);
+			BundleExport export = header.Exports[exportIdx];
+
+			List<BlobHandle> refs = new List<BlobHandle>(export.References.Count);
+			foreach (BundleExportRef reference in export.References)
+			{
+				BlobLocator importBlob;
+				if (reference.ImportIdx == -1)
+				{
+					importBlob = bundleLocator;
+				}
+				else
+				{
+					importBlob = header.Imports[reference.ImportIdx];
+				}
+				Debug.Assert(importBlob.IsValid());
+				refs.Add(new FlushedNodeHandle(this, importBlob, reference.NodeIdx));
+			}
+
+			ReadOnlyMemory<byte> nodeData = ReadOnlyMemory<byte>.Empty;
+			if (export.Length > 0)
+			{
+				ReadOnlyMemory<byte> packetData = await ReadPacketAsync(bundleLocator, export.Packet, cancellationToken);
+				nodeData = packetData.Slice(export.Offset, export.Length);
+			}
+
+			BlobType nodeType = header.Types[export.TypeIdx];
+			return new BlobData(nodeType, nodeData, refs);
+		}
+
+		/// <summary>
+		/// Gets stats for the reader
+		/// </summary>
+		public void GetStats(StorageStats stats)
+		{
+			stats.Add("Num bytes read", _numBytesRead);
+			stats.Add("Num header reads", _numHeaderReads);
+			stats.Add("Num packet reads", _numPacketReads);
+			stats.Add("Decode time (ms)", (_decodeTimeTicks * 1000) / Stopwatch.Frequency);
+		}
+	}
+}
diff --git a/Engine/Source/Programs/Shared/EpicGames.Horde/Storage/Bundles/V1/BundleWriter.cs b/Engine/Source/Programs/Shared/EpicGames.Horde/Storage/Bundles/V1/BundleWriter.cs
new file mode 100644
--- /dev/null
+++ b/Engine/Source/Programs/Shared/EpicGames.Horde/Storage/Bundles/V1/BundleWriter.cs
@@ -0,0 +1,798 @@
+// Copyright Epic Games, Inc. All Rights Reserved.
+
+using System;
+using System.Buffers;
+using System.Buffers.Text;
+using System.Collections.Generic;
+using System.Diagnostics;
+using System.Diagnostics.CodeAnalysis;
+using System.Linq;
+using System.Threading;
+using System.Threading.Tasks;
+using EpicGames.Core;
+using EpicGames.Horde.Storage.Clients;
+using Microsoft.Extensions.Logging;
+
+namespace EpicGames.Horde.Storage.Bundles.V1
+{
+	/// <summary>
+	/// Implementation of <see cref="BlobHandle"/> for nodes which can be read from storage
+	/// </summary>
+	sealed class FlushedNodeHandle : BlobHandle
+	{
+		readonly BundleReader _reader;
+
+		public BlobLocator BundleLocator { get; }
+		public int ExportIdx { get; }
+
+		/// <summary>
+		/// Constructor
+		/// </summary>
+		public FlushedNodeHandle(BundleReader reader, BlobLocator bundleLocator, int exportIdx)
+		{
+			_reader = reader;
+			BundleLocator = bundleLocator;
+			ExportIdx = exportIdx;
+		}
+
+		public static FlushedNodeHandle FromBlobLocator(BundleReader reader, BlobLocator locator)
+		{
+			if (locator.TryUnwrap(out BlobLocator bundleLocator, out Utf8String fragment) && Utf8Parser.TryParse(fragment, out int exportIdx, out int bytesConsumed) && bytesConsumed == fragment.Length)
+			{
+				return new FlushedNodeHandle(reader, bundleLocator, exportIdx);
+			}
+			else
+			{
+				throw new ArgumentException($"Locator {locator} is not a valid bundle node");
+			}
+		}
+
+		/// <inheritdoc/>
+		public override bool TryGetLocator([NotNullWhen(true)] out BlobLocator locator)
+		{
+			locator = new BlobLocator(BundleLocator, ExportIdx.ToString());
+			return true;
+		}
+
+		/// <inheritdoc/>
+		public override ValueTask<BlobData> ReadAsync(CancellationToken cancellationToken = default) => _reader.ReadNodeDataAsync(BundleLocator, ExportIdx, cancellationToken);
+
+		/// <inheritdoc/>
+		public override ValueTask FlushAsync(CancellationToken cancellationToken = default) => new ValueTask();
+	}
+
+	/// <summary>
+	/// Writes nodes of a tree to an <see cref="IStorageClient"/>, packed into bundles. Each <see cref="BundleWriter"/> instance is single threaded,
+	/// but multiple instances may be written to in parallel.
+	/// </summary>
+	public sealed class BundleWriter : IStorageWriter
+	{
+		// Information about a unique output node. Note that multiple node refs may de-duplicate to the same output node.
+		internal class PendingNode : BlobHandle
+		{
+			readonly BundleReader _reader;
+
+			object LockObject => _reader;
+
+			FlushedNodeHandle? _flushedHandle;
+			PendingBundle? _pendingBundle;
+
+			public readonly BlobType BlobType;
+			public readonly int Packet;
+			public readonly int Offset;
+			public readonly int Length;
+			public readonly BlobHandle[] Refs;
+			public readonly AliasInfo[] Aliases;
+
+			public PendingBundle? PendingBundle => _pendingBundle;
+			public FlushedNodeHandle? FlushedNodeHandle => _flushedHandle;
+
+			public PendingNode(BundleReader reader, BlobType blobType, int packet, int offset, int length, IReadOnlyList<BlobHandle> refs, IReadOnlyList<AliasInfo> aliases, PendingBundle pendingBundle)
+			{
+				_reader = reader;
+
+				BlobType = blobType;
+				Packet = packet;
+				Offset = offset;
+				Length = length;
+				Refs = refs.ToArray();
+				Aliases = aliases.ToArray();
+
+				_pendingBundle = pendingBundle;
+			}
+
+			/// <inheritdoc/>
+			public override bool TryGetLocator([NotNullWhen(true)] out BlobLocator locator)
+			{
+				if (_flushedHandle != null && _flushedHandle.TryGetLocator(out BlobLocator flushedLocator))
+				{
+					locator = flushedLocator;
+					return true;
+				}
+				else
+				{
+					Console.WriteLine($"Unable to get locator for node {BlobType}:{Packet}:{Offset}:{Length}. Handle is {FlushedNodeHandle}. Bundle is {PendingBundle}.");
+
+					locator = default;
+					return false;
+				}
+			}
+
+			public void MarkAsWritten(FlushedNodeHandle flushedHandle)
+			{
+				lock (LockObject)
+				{
+					Debug.Assert(_flushedHandle == null);
+					_flushedHandle = flushedHandle;
+					_pendingBundle = null;
+				}
+			}
+
+			/// <inheritdoc/>
+			public override async ValueTask<BlobData> ReadAsync(CancellationToken cancellationToken = default)
+			{
+				if (_flushedHandle == null)
+				{
+					lock (LockObject)
+					{
+						if (_flushedHandle == null)
+						{
+							ReadOnlyMemory<byte> data = _pendingBundle!.GetNodeData(Packet, Offset, Length);
+							return new BlobData(BlobType, data, Refs);
+						}
+					}
+				}
+
+				return await _flushedHandle!.ReadAsync(cancellationToken);
+			}
+
+			/// <inheritdoc/>
+			public override async ValueTask FlushAsync(CancellationToken cancellationToken = default)
+			{
+				if (_flushedHandle != null)
+				{
+					return;
+				}
+
+				foreach (BlobHandle nodeRef in Refs)
+				{
+					await nodeRef.FlushAsync(cancellationToken);
+				}
+
+				PendingBundle? pendingBundle = _pendingBundle;
+				if (pendingBundle != null)
+				{
+					await pendingBundle.FlushAsync(cancellationToken);
+				}
+			}
+		}
+
+		// Information about a bundle being built. Metadata operations are synchronous, compression/writes are asynchronous.
+		internal class PendingBundle : IDisposable
+		{
+			static int s_lastBundleId = 0;
+			public int BundleId { get; } = Interlocked.Increment(ref s_lastBundleId);
+
+			readonly BundleReader _treeReader;
+			readonly BundleWriter _treeWriter;
+			readonly int _maxPacketSize;
+			readonly BundleCompressionFormat _compressionFormat;
+
+			readonly object _lockObject = new object();
+
+			// The current packet being built
+			IMemoryOwner<byte>? _currentPacket;
+
+			// Index of the current packet being written to
+			int _currentPacketIdx;
+
+			// Length of the packet that has been written
+			int _currentPacketLength;
+
+			// Total size of compressed data in the current bundle
+			long _compressedLength;
+
+			// Queue of nodes for the current bundle
+			readonly List<PendingNode> _queue = new List<PendingNode>();
+
+			// Number of references in the queue
+			int _queuedRefs = 0;
+
+			// List of packets in the current bundle
+			readonly List<BundlePacket> _packets = new List<BundlePacket>();
+
+			// List of compressed packets
+			readonly ChunkedMemoryWriter _encodedPacketWriter;
+
+			// Set of all direct dependencies from this bundle
+			readonly HashSet<Task> _dependencies = new HashSet<Task>();
+
+			// Total size of compressed data in the current bundle
+			public long CompressedLength => _compressedLength;
+
+			// Total size of uncompressed data in the current bundle
+			public long UncompressedLength { get; private set; }
+
+			// Task used to compress data in the background
+			Task _compressPacketsTask = Task.CompletedTask;
+
+			// Event which is signalled after the bundle is written to storage
+			readonly TaskCompletionSource<bool> _completeEvent = new TaskCompletionSource<bool>();
+
+			// Task signalled after the write is complete
+			public Task CompleteTask => _completeEvent.Task;
+
+			public PendingBundle(BundleReader treeReader, BundleWriter treeWriter, int maxPacketSize, BundleCompressionFormat compressionFormat)
+			{
+				_treeReader = treeReader;
+				_treeWriter = treeWriter;
+				_maxPacketSize = maxPacketSize;
+				_compressionFormat = compressionFormat;
+				_encodedPacketWriter = new ChunkedMemoryWriter(1024 * 1024);
+			}
+
+			/// <inheritdoc/>
+			public void Dispose()
+			{
+				if (!CompleteTask.IsCompleted)
+				{
+					_completeEvent.TrySetException(new ObjectDisposedException($"Bundle {BundleId} has been disposed. Make sure to call FlushAsync() on forked writers that may have dependencies before disposing.", $"Bundle {BundleId}"));
+				}
+				if (_currentPacket != null)
+				{
+					_currentPacket.Dispose();
+					_currentPacket = null;
+				}
+				_encodedPacketWriter.Dispose();
+			}
+
+			// Whether this bundle is full
+			public bool IsFull() => (_queue.Count + 1000) >= BundleHeader.MaxExports || (_queuedRefs + 1000) >= BundleHeader.MaxExportRefs;
+
+			// Whether this bundle has finished writing
+			public bool IsComplete() => CompleteTask.IsCompleted;
+
+			// Add a dependency onto another bundle
+			public void AddDependencyOn(PendingBundle bundle, ILogger? traceLogger)
+			{
+				if (bundle != this)
+				{
+					if (_dependencies.Add(bundle.CompleteTask))
+					{
+						traceLogger?.LogInformation("Added dependency from bundle {BundleId} on {OtherBundleId}", BundleId, bundle.BundleId);
+					}
+				}
+			}
+
+			// Starts writing a node to the buffer
+			public Memory<byte> GetBuffer(int usedSize, int desiredSize)
+			{
+				if (_currentPacket == null)
+				{
+					_currentPacket = MemoryPool<byte>.Shared.Rent(Math.Max(_maxPacketSize, desiredSize));
+				}
+				else if (_currentPacketLength + desiredSize > _currentPacket.Memory.Length)
+				{
+					IMemoryOwner<byte> nextPacket = MemoryPool<byte>.Shared.Rent(Math.Max(_maxPacketSize, desiredSize));
+					if (usedSize > 0)
+					{
+						_currentPacket.Memory.Slice(_currentPacketLength, usedSize).CopyTo(nextPacket.Memory);
+					}
+					FlushPacket();
+					_currentPacket = nextPacket;
+				}
+				return _currentPacket.Memory.Slice(_currentPacketLength, desiredSize);
+			}
+
+			// Finish a node write
+			public PendingNode WriteNode(BlobType blobType, int size, IReadOnlyList<BlobHandle> refs, IReadOnlyList<AliasInfo> aliases)
+			{
+				PendingNode pendingNode = new PendingNode(_treeReader, blobType, _currentPacketIdx, _currentPacketLength, (int)size, refs, aliases, this);
+				_currentPacketLength += pendingNode.Length;
+				UncompressedLength += pendingNode.Length;
+
+				_queue.Add(pendingNode);
+				_queuedRefs += pendingNode.Refs.Length;
+
+				return pendingNode;
+			}
+
+			public ReadOnlyMemory<byte> GetNodeData(int packetIdx, int offset, int length)
+			{
+				lock (_lockObject)
+				{
+					if (packetIdx < _packets.Count)
+					{
+						BundlePacket packet = _packets[packetIdx];
+
+						ReadOnlySequence<byte> sequence = _encodedPacketWriter.AsSequence(packet.EncodedOffset, packet.EncodedLength);
+						if (!sequence.IsSingleSegment)
+						{
+							throw new InvalidOperationException();
+						}
+
+						using (IMemoryOwner<byte> decodeBuffer = MemoryPool<byte>.Shared.Rent(packet.DecodedLength))
+						{
+							BundleData.Decompress(_compressionFormat, sequence.First, decodeBuffer.Memory);
+							return decodeBuffer.Memory.Slice(offset, length).ToArray();
+						}
+					}
+					else if (packetIdx == _currentPacketIdx)
+					{
+						return _currentPacket!.Memory.Slice(offset, length).ToArray();
+					}
+					else
+					{
+						return ReadOnlyMemory<byte>.Empty;
+					}
+				}
+			}
+
+			// Compresses the current packet and schedule it to be written to storage
+			void FlushPacket()
+			{
+				if (_currentPacket != null)
+				{
+					IMemoryOwner<byte> currentPacket = _currentPacket;
+					int currentPacketLength = _currentPacketLength;
+
+					if (currentPacketLength > 0)
+					{
+						_compressPacketsTask = _compressPacketsTask.ContinueWith(x => CompressPacket(currentPacket, currentPacketLength), TaskScheduler.Default);
+						_currentPacketIdx++;
+					}
+					else
+					{
+						currentPacket.Dispose();
+					}
+
+					_currentPacket = null;
+					_currentPacketLength = 0;
+				}
+			}
+
+			void CompressPacket(IMemoryOwner<byte> packetData, int length)
+			{
+				lock (_lockObject)
+				{
+					int encodedLength = BundleData.Compress(_compressionFormat, packetData.Memory.Slice(0, length), _encodedPacketWriter);
+
+					int encodedOffset = (_packets.Count == 0) ? 0 : _packets[^1].EncodedOffset + _packets[^1].EncodedLength;
+					BundlePacket packet = new BundlePacket(_compressionFormat, encodedOffset, encodedLength, length);
+					_packets.Add(packet);
+					Interlocked.Add(ref _compressedLength, encodedLength);
+
+					packetData.Dispose();
+				}
+			}
+
+			// Wait for the bundle's dependencies to complete
+			public async Task FlushDependenciesAsync(ILogger? traceLogger, CancellationToken cancellationToken)
+			{
+				FlushPacket();
+				await _compressPacketsTask.WaitAsync(cancellationToken);
+
+				if (traceLogger != null)
+				{
+					foreach (Task dependency in _dependencies)
+					{
+						if (!dependency.IsCompleted)
+						{
+							traceLogger.LogInformation("Bundle {BundleId} is stalling waiting for dependencies to flush first", BundleId);
+						}
+					}
+				}
+
+				await Task.WhenAll(_dependencies).WaitAsync(cancellationToken);
+			}
+
+			// Mark the bundle as complete
+			public async Task WriteAsync(IStorageClient store, string? basePath, ILogger? traceLogger)
+			{
+				traceLogger?.LogInformation("Marking bundle {BundleId} as complete ({NumNodes} nodes); adding to write queue.", BundleId, _queue.Count);
+
+				await FlushDependenciesAsync(traceLogger, CancellationToken.None);
+
+				try
+				{
+					(BundleHeader header, List<ReadOnlyMemory<byte>> packets) = CreateBundle();
+
+					// Write the bundle to storage
+					BlobHandle[] imports = new BlobHandle[header.Imports.Count];
+					for (int idx = 0; idx < header.Imports.Count; idx++)
+					{
+						imports[idx] = store.CreateBlobHandle(new BlobLocator(header.Imports[idx].Path));
+					}
+
+					// Create the output sequence
+					ReadOnlySequenceBuilder<byte> sequence = new ReadOnlySequenceBuilder<byte>();
+					header.AppendTo(sequence);
+
+					foreach (ReadOnlyMemory<byte> packet in packets)
+					{
+						sequence.Append(packet);
+					}
+
+					// Write it
+					BlobHandle handle;
+					using (ReadOnlySequenceStream stream = new ReadOnlySequenceStream(sequence.Construct()))
+					{
+						handle = await store.WriteBlobAsync(BundleStorageClient.BundleBlobType, stream, imports, basePath);
+					}
+
+					BlobLocator locator = handle.GetLocator();
+					traceLogger?.LogInformation("Written bundle {BundleId} as {Locator}", BundleId, locator);
+
+					for (int idx = 0; idx < _queue.Count; idx++)
+					{
+						PendingNode node = _queue[idx];
+
+						FlushedNodeHandle flushedHandle = new FlushedNodeHandle(_treeReader, locator, idx);
+						_queue[idx].MarkAsWritten(flushedHandle);
+
+						foreach (AliasInfo alias in node.Aliases)
+						{
+							await store.AddAliasAsync(alias.Name, flushedHandle, alias.Rank, alias.Data, CancellationToken.None);
+						}
+					}
+
+					_completeEvent.SetResult(true);
+				}
+				catch (Exception ex)
+				{
+					_completeEvent.SetException(ex);
+				}
+			}
+
+			// Flush the current write state
+			public Task FlushAsync(CancellationToken cancellationToken) => _treeWriter.FlushAsync(cancellationToken);
+
+			// Mark the bundle as complete and create a bundle with the current state
+			public (BundleHeader, List<ReadOnlyMemory<byte>>) CreateBundle()
+			{
+				lock (_lockObject)
+				{
+					return CreateBundleInternal();
+				}
+			}
+
+			(BundleHeader, List<ReadOnlyMemory<byte>>) CreateBundleInternal()
+			{
+				// List of imported blobs
+				List<BlobLocator> imports = new List<BlobLocator>();
+				Dictionary<BlobLocator, int> importToIndex = new Dictionary<BlobLocator, int>();
+
+				// List of types in the bundle
+				List<BlobType> types = new List<BlobType>();
+				Dictionary<BlobType, int> typeToIndex = new Dictionary<BlobType, int>();
+
+				// Map of node handle to reference
+				Dictionary<BlobHandle, BundleExportRef> nodeHandleToExportRef = new Dictionary<BlobHandle, BundleExportRef>();
+				for (int exportIdx = 0; exportIdx < _queue.Count; exportIdx++)
+				{
+					BlobHandle handle = _queue[exportIdx];
+					nodeHandleToExportRef[handle] = new BundleExportRef(-1, exportIdx);
+				}
+
+				// Create the export list
+				List<BundleExport> exports = new List<BundleExport>(_queue.Count);
+				foreach (PendingNode nodeInfo in _queue)
+				{
+					int typeIdx = FindOrAddItemIndex(nodeInfo.BlobType, types, typeToIndex);
+
+					List<BundleExportRef> exportRefs = new List<BundleExportRef>();
+					foreach (BlobHandle handle in nodeInfo.Refs)
+					{
+						BundleExportRef exportRef;
+						if (!nodeHandleToExportRef.TryGetValue(handle, out exportRef))
+						{
+							FlushedNodeHandle? flushedHandle = handle as FlushedNodeHandle;
+							if (flushedHandle == null)
+							{
+								flushedHandle = (handle as PendingNode)?.FlushedNodeHandle;
+								if (flushedHandle == null)
+								{
+									throw new InvalidOperationException($"Node {handle.GetLocator()} is not a bundle node");
+								}
+							}
+							
+							int importIdx = FindOrAddItemIndex(flushedHandle.BundleLocator, imports, importToIndex);
+							exportRef = new BundleExportRef(importIdx, flushedHandle.ExportIdx);
+						}
+						exportRefs.Add(exportRef);
+					}
+
+					BundleExport export = new BundleExport(typeIdx, nodeInfo.Packet, nodeInfo.Offset, nodeInfo.Length, exportRefs);
+					exports.Add(export);
+				}
+
+				// Get the memory for each packet
+				List<ReadOnlyMemory<byte>> packetData = new List<ReadOnlyMemory<byte>>(_packets.Count);
+				foreach (ReadOnlyMemory<byte> segment in _encodedPacketWriter.AsSequence())
+				{
+					int offset = 0;
+					while(packetData.Count < _packets.Count)
+					{
+						BundlePacket packet = _packets[packetData.Count];
+						int next = offset + packet.EncodedLength;
+
+						if (next > segment.Length)
+						{
+							break;
+						}
+						packetData.Add(segment.Slice(offset, packet.EncodedLength));
+						offset = next;
+					}
+				}
+
+				// Create the bundle
+				BundleHeader header = new BundleHeader(types.ToArray(), imports.ToArray(), exports.ToArray(), _packets.ToArray());
+				return (header, packetData);
+			}
+
+			static int FindOrAddItemIndex<TItem>(TItem item, List<TItem> items, Dictionary<TItem, int> itemToIndex) where TItem : notnull
+			{
+				int index;
+				if (!itemToIndex.TryGetValue(item, out index))
+				{
+					index = items.Count;
+					items.Add(item);
+					itemToIndex.Add(item, index);
+				}
+				return index;
+			}
+		}
+
+		class WriteQueue
+		{
+			long _memoryFootprint;
+			readonly IStorageClient _store;
+			readonly string? _basePath;
+			readonly long _maxMemoryFootprint;
+			readonly ILogger? _traceLogger;
+			readonly AsyncEvent _completeEvent = new AsyncEvent();
+			int _refCount;
+			readonly List<Task> _writeTasks = new List<Task>();
+
+			public WriteQueue(IStorageClient store, string? basePath, long maxMemoryFootprint, ILogger? traceLogger)
+			{
+				_store = store;
+				_basePath = basePath;
+				_maxMemoryFootprint = maxMemoryFootprint;
+				_refCount = 1;
+				_traceLogger = traceLogger;
+			}
+
+			public void AddRef()
+			{
+				Interlocked.Increment(ref _refCount);
+			}
+
+			public async ValueTask ReleaseAsync()
+			{
+				if (Interlocked.Decrement(ref _refCount) == 0)
+				{
+					await FlushAsync(CancellationToken.None);
+				}
+			}
+
+			public async Task AddAsync(PendingBundle pendingBundle, CancellationToken cancellationToken)
+			{
+				lock (_writeTasks)
+				{
+					_writeTasks.RemoveCompleteTasks();
+				}
+
+				await pendingBundle.FlushDependenciesAsync(_traceLogger, cancellationToken);
+
+				for (; ; )
+				{
+					Task completeTask = _completeEvent.Task;
+
+					long memoryFootprint = Interlocked.CompareExchange(ref _memoryFootprint, 0, 0);
+					long newMemoryFootprint = memoryFootprint + pendingBundle.CompressedLength;
+
+					if (memoryFootprint > 0 && newMemoryFootprint > _maxMemoryFootprint)
+					{
+						await completeTask;
+					}
+					else if(Interlocked.CompareExchange(ref _memoryFootprint, newMemoryFootprint, memoryFootprint) == memoryFootprint)
+					{
+						break;
+					}
+				}
+
+				lock (_writeTasks)
+				{
+					_writeTasks.Add(WriteAsync(pendingBundle));
+				}
+			}
+
+			public async Task FlushAsync(CancellationToken cancellationToken)
+			{
+				Task[] writeTasks;
+				lock (_writeTasks)
+				{
+					writeTasks = _writeTasks.ToArray();
+				}
+				await Task.WhenAll(writeTasks).WaitAsync(cancellationToken);
+			}
+
+			async Task WriteAsync(PendingBundle pendingBundle)
+			{
+				await pendingBundle.WriteAsync(_store, _basePath, _traceLogger);
+				Interlocked.Add(ref _memoryFootprint, -pendingBundle.CompressedLength);
+				pendingBundle.Dispose();
+				_completeEvent.Pulse();
+			}
+		}
+
+		static readonly BundleOptions s_defaultOptions = new BundleOptions();
+
+		readonly IStorageClient _store;
+		readonly BundleReader _reader;
+		readonly BundleOptions _options;
+		readonly string? _basePath;
+
+		readonly WriteQueue _writeQueue;
+
+		PendingBundle? _currentBundle;
+
+		bool _disposed;
+
+		internal ILogger? TraceLogger { get; }
+
+		/// <summary>
+		/// Constructor
+		/// </summary>
+		/// <param name="store">Store to write data to</param>
+		/// <param name="reader">Reader for serialized node data</param>
+		/// <param name="basePath">Base path for new nodes</param>
+		/// <param name="options">Options for the writer</param>
+		/// <param name="traceLogger">Optional logger for trace information</param>
+		public BundleWriter(IStorageClient store, BundleReader reader, string? basePath, BundleOptions? options = null, ILogger? traceLogger = null)
+			: this(store, reader, basePath, options, null, traceLogger)
+		{
+		}
+
+		/// <summary>
+		/// Copy constructor
+		/// </summary>
+		/// <param name="other"></param>
+		public BundleWriter(BundleWriter other)
+			: this(other._store, other._reader, other._basePath, other._options, other._writeQueue, other.TraceLogger)
+		{
+			_writeQueue.AddRef();
+		}
+
+		/// <summary>
+		/// Internal constructor
+		/// </summary>
+		private BundleWriter(IStorageClient store, BundleReader reader, string? basePath, BundleOptions? options, WriteQueue? writeQueue, ILogger? traceLogger = null)
+		{
+			_store = store;
+			_reader = reader;
+			_basePath = basePath;
+			_options = options ?? s_defaultOptions;
+			_writeQueue = writeQueue ?? new WriteQueue(store, basePath, _options.MaxWriteQueueLength, traceLogger);
+			TraceLogger = traceLogger;
+		}
+
+		/// <inheritdoc/>
+		IStorageWriter IStorageWriter.Fork() => new BundleWriter(this);
+
+		/// <inheritdoc/>
+		public async ValueTask DisposeAsync()
+		{
+			if (!_disposed)
+			{
+				await FlushAsync();
+				await _writeQueue.ReleaseAsync();
+				_disposed = true;
+			}
+
+			if (_currentBundle != null)
+			{
+				_currentBundle.Dispose();
+				_currentBundle = null;
+			}
+		}
+
+		/// <summary>
+		/// Mark this writer as complete, allowing its data to be serialized.
+		/// </summary>
+		public async Task CompleteAsync(CancellationToken cancellationToken)
+		{
+			if (_currentBundle != null)
+			{
+				await _writeQueue.AddAsync(_currentBundle, cancellationToken);
+				_currentBundle = null;
+			}
+		}
+
+		/// <summary>
+		/// Gets an output buffer for writing.
+		/// </summary>
+		/// <param name="usedSize">Current size in the existing buffer that has been written to</param>
+		/// <param name="desiredSize">Desired size of the returned buffer</param>
+		/// <returns>Buffer to be written into.</returns>
+		public Memory<byte> GetOutputBuffer(int usedSize, int desiredSize)
+		{
+			PendingBundle currentBundle = GetCurrentBundle();
+			return currentBundle.GetBuffer(usedSize, desiredSize);
+		}
+
+		/// <summary>
+		/// Finish writing a node.
+		/// </summary>
+		/// <param name="type">Type of the node that was written</param>
+		/// <param name="size">Used size of the buffer</param>
+		/// <param name="references">References to other nodes</param>
+		/// <param name="aliases">Aliases for the node</param>
+		/// <param name="cancellationToken">Cancellation token for the operation</param>
+		/// <returns>Handle to the written node</returns>
+		public async ValueTask<BlobHandle> WriteBlobAsync(BlobType type, int size, IReadOnlyList<BlobHandle> references, IReadOnlyList<AliasInfo> aliases, CancellationToken cancellationToken = default)
+		{
+			PendingBundle currentBundle = GetCurrentBundle();
+
+			// Append this node data
+			PendingNode pendingNode = currentBundle.WriteNode(type, size, references, aliases);
+			TraceLogger?.LogInformation("Added new node for {NodeKey} in bundle {BundleId}", pendingNode, currentBundle.BundleId);
+
+			// Add dependencies on all bundles containing a dependent node
+			foreach (BlobHandle reference in references)
+			{
+				PendingNode? pendingReference = reference as PendingNode;
+				if (pendingReference?.PendingBundle != null)
+				{
+					currentBundle.AddDependencyOn(pendingReference.PendingBundle, TraceLogger);
+				}
+			}
+
+			// If the bundle is full, start the process of writing it to disk
+			if (currentBundle.UncompressedLength > _options.MaxBlobSize || currentBundle.IsFull())
+			{
+				await CompleteAsync(cancellationToken);
+			}
+
+			return pendingNode;
+		}
+
+		PendingBundle GetCurrentBundle()
+		{
+			if (_currentBundle == null)
+			{
+				int bufferSize = (int)(_options.MinCompressionPacketSize * 1.2);
+				_currentBundle = new PendingBundle(_reader, this, bufferSize, _options.CompressionFormat);
+			}
+			return _currentBundle;
+		}
+
+		/// <summary>
+		/// Flushes all the current nodes to storage
+		/// </summary>
+		/// <param name="cancellationToken">Cancellation token for the operation</param>
+		/// <returns></returns>
+		public async Task FlushAsync(CancellationToken cancellationToken = default)
+		{
+			if (_disposed)
+			{
+				throw new ObjectDisposedException(GetType().Name);
+			}
+
+			await CompleteAsync(cancellationToken);
+			await _writeQueue.FlushAsync(cancellationToken);
+		}
+
+		/// <summary>
+		/// Flushes all the current nodes to storage
+		/// </summary>
+		/// <param name="root">Root for the tree</param>
+		/// <param name="cancellationToken">Cancellation token for the operation</param>
+		/// <returns></returns>
+		public async Task<BlobHandle> FlushAsync(Node root, CancellationToken cancellationToken = default)
+		{
+			return await NodeRefExtensions.FlushAsync(this, root, cancellationToken);
+		}
+	}
+}
diff --git a/Engine/Source/Programs/Shared/EpicGames.Horde/Storage/Clients/BundleStorageClient.cs b/Engine/Source/Programs/Shared/EpicGames.Horde/Storage/Clients/BundleStorageClient.cs
--- a/Engine/Source/Programs/Shared/EpicGames.Horde/Storage/Clients/BundleStorageClient.cs
+++ b/Engine/Source/Programs/Shared/EpicGames.Horde/Storage/Clients/BundleStorageClient.cs
@@ -50,7 +50,7 @@ namespace EpicGames.Horde.Storage.Clients
 				{
 					List<BlobHandle> refs = new List<BlobHandle>();
 
-					BundleHeader header = await _storageClient.ReadHeaderAsync(_locator, cancellationToken);
+					Bundles.V1.BundleHeader header = await _storageClient.ReadHeaderAsync(_locator, cancellationToken);
 					foreach (BlobLocator import in header.Imports)
 					{
 						refs.Add(_storageClient.CreateBlobHandle(new BlobLocator(import.Path)));
@@ -86,7 +86,7 @@ namespace EpicGames.Horde.Storage.Clients
 		}
 
 		readonly IStorageClient _inner;
-		readonly BundleReader _bundleReader;
+		readonly Bundles.V1.BundleReader _bundleReader;
 
 		/// <inheritdoc/>
 		public bool SupportsRedirects { get; } = false;
@@ -97,7 +97,7 @@ namespace EpicGames.Horde.Storage.Clients
 		public BundleStorageClient(IStorageClient inner, BundleReaderCache cache, ILogger logger)
 		{
 			_inner = inner;
-			_bundleReader = new BundleReader(this, cache, logger);
+			_bundleReader = new Bundles.V1.BundleReader(this, cache, logger);
 		}
 
 		/// <inheritdoc/>
@@ -186,7 +186,7 @@ namespace EpicGames.Horde.Storage.Clients
 			=> _inner.WriteBlobAsync(BundleBlobType, stream, references, basePath, cancellationToken);
 
 		/// <inheritdoc/>
-		public Task<BundleHeader> ReadHeaderAsync(BlobLocator locator, CancellationToken cancellationToken) => _bundleReader.ReadHeaderAsync(locator, cancellationToken);
+		public Task<Bundles.V1.BundleHeader> ReadHeaderAsync(BlobLocator locator, CancellationToken cancellationToken) => _bundleReader.ReadHeaderAsync(locator, cancellationToken);
 
 		#endregion
 
@@ -197,7 +197,7 @@ namespace EpicGames.Horde.Storage.Clients
 		{
 			if (locator.CanUnwrap())
 			{
-				return FlushedNodeHandle.FromBlobLocator(_bundleReader, locator);
+				return Bundles.V1.FlushedNodeHandle.FromBlobLocator(_bundleReader, locator);
 			}
 			else
 			{
@@ -206,7 +206,7 @@ namespace EpicGames.Horde.Storage.Clients
 		}
 
 		/// <inheritdoc/>
-		public BundleWriter CreateWriter(string? basePath = null, BundleOptions? options = null) => new BundleWriter(this, _bundleReader, basePath, options);
+		public IStorageWriter CreateWriter(string? basePath = null, BundleOptions? options = null) => new Bundles.V1.BundleWriter(this, _bundleReader, basePath, options);
 
 		/// <inheritdoc/>
 		IStorageWriter IStorageClient.CreateWriter(string? basePath) => CreateWriter(basePath);
