Horde: Add a static blob type to the Bundle class.

diff --git a/Engine/Source/Programs/Horde/Horde.Server/Storage/StorageController.cs b/Engine/Source/Programs/Horde/Horde.Server/Storage/StorageController.cs
--- a/Engine/Source/Programs/Horde/Horde.Server/Storage/StorageController.cs
+++ b/Engine/Source/Programs/Horde/Horde.Server/Storage/StorageController.cs
@@ -12,7 +12,7 @@ using System.Threading.Tasks;
 using Amazon.EC2.Model;
 using EpicGames.Core;
 using EpicGames.Horde.Storage;
-using EpicGames.Horde.Storage.Clients;
+using EpicGames.Horde.Storage.Bundles;
 using EpicGames.Horde.Storage.Nodes;
 using Horde.Server.Acls;
 using Horde.Server.Server;
@@ -98,7 +98,7 @@ namespace Horde.Server.Storage
 			else
 			{
 				using Stream stream = file.OpenReadStream();
-				BlobHandle handle = await storageClient.WriteBlobAsync(BundleStorageClient.BundleBlobType, stream, Array.Empty<BlobHandle>(), prefix, cancellationToken);
+				BlobHandle handle = await storageClient.WriteBlobAsync(Bundle.BlobType, stream, Array.Empty<BlobHandle>(), prefix, cancellationToken);
 				return new WriteBlobResponse { Blob = handle.GetLocator().ToString(), SupportsRedirects = storageClient.SupportsRedirects };
 			}
 		}
diff --git a/Engine/Source/Programs/Shared/EpicGames.Horde/Storage/BlobHandle.cs b/Engine/Source/Programs/Shared/EpicGames.Horde/Storage/BlobHandle.cs
--- a/Engine/Source/Programs/Shared/EpicGames.Horde/Storage/BlobHandle.cs
+++ b/Engine/Source/Programs/Shared/EpicGames.Horde/Storage/BlobHandle.cs
@@ -156,10 +156,10 @@ namespace EpicGames.Horde.Storage
 		/// <inheritdoc/>
 		public override string ToString()
 		{
-			BlobLocator blobId;
-			if (TryGetLocator(out blobId))
+			BlobLocator locator;
+			if (TryGetLocator(out locator))
 			{
-				return blobId.ToString();
+				return locator.ToString();
 			}
 			return base.ToString() ?? "Unknown";
 		}
diff --git a/Engine/Source/Programs/Shared/EpicGames.Horde/Storage/Bundles/Bundle.cs b/Engine/Source/Programs/Shared/EpicGames.Horde/Storage/Bundles/Bundle.cs
--- a/Engine/Source/Programs/Shared/EpicGames.Horde/Storage/Bundles/Bundle.cs
+++ b/Engine/Source/Programs/Shared/EpicGames.Horde/Storage/Bundles/Bundle.cs
@@ -62,6 +62,11 @@ namespace EpicGames.Horde.Storage.Bundles
 	/// </summary>
 	public static class Bundle
 	{
+		/// <summary>
+		/// Blob type for bundles
+		/// </summary>
+		public static BlobType BlobType { get; } = new BlobType(Guid.Parse("{7C5BA294-2D21-4F92-85BE-852F48CC4C1E}"), 1);
+
 		/// <summary>
 		/// Number of bytes in a signature when serialized
 		/// </summary>
diff --git a/Engine/Source/Programs/Shared/EpicGames.Horde/Storage/Bundles/V1/BundleWriter.cs b/Engine/Source/Programs/Shared/EpicGames.Horde/Storage/Bundles/V1/BundleWriter.cs
--- a/Engine/Source/Programs/Shared/EpicGames.Horde/Storage/Bundles/V1/BundleWriter.cs
+++ b/Engine/Source/Programs/Shared/EpicGames.Horde/Storage/Bundles/V1/BundleWriter.cs
@@ -10,7 +10,6 @@ using System.Linq;
 using System.Threading;
 using System.Threading.Tasks;
 using EpicGames.Core;
-using EpicGames.Horde.Storage.Clients;
 using Microsoft.Extensions.Logging;
 
 namespace EpicGames.Horde.Storage.Bundles.V1
@@ -417,7 +416,7 @@ namespace EpicGames.Horde.Storage.Bundles.V1
 					BlobHandle handle;
 					using (ReadOnlySequenceStream stream = new ReadOnlySequenceStream(sequence.Construct()))
 					{
-						handle = await store.WriteBlobAsync(BundleStorageClient.BundleBlobType, stream, imports, basePath);
+						handle = await store.WriteBlobAsync(Bundle.BlobType, stream, imports, basePath);
 					}
 
 					BlobLocator locator = handle.GetLocator();
diff --git a/Engine/Source/Programs/Shared/EpicGames.Horde/Storage/Clients/BundleStorageClient.cs b/Engine/Source/Programs/Shared/EpicGames.Horde/Storage/Clients/BundleStorageClient.cs
--- a/Engine/Source/Programs/Shared/EpicGames.Horde/Storage/Clients/BundleStorageClient.cs
+++ b/Engine/Source/Programs/Shared/EpicGames.Horde/Storage/Clients/BundleStorageClient.cs
@@ -17,11 +17,6 @@ namespace EpicGames.Horde.Storage.Clients
 	/// </summary>
 	public sealed class BundleStorageClient : IStorageClient
 	{
-		/// <summary>
-		/// Blob type for bundles
-		/// </summary>
-		public static BlobType BundleBlobType { get; } = new BlobType(Guid.Parse("{7C5BA294-2D21-4F92-85BE-852F48CC4C1E}"), 1);
-
 		/// <summary>
 		/// Handle to a bundle object
 		/// </summary>
@@ -41,7 +36,7 @@ namespace EpicGames.Horde.Storage.Clients
 			}
 
 			/// <inheritdoc/>
-			public override ValueTask<BlobType> GetTypeAsync(CancellationToken cancellationToken = default) => new ValueTask<BlobType>(BundleBlobType);
+			public override ValueTask<BlobType> GetTypeAsync(CancellationToken cancellationToken = default) => new ValueTask<BlobType>(Bundle.BlobType);
 
 			/// <inheritdoc/>
 			public override async ValueTask<IReadOnlyList<BlobHandle>> GetRefsAsync(CancellationToken cancellationToken = default)
@@ -73,7 +68,7 @@ namespace EpicGames.Horde.Storage.Clients
 				using (Stream stream = await _storageClient.OpenBlobAsync(_locator, 0, cancellationToken: cancellationToken))
 				{
 					byte[] data = await stream.ReadAllBytesAsync(cancellationToken);
-					return new BlobData(BundleBlobType, data, await GetRefsAsync(cancellationToken));
+					return new BlobData(Bundle.BlobType, data, await GetRefsAsync(cancellationToken));
 				}
 			}
 
@@ -140,7 +135,7 @@ namespace EpicGames.Horde.Storage.Clients
 		/// <inheritdoc/>
 		public async ValueTask<BlobHandle> WriteBlobAsync(BlobType type, Stream stream, IReadOnlyList<BlobHandle> references, string? basePath = null, CancellationToken cancellationToken = default)
 		{
-			if (type == BundleBlobType)
+			if (type == Bundle.BlobType)
 			{
 				return await WriteBundleAsync(stream, references, basePath, cancellationToken);
 			}
@@ -183,7 +178,7 @@ namespace EpicGames.Horde.Storage.Clients
 		/// Write a bundle to the underlying storage
 		/// </summary>
 		ValueTask<BlobHandle> WriteBundleAsync(Stream stream, IReadOnlyList<BlobHandle> references, string? basePath = null, CancellationToken cancellationToken = default)
-			=> _inner.WriteBlobAsync(BundleBlobType, stream, references, basePath, cancellationToken);
+			=> _inner.WriteBlobAsync(Bundle.BlobType, stream, references, basePath, cancellationToken);
 
 		/// <inheritdoc/>
 		public Task<Bundles.V1.BundleHeader> ReadHeaderAsync(BlobLocator locator, CancellationToken cancellationToken) => _bundleReader.ReadHeaderAsync(locator, cancellationToken);
