Horde: Return agent categories to the dashboard.

diff --git a/Engine/Source/Programs/Horde/Horde.Server/Dashboard/DashboardController.cs b/Engine/Source/Programs/Horde/Horde.Server/Dashboard/DashboardController.cs
--- a/Engine/Source/Programs/Horde/Horde.Server/Dashboard/DashboardController.cs
+++ b/Engine/Source/Programs/Horde/Horde.Server/Dashboard/DashboardController.cs
@@ -142,6 +142,11 @@ namespace Horde.Server.Dashboard
 
 			dashboardConfigResponse.DeviceProblemCooldownMinutes = _settings.DeviceProblemCooldownMinutes;
 
+			foreach (DashboardAgentCategoryConfig category in _globalConfig.Value.Dashboard.AgentCategories)
+			{
+				dashboardConfigResponse.AgentCategories.Add(new GetDashboardAgentCategoryResponse { Name = category.Name, Condition = category.Condition });
+			}
+
 			return dashboardConfigResponse;
 		}
 
diff --git a/Engine/Source/Programs/Shared/EpicGames.Horde/Dashboard/DashboardMessages.cs b/Engine/Source/Programs/Shared/EpicGames.Horde/Dashboard/DashboardMessages.cs
--- a/Engine/Source/Programs/Shared/EpicGames.Horde/Dashboard/DashboardMessages.cs
+++ b/Engine/Source/Programs/Shared/EpicGames.Horde/Dashboard/DashboardMessages.cs
@@ -1,6 +1,10 @@
 // Copyright Epic Games, Inc. All Rights Reserved.
 
 using System;
+using System.Collections.Generic;
+using EpicGames.Horde.Common;
+
+#pragma warning disable CA2227 // Collection properties should be read only
 
 namespace EpicGames.Horde.Dashboard
 {
@@ -38,6 +42,27 @@ namespace EpicGames.Horde.Dashboard
 		/// Device problem cooldown in minutes
 		/// </summary>
 		public int DeviceProblemCooldownMinutes { get; set; }
+
+		/// <summary>
+		/// Categories to display on the agents page
+		/// </summary>
+		public List<GetDashboardAgentCategoryResponse> AgentCategories { get; set; } = new List<GetDashboardAgentCategoryResponse>();
+	}
+
+	/// <summary>
+	/// Describes a category for the agents page
+	/// </summary>
+	public class GetDashboardAgentCategoryResponse
+	{
+		/// <summary>
+		/// Title for the tab
+		/// </summary>
+		public string Name { get; set; } = "Unnamed";
+
+		/// <summary>
+		/// Condition for agents to be included in this category
+		/// </summary>
+		public Condition? Condition { get; set; }
 	}
 
 	/// <summary>
