EpicGames.Perforce: Add support for updating streams

diff --git a/Engine/Source/Programs/Shared/EpicGames.Perforce.Fixture/PerforceFixture.cs b/Engine/Source/Programs/Shared/EpicGames.Perforce.Fixture/PerforceFixture.cs
--- a/Engine/Source/Programs/Shared/EpicGames.Perforce.Fixture/PerforceFixture.cs
+++ b/Engine/Source/Programs/Shared/EpicGames.Perforce.Fixture/PerforceFixture.cs
@@ -32,6 +32,11 @@ public class DepotFileFixture
 		Digest = PerforceFixture.CalcMd5(content).ToUpperInvariant();
 		Content = content;
 	}
+
+	public override string ToString()
+	{
+		return $"DepotFile={DepotFile}, Size={Size}, Revision={Revision}, Digest={Digest}, Content={Content}";
+	}
 }
 
 public class ChangelistFixture
@@ -131,10 +136,16 @@ public class ChangelistFixture
 	}
 	
 	private (FileSet localFiles, FileSet streamFiles) GetFileSets(string clientRoot)
+	{
+		FileSet streamFiles = new(StreamFiles.Select(x => (x.ClientFile, x.Size, x.Digest)));
+		FileSet localFiles = GetLocalFileSet(clientRoot);
+		return (localFiles, streamFiles);
+	}
+	
+	public static FileSet GetLocalFileSet(string clientRoot)
 	{
 		EnumerationOptions options = new () { RecurseSubdirectories = true };
 
-		FileSet streamFiles = new(StreamFiles.Select(x => (x.ClientFile, x.Size, x.Digest)));
 		FileSet localFiles = new(Directory.EnumerateFiles(clientRoot, "*", options)
 			.Select(x => Path.GetRelativePath(clientRoot, x))
 			.Select(x => x.Replace("\\", "/", StringComparison.Ordinal))
@@ -150,7 +161,7 @@ public class ChangelistFixture
 				return (clientFile, size, PerforceFixture.CalcMd5(content).ToUpperInvariant());
 			}));
 		
-		return (localFiles, streamFiles);
+		return localFiles;
 	}
 
 	public (List<string> localFiles, List<string> streamFiles) GetFiles(string clientRoot)
diff --git a/Engine/Source/Programs/Shared/EpicGames.Perforce/PerforceConnection.cs b/Engine/Source/Programs/Shared/EpicGames.Perforce/PerforceConnection.cs
--- a/Engine/Source/Programs/Shared/EpicGames.Perforce/PerforceConnection.cs
+++ b/Engine/Source/Programs/Shared/EpicGames.Perforce/PerforceConnection.cs
@@ -3361,7 +3361,68 @@ namespace EpicGames.Perforce
 
 			return SingleResponseCommandAsync<StreamRecord>(connection, "stream", arguments, null, cancellationToken);
 		}
+		
+		/// <summary>
+		/// Updates an existing stream
+		/// </summary>
+		/// <param name="connection">Connection to the Perforce server</param>
+		/// <param name="record">Information of the stream to update</param>
+		/// <param name="cancellationToken">Token used to cancel the operation</param>
+		/// <returns>Stream information record</returns>
+		public static Task<PerforceResponse> TryUpdateStreamAsync(this IPerforceConnection connection, StreamRecord record, CancellationToken cancellationToken = default)
+		{
+			List<string> arguments = new () { "-i" };
+			return SingleResponseCommandAsync(connection, "stream", arguments, connection.SerializeRecord(record), null, cancellationToken);
+		}
+		
+		/// <summary>
+		/// Serializes a client record to a byte array
+		/// </summary>
+		/// <param name="connection">Connection to the Perforce server</param>
+		/// <param name="input">The input record</param>
+		/// <returns>Serialized record data</returns>
+		static byte[] SerializeRecord(this IPerforceConnection connection, StreamRecord input)
+		{
+			List<KeyValuePair<string, object>> nameToValue = new List<KeyValuePair<string, object>>();
+
+			void Add(string fieldName, string? value)
+			{
+				if (value != null)
+				{
+					nameToValue.Add(new KeyValuePair<string, object>(fieldName, value));
+				}
+			}
+			
+			Add("Stream", input.Stream);
+			Add("Owner", input.Owner);
+			Add("Name", input.Name);
+			Add("Parent", input.Parent);
+			Add("Type", input.Type);
+			Add("Description", input.Description);
+			Add("ParentView", input.ParentView);
+			
+			if (input.Options != StreamOptions.None)
+			{
+				nameToValue.Add(new KeyValuePair<string, object>("Options", PerforceReflection.GetEnumText(typeof(StreamOptions), input.Options)));
+			}
+			
+			if (input.Paths.Count > 0)
+			{
+				nameToValue.Add(new KeyValuePair<string, object>("Paths", input.Paths));
+			}
+			
+			if (input.View.Count > 0)
+			{
+				nameToValue.Add(new KeyValuePair<string, object>("View", input.View));
+			}
+			
+			if (input.ChangeView.Count > 0)
+			{
+				nameToValue.Add(new KeyValuePair<string, object>("ChangeView", input.ChangeView));
+			}
 
+			return connection.CreateRecord(nameToValue).Serialize();
+		}
 		#endregion
 
 		#region p4 streams
diff --git a/Engine/Source/Programs/Shared/EpicGames.Perforce/Records/StreamRecord.cs b/Engine/Source/Programs/Shared/EpicGames.Perforce/Records/StreamRecord.cs
--- a/Engine/Source/Programs/Shared/EpicGames.Perforce/Records/StreamRecord.cs
+++ b/Engine/Source/Programs/Shared/EpicGames.Perforce/Records/StreamRecord.cs
@@ -63,6 +63,12 @@ namespace EpicGames.Perforce
 		/// </summary>
 		[PerforceTag("Options")]
 		public StreamOptions Options { get; set; }
+		
+		/// <summary>
+		/// Parent view
+		/// </summary>
+		[PerforceTag("ParentView")]
+		public string? ParentView { get; set; }
 
 		/// <summary>
 		/// List of paths in the stream spec
@@ -91,5 +97,26 @@ namespace EpicGames.Perforce
 			Name = null!;
 			Owner = null!;
 		}
+		
+		/// <summary>
+		/// Constructor
+		/// </summary>
+		/// <param name="stream">Stream</param>
+		/// <param name="owner">Owner</param>
+		/// <param name="name">Name of the stream</param>
+		/// <param name="parent">Parent stream</param>
+		/// <param name="type">Stream type</param>
+		/// <param name="parentView">Parent view</param>
+		/// <param name="paths">List of paths in the stream spec</param>
+		public StreamRecord(string stream, string owner, string name, string parent, string type, string parentView, List<string> paths)
+		{
+			Stream = stream;
+			Owner = owner;
+			Name = name;
+			Parent = parent;
+			Type = type;
+			ParentView = parentView;
+			Paths = paths;
+		}
 	}
 }
