Horde: Add value semantics to blob handles, so they can be used as cache keys.

diff --git a/Engine/Source/Programs/Horde/Horde.Server/Storage/StorageService.cs b/Engine/Source/Programs/Horde/Horde.Server/Storage/StorageService.cs
--- a/Engine/Source/Programs/Horde/Horde.Server/Storage/StorageService.cs
+++ b/Engine/Source/Programs/Horde/Horde.Server/Storage/StorageService.cs
@@ -99,6 +99,12 @@ namespace Horde.Server.Storage
 				locator = new BlobLocator(_path);
 				return true;
 			}
+
+			/// <inheritdoc/>
+			public override bool Equals(object? obj) => obj is LeafBlobHandle other && String.Equals(_path, other._path, StringComparison.Ordinal);
+
+			/// <inheritdoc/>
+			public override int GetHashCode() => _path.GetHashCode(StringComparison.Ordinal);
 		}
 
 		sealed class StorageClientImpl : IStorageClient
diff --git a/Engine/Source/Programs/Shared/EpicGames.Horde.Tests/BlobDataTests.cs b/Engine/Source/Programs/Shared/EpicGames.Horde.Tests/BlobDataTests.cs
--- a/Engine/Source/Programs/Shared/EpicGames.Horde.Tests/BlobDataTests.cs
+++ b/Engine/Source/Programs/Shared/EpicGames.Horde.Tests/BlobDataTests.cs
@@ -25,6 +25,9 @@ namespace EpicGames.Horde.Tests
 				locator = _locator;
 				return true;
 			}
+
+			public override bool Equals(object? obj) => obj is DummyHandle other && _locator == other._locator;
+			public override int GetHashCode() => _locator.GetHashCode();
 		}
 
 		[TestMethod]
diff --git a/Engine/Source/Programs/Shared/EpicGames.Horde/Storage/Bundles/V1/BundleWriter.cs b/Engine/Source/Programs/Shared/EpicGames.Horde/Storage/Bundles/V1/BundleWriter.cs
--- a/Engine/Source/Programs/Shared/EpicGames.Horde/Storage/Bundles/V1/BundleWriter.cs
+++ b/Engine/Source/Programs/Shared/EpicGames.Horde/Storage/Bundles/V1/BundleWriter.cs
@@ -7,6 +7,7 @@ using System.Collections.Generic;
 using System.Diagnostics;
 using System.Diagnostics.CodeAnalysis;
 using System.Linq;
+using System.Runtime.CompilerServices;
 using System.Threading;
 using System.Threading.Tasks;
 using EpicGames.Core;
@@ -58,6 +59,12 @@ namespace EpicGames.Horde.Storage.Bundles.V1
 
 		/// <inheritdoc/>
 		public override ValueTask FlushAsync(CancellationToken cancellationToken = default) => new ValueTask();
+
+		/// <inheritdoc/>
+		public override bool Equals(object? obj) => obj is FlushedNodeHandle other && BundleLocator == other.BundleLocator && ExportIdx == other.ExportIdx;
+
+		/// <inheritdoc/>
+		public override int GetHashCode() => HashCode.Combine(BundleLocator, ExportIdx);
 	}
 
 	/// <summary>
@@ -164,6 +171,12 @@ namespace EpicGames.Horde.Storage.Bundles.V1
 					await pendingBundle.FlushAsync(cancellationToken);
 				}
 			}
+
+			/// <inheritdoc/>
+			public override bool Equals(object? obj) => ReferenceEquals(this, obj);
+
+			/// <inheritdoc/>
+			public override int GetHashCode() => RuntimeHelpers.GetHashCode(this);
 		}
 
 		// Information about a bundle being built. Metadata operations are synchronous, compression/writes are asynchronous.
diff --git a/Engine/Source/Programs/Shared/EpicGames.Horde/Storage/Clients/BundleStorageClient.cs b/Engine/Source/Programs/Shared/EpicGames.Horde/Storage/Clients/BundleStorageClient.cs
--- a/Engine/Source/Programs/Shared/EpicGames.Horde/Storage/Clients/BundleStorageClient.cs
+++ b/Engine/Source/Programs/Shared/EpicGames.Horde/Storage/Clients/BundleStorageClient.cs
@@ -78,6 +78,12 @@ namespace EpicGames.Horde.Storage.Clients
 				locator = _locator;
 				return true;
 			}
+
+			/// <inheritdoc/>
+			public override bool Equals(object? obj) => obj is BundleHandle other && _locator == other._locator;
+
+			/// <inheritdoc/>
+			public override int GetHashCode() => _locator.GetHashCode();
 		}
 
 		readonly IStorageClient _inner;
diff --git a/Engine/Source/Programs/Shared/EpicGames.Horde/Storage/Clients/HttpStorageClient.cs b/Engine/Source/Programs/Shared/EpicGames.Horde/Storage/Clients/HttpStorageClient.cs
--- a/Engine/Source/Programs/Shared/EpicGames.Horde/Storage/Clients/HttpStorageClient.cs
+++ b/Engine/Source/Programs/Shared/EpicGames.Horde/Storage/Clients/HttpStorageClient.cs
@@ -33,17 +33,26 @@ namespace EpicGames.Horde.Storage.Clients
 				_locator = locator;
 			}
 
+			/// <inheritdoc/>
 			public override Task<Stream> OpenAsync(int offset = 0, int? length = null, CancellationToken cancellationToken = default)
 				=> _outer._backend.OpenAsync(_locator.ToString(), offset, length, cancellationToken);
 
+			/// <inheritdoc/>
 			public override ValueTask<BlobData> ReadAsync(CancellationToken cancellationToken = default)
 				=> _outer.ReadBlobAsync(_locator, cancellationToken);
 
+			/// <inheritdoc/>
 			public override bool TryGetLocator([NotNullWhen(true)] out BlobLocator locator)
 			{
 				locator = _locator;
 				return true;
 			}
+
+			/// <inheritdoc/>
+			public override bool Equals(object? obj) => obj is Handle other && _locator == other._locator;
+
+			/// <inheritdoc/>
+			public override int GetHashCode() => _locator.GetHashCode();
 		}
 
 		readonly string _basePath;
diff --git a/Engine/Source/Programs/Shared/EpicGames.Horde/Storage/Clients/KeyValueStorageClient.cs b/Engine/Source/Programs/Shared/EpicGames.Horde/Storage/Clients/KeyValueStorageClient.cs
--- a/Engine/Source/Programs/Shared/EpicGames.Horde/Storage/Clients/KeyValueStorageClient.cs
+++ b/Engine/Source/Programs/Shared/EpicGames.Horde/Storage/Clients/KeyValueStorageClient.cs
@@ -26,13 +26,21 @@ namespace EpicGames.Horde.Storage.Clients
 				_locator = locator;
 			}
 
+			/// <inheritdoc/>
 			public override ValueTask<BlobData> ReadAsync(CancellationToken cancellationToken = default) => _outer.ReadBlobAsync(_locator, cancellationToken);
 
+			/// <inheritdoc/>
 			public override bool TryGetLocator([NotNullWhen(true)] out BlobLocator locator)
 			{
 				locator = _locator;
 				return true;
 			}
+
+			/// <inheritdoc/>
+			public override bool Equals(object? obj) => obj is Handle other && _locator == other._locator;
+
+			/// <inheritdoc/>
+			public override int GetHashCode() => _locator.GetHashCode();
 		}
 
 		/// <inheritdoc/>
diff --git a/Engine/Source/Programs/Shared/EpicGames.Horde/Storage/IStorageWriter.cs b/Engine/Source/Programs/Shared/EpicGames.Horde/Storage/IStorageWriter.cs
--- a/Engine/Source/Programs/Shared/EpicGames.Horde/Storage/IStorageWriter.cs
+++ b/Engine/Source/Programs/Shared/EpicGames.Horde/Storage/IStorageWriter.cs
@@ -178,6 +178,7 @@ namespace EpicGames.Horde.Storage
 				}
 			}
 
+			/// <inheritdoc/>
 			public override ValueTask FlushAsync(CancellationToken cancellationToken)
 			{
 				if (_inner == null)
@@ -190,6 +191,7 @@ namespace EpicGames.Horde.Storage
 				}
 			}
 
+			/// <inheritdoc/>
 			public override ValueTask<BlobData> ReadAsync(CancellationToken cancellationToken = default)
 			{
 				if (_inner == null)
@@ -201,6 +203,12 @@ namespace EpicGames.Horde.Storage
 					return _inner.ReadAsync(cancellationToken);
 				}
 			}
+
+			/// <inheritdoc/>
+			public override bool Equals(object? obj) => _inner is not null && obj is WrappedHandle other && _inner == other._inner;
+
+			/// <inheritdoc/>
+			public override int GetHashCode() => HashCode.Combine((_inner is null) ? 0 : _inner.GetHashCode(), 1);
 		}
 
 		/// <summary>
