Horde: Missing file.

diff --git a/Engine/Source/Programs/Shared/EpicGames.Horde/Storage/BlobHandle.cs b/Engine/Source/Programs/Shared/EpicGames.Horde/Storage/BlobHandle.cs
--- a/Engine/Source/Programs/Shared/EpicGames.Horde/Storage/BlobHandle.cs
+++ b/Engine/Source/Programs/Shared/EpicGames.Horde/Storage/BlobHandle.cs
@@ -11,7 +11,7 @@ using EpicGames.Core;
 namespace EpicGames.Horde.Storage
 {
 	/// <summary>
-	/// Handle to a node. Can be used to reference nodes that have not been flushed yet.
+	/// Handle to a node. Can be used to reference nodes that have not been flushed yet. Handles should have value equality semantics, and must override <see cref="Object.Equals(Object?)"/> and <see cref="Object.GetHashCode()"/>.
 	/// </summary>
 	public abstract class BlobHandle
 	{
@@ -163,6 +163,12 @@ namespace EpicGames.Horde.Storage
 			}
 			return base.ToString() ?? "Unknown";
 		}
+
+		/// <inheritdoc/>
+		public abstract override bool Equals(object? obj);
+
+		/// <inheritdoc/>
+		public abstract override int GetHashCode();
 	}
 
 	/// <summary>
@@ -210,6 +216,12 @@ namespace EpicGames.Horde.Storage
 		{
 			throw new NotSupportedException("Blob fragment handles cannot be read directly, and should be deconstructed into more specific types.");
 		}
+
+		/// <inheritdoc/>
+		public override bool Equals(object? obj) => obj is BlobFragmentHandle other && Inner == other.Inner && Fragment == other.Fragment;
+
+		/// <inheritdoc/>
+		public override int GetHashCode() => HashCode.Combine(Inner, Fragment);
 	}
 
 	/// <summary>
