Horde: Add an overload to read data for a blob without reading an entire BlobData object.

diff --git a/Engine/Source/Programs/Horde/Horde.Server/Ddc/BlobService.cs b/Engine/Source/Programs/Horde/Horde.Server/Ddc/BlobService.cs
--- a/Engine/Source/Programs/Horde/Horde.Server/Ddc/BlobService.cs
+++ b/Engine/Source/Programs/Horde/Horde.Server/Ddc/BlobService.cs
@@ -103,7 +103,7 @@ namespace Horde.Server.Ddc
 				throw new BlobNotFoundException(ns, blob);
 			}
 
-			using BlobData data = await alias.Target.ReadAsync(cancellationToken);
+			using BlobData data = await alias.Target.ReadBlobDataAsync(cancellationToken);
 			return new BlobContents(data.Data.ToArray());
 		}
 
diff --git a/Engine/Source/Programs/Horde/Horde.Server/Ddc/RefService.cs b/Engine/Source/Programs/Horde/Horde.Server/Ddc/RefService.cs
--- a/Engine/Source/Programs/Horde/Horde.Server/Ddc/RefService.cs
+++ b/Engine/Source/Programs/Horde/Horde.Server/Ddc/RefService.cs
@@ -49,7 +49,7 @@ namespace Horde.Server.Ddc
 			}
 
 			BlobHandle blobHandle = blobAlias.Target;
-			using BlobData blobContents = await blobHandle.ReadAsync(cancellationToken);
+			using BlobData blobContents = await blobHandle.ReadBlobDataAsync(cancellationToken);
 			CbObject payload = new CbObject(blobContents.Data);
 
 			BlobId[] referencedBlobs = Array.Empty<BlobId>();
@@ -131,7 +131,7 @@ namespace Horde.Server.Ddc
 				throw new RefNotFoundException(ns, bucket, key);
 			}
 
-			BlobData data = await node.References.First(x => x.Hash == node.RootHash).Handle.ReadAsync(cancellationToken);
+			BlobData data = await node.References.First(x => x.Hash == node.RootHash).Handle.ReadBlobDataAsync(cancellationToken);
 			BlobContents contents = new BlobContents(data.Data.ToArray());
 
 			RefRecord record = new RefRecord(ns, bucket, key, DateTime.UtcNow, null, new BlobId(node.RootHash), true);
diff --git a/Engine/Source/Programs/Horde/Horde.Server/Storage/StorageController.cs b/Engine/Source/Programs/Horde/Horde.Server/Storage/StorageController.cs
--- a/Engine/Source/Programs/Horde/Horde.Server/Storage/StorageController.cs
+++ b/Engine/Source/Programs/Horde/Horde.Server/Storage/StorageController.cs
@@ -382,7 +382,7 @@ namespace Horde.Server.Storage
 
 			object content;
 
-			using BlobData blobData = await storageClient.CreateBlobHandle(locator).ReadAsync(cancellationToken);
+			using BlobData blobData = await storageClient.CreateBlobHandle(locator).ReadBlobDataAsync(cancellationToken);
 
 			Node node = Node.Deserialize(blobData);
 			switch (node)
diff --git a/Engine/Source/Programs/Horde/Horde.Server/Storage/StorageService.cs b/Engine/Source/Programs/Horde/Horde.Server/Storage/StorageService.cs
--- a/Engine/Source/Programs/Horde/Horde.Server/Storage/StorageService.cs
+++ b/Engine/Source/Programs/Horde/Horde.Server/Storage/StorageService.cs
@@ -87,7 +87,7 @@ namespace Horde.Server.Storage
 			public override ValueTask<IReadOnlyList<BlobHandle>> GetRefsAsync(CancellationToken cancellationToken = default) => new ValueTask<IReadOnlyList<BlobHandle>>(Array.Empty<BlobHandle>());
 
 			/// <inheritdoc/>
-			public override async ValueTask<BlobData> ReadAsync(CancellationToken cancellationToken = default)
+			public override async ValueTask<BlobData> ReadBlobDataAsync(CancellationToken cancellationToken = default)
 			{
 				IReadOnlyMemoryOwner<byte> obj = await _backend.ReadAsync(_path, cancellationToken);
 				return new ReadOnlyMemoryOwnerBlobData(BlobType.Leaf, obj, Array.Empty<BlobHandle>());
@@ -172,7 +172,7 @@ namespace Horde.Server.Storage
 			public async ValueTask<BlobData> ReadBlobAsync(BlobLocator locator, CancellationToken cancellationToken = default)
 			{
 				BlobHandle handle = CreateBlobHandle(locator);
-				return await handle.ReadAsync(cancellationToken);
+				return await handle.ReadBlobDataAsync(cancellationToken);
 			}
 
 			/// <inheritdoc/>
diff --git a/Engine/Source/Programs/Shared/EpicGames.Horde.Tests/BlobDataTests.cs b/Engine/Source/Programs/Shared/EpicGames.Horde.Tests/BlobDataTests.cs
--- a/Engine/Source/Programs/Shared/EpicGames.Horde.Tests/BlobDataTests.cs
+++ b/Engine/Source/Programs/Shared/EpicGames.Horde.Tests/BlobDataTests.cs
@@ -19,7 +19,7 @@ namespace EpicGames.Horde.Tests
 			readonly BlobLocator _locator;
 
 			public DummyHandle(string locator) => _locator = new BlobLocator(new Utf8String(locator));
-			public override ValueTask<BlobData> ReadAsync(CancellationToken cancellationToken = default) => throw new NotImplementedException();
+			public override ValueTask<BlobData> ReadBlobDataAsync(CancellationToken cancellationToken = default) => throw new NotImplementedException();
 			public override bool TryGetLocator([NotNullWhen(true)] out BlobLocator locator)
 			{
 				locator = _locator;
diff --git a/Engine/Source/Programs/Shared/EpicGames.Horde.Tests/BundleStoreTests.cs b/Engine/Source/Programs/Shared/EpicGames.Horde.Tests/BundleStoreTests.cs
--- a/Engine/Source/Programs/Shared/EpicGames.Horde.Tests/BundleStoreTests.cs
+++ b/Engine/Source/Programs/Shared/EpicGames.Horde.Tests/BundleStoreTests.cs
@@ -65,7 +65,7 @@ namespace EpicGames.Horde.Tests
 			BlobHandle handle = await writer.FlushAsync(node, CancellationToken.None);
 
 			BlobHandle bundleHandle = store.CreateBlobHandle(handle.GetLocator().Outermost);
-			using BlobData blobData = await bundleHandle.ReadAsync();
+			using BlobData blobData = await bundleHandle.ReadBlobDataAsync();
 
 			return blobData.Data.ToArray();
 		}
diff --git a/Engine/Source/Programs/Shared/EpicGames.Horde.Tests/BundleTests.cs b/Engine/Source/Programs/Shared/EpicGames.Horde.Tests/BundleTests.cs
--- a/Engine/Source/Programs/Shared/EpicGames.Horde.Tests/BundleTests.cs
+++ b/Engine/Source/Programs/Shared/EpicGames.Horde.Tests/BundleTests.cs
@@ -221,7 +221,7 @@ namespace EpicGames.Horde.Tests
 			// Check the ref
 			BlobHandle refTarget =  await store.ReadRefTargetAsync(refName);
 			BlobHandle bundleTarget = store.CreateBlobHandle(refTarget.GetLocator().Outermost);
-			using BlobData bundleData = await bundleTarget.ReadAsync();
+			using BlobData bundleData = await bundleTarget.ReadBlobDataAsync();
 
 			BundleHeader bundleHeader = BundleHeader.Read(bundleData.Data);
 			Assert.AreEqual(0, bundleHeader.Imports.Count);
diff --git a/Engine/Source/Programs/Shared/EpicGames.Horde/Storage/BlobHandle.cs b/Engine/Source/Programs/Shared/EpicGames.Horde/Storage/BlobHandle.cs
--- a/Engine/Source/Programs/Shared/EpicGames.Horde/Storage/BlobHandle.cs
+++ b/Engine/Source/Programs/Shared/EpicGames.Horde/Storage/BlobHandle.cs
@@ -21,10 +21,18 @@ namespace EpicGames.Horde.Storage
 
 			public ReadOnlyMemory<byte> Memory { get; }
 
-			public BlobDataFragment(BlobData data, int offset, int length)
+			public BlobDataFragment(BlobData data, int offset, int? length)
 			{
 				_data = data;
-				Memory = data.Data.Slice(offset, length);
+
+				if (length == null)
+				{
+					Memory = data.Data.Slice(offset);
+				}
+				else
+				{
+					Memory = data.Data.Slice(offset, Math.Min(length.Value, data.Data.Length - offset));
+				}
 			}
 
 			public void Dispose() => _data.Dispose();
@@ -76,7 +84,7 @@ namespace EpicGames.Horde.Storage
 		/// <param name="cancellationToken">Cancellation token for the operation</param>
 		public virtual async ValueTask<BlobType> GetTypeAsync(CancellationToken cancellationToken = default)
 		{
-			using BlobData data = await ReadAsync(cancellationToken);
+			using BlobData data = await ReadBlobDataAsync(cancellationToken);
 			return data.Type;
 		}
 
@@ -86,7 +94,7 @@ namespace EpicGames.Horde.Storage
 		/// <param name="cancellationToken">Cancellation token for the operation</param>
 		public virtual async ValueTask<IReadOnlyList<BlobHandle>> GetRefsAsync(CancellationToken cancellationToken = default)
 		{
-			using BlobData data = await ReadAsync(cancellationToken);
+			using BlobData data = await ReadBlobDataAsync(cancellationToken);
 			return data.Refs;
 		}
 
@@ -98,17 +106,21 @@ namespace EpicGames.Horde.Storage
 		/// <param name="cancellationToken">Cancellation token for the operation</param>
 		public virtual async Task<Stream> OpenAsync(int offset = 0, int? length = null, CancellationToken cancellationToken = default)
 		{
-			BlobData blobData = await ReadAsync(cancellationToken);
+			BlobData blobData = await ReadBlobDataAsync(cancellationToken);
 			int maxLength = blobData.Data.Length - offset;
 			ReadOnlyMemory<byte> memory = blobData.Data.Slice(offset, length.HasValue ? Math.Min(length.Value, maxLength) : maxLength);
 			return new BlobDataStream(blobData, memory);
 		}
 
 		/// <summary>
-		/// Reads the blob's data
+		/// Reads part of the blob, and returns a handle that can be used to access the data.
 		/// </summary>
 		/// <param name="cancellationToken">Cancellation token for the operation</param>
-		public abstract ValueTask<BlobData> ReadAsync(CancellationToken cancellationToken = default);
+		/// <returns></returns>
+		public virtual async ValueTask<IReadOnlyMemoryOwner<byte>> ReadAsync(CancellationToken cancellationToken = default)
+		{
+			return await ReadPartialAsync(0, null, cancellationToken);
+		}
 
 		/// <summary>
 		/// Reads part of the blob, and returns a handle that can be used to access the data.
@@ -117,36 +129,17 @@ namespace EpicGames.Horde.Storage
 		/// <param name="length">Length of the data to read</param>
 		/// <param name="cancellationToken">Cancellation token for the operation</param>
 		/// <returns></returns>
-		public virtual async ValueTask<IReadOnlyMemoryOwner<byte>> ReadPartialAsync(int offset, int length, CancellationToken cancellationToken = default)
+		public virtual async ValueTask<IReadOnlyMemoryOwner<byte>> ReadPartialAsync(int offset, int? length, CancellationToken cancellationToken = default)
 		{
-			BlobData data = await ReadAsync(cancellationToken);
+			BlobData data = await ReadBlobDataAsync(cancellationToken);
 			return new BlobDataFragment(data, offset, length);
 		}
 
 		/// <summary>
-		/// Creates a reader for this node's data
+		/// Reads the blob's data
 		/// </summary>
-		/// <param name="offset">Offset within the payload stream to start reading</param>
-		/// <param name="buffer">Buffer to receive the data that was read</param>
 		/// <param name="cancellationToken">Cancellation token for the operation</param>
-		/// <returns>Number of bytes that were read</returns>
-		public virtual async ValueTask<int> ReadPartialAsync(int offset, Memory<byte> buffer, CancellationToken cancellationToken = default)
-		{
-			using BlobData data = await ReadAsync(cancellationToken);
-
-			int length = data.Data.Length - offset;
-			if (length < 0)
-			{
-				return 0;
-			}
-			if (length > buffer.Length)
-			{
-				length = buffer.Length;
-			}
-
-			data.Data.Slice(offset, length).CopyTo(buffer);
-			return length;
-		}
+		public abstract ValueTask<BlobData> ReadBlobDataAsync(CancellationToken cancellationToken = default);
 
 		/// <summary>
 		/// Flush the referenced not to underlying storage
@@ -212,7 +205,7 @@ namespace EpicGames.Horde.Storage
 		}
 
 		/// <inheritdoc/>
-		public override ValueTask<BlobData> ReadAsync(CancellationToken cancellationToken = default)
+		public override ValueTask<BlobData> ReadBlobDataAsync(CancellationToken cancellationToken = default)
 		{
 			throw new NotSupportedException("Blob fragment handles cannot be read directly, and should be deconstructed into more specific types.");
 		}
diff --git a/Engine/Source/Programs/Shared/EpicGames.Horde/Storage/Bundles/V1/BundleWriter.cs b/Engine/Source/Programs/Shared/EpicGames.Horde/Storage/Bundles/V1/BundleWriter.cs
--- a/Engine/Source/Programs/Shared/EpicGames.Horde/Storage/Bundles/V1/BundleWriter.cs
+++ b/Engine/Source/Programs/Shared/EpicGames.Horde/Storage/Bundles/V1/BundleWriter.cs
@@ -55,7 +55,7 @@ namespace EpicGames.Horde.Storage.Bundles.V1
 		}
 
 		/// <inheritdoc/>
-		public override ValueTask<BlobData> ReadAsync(CancellationToken cancellationToken = default) => _reader.ReadNodeDataAsync(BundleLocator, ExportIdx, cancellationToken);
+		public override ValueTask<BlobData> ReadBlobDataAsync(CancellationToken cancellationToken = default) => _reader.ReadNodeDataAsync(BundleLocator, ExportIdx, cancellationToken);
 
 		/// <inheritdoc/>
 		public override ValueTask FlushAsync(CancellationToken cancellationToken = default) => new ValueTask();
@@ -135,7 +135,7 @@ namespace EpicGames.Horde.Storage.Bundles.V1
 			}
 
 			/// <inheritdoc/>
-			public override async ValueTask<BlobData> ReadAsync(CancellationToken cancellationToken = default)
+			public override async ValueTask<BlobData> ReadBlobDataAsync(CancellationToken cancellationToken = default)
 			{
 				if (_flushedHandle == null)
 				{
@@ -149,7 +149,7 @@ namespace EpicGames.Horde.Storage.Bundles.V1
 					}
 				}
 
-				return await _flushedHandle!.ReadAsync(cancellationToken);
+				return await _flushedHandle!.ReadBlobDataAsync(cancellationToken);
 			}
 
 			/// <inheritdoc/>
diff --git a/Engine/Source/Programs/Shared/EpicGames.Horde/Storage/Clients/BundleStorageClient.cs b/Engine/Source/Programs/Shared/EpicGames.Horde/Storage/Clients/BundleStorageClient.cs
--- a/Engine/Source/Programs/Shared/EpicGames.Horde/Storage/Clients/BundleStorageClient.cs
+++ b/Engine/Source/Programs/Shared/EpicGames.Horde/Storage/Clients/BundleStorageClient.cs
@@ -63,7 +63,7 @@ namespace EpicGames.Horde.Storage.Clients
 			}
 
 			/// <inheritdoc/>
-			public override async ValueTask<BlobData> ReadAsync(CancellationToken cancellationToken = default)
+			public override async ValueTask<BlobData> ReadBlobDataAsync(CancellationToken cancellationToken = default)
 			{
 				using (Stream stream = await _storageClient.OpenBlobAsync(_locator, 0, cancellationToken: cancellationToken))
 				{
@@ -135,7 +135,7 @@ namespace EpicGames.Horde.Storage.Clients
 		public async ValueTask<BlobData> ReadBlobAsync(BlobLocator locator, CancellationToken cancellationToken = default)
 		{
 			BlobHandle handle = CreateBlobHandle(locator);
-			return await handle.ReadAsync(cancellationToken);
+			return await handle.ReadBlobDataAsync(cancellationToken);
 		}
 
 		/// <inheritdoc/>
diff --git a/Engine/Source/Programs/Shared/EpicGames.Horde/Storage/Clients/HttpStorageClient.cs b/Engine/Source/Programs/Shared/EpicGames.Horde/Storage/Clients/HttpStorageClient.cs
--- a/Engine/Source/Programs/Shared/EpicGames.Horde/Storage/Clients/HttpStorageClient.cs
+++ b/Engine/Source/Programs/Shared/EpicGames.Horde/Storage/Clients/HttpStorageClient.cs
@@ -38,7 +38,7 @@ namespace EpicGames.Horde.Storage.Clients
 				=> _outer._backend.OpenAsync(_locator.ToString(), offset, length, cancellationToken);
 
 			/// <inheritdoc/>
-			public override ValueTask<BlobData> ReadAsync(CancellationToken cancellationToken = default)
+			public override ValueTask<BlobData> ReadBlobDataAsync(CancellationToken cancellationToken = default)
 				=> _outer.ReadBlobAsync(_locator, cancellationToken);
 
 			/// <inheritdoc/>
diff --git a/Engine/Source/Programs/Shared/EpicGames.Horde/Storage/Clients/KeyValueStorageClient.cs b/Engine/Source/Programs/Shared/EpicGames.Horde/Storage/Clients/KeyValueStorageClient.cs
--- a/Engine/Source/Programs/Shared/EpicGames.Horde/Storage/Clients/KeyValueStorageClient.cs
+++ b/Engine/Source/Programs/Shared/EpicGames.Horde/Storage/Clients/KeyValueStorageClient.cs
@@ -27,7 +27,7 @@ namespace EpicGames.Horde.Storage.Clients
 			}
 
 			/// <inheritdoc/>
-			public override ValueTask<BlobData> ReadAsync(CancellationToken cancellationToken = default) => _outer.ReadBlobAsync(_locator, cancellationToken);
+			public override ValueTask<BlobData> ReadBlobDataAsync(CancellationToken cancellationToken = default) => _outer.ReadBlobAsync(_locator, cancellationToken);
 
 			/// <inheritdoc/>
 			public override bool TryGetLocator([NotNullWhen(true)] out BlobLocator locator)
diff --git a/Engine/Source/Programs/Shared/EpicGames.Horde/Storage/IStorageClient.cs b/Engine/Source/Programs/Shared/EpicGames.Horde/Storage/IStorageClient.cs
--- a/Engine/Source/Programs/Shared/EpicGames.Horde/Storage/IStorageClient.cs
+++ b/Engine/Source/Programs/Shared/EpicGames.Horde/Storage/IStorageClient.cs
@@ -104,7 +104,7 @@ namespace EpicGames.Horde.Storage
 		IStorageWriter CreateWriter(string? basePath = null);
 
 		/// <summary>
-		/// Read a blob from the underlying storage system. Calling <see cref="BlobHandle.ReadAsync(CancellationToken)"/> is more efficient than calling this method repeatedly for small blobs.
+		/// Read a blob from the underlying storage system. Calling <see cref="BlobHandle.ReadBlobDataAsync(CancellationToken)"/> is more efficient than calling this method repeatedly for small blobs.
 		/// </summary>
 		/// <param name="locator">Locator for the blob</param>
 		/// <param name="cancellationToken">Cancellation token for the operation</param>
diff --git a/Engine/Source/Programs/Shared/EpicGames.Horde/Storage/IStorageWriter.cs b/Engine/Source/Programs/Shared/EpicGames.Horde/Storage/IStorageWriter.cs
--- a/Engine/Source/Programs/Shared/EpicGames.Horde/Storage/IStorageWriter.cs
+++ b/Engine/Source/Programs/Shared/EpicGames.Horde/Storage/IStorageWriter.cs
@@ -192,7 +192,7 @@ namespace EpicGames.Horde.Storage
 			}
 
 			/// <inheritdoc/>
-			public override ValueTask<BlobData> ReadAsync(CancellationToken cancellationToken = default)
+			public override ValueTask<BlobData> ReadBlobDataAsync(CancellationToken cancellationToken = default)
 			{
 				if (_inner == null)
 				{
@@ -200,7 +200,7 @@ namespace EpicGames.Horde.Storage
 				}
 				else
 				{
-					return _inner.ReadAsync(cancellationToken);
+					return _inner.ReadBlobDataAsync(cancellationToken);
 				}
 			}
 
diff --git a/Engine/Source/Programs/Shared/EpicGames.Horde/Storage/Node.cs b/Engine/Source/Programs/Shared/EpicGames.Horde/Storage/Node.cs
--- a/Engine/Source/Programs/Shared/EpicGames.Horde/Storage/Node.cs
+++ b/Engine/Source/Programs/Shared/EpicGames.Horde/Storage/Node.cs
@@ -412,7 +412,7 @@ namespace EpicGames.Horde.Storage
 		/// <returns></returns>
 		public static async ValueTask<TNode> ReadNodeAsync<TNode>(this BlobHandle handle, CancellationToken cancellationToken = default) where TNode : Node
 		{
-			using BlobData blobData = await handle.ReadAsync(cancellationToken);
+			using BlobData blobData = await handle.ReadBlobDataAsync(cancellationToken);
 			return Node.Deserialize<TNode>(blobData);
 		}
 
@@ -449,7 +449,7 @@ namespace EpicGames.Horde.Storage
 				return null;
 			}
 
-			using BlobData blobData = await refTarget.ReadAsync(cancellationToken);
+			using BlobData blobData = await refTarget.ReadBlobDataAsync(cancellationToken);
 			return Node.Deserialize<TNode>(blobData);
 		}
 
diff --git a/Engine/Source/Programs/Shared/EpicGames.Horde/Storage/NodeRef.cs b/Engine/Source/Programs/Shared/EpicGames.Horde/Storage/NodeRef.cs
--- a/Engine/Source/Programs/Shared/EpicGames.Horde/Storage/NodeRef.cs
+++ b/Engine/Source/Programs/Shared/EpicGames.Horde/Storage/NodeRef.cs
@@ -48,7 +48,7 @@ namespace EpicGames.Horde.Storage
 		/// <returns></returns>
 		public async ValueTask<Node> ExpandAsync(CancellationToken cancellationToken = default)
 		{
-			using BlobData blobData = await Handle!.ReadAsync(cancellationToken);
+			using BlobData blobData = await Handle!.ReadBlobDataAsync(cancellationToken);
 			return Node.Deserialize(blobData);
 		}
 	}
diff --git a/Engine/Source/Programs/Shared/EpicGames.Horde/Storage/Nodes/ChunkedDataNode.cs b/Engine/Source/Programs/Shared/EpicGames.Horde/Storage/Nodes/ChunkedDataNode.cs
--- a/Engine/Source/Programs/Shared/EpicGames.Horde/Storage/Nodes/ChunkedDataNode.cs
+++ b/Engine/Source/Programs/Shared/EpicGames.Horde/Storage/Nodes/ChunkedDataNode.cs
@@ -37,7 +37,7 @@ namespace EpicGames.Horde.Storage.Nodes
 		/// <param name="cancellationToken">Cancellation token for the operation</param>
 		public static async Task CopyToStreamAsync(BlobHandle handle, Stream outputStream, CancellationToken cancellationToken)
 		{
-			using BlobData blobData = await handle.ReadAsync(cancellationToken);
+			using BlobData blobData = await handle.ReadBlobDataAsync(cancellationToken);
 			if (blobData.Type.Guid == s_leafNodeGuid)
 			{
 				await LeafChunkedDataNode.CopyToStreamAsync(blobData, outputStream, cancellationToken);
diff --git a/Engine/Source/Programs/Shared/EpicGames.Horde/Storage/Workspace.cs b/Engine/Source/Programs/Shared/EpicGames.Horde/Storage/Workspace.cs
--- a/Engine/Source/Programs/Shared/EpicGames.Horde/Storage/Workspace.cs
+++ b/Engine/Source/Programs/Shared/EpicGames.Horde/Storage/Workspace.cs
@@ -755,7 +755,7 @@ namespace EpicGames.Horde.Storage
 			}
 
 			// Otherwise 
-			using BlobData blobData = await nodeRef.Handle.ReadAsync(cancellationToken);
+			using BlobData blobData = await nodeRef.Handle.ReadBlobDataAsync(cancellationToken);
 			await LeafChunkedDataNode.CopyToStreamAsync(blobData, outputStream, cancellationToken);
 		}
 
