Horde: Add support for tunneling in ServerComputeClient

diff --git a/Engine/Source/Programs/Shared/EpicGames.Horde/Compute/Clients/ServerComputeClient.cs b/Engine/Source/Programs/Shared/EpicGames.Horde/Compute/Clients/ServerComputeClient.cs
--- a/Engine/Source/Programs/Shared/EpicGames.Horde/Compute/Clients/ServerComputeClient.cs
+++ b/Engine/Source/Programs/Shared/EpicGames.Horde/Compute/Clients/ServerComputeClient.cs
@@ -5,6 +5,7 @@ using EpicGames.Horde.Compute.Transports;
 using Microsoft.Extensions.Logging;
 using System;
 using System.Collections.Generic;
+using System.IO;
 using System.Net;
 using System.Net.Http;
 using System.Net.Http.Json;
@@ -245,11 +246,30 @@ namespace EpicGames.Horde.Compute.Clients
 				}
 			}
 
-			workerLogger.LogDebug("Connecting to {AgentId} ({Ip}) with nonce {Nonce}...", responseMessage.AgentId, responseMessage.Ip, responseMessage.Nonce);
+			string agentAddress = $"{responseMessage.Ip}:{responseMessage.Port}";
 
 			// Connect to the remote machine
 			using Socket socket = new Socket(SocketType.Stream, ProtocolType.Tcp);
-			await socket.ConnectAsync(IPAddress.Parse(responseMessage.Ip), responseMessage.Port, cancellationToken);
+			
+			workerLogger.LogDebug("Connecting to {AgentId} at {AgentAddress} ({ConnectionType} via {ConnectionAddress}) with nonce {Nonce}...", responseMessage.AgentId, agentAddress, responseMessage.ConnectionMode, responseMessage.ConnectionAddress ?? "None", responseMessage.Nonce);
+			switch (responseMessage.ConnectionMode)
+			{
+				case ConnectionMode.Direct:
+					await socket.ConnectAsync(IPAddress.Parse(responseMessage.Ip), responseMessage.Port, cancellationToken);
+					break;
+
+				case ConnectionMode.Tunnel when !String.IsNullOrEmpty(responseMessage.ConnectionAddress):
+					(string host, int port) = ParseHostPort(responseMessage.ConnectionAddress);
+					await socket.ConnectAsync(host, port, cancellationToken);
+					await TunnelHandshakeAsync(socket, responseMessage, cancellationToken);
+					break;
+				
+				case ConnectionMode.Relay when !String.IsNullOrEmpty(responseMessage.ConnectionAddress):
+					throw new NotImplementedException("Relay connection mode not yet implemented");
+				
+				default:
+					throw new Exception($"Unable to resolve connection mode ({responseMessage.ConnectionMode}/{responseMessage.ConnectionAddress}");
+			}
 
 			// Send the nonce
 			byte[] nonce = StringUtils.ParseHexString(responseMessage.Nonce);
@@ -262,8 +282,45 @@ namespace EpicGames.Horde.Compute.Clients
 			await using RemoteComputeSocket computeSocket = new RemoteComputeSocket(new TcpTransport(socket), workerLogger);
 			yield return new LeaseInfo(responseMessage.Properties, responseMessage.AssignedResources, computeSocket);
 		}
-		
-		//private 
+
+		private static (string host, int port) ParseHostPort(string address)
+		{
+			try
+			{
+				string[] parts = address.Split(":");
+				string host = parts[0];
+				int port = Int32.Parse(parts[1]);
+				return (host, port);
+			}
+			catch (Exception e)
+			{
+				throw new Exception($"Unable to parse host and port for address: {address}", e);
+			}
+		}
+
+		private static async Task TunnelHandshakeAsync(Socket socket, AssignComputeResponse response, CancellationToken cancellationToken)
+		{
+			await using NetworkStream ns = new (socket, false);
+			using StreamReader reader = new (ns);
+			await using StreamWriter writer = new (ns) { AutoFlush = true };
+
+			string request = new TunnelHandshakeRequest(response.Ip, response.Port).Serialize();
+			await writer.WriteLineAsync(request.ToCharArray(), cancellationToken);
+
+			string exceptionMetadata = $"Connection: {response.ConnectionAddress} Target: {response.Ip}:{response.Port}";
+			Task<string?> readTask = reader.ReadLineAsync();
+			Task timeoutTask = Task.Delay(15000, cancellationToken);
+			if (await Task.WhenAny(readTask, timeoutTask) == timeoutTask)
+			{
+				throw new TimeoutException($"Timed out reading tunnel handshake response. {exceptionMetadata}");
+			}
+
+			TunnelHandshakeResponse handshakeResponse = TunnelHandshakeResponse.Deserialize(await readTask);
+			if (!handshakeResponse.IsSuccess)
+			{
+				throw new Exception($"Tunnel handshake failed! Reason: {handshakeResponse.Message} {exceptionMetadata}");
+			}
+		}
 	}
 
 	/// <summary>
diff --git a/Engine/Source/Programs/Shared/EpicGames.Horde/Compute/ComputeMessages.cs b/Engine/Source/Programs/Shared/EpicGames.Horde/Compute/ComputeMessages.cs
--- a/Engine/Source/Programs/Shared/EpicGames.Horde/Compute/ComputeMessages.cs
+++ b/Engine/Source/Programs/Shared/EpicGames.Horde/Compute/ComputeMessages.cs
@@ -2,6 +2,7 @@
 
 using System;
 using System.Collections.Generic;
+using System.Text.Json.Serialization;
 using EpicGames.Horde.Agents;
 using EpicGames.Horde.Agents.Leases;
 
@@ -40,6 +41,16 @@ namespace EpicGames.Horde.Compute
 		/// Port number on the remote machine
 		/// </summary>
 		public int Port { get; set; }
+		
+		/// <summary>
+		/// How to establish a connection to the remote machine
+		/// </summary>
+		public ConnectionMode ConnectionMode { get; set; }
+		
+		/// <summary>
+		/// An optional address (host:port) to use when connecting depending on connection mode
+		/// </summary>
+		public string? ConnectionAddress { get; set; }
 
 		/// <summary>
 		/// Cryptographic nonce to identify the request, as a hex string
@@ -72,6 +83,31 @@ namespace EpicGames.Horde.Compute
 		public IReadOnlyList<string> Properties { get; set; } = new List<string>();
 	}
 	
+	/// <summary>
+	/// Describe how to connect to the remote machine
+	/// </summary>
+	[JsonConverter(typeof(JsonStringEnumConverter))]
+	public enum ConnectionMode
+	{
+		/// <summary>
+		/// Connection is established directly to remote machine, behaving like a normal TCP/UDP connection
+		/// </summary>
+		Direct,
+		
+		/// <summary>
+		/// Connection is tunneled through Horde server.
+		/// When connecting, initiator must send a tunnel handshake request indicating which machine/IP to tunnel to.
+		/// Once handshake is complete, TCP connection behaves as normal.
+		/// </summary>
+		Tunnel,
+		
+		/// <summary>
+		/// Connection is established to remote machine via a relay.
+		/// Forwarding is transparent and behaves like a normal TCP/UDP connection.
+		/// </summary>
+		Relay
+	}
+	
 	/// <summary>
 	/// Resource needs declaration request
 	/// </summary>
