Improved error message with TObjectPtr appearing in function signatures.

diff --git a/Engine/Source/Programs/Shared/EpicGames.UHT/Types/Properties/UhtClassPtrProperties.cs b/Engine/Source/Programs/Shared/EpicGames.UHT/Types/Properties/UhtClassPtrProperties.cs
--- a/Engine/Source/Programs/Shared/EpicGames.UHT/Types/Properties/UhtClassPtrProperties.cs
+++ b/Engine/Source/Programs/Shared/EpicGames.UHT/Types/Properties/UhtClassPtrProperties.cs
@@ -77,7 +77,7 @@ namespace EpicGames.UHT.Types
 			// UFunctions with a smart pointer as return type will crash when called via blueprint, because they are not supported in VM.
 			if (!options.HasAnyFlags(UhtValidationOptions.IsKey) && PropertyCategory != UhtPropertyCategory.Member)
 			{
-				outerStruct.LogError("UFunctions cannot take a TObjectPtr as a parameter.");
+				outerStruct.LogError("UFunctions cannot take a TObjectPtr as a function parameter or return value.");
 			}
 		}
 	}
diff --git a/Engine/Source/Programs/Shared/EpicGames.UHT/Types/Properties/UhtObjectPtrProperty.cs b/Engine/Source/Programs/Shared/EpicGames.UHT/Types/Properties/UhtObjectPtrProperty.cs
--- a/Engine/Source/Programs/Shared/EpicGames.UHT/Types/Properties/UhtObjectPtrProperty.cs
+++ b/Engine/Source/Programs/Shared/EpicGames.UHT/Types/Properties/UhtObjectPtrProperty.cs
@@ -103,7 +103,7 @@ namespace EpicGames.UHT.Types
 					outerStruct.Outer == null ||
 					!outerStruct.Outer.SourceName.Equals("UPlayerMappableInputConfig", StringComparison.Ordinal))
 				{
-					outerStruct.LogError("UFunctions cannot take a TObjectPtr as a parameter.");
+					outerStruct.LogError("UFunctions cannot take a TObjectPtr as a function parameter or return value.");
 				}
 			}
 		}
