Horde: Compute service fixes for older clients.

* Pass 128mb as the read length when reading entire bundles. Older clients do not support passing 0 to indicate 'entire blob'.
* Fix requests to read zero bytes resulting in remote not sending any response packets.
* Always create a cache for bundle data when reading from the remote. Since we always over-read from the source bundle, disabling the cache results in re-reading entire bundles multiple times. Not observed previously due to BundleReaderCache.None actually using the default cache size.

diff --git a/Engine/Source/Programs/Shared/EpicGames.Horde/Compute/AgentMessage.cs b/Engine/Source/Programs/Shared/EpicGames.Horde/Compute/AgentMessage.cs
--- a/Engine/Source/Programs/Shared/EpicGames.Horde/Compute/AgentMessage.cs
+++ b/Engine/Source/Programs/Shared/EpicGames.Horde/Compute/AgentMessage.cs
@@ -656,7 +656,7 @@ namespace EpicGames.Horde.Compute
 			using Stream stream = await handle.OpenAsync(offset, (length == 0)? null : length, cancellationToken);
 
 			const int MaxChunkSize = 512 * 1024;
-			for (int chunkOffset = 0; chunkOffset < stream.Length;)
+			for (int chunkOffset = 0; ;)
 			{
 				int chunkLength = (int)Math.Min(stream.Length - chunkOffset, MaxChunkSize);
 				using (IAgentMessageBuilder response = await channel.CreateMessageAsync(AgentMessageType.ReadBlobResponse, chunkLength + 128, cancellationToken))
@@ -669,7 +669,12 @@ namespace EpicGames.Horde.Compute
 
 					response.Send();
 				}
+
 				chunkOffset += chunkLength;
+				if(chunkOffset == stream.Length)
+				{
+					break;
+				}
 			}
 		}
 
diff --git a/Engine/Source/Programs/Shared/EpicGames.Horde/Compute/AgentMessageHandler.cs b/Engine/Source/Programs/Shared/EpicGames.Horde/Compute/AgentMessageHandler.cs
--- a/Engine/Source/Programs/Shared/EpicGames.Horde/Compute/AgentMessageHandler.cs
+++ b/Engine/Source/Programs/Shared/EpicGames.Horde/Compute/AgentMessageHandler.cs
@@ -136,7 +136,8 @@ namespace EpicGames.Horde.Compute
 		async Task WriteFilesAsync(AgentMessageChannel channel, string path, BlobLocator locator, CancellationToken cancellationToken)
 		{
 			using AgentStorageClient innerStore = new AgentStorageClient(channel);
-			using BundleStorageClient store = new BundleStorageClient(innerStore, BundleReaderCache.None, _logger);
+			using BundleReaderCache cache = new BundleReaderCache(new BundleReaderCacheOptions { HeaderCacheSize = 10 * 1024 * 1024, PacketCacheSize = 128 * 1024 * 1024 });
+			using BundleStorageClient store = new BundleStorageClient(innerStore, cache, _logger);
 
 			BlobHandle handle = store.CreateBlobHandle(locator);
 			DirectoryNode directoryNode = await handle.ReadNodeAsync<DirectoryNode>(cancellationToken);
diff --git a/Engine/Source/Programs/Shared/EpicGames.Horde/Compute/AgentStorageClient.cs b/Engine/Source/Programs/Shared/EpicGames.Horde/Compute/AgentStorageClient.cs
--- a/Engine/Source/Programs/Shared/EpicGames.Horde/Compute/AgentStorageClient.cs
+++ b/Engine/Source/Programs/Shared/EpicGames.Horde/Compute/AgentStorageClient.cs
@@ -50,7 +50,8 @@ namespace EpicGames.Horde.Compute
 			await _semaphore.WaitAsync(cancellationToken);
 			try
 			{
-				ReadOnlyMemory<byte> data = await _channel.ReadBlobAsync(locator.ToString(), 0, 0, cancellationToken);
+				// TODO: Want to pass 0 for length to read here (meaning entire blob), but older streams misinterpret this as a 0 byte read.
+				ReadOnlyMemory<byte> data = await _channel.ReadBlobAsync(locator.ToString(), 0, 128 * 1024 * 1024, cancellationToken);
 				return new BlobData(BlobType.Leaf, data, Array.Empty<BlobHandle>());
 			}
 			finally
