EpicGames.Core: Trim whitespace from the end of log lines before trying to parse JSON log events. Fixes output from child UBT processes (which have trailing CR/LF characters) from being parsed of JSON log events.

diff --git a/Engine/Source/Programs/Shared/EpicGames.Core/LogEventParser.cs b/Engine/Source/Programs/Shared/EpicGames.Core/LogEventParser.cs
--- a/Engine/Source/Programs/Shared/EpicGames.Core/LogEventParser.cs
+++ b/Engine/Source/Programs/Shared/EpicGames.Core/LogEventParser.cs
@@ -314,7 +314,13 @@ namespace EpicGames.Core
 		{
 			if (line.Length > 0 && line[0] == '{')
 			{
-				byte[] data = Encoding.UTF8.GetBytes(line);
+				int length = line.Length;
+				while(length > 0 && Char.IsWhiteSpace(line[length - 1]))
+				{
+					length--;
+				}
+
+				byte[] data = Encoding.UTF8.GetBytes(line, 0, length);
 				try
 				{
 					JsonLogEvent jsonEvent;
