UnrealBuildTool: Rename some shared classes

#rnx

diff --git a/Engine/Source/Programs/Shared/EpicGames.Box/Utils.cs b/Engine/Source/Programs/Shared/EpicGames.Box/Utils.cs
--- a/Engine/Source/Programs/Shared/EpicGames.Box/Utils.cs
+++ b/Engine/Source/Programs/Shared/EpicGames.Box/Utils.cs
@@ -5,14 +5,14 @@ using System;
 namespace EpicGames.Box
 {
 	/// <summary>
-	/// Box utils
+	/// Utils
 	/// </summary>
 	public static partial class Utils
 	{
 		static Lazy<bool> Available { get; set; } = new Lazy<bool>(() => false);
 
 		/// <summary>
-		/// Is box available?
+		/// Is available?
 		/// </summary>
 		public static bool IsAvailable() => Available.Value;
 	}
