Rename Box -> UBA

#rnx

diff --git a/Engine/Build/BatchFiles/BuildUAT.bat b/Engine/Build/BatchFiles/BuildUAT.bat
--- a/Engine/Build/BatchFiles/BuildUAT.bat
+++ b/Engine/Build/BatchFiles/BuildUAT.bat
@@ -24,8 +24,6 @@ rem find ".cs" files to only lines that match those names - excludes lines that
 md ..\Intermediate\Build >nul 2>nul
 
 dir /s ^
- Programs\Shared\EpicGames.Box\*.cs ^
- Programs\Shared\EpicGames.Box\*.csproj ^
  Programs\Shared\EpicGames.Build\*.cs ^
  Programs\Shared\EpicGames.Build\*.csproj ^
  Programs\Shared\EpicGames.Core\*.cs ^
@@ -40,6 +38,8 @@ dir /s ^
  Programs\Shared\EpicGames.OIDC\*.csproj ^
  Programs\Shared\EpicGames.Serialization\*.cs ^
  Programs\Shared\EpicGames.Serialization\*.csproj ^
+ Programs\Shared\EpicGames.UBA\*.cs ^
+ Programs\Shared\EpicGames.UBA\*.csproj ^
  Programs\Shared\EpicGames.UHT\*.cs ^
  Programs\Shared\EpicGames.UHT\*.csproj ^
  Programs\UnrealBuildTool\*.cs ^
diff --git a/Engine/Build/BatchFiles/BuildUAT.sh b/Engine/Build/BatchFiles/BuildUAT.sh
--- a/Engine/Build/BatchFiles/BuildUAT.sh
+++ b/Engine/Build/BatchFiles/BuildUAT.sh
@@ -37,7 +37,6 @@ elif [ ! -f ../Binaries/DotNET/AutomationTool/AutomationTool.dll ]; then
 
 elif [ -f ../Intermediate/Build/AutomationToolLastBuildTime ]; then
   UPDATED_DEP_FILES="$(find \
-    Programs/Shared/EpicGames.Box \
     Programs/Shared/EpicGames.Build \
     Programs/Shared/EpicGames.Core \
     Programs/Shared/EpicGames.Horde \
@@ -45,6 +44,7 @@ elif [ -f ../Intermediate/Build/AutomationToolLastBuildTime ]; then
     Programs/Shared/EpicGames.MsBuild \
     Programs/Shared/EpicGames.OIDC \
     Programs/Shared/EpicGames.Serialization \
+    Programs/Shared/EpicGames.UBA \
     Programs/Shared/EpicGames.UHT \
     Programs/UnrealBuildTool \
     ../Restricted/**/Source/Programs/UnrealBuildTool \
diff --git a/Engine/Build/BatchFiles/BuildUBT.bat b/Engine/Build/BatchFiles/BuildUBT.bat
--- a/Engine/Build/BatchFiles/BuildUBT.bat
+++ b/Engine/Build/BatchFiles/BuildUBT.bat
@@ -24,8 +24,6 @@ rem find ".cs" files to only lines that match those names - excludes lines that
 md ..\Intermediate\Build >nul 2>nul
 
 dir /s ^
- Programs\Shared\EpicGames.Box\*.cs ^
- Programs\Shared\EpicGames.Box\*.csproj ^
  Programs\Shared\EpicGames.Build\*.cs ^
  Programs\Shared\EpicGames.Build\*.csproj ^
  Programs\Shared\EpicGames.Core\*.cs ^
@@ -40,6 +38,8 @@ dir /s ^
  Programs\Shared\EpicGames.OIDC\*.csproj ^
  Programs\Shared\EpicGames.Serialization\*.cs ^
  Programs\Shared\EpicGames.Serialization\*.csproj ^
+ Programs\Shared\EpicGames.UBA\*.cs ^
+ Programs\Shared\EpicGames.UBA\*.csproj ^
  Programs\Shared\EpicGames.UHT\*.cs ^
  Programs\Shared\EpicGames.UHT\*.csproj ^
  Programs\UnrealBuildTool\*.cs ^
diff --git a/Engine/Build/BatchFiles/BuildUBT.sh b/Engine/Build/BatchFiles/BuildUBT.sh
--- a/Engine/Build/BatchFiles/BuildUBT.sh
+++ b/Engine/Build/BatchFiles/BuildUBT.sh
@@ -37,7 +37,6 @@ elif [ ! -f ../Binaries/DotNET/UnrealBuildTool/UnrealBuildTool.dll ]; then
 
 elif [ -f ../Intermediate/Build/UnrealBuildToolLastBuildTime ]; then
   UPDATED_DEP_FILES="$(find \
-    Programs/Shared/EpicGames.Box \
     Programs/Shared/EpicGames.Build \
     Programs/Shared/EpicGames.Core \
     Programs/Shared/EpicGames.Horde \
@@ -45,6 +44,7 @@ elif [ -f ../Intermediate/Build/UnrealBuildToolLastBuildTime ]; then
     Programs/Shared/EpicGames.MsBuild \
     Programs/Shared/EpicGames.OIDC \
     Programs/Shared/EpicGames.Serialization \
+    Programs/Shared/EpicGames.UBA \
     Programs/Shared/EpicGames.UHT \
     Programs/UnrealBuildTool \
     ../Restricted/**/Source/Programs/UnrealBuildTool \
diff --git a/Engine/Source/Programs/Shared/EpicGames.Box/EpicGames.Box.csproj b/Engine/Source/Programs/Shared/EpicGames.Box/EpicGames.Box.csproj
deleted file mode 100644
diff --git a/Engine/Source/Programs/Shared/EpicGames.Box/Utils.cs b/Engine/Source/Programs/Shared/EpicGames.Box/Utils.cs
deleted file mode 100644
diff --git a/Engine/Source/Programs/Shared/EpicGames.UBA/EpicGames.UBA.csproj b/Engine/Source/Programs/Shared/EpicGames.UBA/EpicGames.UBA.csproj
new file mode 100644
--- /dev/null
+++ b/Engine/Source/Programs/Shared/EpicGames.UBA/EpicGames.UBA.csproj
@@ -0,0 +1,49 @@
+<Project Sdk="Microsoft.NET.Sdk">
+  <Import Project="..\UnrealEngine.csproj.props" />
+
+  <PropertyGroup>
+    <TargetFramework>net6.0</TargetFramework>
+    <Configuration Condition=" '$(Configuration)' == '' ">Development</Configuration>
+    <OutputType>Library</OutputType>
+    <RootNamespace>EpicGames.UBA</RootNamespace>
+    <GenerateAssemblyInfo>false</GenerateAssemblyInfo>
+    <GenerateTargetFrameworkAttribute>false</GenerateTargetFrameworkAttribute>
+    <Configurations>Debug;Release;Development</Configurations>
+    <Nullable>enable</Nullable>
+    <AllowUnsafeBlocks>true</AllowUnsafeBlocks>
+    <GenerateDocumentationFile>True</GenerateDocumentationFile>
+    <DebugType>pdbonly</DebugType>
+    <DebugType Condition="'$([System.Runtime.InteropServices.RuntimeInformation]::IsOSPlatform($([System.Runtime.InteropServices.OSPlatform]::Windows)))' == 'true' And '$([System.Runtime.InteropServices.RuntimeInformation]::OSArchitecture)' == 'Arm64'">portable</DebugType>
+    <EnforceCodeStyleInBuild>False</EnforceCodeStyleInBuild>
+    <ProduceReferenceAssemblyInOutDir>true</ProduceReferenceAssemblyInOutDir>
+    <RunAnalyzersDuringBuild>false</RunAnalyzersDuringBuild>
+  </PropertyGroup>
+  <PropertyGroup Condition=" '$(Configuration)|$(Platform)' == 'Debug|AnyCPU' ">
+    <DocumentationFile></DocumentationFile>
+  </PropertyGroup>
+  <PropertyGroup Condition="'$(Configuration)|$(Platform)' == 'Development|AnyCPU'">
+    <DocumentationFile></DocumentationFile>
+    <Optimize>true</Optimize>
+  </PropertyGroup>
+  <PropertyGroup Condition="'$(Configuration)|$(Platform)' == 'Release|AnyCPU'">
+    <DocumentationFile></DocumentationFile>
+    <Optimize>true</Optimize>
+  </PropertyGroup>
+
+  <ItemGroup>
+    <Compile Include="..\MetaData.cs">
+      <Link>Properties\MetaData.cs</Link>
+    </Compile>
+  </ItemGroup>
+
+  <Import Condition="Exists('..\..\..\..\Restricted\NotForLicensees\Source\Programs\Shared\EpicGames.UBA\Project.props')" Project="..\..\..\..\Restricted\NotForLicensees\Source\Programs\Shared\EpicGames.UBA\Project.props" />
+
+  <ItemGroup>
+    <ProjectReference Include="..\..\Shared\EpicGames.Core\EpicGames.Core.csproj" />
+  </ItemGroup>
+  <ItemGroup>
+    <PackageReference Include="Microsoft.CSharp" Version="4.7.0" />
+    <PackageReference Include="Microsoft.Win32.Registry" Version="5.0.0" />
+    <PackageReference Include="System.Data.DataSetExtensions" Version="4.5.0" />
+  </ItemGroup>
+</Project>
\ No newline at end of file
diff --git a/Engine/Source/Programs/Shared/EpicGames.UBA/Utils.cs b/Engine/Source/Programs/Shared/EpicGames.UBA/Utils.cs
new file mode 100644
--- /dev/null
+++ b/Engine/Source/Programs/Shared/EpicGames.UBA/Utils.cs
@@ -0,0 +1,19 @@
+﻿// Copyright Epic Games, Inc. All Rights Reserved.
+
+using System;
+
+namespace EpicGames.UBA
+{
+	/// <summary>
+	/// Utils
+	/// </summary>
+	public static partial class Utils
+	{
+		static Lazy<bool> Available { get; set; } = new Lazy<bool>(() => false);
+
+		/// <summary>
+		/// Is available?
+		/// </summary>
+		public static bool IsAvailable() => Available.Value;
+	}
+}
diff --git a/Engine/Source/Programs/UnrealBuildTool/Executors/ImmediateActionQueue.cs b/Engine/Source/Programs/UnrealBuildTool/Executors/ImmediateActionQueue.cs
--- a/Engine/Source/Programs/UnrealBuildTool/Executors/ImmediateActionQueue.cs
+++ b/Engine/Source/Programs/UnrealBuildTool/Executors/ImmediateActionQueue.cs
@@ -101,7 +101,7 @@ namespace UnrealBuildTool
 	/// For example:
 	/// 
 	///		ParallelExecutor uses an automatic runner exclusively.
-	///		BoxExecutor uses an automatic runner to run jobs locally and a manual runner to run jobs remotely as processes 
+	///		UBAExecutor uses an automatic runner to run jobs locally and a manual runner to run jobs remotely as processes 
 	///			become available.
 	/// </summary>
 	class ImmediateActionQueue : IDisposable
diff --git a/Engine/Source/Programs/UnrealBuildTool/System/ActionGraph.cs b/Engine/Source/Programs/UnrealBuildTool/System/ActionGraph.cs
--- a/Engine/Source/Programs/UnrealBuildTool/System/ActionGraph.cs
+++ b/Engine/Source/Programs/UnrealBuildTool/System/ActionGraph.cs
@@ -363,12 +363,12 @@ namespace UnrealBuildTool
 		{
 			if (ActionCount > ParallelExecutor.GetDefaultNumParallelProcesses(BuildConfiguration.MaxParallelActions, BuildConfiguration.bAllCores, Logger))
 			{
-#if __BOXEXECUTOR_AVAILABLE__
-				if (BuildConfiguration.bAllowBoxExecutor && BoxExecutor.IsAvailable(Logger))
+#if __UBAEXECUTOR_AVAILABLE__
+				if (BuildConfiguration.bAllowBoxExecutor && UBAExecutor.IsAvailable(Logger))
 				{
-					return new BoxExecutor(BuildConfiguration.MaxParallelActions, BuildConfiguration.bAllCores, BuildConfiguration.bCompactOutput, Logger, TargetDescriptors.FirstOrDefault()?.AdditionalArguments);
+					return new UBAExecutor(BuildConfiguration.MaxParallelActions, BuildConfiguration.bAllCores, BuildConfiguration.bCompactOutput, Logger, TargetDescriptors.FirstOrDefault()?.AdditionalArguments);
 				}
-#endif // #if __BOXEXECUTOR_AVAILABLE__
+#endif // #if __UBAEXECUTOR_AVAILABLE__
 
 				if (BuildConfiguration.bAllowXGE && XGE.IsAvailable(Logger) && ActionCount >= XGE.MinActions)
 				{
@@ -384,12 +384,12 @@ namespace UnrealBuildTool
 				}
 			}
 
-#if __BOXEXECUTOR_AVAILABLE__
-			if (BuildConfiguration.bAllowBoxLocalExecutor && BoxLocalExecutor.IsAvailable(Logger))
+#if __UBAEXECUTOR_AVAILABLE__
+			if (BuildConfiguration.bAllowBoxLocalExecutor && UBALocalExecutor.IsAvailable(Logger))
 			{
-				return new BoxLocalExecutor(BuildConfiguration.MaxParallelActions, BuildConfiguration.bAllCores, BuildConfiguration.bCompactOutput, Logger, TargetDescriptors.FirstOrDefault()?.AdditionalArguments);
+				return new UBALocalExecutor(BuildConfiguration.MaxParallelActions, BuildConfiguration.bAllCores, BuildConfiguration.bCompactOutput, Logger, TargetDescriptors.FirstOrDefault()?.AdditionalArguments);
 			}
-#endif // #if __BOXEXECUTOR_AVAILABLE__
+#endif // #if __UBAEXECUTOR_AVAILABLE__
 
 			return new ParallelExecutor(BuildConfiguration.MaxParallelActions, BuildConfiguration.bAllCores, BuildConfiguration.bCompactOutput, Logger);
 		}
diff --git a/Engine/Source/Programs/UnrealBuildTool/UnrealBuildTool.csproj b/Engine/Source/Programs/UnrealBuildTool/UnrealBuildTool.csproj
--- a/Engine/Source/Programs/UnrealBuildTool/UnrealBuildTool.csproj
+++ b/Engine/Source/Programs/UnrealBuildTool/UnrealBuildTool.csproj
@@ -72,11 +72,11 @@
     </When>
   </Choose>
 
-  <PropertyGroup Condition="Exists('..\..\..\Restricted\NotForLicensees\Source\Programs\UnrealBuildTool\Executors\Experimental\BoxExecutor.cs')">
-    <DefineConstants>$(DefineConstants);__BOXEXECUTOR_AVAILABLE__</DefineConstants>
+  <PropertyGroup Condition="Exists('..\..\..\Restricted\NotForLicensees\Source\Programs\UnrealBuildTool\Executors\Experimental\UBAExecutor.cs')">
+    <DefineConstants>$(DefineConstants);__UBAEXECUTOR_AVAILABLE__</DefineConstants>
   </PropertyGroup>
-  <PropertyGroup Condition="Exists('..\..\..\Restricted\NotForLicensees\Source\Programs\Shared\EpicGames.Box\Box.props')">
-    <DefineConstants>$(DefineConstants);__BOX_AVAILABLE__</DefineConstants>
+  <PropertyGroup Condition="Exists('..\..\..\Restricted\NotForLicensees\Source\Programs\Shared\EpicGames.UBA\Project.props')">
+    <DefineConstants>$(DefineConstants);__UBA_AVAILABLE__</DefineConstants>
   </PropertyGroup>
 
   <ItemGroup>
@@ -111,13 +111,13 @@
 
   <ItemGroup>
     <!-- If you add a project reference, BuildUAT.bat/.sh and BuildUBT.bat/.sh need to be updated to also test for those projects.-->
-    <ProjectReference Include="..\Shared\EpicGames.Box\EpicGames.Box.csproj" />
     <ProjectReference Include="..\Shared\EpicGames.Build\EpicGames.Build.csproj" />
     <ProjectReference Include="..\Shared\EpicGames.Core\EpicGames.Core.csproj" />
     <ProjectReference Include="..\Shared\EpicGames.Horde\EpicGames.Horde.csproj" />
     <ProjectReference Include="..\Shared\EpicGames.IoHash\EpicGames.IoHash.csproj" />
     <ProjectReference Include="..\Shared\EpicGames.OIDC\EpicGames.OIDC.csproj" />
     <ProjectReference Include="..\Shared\EpicGames.Serialization\EpicGames.Serialization.csproj" />
+    <ProjectReference Include="..\Shared\EpicGames.UBA\EpicGames.UBA.csproj" />
     <ProjectReference Include="..\Shared\EpicGames.UHT\EpicGames.UHT.csproj" />
   </ItemGroup>
   <ItemGroup>
