Horde: Add connection preference parameter to compute client

diff --git a/Engine/Source/Programs/Shared/EpicGames.Horde/Compute/Clients/AgentComputeClient.cs b/Engine/Source/Programs/Shared/EpicGames.Horde/Compute/Clients/AgentComputeClient.cs
--- a/Engine/Source/Programs/Shared/EpicGames.Horde/Compute/Clients/AgentComputeClient.cs
+++ b/Engine/Source/Programs/Shared/EpicGames.Horde/Compute/Clients/AgentComputeClient.cs
@@ -66,7 +66,7 @@ namespace EpicGames.Horde.Compute.Clients
 		public ValueTask DisposeAsync() => new ValueTask();
 		
 		/// <inheritdoc/>
-		public async Task<IComputeLease?> TryAssignWorkerAsync(ClusterId clusterId, Requirements? requirements, string? requestId, ILogger logger, CancellationToken cancellationToken)
+		public async Task<IComputeLease?> TryAssignWorkerAsync(ClusterId clusterId, Requirements? requirements, string? requestId, ConnectionMode? connectionPreference, ILogger logger, CancellationToken cancellationToken)
 		{
 			logger.LogInformation("** CLIENT **");
 			logger.LogInformation("Launching {Path} to handle remote", _hordeAgentAssembly);
diff --git a/Engine/Source/Programs/Shared/EpicGames.Horde/Compute/Clients/LocalComputeClient.cs b/Engine/Source/Programs/Shared/EpicGames.Horde/Compute/Clients/LocalComputeClient.cs
--- a/Engine/Source/Programs/Shared/EpicGames.Horde/Compute/Clients/LocalComputeClient.cs
+++ b/Engine/Source/Programs/Shared/EpicGames.Horde/Compute/Clients/LocalComputeClient.cs
@@ -83,7 +83,7 @@ namespace EpicGames.Horde.Compute.Clients
 		}
 
 		/// <inheritdoc/>
-		public Task<IComputeLease?> TryAssignWorkerAsync(ClusterId clusterId, Requirements? requirements, string? requestId, ILogger logger, CancellationToken cancellationToken)
+		public Task<IComputeLease?> TryAssignWorkerAsync(ClusterId clusterId, Requirements? requirements, string? requestId, ConnectionMode? connectionPreference, ILogger logger, CancellationToken cancellationToken)
 		{
 #pragma warning disable CA2000 // Dispose objects before losing scope
 			RemoteComputeSocket socket = new RemoteComputeSocket(new TcpTransport(_socket), logger);
diff --git a/Engine/Source/Programs/Shared/EpicGames.Horde/Compute/Clients/ServerComputeClient.cs b/Engine/Source/Programs/Shared/EpicGames.Horde/Compute/Clients/ServerComputeClient.cs
--- a/Engine/Source/Programs/Shared/EpicGames.Horde/Compute/Clients/ServerComputeClient.cs
+++ b/Engine/Source/Programs/Shared/EpicGames.Horde/Compute/Clients/ServerComputeClient.cs
@@ -192,9 +192,9 @@ namespace EpicGames.Horde.Compute.Clients
 		}
 
 		/// <inheritdoc/>
-		public async Task<IComputeLease?> TryAssignWorkerAsync(ClusterId clusterId, Requirements? requirements, string? requestId, ILogger logger, CancellationToken cancellationToken)
+		public async Task<IComputeLease?> TryAssignWorkerAsync(ClusterId clusterId, Requirements? requirements, string? requestId, ConnectionMode? connectionPreference, ILogger logger, CancellationToken cancellationToken)
 		{
-			IAsyncEnumerator<LeaseInfo> source = ConnectAsync(clusterId, requirements, requestId, logger, cancellationToken).GetAsyncEnumerator(cancellationToken);
+			IAsyncEnumerator<LeaseInfo> source = ConnectAsync(clusterId, requirements, requestId, connectionPreference, logger, cancellationToken).GetAsyncEnumerator(cancellationToken);
 			if (!await source.MoveNextAsync())
 			{
 				await source.DisposeAsync();
@@ -212,7 +212,7 @@ namespace EpicGames.Horde.Compute.Clients
 			response.EnsureSuccessStatusCode();
 		}
 
-		async IAsyncEnumerable<LeaseInfo> ConnectAsync(ClusterId clusterId, Requirements? requirements, string? requestId, ILogger workerLogger, [EnumeratorCancellation] CancellationToken cancellationToken)
+		async IAsyncEnumerable<LeaseInfo> ConnectAsync(ClusterId clusterId, Requirements? requirements, string? requestId, ConnectionMode? connectionPreference, ILogger workerLogger, [EnumeratorCancellation] CancellationToken cancellationToken)
 		{
 			_logger.LogDebug("Requesting compute resource");
 
@@ -222,6 +222,7 @@ namespace EpicGames.Horde.Compute.Clients
 			AssignComputeRequest request = new AssignComputeRequest();
 			request.Requirements = requirements;
 			request.RequestId = requestId;
+			request.ConnectionPreference = connectionPreference;
 
 			AssignComputeResponse? responseMessage;
 			using (HttpResponseMessage response = await HordeHttpClient.PostAsync(client, $"api/v2/compute/{clusterId}", request, _cancellationSource.Token))
diff --git a/Engine/Source/Programs/Shared/EpicGames.Horde/Compute/IComputeClient.cs b/Engine/Source/Programs/Shared/EpicGames.Horde/Compute/IComputeClient.cs
--- a/Engine/Source/Programs/Shared/EpicGames.Horde/Compute/IComputeClient.cs
+++ b/Engine/Source/Programs/Shared/EpicGames.Horde/Compute/IComputeClient.cs
@@ -19,9 +19,10 @@ namespace EpicGames.Horde.Compute
 		/// <param name="clusterId">Cluster to execute the request</param>
 		/// <param name="requirements">Requirements for the agent</param>
 		/// <param name="requestId">Optional ID identifying the request over multiple calls, such as retrying the same request</param>
+		/// <param name="connectionPreference">Optional preference of connection mode to use. Server still decides</param>
 		/// <param name="logger">Logger for output from this worker</param>
 		/// <param name="cancellationToken">Cancellation token for the operation</param>
-		public Task<IComputeLease?> TryAssignWorkerAsync(ClusterId clusterId, Requirements? requirements, string? requestId, ILogger logger, CancellationToken cancellationToken = default);
+		public Task<IComputeLease?> TryAssignWorkerAsync(ClusterId clusterId, Requirements? requirements, string? requestId, ConnectionMode? connectionPreference, ILogger logger, CancellationToken cancellationToken = default);
 		
 		/// <summary>
 		/// Declare resource needs for current client
@@ -40,11 +41,17 @@ namespace EpicGames.Horde.Compute
 	/// </summary>
 	public static class ComputeClientExtensions
 	{
-		/// <inheritdoc/>
+		/// <inheritdoc cref="IComputeClient.TryAssignWorkerAsync" />
 		[Obsolete("Prefer the overload taking a requestId parameter")]
 		public static Task<IComputeLease?> TryAssignWorkerAsync(this IComputeClient computeClient, ClusterId clusterId, Requirements? requirements, ILogger logger, CancellationToken cancellationToken)
 		{
 			return computeClient.TryAssignWorkerAsync(clusterId, requirements, null, logger, cancellationToken);
 		}
+
+		/// <inheritdoc cref="IComputeClient.TryAssignWorkerAsync" />
+		public static Task<IComputeLease?> TryAssignWorkerAsync(this IComputeClient computeClient, ClusterId clusterId, Requirements? requirements, string? requestId, ILogger logger, CancellationToken cancellationToken = default)
+		{
+			return computeClient.TryAssignWorkerAsync(clusterId, requirements, requestId, null, logger, cancellationToken);
+		}
 	}
 }
