Cleaned up wording of confusing error message.

#rnx

diff --git a/Engine/Source/Programs/Shared/EpicGames.UHT/Parsers/UhtHeaderFileParser.cs b/Engine/Source/Programs/Shared/EpicGames.UHT/Parsers/UhtHeaderFileParser.cs
--- a/Engine/Source/Programs/Shared/EpicGames.UHT/Parsers/UhtHeaderFileParser.cs
+++ b/Engine/Source/Programs/Shared/EpicGames.UHT/Parsers/UhtHeaderFileParser.cs
@@ -364,7 +364,7 @@ namespace EpicGames.UHT.Parsers
 					}
 				}
 
-				string logMessage = $"Expected an include at the top of the header the follows all other includes: '#include \"{headerParser.HeaderFile.GeneratedHeaderFileName}\"'";
+				string logMessage = $"The given include must appear at the top of the header following all other includes: '#include \"{headerParser.HeaderFile.GeneratedHeaderFileName}\"'";
 				if (!noExportClassesOnly)
 				{
 					headerParser.HeaderFile.LogError(logMessage);
