Fix race waiting for empty event in ThreadPoolWorkQueue before disposing it.

diff --git a/Engine/Source/Programs/Shared/EpicGames.Core/ThreadPoolWorkQueue.cs b/Engine/Source/Programs/Shared/EpicGames.Core/ThreadPoolWorkQueue.cs
--- a/Engine/Source/Programs/Shared/EpicGames.Core/ThreadPoolWorkQueue.cs
+++ b/Engine/Source/Programs/Shared/EpicGames.Core/ThreadPoolWorkQueue.cs
@@ -46,11 +46,20 @@ namespace EpicGames.Core
 		/// </summary>
 		public void Dispose()
 		{
-			ManualResetEvent? emptyEvent = Interlocked.CompareExchange(ref _emptyEvent, null, null);
-			if (emptyEvent != null)
+			if (_emptyEvent != null) // Check we haven't disposed already
 			{
-				emptyEvent.WaitOne();
-				Interlocked.CompareExchange(ref _emptyEvent, null, emptyEvent)?.Dispose();
+				// Ensure the event state is in sync with the counter before we wait on it. Its state can lag behind the _numOutstandingJobs
+				// field because we only acquire the lock and update it after modifying _numOutstandingJobs.
+				SetEventState();
+				_emptyEvent.WaitOne();
+
+				// Acquire the lock before disposing in case any background threads are about to execute SetEventState.
+				lock (_lockObject)
+				{
+					_emptyEvent.Dispose();
+					_emptyEvent = null;
+				}
+
 				RethrowExceptions();
 			}
 		}
@@ -106,7 +115,7 @@ namespace EpicGames.Core
 		/// </summary>
 		ManualResetEvent GetEmptyEvent()
 		{
-			ManualResetEvent? emptyEvent = Interlocked.CompareExchange(ref _emptyEvent, null, null);
+			ManualResetEvent? emptyEvent = _emptyEvent;
 			if (emptyEvent == null)
 			{
 				throw new ObjectDisposedException(typeof(ThreadPoolWorkQueue).Name);
@@ -168,14 +177,16 @@ namespace EpicGames.Core
 		{
 			lock (_lockObject)
 			{
-				ManualResetEvent emptyEvent = GetEmptyEvent();
-				if (Interlocked.CompareExchange(ref _numOutstandingJobs, 0, 0) > 0)
-				{
-					emptyEvent.Reset();
-				}
-				else
+				if (_emptyEvent != null)
 				{
-					emptyEvent.Set();
+					if (Interlocked.CompareExchange(ref _numOutstandingJobs, 0, 0) > 0)
+					{
+						_emptyEvent.Reset();
+					}
+					else
+					{
+						_emptyEvent.Set();
+					}
 				}
 			}
 		}
