Horde: Rename BundleReaderCache to BundleCache. Will be used for tracking write data also.

diff --git a/Engine/Source/Programs/AutomationTool/BuildGraph/Tasks/DeployToolTask.cs b/Engine/Source/Programs/AutomationTool/BuildGraph/Tasks/DeployToolTask.cs
--- a/Engine/Source/Programs/AutomationTool/BuildGraph/Tasks/DeployToolTask.cs
+++ b/Engine/Source/Programs/AutomationTool/BuildGraph/Tasks/DeployToolTask.cs
@@ -155,7 +155,7 @@ namespace AutomationTool.Tasks
 			string basePath = $"api/v1/tools/{Parameters.Id}";
 			using HttpStorageBackend httpStorageBackend = new HttpStorageBackend(basePath, CreateHttpClient, Logger);
 			using HttpStorageClient httpStorageClient = new HttpStorageClient(basePath, CreateHttpClient, httpStorageBackend, Logger);
-			using BundleStorageClient storageClient = new BundleStorageClient(httpStorageClient, BundleReaderCache.None, Logger);
+			using BundleStorageClient storageClient = new BundleStorageClient(httpStorageClient, BundleCache.None, Logger);
 
 			await using (IStorageWriter treeWriter = storageClient.CreateWriter())
 			{
diff --git a/Engine/Source/Programs/Horde/Horde.Agent.Tests/JsonRpcLoggerTests.cs b/Engine/Source/Programs/Horde/Horde.Agent.Tests/JsonRpcLoggerTests.cs
--- a/Engine/Source/Programs/Horde/Horde.Agent.Tests/JsonRpcLoggerTests.cs
+++ b/Engine/Source/Programs/Horde/Horde.Agent.Tests/JsonRpcLoggerTests.cs
@@ -138,7 +138,7 @@ namespace Horde.Agent.Tests
 		[TestMethod]
 		public async Task StorageLoggerTestAsync()
 		{
-			using BundleReaderCache cache = new BundleReaderCache();
+			using BundleCache cache = new BundleCache();
 			using MemoryStorageClient store = new MemoryStorageClient();
 
 			BundleReader reader = new BundleReader(store, cache, NullLogger.Instance);
diff --git a/Engine/Source/Programs/Horde/Horde.Agent.Tests/WorkerServiceTest.cs b/Engine/Source/Programs/Horde/Horde.Agent.Tests/WorkerServiceTest.cs
--- a/Engine/Source/Programs/Horde/Horde.Agent.Tests/WorkerServiceTest.cs
+++ b/Engine/Source/Programs/Horde/Horde.Agent.Tests/WorkerServiceTest.cs
@@ -68,7 +68,7 @@ namespace Horde.Agent.Tests
 			_serviceCollection.AddLogging();
 			_serviceCollection.AddHordeHttpClient();
 			_serviceCollection.AddSingleton<IServerLoggerFactory, FakeServerLoggerFactory>();
-			_serviceCollection.AddSingleton<BundleReaderCache>();
+			_serviceCollection.AddSingleton<BundleCache>();
 			_serviceCollection.AddSingleton<StorageBackendCache>();
 			_serviceCollection.AddSingleton<HttpStorageClientFactory>();
 
diff --git a/Engine/Source/Programs/Horde/Horde.Agent/AgentApp.cs b/Engine/Source/Programs/Horde/Horde.Agent/AgentApp.cs
--- a/Engine/Source/Programs/Horde/Horde.Agent/AgentApp.cs
+++ b/Engine/Source/Programs/Horde/Horde.Agent/AgentApp.cs
@@ -214,7 +214,7 @@ namespace Horde.Agent
 			services.AddSingleton<WorkerService>();
 			services.AddHostedService(sp => sp.GetRequiredService<WorkerService>());
 
-			services.AddSingleton<BundleReaderCache>();
+			services.AddSingleton<BundleCache>();
 			services.AddSingleton<StorageBackendCache>(CreateStorageBackendCache);
 			services.AddSingleton<HttpStorageClientFactory>();
 
diff --git a/Engine/Source/Programs/Horde/Horde.Server.Tests/TestSetup.cs b/Engine/Source/Programs/Horde/Horde.Server.Tests/TestSetup.cs
--- a/Engine/Source/Programs/Horde/Horde.Server.Tests/TestSetup.cs
+++ b/Engine/Source/Programs/Horde/Horde.Server.Tests/TestSetup.cs
@@ -276,7 +276,7 @@ namespace Horde.Server.Tests
 
 			services.AddSingleton<StorageService>();
 			services.AddSingleton<StorageBackendCache>();
-			services.AddSingleton<BundleReaderCache>();
+			services.AddSingleton<BundleCache>();
 		}
 
 		public Task<Fixture> CreateFixtureAsync()
diff --git a/Engine/Source/Programs/Horde/Horde.Server/Commands/Bundles/BundleWriteTest.cs b/Engine/Source/Programs/Horde/Horde.Server/Commands/Bundles/BundleWriteTest.cs
--- a/Engine/Source/Programs/Horde/Horde.Server/Commands/Bundles/BundleWriteTest.cs
+++ b/Engine/Source/Programs/Horde/Horde.Server/Commands/Bundles/BundleWriteTest.cs
@@ -50,7 +50,7 @@ namespace Horde.Server.Commands.Bundles
 		public override async Task<int> ExecuteAsync(ILogger logger)
 		{
 			using NullStorageClient nullStore = new NullStorageClient();
-			using BundleStorageClient store = new BundleStorageClient(nullStore, BundleReaderCache.None, logger);
+			using BundleStorageClient store = new BundleStorageClient(nullStore, BundleCache.None, logger);
 
 			BundleOptions options = new BundleOptions();
 			options.CompressionFormat = BundleCompressionFormat.None;
diff --git a/Engine/Source/Programs/Horde/Horde.Server/Commands/Bundles/CreateCommand.cs b/Engine/Source/Programs/Horde/Horde.Server/Commands/Bundles/CreateCommand.cs
--- a/Engine/Source/Programs/Horde/Horde.Server/Commands/Bundles/CreateCommand.cs
+++ b/Engine/Source/Programs/Horde/Horde.Server/Commands/Bundles/CreateCommand.cs
@@ -47,7 +47,7 @@ namespace Horde.Server.Commands.Bundles
 			}
 			else
 			{
-				using IStorageClient storageClient = BundleStorageClient.CreateFromDirectory(OutputDir, serviceProvider.GetRequiredService<BundleReaderCache>(), logger);
+				using IStorageClient storageClient = BundleStorageClient.CreateFromDirectory(OutputDir, serviceProvider.GetRequiredService<BundleCache>(), logger);
 				await ExecuteInternalAsync(storageClient, logger);
 			}
 
diff --git a/Engine/Source/Programs/Horde/Horde.Server/Commands/Install/SetupCommand.cs b/Engine/Source/Programs/Horde/Horde.Server/Commands/Install/SetupCommand.cs
--- a/Engine/Source/Programs/Horde/Horde.Server/Commands/Install/SetupCommand.cs
+++ b/Engine/Source/Programs/Horde/Horde.Server/Commands/Install/SetupCommand.cs
@@ -58,9 +58,9 @@ namespace Horde.Server.Commands.Install
 
 			RefName refName = new RefName("latest");
 
-			using BundleReaderCache bundleReaderCache = new BundleReaderCache();
+			using BundleCache bundleCache = new BundleCache();
 
-			using IStorageClient client = BundleStorageClient.CreateFromDirectory(bundleDir, bundleReaderCache, logger);
+			using IStorageClient client = BundleStorageClient.CreateFromDirectory(bundleDir, bundleCache, logger);
 			await using (IStorageWriter writer = client.CreateWriter(refName))
 			{
 				DirectoryNode dirNode = new DirectoryNode();
diff --git a/Engine/Source/Programs/Horde/Horde.Server/Startup.cs b/Engine/Source/Programs/Horde/Horde.Server/Startup.cs
--- a/Engine/Source/Programs/Horde/Horde.Server/Startup.cs
+++ b/Engine/Source/Programs/Horde/Horde.Server/Startup.cs
@@ -529,7 +529,7 @@ namespace Horde.Server
 			services.AddSingleton<StorageService>();
 			services.AddScoped(sp => sp.GetRequiredService<StorageService>().CreateStorageClientFactory(sp.GetRequiredService<IOptionsSnapshot<GlobalConfig>>().Value));
 			services.AddSingleton<TestDataService>();
-			services.AddSingleton<BundleReaderCache>();
+			services.AddSingleton<BundleCache>();
 			services.AddSingleton<StorageBackendCache>(CreateStorageBackendCache);
 
 			if (settings.JiraUrl != null)
diff --git a/Engine/Source/Programs/Horde/Horde.Server/Storage/StorageService.cs b/Engine/Source/Programs/Horde/Horde.Server/Storage/StorageService.cs
--- a/Engine/Source/Programs/Horde/Horde.Server/Storage/StorageService.cs
+++ b/Engine/Source/Programs/Horde/Horde.Server/Storage/StorageService.cs
@@ -488,7 +488,7 @@ namespace Horde.Server.Storage
 
 		readonly RedisService _redisService;
 		readonly IClock _clock;
-		readonly BundleReaderCache _bundleReaderCache;
+		readonly BundleCache _bundleCache;
 		readonly IMemoryCache _memoryCache;
 		readonly IStorageBackendProvider _storageBackendProvider;
 		readonly IOptionsMonitor<GlobalConfig> _globalConfig;
@@ -512,11 +512,11 @@ namespace Horde.Server.Storage
 		/// <summary>
 		/// Constructor
 		/// </summary>
-		public StorageService(MongoService mongoService, RedisService redisService, IClock clock, BundleReaderCache bundleReaderCache, IMemoryCache memoryCache, IStorageBackendProvider storageBackendProvider, IOptionsMonitor<GlobalConfig> globalConfig, Tracer tracer, ILogger<StorageService> logger)
+		public StorageService(MongoService mongoService, RedisService redisService, IClock clock, BundleCache bundleCache, IMemoryCache memoryCache, IStorageBackendProvider storageBackendProvider, IOptionsMonitor<GlobalConfig> globalConfig, Tracer tracer, ILogger<StorageService> logger)
 		{
 			_redisService = redisService;
 			_clock = clock;
-			_bundleReaderCache = bundleReaderCache;
+			_bundleCache = bundleCache;
 			_memoryCache = memoryCache;
 			_storageBackendProvider = storageBackendProvider;
 			_globalConfig = globalConfig;
@@ -615,7 +615,7 @@ namespace Horde.Server.Storage
 
 #pragma warning disable CA2000 // Dispose objects before losing scope
 							client = new StorageClientImpl(this, namespaceConfig, backend, _tracer);
-							client = new BundleStorageClient(client, _bundleReaderCache, _logger);
+							client = new BundleStorageClient(client, _bundleCache, _logger);
 #pragma warning restore CA2000 // Dispose objects before losing scope
 
 							NamespaceInfo namespaceInfo = new NamespaceInfo(namespaceConfig, backend, new SharedStorageClient(client));
diff --git a/Engine/Source/Programs/Horde/Horde.Server/Tools/ToolCollection.cs b/Engine/Source/Programs/Horde/Horde.Server/Tools/ToolCollection.cs
--- a/Engine/Source/Programs/Horde/Horde.Server/Tools/ToolCollection.cs
+++ b/Engine/Source/Programs/Horde/Horde.Server/Tools/ToolCollection.cs
@@ -162,7 +162,7 @@ namespace Horde.Server.Tools
 		private readonly VersionedCollection<ToolId, Tool> _tools;
 		private readonly StorageService _storageService;
 		private readonly IClock _clock;
-		private readonly BundleReaderCache _cache;
+		private readonly BundleCache _cache;
 		private readonly ILogger _logger;
 
 		private static readonly RedisKey s_baseKey = "tools/v1/";
@@ -172,7 +172,7 @@ namespace Horde.Server.Tools
 		/// <summary>
 		/// Constructor
 		/// </summary>
-		public ToolCollection(MongoService mongoService, RedisService redisService, StorageService storageService, BundleReaderCache cache, IClock clock, ILogger<ToolCollection> logger)
+		public ToolCollection(MongoService mongoService, RedisService redisService, StorageService storageService, BundleCache cache, IClock clock, ILogger<ToolCollection> logger)
 		{
 			_tools = new VersionedCollection<ToolId, Tool>(mongoService, "Tools", redisService, s_baseKey, s_types);
 			_storageService = storageService;
diff --git a/Engine/Source/Programs/Horde/Horde/CmdApp.cs b/Engine/Source/Programs/Horde/Horde/CmdApp.cs
--- a/Engine/Source/Programs/Horde/Horde/CmdApp.cs
+++ b/Engine/Source/Programs/Horde/Horde/CmdApp.cs
@@ -38,7 +38,7 @@ namespace Horde
 			services.AddMemoryCache();
 			services.AddSingleton(sp => Options.Create(CmdConfig.Read()));
 			services.AddHordeHttpClient((sp, client) => client.BaseAddress = sp.GetRequiredService<IOptions<CmdConfig>>().Value.Server);
-			services.AddSingleton<BundleReaderCache>(CreateStorageClientCache);
+			services.AddSingleton<BundleCache>(CreateStorageClientCache);
 			services.AddSingleton<StorageBackendCache>(CreateStorageBackendCache);
 			services.AddSingleton<HttpStorageClientFactory>();
 
@@ -47,11 +47,11 @@ namespace Horde
 			return await CommandHost.RunAsync(arguments, serviceProvider, null);
 		}
 
-		static BundleReaderCache CreateStorageClientCache(IServiceProvider serviceProvider)
+		static BundleCache CreateStorageClientCache(IServiceProvider serviceProvider)
 		{
 			CmdConfig cmdConfig = serviceProvider.GetRequiredService<IOptions<CmdConfig>>().Value;
 
-			BundleReaderCacheOptions options = new BundleReaderCacheOptions();
+			BundleCacheOptions options = new BundleCacheOptions();
 			if (cmdConfig.Cache.HeaderCacheSize.HasValue)
 			{
 				options.HeaderCacheSize = cmdConfig.Cache.HeaderCacheSize.Value * 1024 * 1024;
@@ -61,7 +61,7 @@ namespace Horde
 				options.PacketCacheSize = cmdConfig.Cache.PacketCacheSize.Value * 1024 * 1024;
 			}
 
-			return new BundleReaderCache(options);
+			return new BundleCache(options);
 		}
 
 		static StorageBackendCache CreateStorageBackendCache(IServiceProvider serviceProvider)
diff --git a/Engine/Source/Programs/Horde/Horde/Commands/Bundles/BundleCreate.cs b/Engine/Source/Programs/Horde/Horde/Commands/Bundles/BundleCreate.cs
--- a/Engine/Source/Programs/Horde/Horde/Commands/Bundles/BundleCreate.cs
+++ b/Engine/Source/Programs/Horde/Horde/Commands/Bundles/BundleCreate.cs
@@ -25,8 +25,8 @@ namespace Horde.Commands.Bundles
 		[CommandLine("-Filter=", Description = "Filter for files to include, in P4 syntax (eg. Foo/...).")]
 		public string Filter { get; set; } = "...";
 
-		public BundleCreate(HttpStorageClientFactory storageClientFactory, BundleReaderCache bundleReaderCache, IOptions<CmdConfig> config)
-			: base(storageClientFactory, bundleReaderCache, config)
+		public BundleCreate(HttpStorageClientFactory storageClientFactory, BundleCache bundleCache, IOptions<CmdConfig> config)
+			: base(storageClientFactory, bundleCache, config)
 		{
 		}
 
@@ -34,7 +34,7 @@ namespace Horde.Commands.Bundles
 		{
 			if (File != null)
 			{
-				using IStorageClient store = BundleStorageClient.CreateFromDirectory(File.Directory, BundleReaderCache, logger);
+				using IStorageClient store = BundleStorageClient.CreateFromDirectory(File.Directory, BundleCache, logger);
 				return await ExecuteInternalAsync(store, logger);
 			}
 			else if (!String.IsNullOrEmpty(Ref))
diff --git a/Engine/Source/Programs/Horde/Horde/Commands/Bundles/BundleExtract.cs b/Engine/Source/Programs/Horde/Horde/Commands/Bundles/BundleExtract.cs
--- a/Engine/Source/Programs/Horde/Horde/Commands/Bundles/BundleExtract.cs
+++ b/Engine/Source/Programs/Horde/Horde/Commands/Bundles/BundleExtract.cs
@@ -28,8 +28,8 @@ namespace Horde.Commands.Bundles
 		[CommandLine("-OutputDir=", Required = true)]
 		public DirectoryReference OutputDir { get; set; } = null!;
 
-		public BundleExtract(HttpStorageClientFactory storageClientFactory, BundleReaderCache bundleReaderCache, IOptions<CmdConfig> config)
-			: base(storageClientFactory, bundleReaderCache, config)
+		public BundleExtract(HttpStorageClientFactory storageClientFactory, BundleCache bundleCache, IOptions<CmdConfig> config)
+			: base(storageClientFactory, bundleCache, config)
 		{
 		}
 
@@ -37,7 +37,7 @@ namespace Horde.Commands.Bundles
 		{
 			if (File != null)
 			{
-				using IStorageClient store = BundleStorageClient.CreateFromDirectory(File.Directory, BundleReaderCache, logger);
+				using IStorageClient store = BundleStorageClient.CreateFromDirectory(File.Directory, BundleCache, logger);
 				BlobHandle handle = store.CreateBlobHandle(await FileStorageClient.ReadRefAsync(File));
 				await ExecuteInternalAsync(store, handle, logger);
 			}
diff --git a/Engine/Source/Programs/Horde/Horde/Commands/Compute/ComputeRun.cs b/Engine/Source/Programs/Horde/Horde/Commands/Compute/ComputeRun.cs
--- a/Engine/Source/Programs/Horde/Horde/Commands/Compute/ComputeRun.cs
+++ b/Engine/Source/Programs/Horde/Horde/Commands/Compute/ComputeRun.cs
@@ -107,7 +107,7 @@ namespace Horde.Commands.Compute
 
 			// Create a sandbox from the data to be uploaded
 			using MemoryStorageClient memoryStorage = new MemoryStorageClient();
-			using BundleStorageClient storage = new BundleStorageClient(memoryStorage, BundleReaderCache.None, logger);
+			using BundleStorageClient storage = new BundleStorageClient(memoryStorage, BundleCache.None, logger);
 			BlobLocator sandbox = await CreateSandboxAsync(TaskFile, storage, cancellationToken);
 
 			// Open a socket and upload the sandbox
diff --git a/Engine/Source/Programs/Horde/Horde/Commands/StorageCommandBase.cs b/Engine/Source/Programs/Horde/Horde/Commands/StorageCommandBase.cs
--- a/Engine/Source/Programs/Horde/Horde/Commands/StorageCommandBase.cs
+++ b/Engine/Source/Programs/Horde/Horde/Commands/StorageCommandBase.cs
@@ -27,7 +27,7 @@ namespace Horde.Commands
 		/// <summary>
 		/// Cache for storage
 		/// </summary>
-		public BundleReaderCache BundleReaderCache { get; }
+		public BundleCache BundleCache { get; }
 
 		/// <summary>
 		/// Configuration for the tool
@@ -39,11 +39,11 @@ namespace Horde.Commands
 		/// <summary>
 		/// Constructor
 		/// </summary>
-		public StorageCommandBase(HttpStorageClientFactory storageClientFactory, BundleReaderCache bundleReaderCache, IOptions<CmdConfig> config)
+		public StorageCommandBase(HttpStorageClientFactory storageClientFactory, BundleCache bundleCache, IOptions<CmdConfig> config)
 		{
 			_storageClientFactory = storageClientFactory;
 
-			BundleReaderCache = bundleReaderCache;
+			BundleCache = bundleCache;
 			Config = config.Value;
 		}
 
diff --git a/Engine/Source/Programs/Horde/Horde/Commands/Workspace/WorkspaceSync.cs b/Engine/Source/Programs/Horde/Horde/Commands/Workspace/WorkspaceSync.cs
--- a/Engine/Source/Programs/Horde/Horde/Commands/Workspace/WorkspaceSync.cs
+++ b/Engine/Source/Programs/Horde/Horde/Commands/Workspace/WorkspaceSync.cs
@@ -33,8 +33,8 @@ namespace Horde.Commands.Workspace
 		[CommandLine("-Stats")]
 		public bool Stats { get; set; }
 
-		public WorkspaceSync(HttpStorageClientFactory storageClientFactory, BundleReaderCache bundleReaderCache, IOptions<CmdConfig> config)
-			: base(storageClientFactory, bundleReaderCache, config)
+		public WorkspaceSync(HttpStorageClientFactory storageClientFactory, BundleCache bundleCache, IOptions<CmdConfig> config)
+			: base(storageClientFactory, bundleCache, config)
 		{
 		}
 
@@ -42,7 +42,7 @@ namespace Horde.Commands.Workspace
 		{
 			if (File != null)
 			{
-				using IStorageClient store = BundleStorageClient.CreateFromDirectory(File.Directory, BundleReaderCache, logger);
+				using IStorageClient store = BundleStorageClient.CreateFromDirectory(File.Directory, BundleCache, logger);
 				BlobHandle handle = store.CreateBlobHandle(await FileStorageClient.ReadRefAsync(File));
 				return await ExecuteInternalAsync(store, handle, logger);
 			}
diff --git a/Engine/Source/Programs/Shared/EpicGames.Horde.Tests/BundleStoreTests.cs b/Engine/Source/Programs/Shared/EpicGames.Horde.Tests/BundleStoreTests.cs
--- a/Engine/Source/Programs/Shared/EpicGames.Horde.Tests/BundleStoreTests.cs
+++ b/Engine/Source/Programs/Shared/EpicGames.Horde.Tests/BundleStoreTests.cs
@@ -58,7 +58,7 @@ namespace EpicGames.Horde.Tests
 		static async Task<byte[]> CreateBundleNormalAsync()
 		{
 			using MemoryStorageClient memoryStore = new MemoryStorageClient();
-			using BundleStorageClient store = new BundleStorageClient(memoryStore, BundleReaderCache.None, NullLogger.Instance);
+			using BundleStorageClient store = new BundleStorageClient(memoryStore, BundleCache.None, NullLogger.Instance);
 			await using IStorageWriter writer = store.CreateWriter(options: new BundleOptions { CompressionFormat = BundleCompressionFormat.None });
 
 			TextNode node = new TextNode("Hello world");
@@ -99,7 +99,7 @@ namespace EpicGames.Horde.Tests
 			using IMemoryCache cache = new MemoryCache(new MemoryCacheOptions());
 
 			using MemoryStorageClient blobStore = new MemoryStorageClient();
-			using BundleStorageClient bundleStore = new BundleStorageClient(blobStore, BundleReaderCache.None, NullLogger.Instance);
+			using BundleStorageClient bundleStore = new BundleStorageClient(blobStore, BundleCache.None, NullLogger.Instance);
 
 			await TestTreeAsync(bundleStore, new BundleOptions { MaxBlobSize = 1024 * 1024 });
 
@@ -113,7 +113,7 @@ namespace EpicGames.Horde.Tests
 			using IMemoryCache cache = new MemoryCache(new MemoryCacheOptions());
 
 			using MemoryStorageClient blobStore = new MemoryStorageClient();
-			using BundleStorageClient bundleStore = new BundleStorageClient(blobStore, BundleReaderCache.None, NullLogger.Instance);
+			using BundleStorageClient bundleStore = new BundleStorageClient(blobStore, BundleCache.None, NullLogger.Instance);
 
 			await TestTreeAsync(bundleStore, new BundleOptions { MaxBlobSize = 1 });
 
@@ -161,7 +161,7 @@ namespace EpicGames.Horde.Tests
 
 				await store.WriteRefTargetAsync(new RefName("test"), await writer.WriteHashedNodeAsync(root));
 
-				BundleReader reader = new BundleReader(store, BundleReaderCache.None, NullLogger.Instance);
+				BundleReader reader = new BundleReader(store, BundleCache.None, NullLogger.Instance);
 				await CheckTreeAsync(root);
 			}
 
@@ -269,7 +269,7 @@ namespace EpicGames.Horde.Tests
 		{
 			using IMemoryCache cache = new MemoryCache(new MemoryCacheOptions());
 			using MemoryStorageClient store = new MemoryStorageClient();
-			BundleReader reader = new BundleReader(store, BundleReaderCache.None, NullLogger.Instance);
+			BundleReader reader = new BundleReader(store, BundleCache.None, NullLogger.Instance);
 
 			// Generate a tree
 			{
@@ -317,7 +317,7 @@ namespace EpicGames.Horde.Tests
 		public async Task StreamTestAsync()
 		{
 			using MemoryStorageClient memoryStore = new MemoryStorageClient();
-			using BundleStorageClient store = new BundleStorageClient(memoryStore, BundleReaderCache.None, NullLogger.Instance);
+			using BundleStorageClient store = new BundleStorageClient(memoryStore, BundleCache.None, NullLogger.Instance);
 
 			const int Length = 4096;
 
@@ -358,7 +358,7 @@ namespace EpicGames.Horde.Tests
 		{
 			using IMemoryCache cache = new MemoryCache(new MemoryCacheOptions());
 			using MemoryStorageClient memoryStore = new MemoryStorageClient();
-			using BundleStorageClient store = new BundleStorageClient(memoryStore, BundleReaderCache.None, NullLogger.Instance);
+			using BundleStorageClient store = new BundleStorageClient(memoryStore, BundleCache.None, NullLogger.Instance);
 
 			const int Length = 1024;
 			const int Copies = 4096;
diff --git a/Engine/Source/Programs/Shared/EpicGames.Horde.Tests/BundleTests.cs b/Engine/Source/Programs/Shared/EpicGames.Horde.Tests/BundleTests.cs
--- a/Engine/Source/Programs/Shared/EpicGames.Horde.Tests/BundleTests.cs
+++ b/Engine/Source/Programs/Shared/EpicGames.Horde.Tests/BundleTests.cs
@@ -29,7 +29,7 @@ namespace EpicGames.Horde.Tests
 		{
 			_cache = new MemoryCache(new MemoryCacheOptions());
 			_memoryStore = new MemoryStorageClient();
-			_storage = new BundleStorageClient(_memoryStore, BundleReaderCache.None, NullLogger.Instance);
+			_storage = new BundleStorageClient(_memoryStore, BundleCache.None, NullLogger.Instance);
 		}
 
 		public void Dispose()
diff --git a/Engine/Source/Programs/Shared/EpicGames.Horde.Tests/ComputeSocketTests.cs b/Engine/Source/Programs/Shared/EpicGames.Horde.Tests/ComputeSocketTests.cs
--- a/Engine/Source/Programs/Shared/EpicGames.Horde.Tests/ComputeSocketTests.cs
+++ b/Engine/Source/Programs/Shared/EpicGames.Horde.Tests/ComputeSocketTests.cs
@@ -140,7 +140,7 @@ namespace EpicGames.Horde.Tests
 					}
 
 					using MemoryStorageClient memoryStorage = new MemoryStorageClient();
-					using BundleStorageClient storage = new BundleStorageClient(memoryStorage, BundleReaderCache.None, NullLogger.Instance);
+					using BundleStorageClient storage = new BundleStorageClient(memoryStorage, BundleCache.None, NullLogger.Instance);
 					await using (IStorageWriter treeWriter = storage.CreateWriter())
 					{
 						FileReference file = FileReference.Combine(tempDir, "subdir/hello.txt");
diff --git a/Engine/Source/Programs/Shared/EpicGames.Horde/Compute/AgentMessageHandler.cs b/Engine/Source/Programs/Shared/EpicGames.Horde/Compute/AgentMessageHandler.cs
--- a/Engine/Source/Programs/Shared/EpicGames.Horde/Compute/AgentMessageHandler.cs
+++ b/Engine/Source/Programs/Shared/EpicGames.Horde/Compute/AgentMessageHandler.cs
@@ -136,7 +136,7 @@ namespace EpicGames.Horde.Compute
 		async Task WriteFilesAsync(AgentMessageChannel channel, string path, BlobLocator locator, CancellationToken cancellationToken)
 		{
 			using AgentStorageClient innerStore = new AgentStorageClient(channel);
-			using BundleReaderCache cache = new BundleReaderCache(new BundleReaderCacheOptions { HeaderCacheSize = 10 * 1024 * 1024, PacketCacheSize = 128 * 1024 * 1024 });
+			using BundleCache cache = new BundleCache(new BundleCacheOptions { HeaderCacheSize = 10 * 1024 * 1024, PacketCacheSize = 128 * 1024 * 1024 });
 			using BundleStorageClient store = new BundleStorageClient(innerStore, cache, _logger);
 
 			BlobHandle handle = store.CreateBlobHandle(locator);
diff --git a/Engine/Source/Programs/Shared/EpicGames.Horde/Storage/Bundles/BundleCache.cs b/Engine/Source/Programs/Shared/EpicGames.Horde/Storage/Bundles/BundleCache.cs
new file mode 100644
--- /dev/null
+++ b/Engine/Source/Programs/Shared/EpicGames.Horde/Storage/Bundles/BundleCache.cs
@@ -0,0 +1,158 @@
+// Copyright Epic Games, Inc. All Rights Reserved.
+
+using System;
+using System.Diagnostics.CodeAnalysis;
+using Microsoft.Extensions.Caching.Memory;
+
+namespace EpicGames.Horde.Storage
+{
+	/// <summary>
+	/// Options for creating a storage cache
+	/// </summary>
+	public class BundleCacheOptions
+	{
+		/// <summary>
+		/// Size of the header cache
+		/// </summary>
+		public long HeaderCacheSize { get; set; } = 64 * 1024 * 1024;
+
+		/// <summary>
+		/// Size of the packet cache
+		/// </summary>
+		public long PacketCacheSize { get; set; } = 192 * 1024 * 1024;
+	}
+
+	/// <summary>
+	/// Implements caching functionality for reading from a storage client
+	/// </summary>
+	public sealed class BundleCache : IDisposable
+	{
+		readonly MemoryCache? _headerCache;
+		readonly MemoryCache? _packetCache;
+
+		/// <summary>
+		/// Instance of an empty cache
+		/// </summary>
+		public static BundleCache None { get; } = new BundleCache(new BundleCacheOptions { HeaderCacheSize = 0, PacketCacheSize = 0 });
+
+		/// <summary>
+		/// Size of the configured header cache
+		/// </summary>
+		public long HeaderCacheSize { get; }
+
+		/// <summary>
+		/// Size of the configured packet cache
+		/// </summary>
+		public long PacketCacheSize { get; }
+
+		/// <summary>
+		/// Whether there is a packet cache present
+		/// </summary>
+		public bool HasPacketCache => _packetCache != null;
+
+		/// <summary>
+		/// Constructor
+		/// </summary>
+		public BundleCache() : this(new BundleCacheOptions()) 
+		{ 
+		}
+
+		/// <summary>
+		/// Constructor
+		/// </summary>
+		public BundleCache(BundleCacheOptions options)
+		{
+			if (options.HeaderCacheSize > 0)
+			{
+				HeaderCacheSize = options.HeaderCacheSize;
+				_headerCache = new MemoryCache(new MemoryCacheOptions { SizeLimit = options.HeaderCacheSize });
+			}
+			if (options.PacketCacheSize > 0)
+			{
+				PacketCacheSize = options.PacketCacheSize;
+				_packetCache = new MemoryCache(new MemoryCacheOptions { SizeLimit = options.PacketCacheSize });
+			}
+		}
+
+		/// <inheritdoc/>
+		public void Dispose()
+		{
+			_headerCache?.Dispose();
+			_packetCache?.Dispose();
+		}
+
+		static void AddCachedValue(IMemoryCache? cache, string cacheKey, object value, int size)
+		{
+			if (cache != null)
+			{
+				using (ICacheEntry entry = cache.CreateEntry(cacheKey))
+				{
+					entry.SetValue(value);
+					entry.SetSize(size);
+				}
+			}
+		}
+
+		static bool TryGetCachedValue<T>(IMemoryCache? cache, string cacheKey, [MaybeNull, NotNullWhen(true)] out T value)
+		{
+			if (cache == null)
+			{
+				value = default;
+				return false;
+			}
+			return cache.TryGetValue(cacheKey, out value);
+		}
+
+		static string GetBundleInfoCacheKey(BlobLocator locator) => $"bundle:{locator}";
+		static string GetEncodedPacketCacheKey(BlobLocator locator, int packetIdx) => $"encoded-packet:{locator}#{packetIdx}";
+		static string GetDecodedPacketCacheKey(BlobLocator locator, int packetIdx) => $"decoded-packet:{locator}#{packetIdx}";
+
+		/// <summary>
+		/// Adds a bundle info object to the cache
+		/// </summary>
+		public void AddCachedHeader(BlobLocator locator, Bundles.V1.BundleInfo bundleInfo)
+		{
+			AddCachedValue(_headerCache, GetBundleInfoCacheKey(locator), bundleInfo, bundleInfo.HeaderLength);
+		}
+
+		/// <summary>
+		/// Try to read a bundle info object from the cache
+		/// </summary>
+		public bool TryGetCachedHeader(BlobLocator locator, [NotNullWhen(true)] out Bundles.V1.BundleInfo? bundleInfo)
+		{
+			return TryGetCachedValue(_headerCache, GetBundleInfoCacheKey(locator), out bundleInfo);
+		}
+
+		/// <summary>
+		/// Adds an encoded bundle packet to the cache
+		/// </summary>
+		public void AddCachedEncodedPacket(BlobLocator locator, int packetIdx, ReadOnlyMemory<byte> data)
+		{
+			AddCachedValue(_packetCache, GetEncodedPacketCacheKey(locator, packetIdx), data, data.Length);
+		}
+
+		/// <summary>
+		/// Try to read an encoded bundle packet from the cache
+		/// </summary>
+		public bool TryGetCachedEncodedPacket(BlobLocator locator, int packetIdx, out ReadOnlyMemory<byte> data)
+		{
+			return TryGetCachedValue(_packetCache, GetEncodedPacketCacheKey(locator, packetIdx), out data);
+		}
+
+		/// <summary>
+		/// Adds a decoded bundle packet to the cache
+		/// </summary>
+		public void AddCachedDecodedPacket(BlobLocator locator, int packetIdx, ReadOnlyMemory<byte> data)
+		{
+			AddCachedValue(_packetCache, GetDecodedPacketCacheKey(locator, packetIdx), data, data.Length);
+		}
+
+		/// <summary>
+		/// Try to read a decoded bundle packet from the cache
+		/// </summary>
+		public bool TryGetCachedDecodedPacket(BlobLocator locator, int packetIdx, out ReadOnlyMemory<byte> data)
+		{
+			return TryGetCachedValue(_packetCache, GetDecodedPacketCacheKey(locator, packetIdx), out data);
+		}
+	}
+}
diff --git a/Engine/Source/Programs/Shared/EpicGames.Horde/Storage/Bundles/BundleReaderCache.cs b/Engine/Source/Programs/Shared/EpicGames.Horde/Storage/Bundles/BundleReaderCache.cs
deleted file mode 100644
diff --git a/Engine/Source/Programs/Shared/EpicGames.Horde/Storage/Bundles/V1/BundleReader.cs b/Engine/Source/Programs/Shared/EpicGames.Horde/Storage/Bundles/V1/BundleReader.cs
--- a/Engine/Source/Programs/Shared/EpicGames.Horde/Storage/Bundles/V1/BundleReader.cs
+++ b/Engine/Source/Programs/Shared/EpicGames.Horde/Storage/Bundles/V1/BundleReader.cs
@@ -127,7 +127,7 @@ namespace EpicGames.Horde.Storage.Bundles.V1
 		const int DefaultUncachedFetchSize = 1 * 1024 * 1024;
 
 		readonly IStorageClient _store;
-		readonly BundleReaderCache _cache;
+		readonly BundleCache _cache;
 		readonly ILogger _logger;
 
 		readonly object _queueLock = new object();
@@ -142,7 +142,7 @@ namespace EpicGames.Horde.Storage.Bundles.V1
 		/// <summary>
 		/// Accessor for the cache
 		/// </summary>
-		public BundleReaderCache Cache => _cache;
+		public BundleCache Cache => _cache;
 
 		/// <summary>
 		/// Constructor
@@ -150,7 +150,7 @@ namespace EpicGames.Horde.Storage.Bundles.V1
 		/// <param name="store"></param>
 		/// <param name="cache">Cache for data</param>
 		/// <param name="logger">Logger for output</param>
-		public BundleReader(IStorageClient store, BundleReaderCache cache, ILogger logger)
+		public BundleReader(IStorageClient store, BundleCache cache, ILogger logger)
 		{
 			_store = store;
 			_cache = cache;
diff --git a/Engine/Source/Programs/Shared/EpicGames.Horde/Storage/Clients/BundleStorageClient.cs b/Engine/Source/Programs/Shared/EpicGames.Horde/Storage/Clients/BundleStorageClient.cs
--- a/Engine/Source/Programs/Shared/EpicGames.Horde/Storage/Clients/BundleStorageClient.cs
+++ b/Engine/Source/Programs/Shared/EpicGames.Horde/Storage/Clients/BundleStorageClient.cs
@@ -95,7 +95,7 @@ namespace EpicGames.Horde.Storage.Clients
 		/// <summary>
 		/// Constructor
 		/// </summary>
-		public BundleStorageClient(IStorageClient inner, BundleReaderCache cache, ILogger logger)
+		public BundleStorageClient(IStorageClient inner, BundleCache cache, ILogger logger)
 		{
 			_inner = inner;
 			_bundleReader = new Bundles.V1.BundleReader(this, cache, logger);
@@ -113,13 +113,13 @@ namespace EpicGames.Horde.Storage.Clients
 		public static BundleStorageClient CreateFromMemory(ILogger logger)
 		{
 			MemoryStorageClient blobStore = new MemoryStorageClient();
-			return new BundleStorageClient(blobStore, BundleReaderCache.None, logger);
+			return new BundleStorageClient(blobStore, BundleCache.None, logger);
 		}
 
 		/// <summary>
 		/// Creates a bundle storage client around a directory on the filesystem
 		/// </summary>
-		public static BundleStorageClient CreateFromDirectory(DirectoryReference rootDir, BundleReaderCache cache, ILogger logger)
+		public static BundleStorageClient CreateFromDirectory(DirectoryReference rootDir, BundleCache cache, ILogger logger)
 		{
 			FileStorageClient fileStorageClient = new FileStorageClient(rootDir, logger);
 			return new BundleStorageClient(fileStorageClient, cache, logger);
diff --git a/Engine/Source/Programs/Shared/EpicGames.Horde/Storage/Clients/HttpStorageClient.cs b/Engine/Source/Programs/Shared/EpicGames.Horde/Storage/Clients/HttpStorageClient.cs
--- a/Engine/Source/Programs/Shared/EpicGames.Horde/Storage/Clients/HttpStorageClient.cs
+++ b/Engine/Source/Programs/Shared/EpicGames.Horde/Storage/Clients/HttpStorageClient.cs
@@ -280,18 +280,18 @@ namespace EpicGames.Horde.Storage.Clients
 	{
 		readonly IHttpClientFactory _httpClientFactory;
 		readonly StorageBackendCache _backendCache;
-		readonly BundleReaderCache _readerCache;
+		readonly BundleCache _bundleCache;
 		readonly ILogger<HttpStorageBackend> _backendLogger;
 		readonly ILogger<HttpStorageClient> _clientLogger;
 
 		/// <summary>
 		/// Constructor
 		/// </summary>
-		public HttpStorageClientFactory(IHttpClientFactory httpClientFactory, StorageBackendCache backendCache, BundleReaderCache readerCache, ILogger<HttpStorageBackend> backendLogger, ILogger<HttpStorageClient> clientLogger)
+		public HttpStorageClientFactory(IHttpClientFactory httpClientFactory, StorageBackendCache backendCache, BundleCache bundleCache, ILogger<HttpStorageBackend> backendLogger, ILogger<HttpStorageClient> clientLogger)
 		{
 			_httpClientFactory = httpClientFactory;
 			_backendCache = backendCache;
-			_readerCache = readerCache;
+			_bundleCache = bundleCache;
 			_backendLogger = backendLogger;
 			_clientLogger = clientLogger;
 		}
@@ -320,7 +320,7 @@ namespace EpicGames.Horde.Storage.Clients
 			}
 
 			HttpStorageClient client = new HttpStorageClient(basePath, CreateClient, backend, _clientLogger);
-			return new BundleStorageClient(client, _readerCache, _clientLogger);
+			return new BundleStorageClient(client, _bundleCache, _clientLogger);
 		}
 
 		/// <summary>
diff --git a/Engine/Source/Programs/UnrealBuildTool/Artifacts/HordeStorageArtifactCache.cs b/Engine/Source/Programs/UnrealBuildTool/Artifacts/HordeStorageArtifactCache.cs
--- a/Engine/Source/Programs/UnrealBuildTool/Artifacts/HordeStorageArtifactCache.cs
+++ b/Engine/Source/Programs/UnrealBuildTool/Artifacts/HordeStorageArtifactCache.cs
@@ -526,7 +526,7 @@ namespace UnrealBuildTool.Artifacts
 				}
 				Directory.CreateDirectory(directory.FullName);
 
-				_store = BundleStorageClient.CreateFromDirectory(directory, BundleReaderCache.None, logger);
+				_store = BundleStorageClient.CreateFromDirectory(directory, BundleCache.None, logger);
 
 				State = ArtifactCacheState.Available;
 				return State;
