Horde: Support nested blob handles, and return separate identifiers from each.

diff --git a/Engine/Source/Programs/Horde/Horde.Server/Storage/StorageService.cs b/Engine/Source/Programs/Horde/Horde.Server/Storage/StorageService.cs
--- a/Engine/Source/Programs/Horde/Horde.Server/Storage/StorageService.cs
+++ b/Engine/Source/Programs/Horde/Horde.Server/Storage/StorageService.cs
@@ -4,7 +4,6 @@ using System;
 using System.Buffers;
 using System.Collections.Generic;
 using System.ComponentModel;
-using System.Diagnostics.CodeAnalysis;
 using System.IO;
 using System.Linq;
 using System.Threading;
@@ -57,6 +56,8 @@ namespace Horde.Server.Storage
 			readonly string _path;
 			readonly Tracer _tracer;
 
+			public override BlobHandle? Outer => null;
+
 			public LeafBlobHandle(IStorageBackend backend, string path, Tracer tracer)
 			{
 				_backend = backend;
@@ -94,9 +95,9 @@ namespace Horde.Server.Storage
 			}
 
 			/// <inheritdoc/>
-			public override bool TryGetLocator([NotNullWhen(true)] out BlobLocator locator)
+			public override bool TryAppendIdentifier(Utf8StringBuilder builder)
 			{
-				locator = new BlobLocator(_path);
+				builder.Append(_path);
 				return true;
 			}
 
@@ -152,9 +153,9 @@ namespace Horde.Server.Storage
 			/// <inheritdoc/>
 			public BlobHandle CreateBlobHandle(BlobLocator locator)
 			{
-				if (locator.TryUnwrapFull(out BlobLocator outer, out Utf8String fragment))
+				if (locator.TryUnwrap(out BlobLocator baseLocator, out Utf8String fragment))
 				{
-					return new BlobFragmentHandle(new LeafBlobHandle(_backend, outer.ToString(), _tracer), fragment);
+					return new BlobFragmentHandle(new LeafBlobHandle(_backend, baseLocator.ToString(), _tracer), fragment);
 				}
 				else
 				{
@@ -695,7 +696,7 @@ namespace Horde.Server.Storage
 							foreach (BlobHandle import in imports)
 							{
 								BlobLocator importLocator = import.GetLocator();
-								string importPath = importLocator.Outermost.ToString();
+								string importPath = importLocator.BaseLocator.ToString();
 
 								FilterDefinition<BlobInfo> filter = Builders<BlobInfo>.Filter.Expr(x => x.NamespaceId == blobInfo.NamespaceId && x.Path == importPath);
 								UpdateDefinition<BlobInfo> update = Builders<BlobInfo>.Update.SetOnInsert(x => x.Imports, null);
@@ -733,8 +734,8 @@ namespace Horde.Server.Storage
 		{
 			BlobLocator locator = target.GetLocator();
 
-			string blobPath = locator.Outermost.ToString();
-			string blobFragment = locator.OutermostFragment.ToString();
+			string blobPath = locator.BaseLocator.ToString();
+			string blobFragment = locator.Fragment.ToString();
 
 			BlobInfo? blobInfo = await _blobCollection.Find(x => x.NamespaceId == namespaceId && x.Path == blobPath).FirstOrDefaultAsync(cancellationToken);
 			if (blobInfo == null)
@@ -763,8 +764,8 @@ namespace Horde.Server.Storage
 		{
 			BlobLocator locator = target.GetLocator();
 
-			string blobPath = locator.Outermost.ToString();
-			string blobFragment = locator.OutermostFragment.ToString();
+			string blobPath = locator.BaseLocator.ToString();
+			string blobFragment = locator.Fragment.ToString();
 
 			FilterDefinition<BlobInfo> filter = Builders<BlobInfo>.Filter.Expr(x => x.NamespaceId == namespaceId && x.Path == blobPath);
 			UpdateDefinition<BlobInfo> update = Builders<BlobInfo>.Update.PullFilter(x => x.Aliases, Builders<AliasInfo>.Filter.Expr(x => x.Name == name && x.Fragment == blobFragment));
@@ -921,7 +922,7 @@ namespace Horde.Server.Storage
 		/// <inheritdoc/>
 		async Task WriteRefAsync(NamespaceId namespaceId, RefName name, BlobLocator target, RefOptions? options = null, CancellationToken cancellationToken = default)
 		{
-			string path = target.Outermost.ToString();
+			string path = target.BaseLocator.ToString();
 
 			BlobInfo? newBlobInfo = await _blobCollection.Find(x => x.NamespaceId == namespaceId && x.Path == path).FirstOrDefaultAsync(cancellationToken);
 			if (newBlobInfo == null)
diff --git a/Engine/Source/Programs/Shared/EpicGames.Core/Utf8StringBuilder.cs b/Engine/Source/Programs/Shared/EpicGames.Core/Utf8StringBuilder.cs
--- a/Engine/Source/Programs/Shared/EpicGames.Core/Utf8StringBuilder.cs
+++ b/Engine/Source/Programs/Shared/EpicGames.Core/Utf8StringBuilder.cs
@@ -65,6 +65,24 @@ namespace EpicGames.Core
 			_writer.WriteUInt8(ch);
 		}
 
+		/// <summary>
+		/// Append a character to the end of this builder
+		/// </summary>
+		/// <param name="ch">Character to append</param>
+		public void Append(char ch)
+		{
+			if (ch < 127)
+			{
+				_writer.WriteUInt8((byte)ch);
+			}
+			else
+			{
+				Span<char> buffer = stackalloc char[1];
+				buffer[0] = ch;
+				Append(buffer);
+			}
+		}
+
 		/// <summary>
 		/// Appends a string to the end of this builder
 		/// </summary>
@@ -131,10 +149,10 @@ namespace EpicGames.Core
 		/// Appends a string to the end of this builder
 		/// </summary>
 		/// <param name="text">Text to append</param>
-		public void Append(string text)
+		public void Append(ReadOnlySpan<char> text)
 		{
 			Span<byte> span = _writer.GetSpanAndAdvance(Encoding.UTF8.GetByteCount(text));
-			Encoding.UTF8.GetBytes(text.AsSpan(), span);
+			Encoding.UTF8.GetBytes(text, span);
 		}
 
 		/// <summary>
diff --git a/Engine/Source/Programs/Shared/EpicGames.Horde.Tests/BlobDataTests.cs b/Engine/Source/Programs/Shared/EpicGames.Horde.Tests/BlobDataTests.cs
--- a/Engine/Source/Programs/Shared/EpicGames.Horde.Tests/BlobDataTests.cs
+++ b/Engine/Source/Programs/Shared/EpicGames.Horde.Tests/BlobDataTests.cs
@@ -3,7 +3,6 @@
 using EpicGames.Core;
 using System;
 using System.Collections.Generic;
-using System.Diagnostics.CodeAnalysis;
 using System.Threading;
 using System.Threading.Tasks;
 using EpicGames.Horde.Storage;
@@ -18,11 +17,14 @@ namespace EpicGames.Horde.Tests
 		{
 			readonly BlobLocator _locator;
 
+			public override BlobHandle? Outer => null;
+
 			public DummyHandle(string locator) => _locator = new BlobLocator(new Utf8String(locator));
 			public override ValueTask<BlobData> ReadBlobDataAsync(CancellationToken cancellationToken = default) => throw new NotImplementedException();
-			public override bool TryGetLocator([NotNullWhen(true)] out BlobLocator locator)
+
+			public override bool TryAppendIdentifier(Utf8StringBuilder builder)
 			{
-				locator = _locator;
+				builder.Append(_locator.Path);
 				return true;
 			}
 
diff --git a/Engine/Source/Programs/Shared/EpicGames.Horde.Tests/BundleStoreTests.cs b/Engine/Source/Programs/Shared/EpicGames.Horde.Tests/BundleStoreTests.cs
--- a/Engine/Source/Programs/Shared/EpicGames.Horde.Tests/BundleStoreTests.cs
+++ b/Engine/Source/Programs/Shared/EpicGames.Horde.Tests/BundleStoreTests.cs
@@ -64,7 +64,7 @@ namespace EpicGames.Horde.Tests
 			TextNode node = new TextNode("Hello world");
 			BlobHandle handle = await writer.FlushAsync(node, CancellationToken.None);
 
-			BlobHandle bundleHandle = store.CreateBlobHandle(handle.GetLocator().Outermost);
+			BlobHandle bundleHandle = store.CreateBlobHandle(handle.GetLocator().BaseLocator);
 			using BlobData blobData = await bundleHandle.ReadBlobDataAsync();
 
 			return blobData.Data.ToArray();
diff --git a/Engine/Source/Programs/Shared/EpicGames.Horde.Tests/BundleTests.cs b/Engine/Source/Programs/Shared/EpicGames.Horde.Tests/BundleTests.cs
--- a/Engine/Source/Programs/Shared/EpicGames.Horde.Tests/BundleTests.cs
+++ b/Engine/Source/Programs/Shared/EpicGames.Horde.Tests/BundleTests.cs
@@ -220,7 +220,7 @@ namespace EpicGames.Horde.Tests
 
 			// Check the ref
 			BlobHandle refTarget =  await store.ReadRefTargetAsync(refName);
-			BlobHandle bundleTarget = store.CreateBlobHandle(refTarget.GetLocator().Outermost);
+			BlobHandle bundleTarget = store.CreateBlobHandle(refTarget.GetLocator().BaseLocator);
 			using BlobData bundleData = await bundleTarget.ReadBlobDataAsync();
 
 			BundleHeader bundleHeader = BundleHeader.Read(bundleData.Data);
diff --git a/Engine/Source/Programs/Shared/EpicGames.Horde/Storage/BlobHandle.cs b/Engine/Source/Programs/Shared/EpicGames.Horde/Storage/BlobHandle.cs
--- a/Engine/Source/Programs/Shared/EpicGames.Horde/Storage/BlobHandle.cs
+++ b/Engine/Source/Programs/Shared/EpicGames.Horde/Storage/BlobHandle.cs
@@ -59,24 +59,9 @@ namespace EpicGames.Horde.Storage
 		}
 
 		/// <summary>
-		/// Gets a path to this blob that can be used to describe blob references over the wire.
-		/// </summary>
-		public BlobLocator GetLocator()
-		{
-			BlobLocator locator;
-			if (!TryGetLocator(out locator))
-			{
-				throw new InvalidOperationException("Blob has not yet been written to storage");
-			}
-			return locator;
-		}
-
-		/// <summary>
-		/// Attempt to get a path for this blob.
+		/// For a blob nested within another blob, gets a handle to the containing blob (eg. For a bundle node, will return the packet. For a bundle packet, will return the bundle. For a bundle or other non-nested blob, returns null.)
 		/// </summary>
-		/// <param name="locator">Receives the blob path on success.</param>
-		/// <returns>True if a path was available, false if the blob has not yet been flushed to storage.</returns>
-		public abstract bool TryGetLocator([NotNullWhen(true)] out BlobLocator locator);
+		public abstract BlobHandle? Outer { get; }
 
 		/// <summary>
 		/// Gets the type of this blob
@@ -146,11 +131,18 @@ namespace EpicGames.Horde.Storage
 		/// </summary>
 		public virtual ValueTask FlushAsync(CancellationToken cancellationToken = default) => new ValueTask();
 
+		/// <summary>
+		/// Gets an identifier for this blob, relative to its outer
+		/// </summary>
+		/// <param name="builder">Builder for appending the identifier to</param>
+		/// <returns>True if an identifier was returned, false otherwise</returns>
+		public abstract bool TryAppendIdentifier(Utf8StringBuilder builder);
+
 		/// <inheritdoc/>
 		public override string ToString()
 		{
 			BlobLocator locator;
-			if (TryGetLocator(out locator))
+			if (this.TryGetLocator(out locator))
 			{
 				return locator.ToString();
 			}
@@ -170,9 +162,9 @@ namespace EpicGames.Horde.Storage
 	public class BlobFragmentHandle : BlobHandle
 	{
 		/// <summary>
-		/// Handle to the inner blob
+		/// Handle to the outer blob
 		/// </summary>
-		public BlobHandle Inner { get; }
+		public override BlobHandle Outer { get; }
 
 		/// <summary>
 		/// The fragment portion of the handle
@@ -182,26 +174,17 @@ namespace EpicGames.Horde.Storage
 		/// <summary>
 		/// Constructor
 		/// </summary>
-		public BlobFragmentHandle(BlobHandle inner, Utf8String fragment)
+		public BlobFragmentHandle(BlobHandle outer, Utf8String fragment)
 		{
-			Inner = inner;
+			Outer = outer;
 			Fragment = fragment;
 		}
 
 		/// <inheritdoc/>
-		public override bool TryGetLocator([NotNullWhen(true)] out BlobLocator locator)
+		public override bool TryAppendIdentifier(Utf8StringBuilder builder)
 		{
-			BlobLocator innerLocator;
-			if (Inner.TryGetLocator(out innerLocator))
-			{
-				locator = new BlobLocator(innerLocator, Fragment.Span);
-				return true;
-			}
-			else
-			{
-				locator = default;
-				return false;
-			}
+			builder.Append(Fragment);
+			return true;
 		}
 
 		/// <inheritdoc/>
@@ -211,10 +194,10 @@ namespace EpicGames.Horde.Storage
 		}
 
 		/// <inheritdoc/>
-		public override bool Equals(object? obj) => obj is BlobFragmentHandle other && Inner == other.Inner && Fragment == other.Fragment;
+		public override bool Equals(object? obj) => obj is BlobFragmentHandle other && Outer.Equals(other.Outer) && Fragment == other.Fragment;
 
 		/// <inheritdoc/>
-		public override int GetHashCode() => HashCode.Combine(Inner, Fragment);
+		public override int GetHashCode() => HashCode.Combine(Outer, Fragment);
 	}
 
 	/// <summary>
@@ -222,6 +205,86 @@ namespace EpicGames.Horde.Storage
 	/// </summary>
 	public static class BlobHandleExtensions
 	{
+		/// <summary>
+		/// Gets a path to this blob that can be used to describe blob references over the wire.
+		/// </summary>
+		/// <param name="handle">Handle to query</param>
+		public static BlobLocator GetLocator(this BlobHandle handle)
+		{
+			BlobLocator locator;
+			if (!TryGetLocator(handle, out locator))
+			{
+				throw new InvalidOperationException("Blob has not yet been written to storage");
+			}
+			return locator;
+		}
+
+		/// <summary>
+		/// Attempt to get a path for this blob.
+		/// </summary>
+		/// <param name="handle">Handle to query</param>
+		/// <param name="locator">Receives the blob path on success.</param>
+		/// <returns>True if a path was available, false if the blob has not yet been flushed to storage.</returns>
+		public static bool TryGetLocator(this BlobHandle handle, [NotNullWhen(true)] out BlobLocator locator)
+		{
+			Utf8StringBuilder builder = new Utf8StringBuilder();
+			if (AppendLocator(handle, builder))
+			{
+				locator = new BlobLocator(builder.ToUtf8String());
+				return true;
+			}
+			else
+			{
+				locator = default;
+				return false;
+			}
+		}
+
+		/// <summary>
+		/// Builds a full locator for a blob by traversing the outer chain
+		/// </summary>
+		static bool AppendLocator(BlobHandle handle, Utf8StringBuilder builder)
+		{
+			BlobHandle? outer = handle.Outer;
+			if (outer != null)
+			{
+				if (!AppendLocator(outer, builder))
+				{
+					return false;
+				}
+				if (outer.Outer == null)
+				{
+					builder.Append('#');
+				}
+				else
+				{
+					builder.Append('&');
+				}
+			}
+			return handle.TryAppendIdentifier(builder);
+		}
+
+		/// <summary>
+		/// Gets an identifier for a blob handle
+		/// </summary>
+		/// <param name="handle">Handle to the blob</param>
+		/// <param name="fragment">On success, receives the fragment path</param>
+		/// <returns>True if the handle has an identifier</returns>
+		public static bool TryGetIdentifier(this BlobHandle handle, out Utf8String fragment)
+		{
+			Utf8StringBuilder builder = new Utf8StringBuilder();
+			if (handle.TryAppendIdentifier(builder))
+			{
+				fragment = builder.ToUtf8String();
+				return true;
+			}
+			else
+			{
+				fragment = default;
+				return false;
+			}
+		}
+
 		/// <summary>
 		/// Helper method for awaiting a handle and returning its locator
 		/// </summary>
diff --git a/Engine/Source/Programs/Shared/EpicGames.Horde/Storage/BlobLocator.cs b/Engine/Source/Programs/Shared/EpicGames.Horde/Storage/BlobLocator.cs
--- a/Engine/Source/Programs/Shared/EpicGames.Horde/Storage/BlobLocator.cs
+++ b/Engine/Source/Programs/Shared/EpicGames.Horde/Storage/BlobLocator.cs
@@ -11,7 +11,7 @@ using EpicGames.Core;
 namespace EpicGames.Horde.Storage
 {
 	/// <summary>
-	/// Identifier for a blob within a particular namespace. Locators can be nested to allow forming a hierarchy of storage clients using the <see cref="Outer"/> and <see cref="Fragment"/> properties.
+	/// Identifier for a blob within a particular namespace.
 	/// </summary>
 	[JsonSchemaString]
 	[TypeConverter(typeof(BlobLocatorTypeConverter))]
@@ -67,9 +67,9 @@ namespace EpicGames.Horde.Storage
 		public bool IsValid() => !_path.IsEmpty;
 
 		/// <summary>
-		/// The outermost blob locator
+		/// The base blob locator
 		/// </summary>
-		public BlobLocator Outermost
+		public BlobLocator BaseLocator
 		{
 			get
 			{
@@ -78,30 +78,6 @@ namespace EpicGames.Horde.Storage
 			}
 		}
 
-		/// <summary>
-		/// Fragment within the base blob
-		/// </summary>
-		public Utf8String OutermostFragment
-		{
-			get
-			{
-				int hashIdx = _path.IndexOf('#');
-				return (hashIdx == -1) ? Utf8String.Empty : _path.Slice(hashIdx + 1);
-			}
-		}
-
-		/// <summary>
-		/// The containing blob locator
-		/// </summary>
-		public BlobLocator Outer
-		{
-			get
-			{
-				int hashIdx = _path.LastIndexOf('#');
-				return (hashIdx == -1) ? new BlobLocator(Utf8String.Empty) : new BlobLocator(_path.Slice(0, hashIdx));
-			}
-		}
-
 		/// <summary>
 		/// Fragment within the base blob
 		/// </summary>
@@ -109,7 +85,7 @@ namespace EpicGames.Horde.Storage
 		{
 			get
 			{
-				int hashIdx = _path.LastIndexOf('#');
+				int hashIdx = _path.IndexOf('#');
 				return (hashIdx == -1) ? Utf8String.Empty : _path.Slice(hashIdx + 1);
 			}
 		}
@@ -120,40 +96,23 @@ namespace EpicGames.Horde.Storage
 		public bool CanUnwrap() => _path.IndexOf('#') != -1;
 
 		/// <summary>
-		/// Split this locator into 
+		/// Split this locator into a locator and fragment
 		/// </summary>
-		/// <param name="outer"></param>
-		/// <param name="fragment"></param>
-		/// <returns></returns>
-		public bool TryUnwrap(out BlobLocator outer, out Utf8String fragment)
+		/// <param name="baseLocator">Receives the base blob locator</param>
+		/// <param name="fragment">Receives the blob fragment</param>
+		/// <returns>True if the locator was unwrapped, false otherwise</returns>
+		public bool TryUnwrap(out BlobLocator baseLocator, out Utf8String fragment)
 		{
 			int hashIdx = _path.LastIndexOf('#');
-			return Split(hashIdx, out outer, out fragment);
-		}
-
-		/// <summary>
-		/// Split this locator into the outermost fragment and locator
-		/// </summary>
-		/// <param name="outer">The outermost blob locator</param>
-		/// <param name="fragment">The corresponding fragment</param>
-		/// <returns>True if </returns>
-		public bool TryUnwrapFull(out BlobLocator outer, out Utf8String fragment)
-		{
-			int hashIdx = _path.IndexOf('#');
-			return Split(hashIdx, out outer, out fragment);
-		}
-
-		bool Split(int hashIdx, out BlobLocator outer, out Utf8String fragment)
-		{
 			if (hashIdx == -1)
 			{
-				outer = this;
+				baseLocator = this;
 				fragment = default;
 				return false;
 			}
 			else
 			{
-				outer = new BlobLocator(_path.Slice(0, hashIdx));
+				baseLocator = new BlobLocator(_path.Slice(0, hashIdx));
 				fragment = _path.Slice(hashIdx + 1);
 				return true;
 			}
diff --git a/Engine/Source/Programs/Shared/EpicGames.Horde/Storage/Bundles/V1/BundleReader.cs b/Engine/Source/Programs/Shared/EpicGames.Horde/Storage/Bundles/V1/BundleReader.cs
--- a/Engine/Source/Programs/Shared/EpicGames.Horde/Storage/Bundles/V1/BundleReader.cs
+++ b/Engine/Source/Programs/Shared/EpicGames.Horde/Storage/Bundles/V1/BundleReader.cs
@@ -558,8 +558,9 @@ namespace EpicGames.Horde.Storage.Bundles.V1
 				{
 					importBlob = header.Imports[reference.ImportIdx];
 				}
+
 				Debug.Assert(importBlob.IsValid());
-				refs.Add(new FlushedNodeHandle(this, importBlob, reference.NodeIdx));
+				refs.Add(new FlushedNodeHandle(this, importBlob, _store.CreateBlobHandle(importBlob), reference.NodeIdx));
 			}
 
 			ReadOnlyMemory<byte> nodeData = ReadOnlyMemory<byte>.Empty;
diff --git a/Engine/Source/Programs/Shared/EpicGames.Horde/Storage/Bundles/V1/BundleWriter.cs b/Engine/Source/Programs/Shared/EpicGames.Horde/Storage/Bundles/V1/BundleWriter.cs
--- a/Engine/Source/Programs/Shared/EpicGames.Horde/Storage/Bundles/V1/BundleWriter.cs
+++ b/Engine/Source/Programs/Shared/EpicGames.Horde/Storage/Bundles/V1/BundleWriter.cs
@@ -5,9 +5,9 @@ using System.Buffers;
 using System.Buffers.Text;
 using System.Collections.Generic;
 using System.Diagnostics;
-using System.Diagnostics.CodeAnalysis;
 using System.Linq;
 using System.Runtime.CompilerServices;
+using System.Text;
 using System.Threading;
 using System.Threading.Tasks;
 using EpicGames.Core;
@@ -25,32 +25,42 @@ namespace EpicGames.Horde.Storage.Bundles.V1
 		public BlobLocator BundleLocator { get; }
 		public int ExportIdx { get; }
 
+		/// <inheritdoc/>
+		public override BlobHandle Outer { get; }
+
 		/// <summary>
 		/// Constructor
 		/// </summary>
-		public FlushedNodeHandle(BundleReader reader, BlobLocator bundleLocator, int exportIdx)
+		public FlushedNodeHandle(BundleReader reader, BlobLocator bundleLocator, BlobHandle bundleHandle, int exportIdx)
 		{
+			Debug.Assert(!bundleLocator.CanUnwrap());
+
 			_reader = reader;
 			BundleLocator = bundleLocator;
+			Outer = bundleHandle;
 			ExportIdx = exportIdx;
 		}
 
-		public static FlushedNodeHandle FromBlobLocator(BundleReader reader, BlobLocator locator)
+		public FlushedNodeHandle(BundleReader reader, BlobLocator bundleLocator, BlobHandle bundleHandle, ReadOnlySpan<byte> fragment)
 		{
-			if (locator.TryUnwrap(out BlobLocator bundleLocator, out Utf8String fragment) && Utf8Parser.TryParse(fragment, out int exportIdx, out int bytesConsumed) && bytesConsumed == fragment.Length)
-			{
-				return new FlushedNodeHandle(reader, bundleLocator, exportIdx);
-			}
-			else
+			Debug.Assert(!bundleLocator.CanUnwrap());
+
+			_reader = reader;
+			BundleLocator = bundleLocator;
+			Outer = bundleHandle;
+
+			if (!Utf8Parser.TryParse(fragment, out int exportIdx, out int bytesConsumed) || bytesConsumed != fragment.Length)
 			{
-				throw new ArgumentException($"Locator {locator} is not a valid bundle node");
+				throw new ArgumentException($"Fragment {Encoding.UTF8.GetString(fragment)} is not valid for a bundle export");
 			}
+
+			ExportIdx = exportIdx;
 		}
 
 		/// <inheritdoc/>
-		public override bool TryGetLocator([NotNullWhen(true)] out BlobLocator locator)
+		public override bool TryAppendIdentifier(Utf8StringBuilder builder)
 		{
-			locator = new BlobLocator(BundleLocator, ExportIdx.ToString());
+			builder.Append(ExportIdx);
 			return true;
 		}
 
@@ -93,6 +103,8 @@ namespace EpicGames.Horde.Storage.Bundles.V1
 			public PendingBundle? PendingBundle => _pendingBundle;
 			public FlushedNodeHandle? FlushedNodeHandle => _flushedHandle;
 
+			public override BlobHandle? Outer => (_flushedHandle != null) ? _flushedHandle.Outer : throw new NotSupportedException();
+
 			public PendingNode(BundleReader reader, BlobType blobType, int packet, int offset, int length, IReadOnlyList<BlobHandle> refs, IReadOnlyList<AliasInfo> aliases, PendingBundle pendingBundle)
 			{
 				_reader = reader;
@@ -108,20 +120,9 @@ namespace EpicGames.Horde.Storage.Bundles.V1
 			}
 
 			/// <inheritdoc/>
-			public override bool TryGetLocator([NotNullWhen(true)] out BlobLocator locator)
+			public override bool TryAppendIdentifier(Utf8StringBuilder builder)
 			{
-				if (_flushedHandle != null && _flushedHandle.TryGetLocator(out BlobLocator flushedLocator))
-				{
-					locator = flushedLocator;
-					return true;
-				}
-				else
-				{
-					Console.WriteLine($"Unable to get locator for node {BlobType}:{Packet}:{Offset}:{Length}. Handle is {FlushedNodeHandle}. Bundle is {PendingBundle}.");
-
-					locator = default;
-					return false;
-				}
+				return _flushedHandle?.TryAppendIdentifier(builder) ?? false;
 			}
 
 			public void MarkAsWritten(FlushedNodeHandle flushedHandle)
@@ -439,7 +440,7 @@ namespace EpicGames.Horde.Storage.Bundles.V1
 					{
 						PendingNode node = _queue[idx];
 
-						FlushedNodeHandle flushedHandle = new FlushedNodeHandle(_treeReader, locator, idx);
+						FlushedNodeHandle flushedHandle = new FlushedNodeHandle(_treeReader, locator, handle, idx);
 						_queue[idx].MarkAsWritten(flushedHandle);
 
 						foreach (AliasInfo alias in node.Aliases)
diff --git a/Engine/Source/Programs/Shared/EpicGames.Horde/Storage/Clients/BundleStorageClient.cs b/Engine/Source/Programs/Shared/EpicGames.Horde/Storage/Clients/BundleStorageClient.cs
--- a/Engine/Source/Programs/Shared/EpicGames.Horde/Storage/Clients/BundleStorageClient.cs
+++ b/Engine/Source/Programs/Shared/EpicGames.Horde/Storage/Clients/BundleStorageClient.cs
@@ -2,7 +2,6 @@
 
 using System;
 using System.Collections.Generic;
-using System.Diagnostics.CodeAnalysis;
 using System.IO;
 using System.Threading;
 using System.Threading.Tasks;
@@ -26,6 +25,9 @@ namespace EpicGames.Horde.Storage.Clients
 			readonly BlobLocator _locator;
 			List<BlobHandle>? _refs;
 
+			/// <inheritdoc/>
+			public override BlobHandle? Outer => null;
+
 			/// <summary>
 			/// Constructor
 			/// </summary>
@@ -73,9 +75,9 @@ namespace EpicGames.Horde.Storage.Clients
 			}
 
 			/// <inheritdoc/>
-			public override bool TryGetLocator([NotNullWhen(true)] out BlobLocator locator)
+			public override bool TryAppendIdentifier(Utf8StringBuilder builder)
 			{
-				locator = _locator;
+				builder.Append(_locator.Path);
 				return true;
 			}
 
@@ -196,9 +198,9 @@ namespace EpicGames.Horde.Storage.Clients
 		/// <inheritdoc/>
 		public BlobHandle CreateBlobHandle(BlobLocator locator)
 		{
-			if (locator.CanUnwrap())
+			if (locator.TryUnwrap(out BlobLocator baseLocator, out Utf8String fragment))
 			{
-				return Bundles.V1.FlushedNodeHandle.FromBlobLocator(_bundleReader, locator);
+				return new Bundles.V1.FlushedNodeHandle(_bundleReader, baseLocator, new BundleHandle(this, baseLocator), fragment.Span);
 			}
 			else
 			{
diff --git a/Engine/Source/Programs/Shared/EpicGames.Horde/Storage/Clients/HttpStorageClient.cs b/Engine/Source/Programs/Shared/EpicGames.Horde/Storage/Clients/HttpStorageClient.cs
--- a/Engine/Source/Programs/Shared/EpicGames.Horde/Storage/Clients/HttpStorageClient.cs
+++ b/Engine/Source/Programs/Shared/EpicGames.Horde/Storage/Clients/HttpStorageClient.cs
@@ -13,7 +13,6 @@ using Microsoft.Extensions.Logging;
 using EpicGames.Horde.Storage.Backends;
 using System.IO;
 using System.Collections.Generic;
-using System.Diagnostics.CodeAnalysis;
 
 namespace EpicGames.Horde.Storage.Clients
 {
@@ -24,27 +23,30 @@ namespace EpicGames.Horde.Storage.Clients
 	{
 		class Handle : BlobHandle
 		{
-			readonly HttpStorageClient _outer;
+			readonly HttpStorageClient _httpStorageClient;
 			readonly BlobLocator _locator;
 
-			public Handle(HttpStorageClient outer, BlobLocator locator)
+			/// <inheritdoc/>
+			public override BlobHandle? Outer => null;
+
+			public Handle(HttpStorageClient httpStorageClient, BlobLocator locator)
 			{
-				_outer = outer;
+				_httpStorageClient = httpStorageClient;
 				_locator = locator;
 			}
 
 			/// <inheritdoc/>
 			public override Task<Stream> OpenAsync(int offset = 0, int? length = null, CancellationToken cancellationToken = default)
-				=> _outer._backend.OpenAsync(_locator.ToString(), offset, length, cancellationToken);
+				=> _httpStorageClient._backend.OpenAsync(_locator.ToString(), offset, length, cancellationToken);
 
 			/// <inheritdoc/>
 			public override ValueTask<BlobData> ReadBlobDataAsync(CancellationToken cancellationToken = default)
-				=> _outer.ReadBlobAsync(_locator, cancellationToken);
+				=> _httpStorageClient.ReadBlobAsync(_locator, cancellationToken);
 
 			/// <inheritdoc/>
-			public override bool TryGetLocator([NotNullWhen(true)] out BlobLocator locator)
+			public override bool TryAppendIdentifier(Utf8StringBuilder builder)
 			{
-				locator = _locator;
+				builder.Append(_locator.Path);
 				return true;
 			}
 
@@ -85,7 +87,7 @@ namespace EpicGames.Horde.Storage.Clients
 		/// <inheritdoc/>
 		public BlobHandle CreateBlobHandle(BlobLocator locator)
 		{
-			if (locator.TryUnwrapFull(out BlobLocator outer, out Utf8String fragment))
+			if (locator.TryUnwrap(out BlobLocator outer, out Utf8String fragment))
 			{
 				return new BlobFragmentHandle(new Handle(this, outer), fragment);
 			}
diff --git a/Engine/Source/Programs/Shared/EpicGames.Horde/Storage/Clients/KeyValueStorageClient.cs b/Engine/Source/Programs/Shared/EpicGames.Horde/Storage/Clients/KeyValueStorageClient.cs
--- a/Engine/Source/Programs/Shared/EpicGames.Horde/Storage/Clients/KeyValueStorageClient.cs
+++ b/Engine/Source/Programs/Shared/EpicGames.Horde/Storage/Clients/KeyValueStorageClient.cs
@@ -2,7 +2,6 @@
 
 using System;
 using System.Collections.Generic;
-using System.Diagnostics.CodeAnalysis;
 using System.IO;
 using System.Threading;
 using System.Threading.Tasks;
@@ -17,22 +16,28 @@ namespace EpicGames.Horde.Storage.Clients
 	{
 		class Handle : BlobHandle
 		{
-			readonly KeyValueStorageClient _outer;
+			readonly KeyValueStorageClient _keyValueStorageClient;
 			readonly BlobLocator _locator;
 
-			public Handle(KeyValueStorageClient outer, BlobLocator locator)
+			/// <inheritdoc/>
+			public override BlobHandle? Outer => null;
+
+			/// <summary>
+			/// Constructor
+			/// </summary>
+			public Handle(KeyValueStorageClient keyValueStorageClient, BlobLocator locator)
 			{
-				_outer = outer;
+				_keyValueStorageClient = keyValueStorageClient;
 				_locator = locator;
 			}
 
 			/// <inheritdoc/>
-			public override ValueTask<BlobData> ReadBlobDataAsync(CancellationToken cancellationToken = default) => _outer.ReadBlobAsync(_locator, cancellationToken);
+			public override ValueTask<BlobData> ReadBlobDataAsync(CancellationToken cancellationToken = default) => _keyValueStorageClient.ReadBlobAsync(_locator, cancellationToken);
 
 			/// <inheritdoc/>
-			public override bool TryGetLocator([NotNullWhen(true)] out BlobLocator locator)
+			public override bool TryAppendIdentifier(Utf8StringBuilder builder)
 			{
-				locator = _locator;
+				builder.Append(_locator.Path);
 				return true;
 			}
 
@@ -65,9 +70,9 @@ namespace EpicGames.Horde.Storage.Clients
 		/// <inheritdoc/>
 		public virtual BlobHandle CreateBlobHandle(BlobLocator locator)
 		{
-			if (locator.TryUnwrapFull(out BlobLocator outer, out Utf8String fragment))
+			if (locator.TryUnwrap(out BlobLocator baseLocator, out Utf8String fragment))
 			{
-				return new BlobFragmentHandle(new Handle(this, outer), fragment);
+				return new BlobFragmentHandle(new Handle(this, baseLocator), fragment);
 			}
 			else
 			{
diff --git a/Engine/Source/Programs/Shared/EpicGames.Horde/Storage/IStorageWriter.cs b/Engine/Source/Programs/Shared/EpicGames.Horde/Storage/IStorageWriter.cs
--- a/Engine/Source/Programs/Shared/EpicGames.Horde/Storage/IStorageWriter.cs
+++ b/Engine/Source/Programs/Shared/EpicGames.Horde/Storage/IStorageWriter.cs
@@ -163,19 +163,12 @@ namespace EpicGames.Horde.Storage
 			public BlobHandle? _inner;
 
 			/// <inheritdoc/>
-			public override bool TryGetLocator([NotNullWhen(true)] out BlobLocator blobId)
+			public override BlobHandle? Outer => _inner?.Outer;
+
+			/// <inheritdoc/>
+			public override bool TryAppendIdentifier(Utf8StringBuilder builder)
 			{
-				BlobLocator innerBlobId;
-				if (_inner != null && _inner.TryGetLocator(out innerBlobId))
-				{
-					blobId = innerBlobId;
-					return true;
-				}
-				else
-				{
-					blobId = default;
-					return false;
-				}
+				return _inner?.TryAppendIdentifier(builder) ?? false;
 			}
 
 			/// <inheritdoc/>
