EpicGames.Core: Add a utility class for managing a reference counted object, running a custom disposer when the refcount reaches zero. Also modify IReadOnlyMemoryOwner to support a custom disposer (which may be a ref-counted handle).

diff --git a/Engine/Source/Programs/Shared/EpicGames.Core/IReadOnlyMemoryOwner.cs b/Engine/Source/Programs/Shared/EpicGames.Core/IReadOnlyMemoryOwner.cs
--- a/Engine/Source/Programs/Shared/EpicGames.Core/IReadOnlyMemoryOwner.cs
+++ b/Engine/Source/Programs/Shared/EpicGames.Core/IReadOnlyMemoryOwner.cs
@@ -24,25 +24,31 @@ namespace EpicGames.Core
 		class DefaultReadOnlyMemoryOwner<T> : IReadOnlyMemoryOwner<T>
 		{
 			public ReadOnlyMemory<T> Memory { get; }
+			readonly IDisposable? _owner;
 
-			public DefaultReadOnlyMemoryOwner(ReadOnlyMemory<T> memory) => Memory = memory;
-			public void Dispose() { }
+			public DefaultReadOnlyMemoryOwner(ReadOnlyMemory<T> memory, IDisposable? owner)
+			{
+				Memory = memory;
+				_owner = owner;
+			}
+
+			public void Dispose() => _owner?.Dispose();
 		}
 
 		/// <summary>
 		/// Wrap an array as a <see cref="IReadOnlyMemoryOwner{T}"/>
 		/// </summary>
-		public static IReadOnlyMemoryOwner<T> Create<T>(T[] memory)
+		public static IReadOnlyMemoryOwner<T> Create<T>(T[] memory, IDisposable? owner = null)
 		{
-			return new DefaultReadOnlyMemoryOwner<T>(memory);
+			return new DefaultReadOnlyMemoryOwner<T>(memory, owner);
 		}
 
 		/// <summary>
 		/// Wrap a <see cref="ReadOnlyMemory{T}"/> as a <see cref="IReadOnlyMemoryOwner{T}"/>
 		/// </summary>
-		public static IReadOnlyMemoryOwner<T> Create<T>(ReadOnlyMemory<T> memory)
+		public static IReadOnlyMemoryOwner<T> Create<T>(ReadOnlyMemory<T> memory, IDisposable? owner = null)
 		{
-			return new DefaultReadOnlyMemoryOwner<T>(memory);
+			return new DefaultReadOnlyMemoryOwner<T>(memory, owner);
 		}
 
 		// Wraps a ReadOnlyMemoryOwner in a stream
diff --git a/Engine/Source/Programs/Shared/EpicGames.Core/RefCountedHandle.cs b/Engine/Source/Programs/Shared/EpicGames.Core/RefCountedHandle.cs
new file mode 100644
--- /dev/null
+++ b/Engine/Source/Programs/Shared/EpicGames.Core/RefCountedHandle.cs
@@ -0,0 +1,148 @@
+// Copyright Epic Games, Inc. All Rights Reserved.
+
+using System;
+using System.Threading;
+
+namespace EpicGames.Core
+{
+	/// <summary>
+	/// Disposable reference to an underlying resource. The underlying resource is disposed once the last reference to it is disposed.
+	/// </summary>
+	public interface IRefCountedHandle : IDisposable
+	{
+		/// <summary>
+		/// Increment the reference count and returns a new handle to it
+		/// </summary>
+		IRefCountedHandle AddRef();
+	}
+
+	/// <summary>
+	/// Disposable reference to an underlying resource.
+	/// </summary>
+	/// <typeparam name="T">Type of the target resource</typeparam>
+	public interface IRefCountedHandle<T> : IRefCountedHandle
+	{
+		/// <summary>
+		/// Target of the reference
+		/// </summary>
+		T Target { get; }
+
+		/// <inheritdoc cref="IRefCountedHandle.AddRef"/>
+		new IRefCountedHandle<T> AddRef();
+
+		/// <inheritdoc/>
+		IRefCountedHandle IRefCountedHandle.AddRef() => AddRef();
+	}
+
+	/// <summary>
+	/// Utility methods for creating ref counted handles
+	/// </summary>
+	public static class RefCountedHandle
+	{
+		/// <summary>
+		/// Creates a reference counted handle to a disposable resource.
+		/// </summary>
+		public static IRefCountedHandle<T> Create<T>(T target) where T : IDisposable
+		{
+			return new RefCountedHandle<T>(target, target);
+		}
+
+		/// <summary>
+		/// Creates a reference counted handle to a resource with an explicit disposer.
+		/// </summary>
+		public static IRefCountedHandle<T> Create<T>(T target, IDisposable? owner)
+		{
+			return new RefCountedHandle<T>(target, owner);
+		}
+	}
+
+	/// <summary>
+	/// Concrete implementation of <see cref="IRefCountedHandle{T}"/>
+	/// </summary>
+	public sealed class RefCountedHandle<T> : IRefCountedHandle<T>
+	{
+		readonly RefCountedDisposer? _refCountedDisposer;
+		T _target;
+
+		/// <inheritdoc/>
+		public T Target => _target ?? throw new ObjectDisposedException(typeof(RefCountedHandle<T>).Name);
+
+		/// <summary>
+		/// Constructor
+		/// </summary>
+		/// <param name="target">Target value</param>
+		/// <param name="owner">Owner of the target resource</param>
+		public RefCountedHandle(T target, IDisposable? owner)
+		{
+			_refCountedDisposer = new RefCountedDisposer(owner);
+			_target = target;
+		}
+
+		private RefCountedHandle(T value, RefCountedDisposer disposer)
+		{
+			_refCountedDisposer = disposer;
+			_target = value;
+		}
+
+		/// <inheritdoc/>
+		public void Dispose()
+		{
+			if (_refCountedDisposer != null)
+			{
+				_refCountedDisposer.Release();
+				_target = default!;
+			}
+		}
+
+		/// <inheritdoc/>
+		public IRefCountedHandle<T> AddRef()
+		{
+			if (_refCountedDisposer == null)
+			{
+				throw new ObjectDisposedException(typeof(RefCountedHandle<T>).Name);
+			}
+
+			_refCountedDisposer.AddRef();
+			return new RefCountedHandle<T>(Target, _refCountedDisposer);
+		}
+	}
+
+	/// <summary>
+	/// Tracks a ref count value and disposes of a resource once the reference count is zero
+	/// </summary>
+	public sealed class RefCountedDisposer
+	{
+		IDisposable? _owner;
+		int _value = 1;
+
+		/// <summary>
+		/// Constructor
+		/// </summary>
+		public RefCountedDisposer(IDisposable? owner, int initialValue = 1)
+		{
+			_owner = owner;
+			_value = initialValue;
+		}
+
+		/// <summary>
+		/// Increment the current value
+		/// </summary>
+		/// <returns>The incremented value</returns>
+		public int AddRef() => Interlocked.Increment(ref _value);
+
+		/// <summary>
+		/// Decrement the current value
+		/// </summary>
+		/// <returns>The decremented value</returns>
+		public int Release()
+		{
+			int count = Interlocked.Decrement(ref _value);
+			if (count == 0)
+			{
+				_owner?.Dispose();
+				_owner = null;
+			}
+			return count;
+		}
+	}
+}
