Horde: Change BlobHandle to be an interface so writer interfaces can also implement it.

diff --git a/Engine/Source/Programs/AutomationTool/BuildGraph/Tasks/DeployToolTask.cs b/Engine/Source/Programs/AutomationTool/BuildGraph/Tasks/DeployToolTask.cs
--- a/Engine/Source/Programs/AutomationTool/BuildGraph/Tasks/DeployToolTask.cs
+++ b/Engine/Source/Programs/AutomationTool/BuildGraph/Tasks/DeployToolTask.cs
@@ -150,7 +150,7 @@ namespace AutomationTool.Tasks
 				return httpClient;
 			}
 
-			BlobHandle handle;
+			IBlobHandle handle;
 
 			string basePath = $"api/v1/tools/{Parameters.Id}";
 			using HttpStorageBackend httpStorageBackend = new HttpStorageBackend(basePath, CreateHttpClient, Logger);
diff --git a/Engine/Source/Programs/Horde/Horde.Agent.Tests/JsonRpcLoggerTests.cs b/Engine/Source/Programs/Horde/Horde.Agent.Tests/JsonRpcLoggerTests.cs
--- a/Engine/Source/Programs/Horde/Horde.Agent.Tests/JsonRpcLoggerTests.cs
+++ b/Engine/Source/Programs/Horde/Horde.Agent.Tests/JsonRpcLoggerTests.cs
@@ -104,14 +104,14 @@ namespace Horde.Agent.Tests
 
 		class FakeJsonRpcLoggerBackend : JsonRpcAndStorageLogSink
 		{
-			public BlobHandle? Target { get; private set; }
+			public IBlobHandle? Target { get; private set; }
 
 			public FakeJsonRpcLoggerBackend(IRpcConnection connection, string logId, IJsonRpcLogSink inner, IStorageClient store, ILogger logger)
 				: base(connection, logId, inner, store, logger)
 			{
 			}
 
-			protected override Task UpdateLogAsync(BlobHandle target, int lineCount, bool complete, CancellationToken cancellationToken)
+			protected override Task UpdateLogAsync(IBlobHandle target, int lineCount, bool complete, CancellationToken cancellationToken)
 			{
 				Target = target;
 				return Task.CompletedTask;
diff --git a/Engine/Source/Programs/Horde/Horde.Agent/Utility/JsonRpcLogSink.cs b/Engine/Source/Programs/Horde/Horde.Agent/Utility/JsonRpcLogSink.cs
--- a/Engine/Source/Programs/Horde/Horde.Agent/Utility/JsonRpcLogSink.cs
+++ b/Engine/Source/Programs/Horde/Horde.Agent/Utility/JsonRpcLogSink.cs
@@ -240,7 +240,7 @@ namespace Horde.Agent.Utility
 
 		#region RPC calls
 
-		protected virtual async Task UpdateLogAsync(BlobHandle target, int lineCount, bool complete, CancellationToken cancellationToken)
+		protected virtual async Task UpdateLogAsync(IBlobHandle target, int lineCount, bool complete, CancellationToken cancellationToken)
 		{
 			_logger.LogInformation("Updating log {LogId} to line {LineCount}, target {Locator}", _logId, lineCount, target);
 
diff --git a/Engine/Source/Programs/Horde/Horde.Server.Tests/BlobStoreTests.cs b/Engine/Source/Programs/Horde/Horde.Server.Tests/BlobStoreTests.cs
--- a/Engine/Source/Programs/Horde/Horde.Server.Tests/BlobStoreTests.cs
+++ b/Engine/Source/Programs/Horde/Horde.Server.Tests/BlobStoreTests.cs
@@ -65,7 +65,7 @@ namespace Horde.Server.Tests
 			await using IStorageWriter writer = store.CreateWriter();
 			blob.Data.CopyTo(writer.GetOutputBuffer(0, blob.Data.Length));
 
-			BlobHandle handle = await writer.WriteBlobAsync(s_blobType, blob.Data.Length, blob.References.ConvertAll(x => store.CreateBlobHandle(x)));
+			IBlobHandle handle = await writer.WriteBlobAsync(s_blobType, blob.Data.Length, blob.References.ConvertAll(x => store.CreateBlobHandle(x)));
 			await handle.FlushAsync();
 
 			return handle.GetLocator();
@@ -113,7 +113,7 @@ namespace Horde.Server.Tests
 			{
 				RefName refName = new RefName("hello");
 				await store.WriteRefTargetAsync(refName, store.CreateBlobHandle(locator3));
-				BlobHandle refTarget = await store.ReadRefTargetAsync(refName);
+				IBlobHandle refTarget = await store.ReadRefTargetAsync(refName);
 				Assert.AreEqual(locator3, refTarget.GetLocator());
 			}
 		}
@@ -155,7 +155,7 @@ namespace Horde.Server.Tests
 
 		static async Task<BlobLocator> TryReadRefTargetAsync(IStorageClient store, RefName name)
 		{
-			BlobHandle? handle = await store.TryReadRefTargetAsync(name);
+			IBlobHandle? handle = await store.TryReadRefTargetAsync(name);
 			if (handle == null)
 			{
 				return default;
diff --git a/Engine/Source/Programs/Horde/Horde.Server.Tests/GcServiceTests.cs b/Engine/Source/Programs/Horde/Horde.Server.Tests/GcServiceTests.cs
--- a/Engine/Source/Programs/Horde/Horde.Server.Tests/GcServiceTests.cs
+++ b/Engine/Source/Programs/Horde/Horde.Server.Tests/GcServiceTests.cs
@@ -38,7 +38,7 @@ namespace Horde.Server.Tests
 				int blobIdx = (int)(random.NextDouble() * blobs.Length);
 				if (roots.Add(blobs[blobIdx]))
 				{
-					BlobHandle handle = store.CreateBlobHandle(blobs[blobIdx]);
+					IBlobHandle handle = store.CreateBlobHandle(blobs[blobIdx]);
 					await store.WriteRefTargetAsync(new RefName($"ref-{idx}"), handle);
 				}
 			}
@@ -67,8 +67,8 @@ namespace Horde.Server.Tests
 			{
 				if (nodes.Add(root))
 				{
-					BlobHandle handle = store.CreateBlobHandle(root);
-					IReadOnlyList<BlobHandle> refs = await handle.GetRefsAsync();
+					IBlobHandle handle = store.CreateBlobHandle(root);
+					IReadOnlyList<IBlobHandle> refs = await handle.ReadImportsAsync();
 					await FindNodesAsync(store, refs.ConvertAll(x => x.GetLocator()), nodes);
 				}
 			}
@@ -103,10 +103,10 @@ namespace Horde.Server.Tests
 			BlobLocator[] locators = new BlobLocator[children.Length];
 			for (int idx = numNodes - 1; idx >= 0; idx--)
 			{
-				BlobHandle handle;
+				IBlobHandle handle;
 				await using (IStorageWriter writer = store.CreateWriter("gctest"))
 				{
-					List<BlobHandle> imports = children[idx].ConvertAll(x => store.CreateBlobHandle(locators[x]));
+					List<IBlobHandle> imports = children[idx].ConvertAll(x => store.CreateBlobHandle(locators[x]));
 					handle = await writer.WriteBlobAsync(blobType, 0, imports);
 				}
 				locators[idx] = handle.GetLocator();
diff --git a/Engine/Source/Programs/Horde/Horde.Server.Tests/StorageServiceTests.cs b/Engine/Source/Programs/Horde/Horde.Server.Tests/StorageServiceTests.cs
--- a/Engine/Source/Programs/Horde/Horde.Server.Tests/StorageServiceTests.cs
+++ b/Engine/Source/Programs/Horde/Horde.Server.Tests/StorageServiceTests.cs
@@ -24,19 +24,19 @@ namespace Horde.Server.Tests
 			BlobType type2 = new BlobType(Guid.Parse("{6CB3A005-26BA-4787-86D2-793ED13771CB}"), 0);
 			byte[] data2 = new byte[] { 4, 5, 6 };
 
-			BlobHandle handle1a;
-			BlobHandle handle1b;
-			BlobHandle handle2;
+			IBlobHandle handle1a;
+			IBlobHandle handle1b;
+			IBlobHandle handle2;
 			await using (IStorageWriter writer = client.CreateWriter())
 			{
 				data1.CopyTo(writer.GetOutputBuffer(0, data1.Length));
-				handle1a = await writer.WriteBlobAsync(type1, data1.Length, Array.Empty<BlobHandle>(), new AliasInfo[] { new AliasInfo("foo", 2) });
+				handle1a = await writer.WriteBlobAsync(type1, data1.Length, Array.Empty<IBlobHandle>(), new AliasInfo[] { new AliasInfo("foo", 2) });
 
 				data1.CopyTo(writer.GetOutputBuffer(0, data1.Length));
-				handle1b = await writer.WriteBlobAsync(type1, data1.Length, Array.Empty<BlobHandle>(), new AliasInfo[] { new AliasInfo("foo", 1) });
+				handle1b = await writer.WriteBlobAsync(type1, data1.Length, Array.Empty<IBlobHandle>(), new AliasInfo[] { new AliasInfo("foo", 1) });
 
 				data2.CopyTo(writer.GetOutputBuffer(0, data2.Length));
-				handle2 = await writer.WriteBlobAsync(type2, data2.Length, Array.Empty<BlobHandle>(), new AliasInfo[] { new AliasInfo("bar") });
+				handle2 = await writer.WriteBlobAsync(type2, data2.Length, Array.Empty<IBlobHandle>(), new AliasInfo[] { new AliasInfo("bar") });
 			}
 
 			BlobAlias[] aliases;
diff --git a/Engine/Source/Programs/Horde/Horde.Server/Commands/Bundles/BundleWriteTest.cs b/Engine/Source/Programs/Horde/Horde.Server/Commands/Bundles/BundleWriteTest.cs
--- a/Engine/Source/Programs/Horde/Horde.Server/Commands/Bundles/BundleWriteTest.cs
+++ b/Engine/Source/Programs/Horde/Horde.Server/Commands/Bundles/BundleWriteTest.cs
@@ -29,18 +29,18 @@ namespace Horde.Server.Commands.Bundles
 
 			public void Dispose() { }
 
-			public BlobHandle CreateBlobHandle(BlobLocator locator) => throw new NotImplementedException();
+			public IBlobHandle CreateBlobHandle(BlobLocator locator) => throw new NotImplementedException();
 			public Task<bool> DeleteRefAsync(RefName name, CancellationToken cancellationToken = default) => Task.FromResult(true);
 			public ValueTask<BlobData> ReadBlobAsync(BlobLocator locator, CancellationToken cancellationToken = default) => throw new NotImplementedException();
-			public ValueTask<BlobHandle> WriteBlobAsync(BlobType type, Stream stream, IReadOnlyList<BlobHandle> references, string? basePath = null, CancellationToken cancellationToken = default) => throw new NotImplementedException();
+			public ValueTask<IBlobHandle> WriteBlobAsync(BlobType type, Stream stream, IReadOnlyList<IBlobHandle> references, string? basePath = null, CancellationToken cancellationToken = default) => throw new NotImplementedException();
 			public ValueTask<Uri?> TryGetReadRedirectAsync(BlobLocator locator, CancellationToken cancellationToken = default) => new ValueTask<Uri?>();
 			public ValueTask<(BlobLocator, Uri)?> TryGetWriteRedirectAsync(string? prefix = null, CancellationToken cancellationToken = default) => new ValueTask<(BlobLocator, Uri)?>();
 
-			public Task AddAliasAsync(string name, BlobHandle handle, int rank = 0, ReadOnlyMemory<byte> data = default, CancellationToken cancellationToken = default) => throw new NotImplementedException();
-			public Task RemoveAliasAsync(string name, BlobHandle handle, CancellationToken cancellationToken = default) => throw new NotImplementedException();
+			public Task AddAliasAsync(string name, IBlobHandle handle, int rank = 0, ReadOnlyMemory<byte> data = default, CancellationToken cancellationToken = default) => throw new NotImplementedException();
+			public Task RemoveAliasAsync(string name, IBlobHandle handle, CancellationToken cancellationToken = default) => throw new NotImplementedException();
 			public Task<BlobAlias[]> FindAliasesAsync(string alias, int? maxResults = null, CancellationToken cancellationToken = default) => throw new NotImplementedException();
-			public Task<BlobHandle?> TryReadRefAsync(RefName name, RefCacheTime cacheTime = default, CancellationToken cancellationToken = default) => throw new NotImplementedException();
-			public Task WriteRefAsync(RefName name, BlobHandle target, RefOptions? options = null, CancellationToken cancellationToken = default) => Task.CompletedTask;
+			public Task<IBlobHandle?> TryReadRefAsync(RefName name, RefCacheTime cacheTime = default, CancellationToken cancellationToken = default) => throw new NotImplementedException();
+			public Task WriteRefAsync(RefName name, IBlobHandle target, RefOptions? options = null, CancellationToken cancellationToken = default) => Task.CompletedTask;
 
 			public IStorageWriter CreateWriter(string? basePath) => new DefaultStorageWriter(this, basePath);
 
diff --git a/Engine/Source/Programs/Horde/Horde.Server/Ddc/BlobService.cs b/Engine/Source/Programs/Horde/Horde.Server/Ddc/BlobService.cs
--- a/Engine/Source/Programs/Horde/Horde.Server/Ddc/BlobService.cs
+++ b/Engine/Source/Programs/Horde/Horde.Server/Ddc/BlobService.cs
@@ -103,7 +103,7 @@ namespace Horde.Server.Ddc
 				throw new BlobNotFoundException(ns, blob);
 			}
 
-			using BlobData data = await alias.Target.ReadBlobDataAsync(cancellationToken);
+			using BlobData data = await alias.Target.ReadAsync(cancellationToken);
 			return new BlobContents(data.Data.ToArray());
 		}
 
@@ -143,7 +143,7 @@ namespace Horde.Server.Ddc
 		{
 			using IStorageClient storageClient = _storageClientFactory.CreateClient(ns);
 
-			BlobHandle blobHandle;
+			IBlobHandle blobHandle;
 			await using (IStorageWriter writer = storageClient.CreateWriter())
 			{
 				Memory<byte> memory = writer.GetOutputBuffer(0, (int)content.Length);
@@ -151,7 +151,7 @@ namespace Horde.Server.Ddc
 				using Stream stream = content.GetStream();
 				await stream.ReadAsync(memory, cancellationToken);
 
-				blobHandle = await writer.WriteBlobAsync(s_rawBlobType, memory.Length, Array.Empty<BlobHandle>(), cancellationToken);
+				blobHandle = await writer.WriteBlobAsync(s_rawBlobType, memory.Length, Array.Empty<IBlobHandle>(), cancellationToken);
 				await writer.FlushAsync(cancellationToken);
 			}
 
diff --git a/Engine/Source/Programs/Horde/Horde.Server/Ddc/RefService.cs b/Engine/Source/Programs/Horde/Horde.Server/Ddc/RefService.cs
--- a/Engine/Source/Programs/Horde/Horde.Server/Ddc/RefService.cs
+++ b/Engine/Source/Programs/Horde/Horde.Server/Ddc/RefService.cs
@@ -34,7 +34,7 @@ namespace Horde.Server.Ddc
 		public async Task<bool> ExistsAsync(NamespaceId ns, BucketId bucket, RefId key, CancellationToken cancellationToken)
 		{
 			using IStorageClient storageClient = _storageClientFactory.CreateClient(ns);
-			BlobHandle? handle = await storageClient.TryReadRefTargetAsync(GetRefName(bucket, key), cancellationToken: cancellationToken);
+			IBlobHandle? handle = await storageClient.TryReadRefTargetAsync(GetRefName(bucket, key), cancellationToken: cancellationToken);
 			return handle != null;
 		}
 
@@ -48,8 +48,8 @@ namespace Horde.Server.Ddc
 				throw new BlobNotFoundException(ns, blobHash);
 			}
 
-			BlobHandle blobHandle = blobAlias.Target;
-			using BlobData blobContents = await blobHandle.ReadBlobDataAsync(cancellationToken);
+			IBlobHandle blobHandle = blobAlias.Target;
+			using BlobData blobContents = await blobHandle.ReadAsync(cancellationToken);
 			CbObject payload = new CbObject(blobContents.Data);
 
 			BlobId[] referencedBlobs = Array.Empty<BlobId>();
@@ -131,7 +131,7 @@ namespace Horde.Server.Ddc
 				throw new RefNotFoundException(ns, bucket, key);
 			}
 
-			BlobData data = await node.References.First(x => x.Hash == node.RootHash).Handle.ReadBlobDataAsync(cancellationToken);
+			BlobData data = await node.References.First(x => x.Hash == node.RootHash).Handle.ReadAsync(cancellationToken);
 			BlobContents contents = new BlobContents(data.Data.ToArray());
 
 			RefRecord record = new RefRecord(ns, bucket, key, DateTime.UtcNow, null, new BlobId(node.RootHash), true);
diff --git a/Engine/Source/Programs/Horde/Horde.Server/Perforce/PerforceReplicator.cs b/Engine/Source/Programs/Horde/Horde.Server/Perforce/PerforceReplicator.cs
--- a/Engine/Source/Programs/Horde/Horde.Server/Perforce/PerforceReplicator.cs
+++ b/Engine/Source/Programs/Horde/Horde.Server/Perforce/PerforceReplicator.cs
@@ -547,7 +547,7 @@ namespace Horde.Server.Perforce
 			ChangeRecord changeRecord = await perforce.GetChangeAsync(GetChangeOptions.None, change, cancellationToken);
 			DirectoryNodeRef rootRef = new DirectoryNodeRef(root.Length, await writer.WriteHashedNodeAsync(root, cancellationToken));
 			CommitNode commitNode = new CommitNode(change, parentRef, changeRecord.User ?? "Unknown", changeRecord.Description ?? String.Empty, changeRecord.Date, rootRef);
-			BlobHandle commitHandle = await store.WriteRefAsync(refName, commitNode, refOptions: options.RefOptions, cancellationToken: cancellationToken);
+			IBlobHandle commitHandle = await store.WriteRefAsync(refName, commitNode, refOptions: options.RefOptions, cancellationToken: cancellationToken);
 
 			// Log the snapshot info
 			_logger.LogInformation("Snapshot for {StreamId} CL {Change} is ref {RefName} (commit: {CommitHandle}, root: {RootHandle})", streamConfig.Id, change, refName, commitHandle.GetLocator(), rootRef.Handle.GetLocator());
diff --git a/Engine/Source/Programs/Horde/Horde.Server/Storage/StorageController.cs b/Engine/Source/Programs/Horde/Horde.Server/Storage/StorageController.cs
--- a/Engine/Source/Programs/Horde/Horde.Server/Storage/StorageController.cs
+++ b/Engine/Source/Programs/Horde/Horde.Server/Storage/StorageController.cs
@@ -98,7 +98,7 @@ namespace Horde.Server.Storage
 			else
 			{
 				using Stream stream = file.OpenReadStream();
-				BlobHandle handle = await storageClient.WriteBlobAsync(Bundle.BlobType, stream, Array.Empty<BlobHandle>(), prefix, cancellationToken);
+				IBlobHandle handle = await storageClient.WriteBlobAsync(Bundle.BlobType, stream, Array.Empty<IBlobHandle>(), prefix, cancellationToken);
 				return new WriteBlobResponse { Blob = handle.GetLocator().ToString(), SupportsRedirects = storageClient.SupportsRedirects };
 			}
 		}
@@ -176,8 +176,8 @@ namespace Horde.Server.Storage
 			}
 
 #pragma warning disable CA2000 // Dispose objects before losing scope
-			BlobHandle handle = storageClient.CreateBlobHandle(new BlobLocator(locator.Path));
-			Stream stream = await handle.OpenAsync(offset, length, cancellationToken);
+			IBlobHandle handle = storageClient.CreateBlobHandle(new BlobLocator(locator.Path));
+			Stream stream = await handle.OpenBodyAsync(offset, length, cancellationToken);
 			return new FileStreamResult(stream, "application/octet-stream");
 #pragma warning restore CA2000 // Dispose objects before losing scope
 		}
@@ -244,7 +244,7 @@ namespace Horde.Server.Storage
 			}
 #pragma warning restore CS0618 // Type or member is obsolete
 
-			BlobHandle target = client.CreateBlobHandle(request.Target);
+			IBlobHandle target = client.CreateBlobHandle(request.Target);
 			await client.WriteRefAsync(refName, target, request.Options, cancellationToken);
 
 			return Ok();
@@ -289,7 +289,7 @@ namespace Horde.Server.Storage
 				}
 			}
 
-			BlobHandle? target = await client.TryReadRefAsync(refName, cacheTime, cancellationToken: cancellationToken);
+			IBlobHandle? target = await client.TryReadRefAsync(refName, cacheTime, cancellationToken: cancellationToken);
 			if (target == null)
 			{
 				return new NotFoundResult();
@@ -301,13 +301,11 @@ namespace Horde.Server.Storage
 #pragma warning disable CS0618 // Type or member is obsolete
 			try
 			{
-				string locator = target.ToString();
-				int hashIdx = locator.LastIndexOf('#');
-
-				if (hashIdx != -1)
+				BlobLocator locator = target.GetLocator();
+				if(locator.TryUnwrap(out BlobLocator blob, out Utf8String fragment) && Int32.TryParse(fragment.ToString(), out int exportIdx))
 				{
-					response.Blob = new BlobLocator(locator.Substring(0, hashIdx));
-					response.ExportIdx = Int32.Parse(locator.Substring(hashIdx + 1));
+					response.Blob = blob;
+					response.ExportIdx = exportIdx;
 				}
 			}
 			catch { }
@@ -382,7 +380,7 @@ namespace Horde.Server.Storage
 
 			object content;
 
-			using BlobData blobData = await storageClient.CreateBlobHandle(locator).ReadBlobDataAsync(cancellationToken);
+			using BlobData blobData = await storageClient.CreateBlobHandle(locator).ReadAsync(cancellationToken);
 
 			Node node = Node.Deserialize(blobData);
 			switch (node)
@@ -412,7 +410,7 @@ namespace Horde.Server.Storage
 			return new { type = blobData.Type.Guid, @class = node.GetType().Name, content = content };
 		}
 
-		static string GetNodeLink(string linkBase, BlobHandle handle) => GetNodeLink(linkBase, handle.GetLocator());
+		static string GetNodeLink(string linkBase, IBlobHandle handle) => GetNodeLink(linkBase, handle.GetLocator());
 		
 		static string GetNodeLink(string linkBase, BlobLocator locator) => $"{linkBase}/nodes/{locator}";
 	}
diff --git a/Engine/Source/Programs/Horde/Horde.Server/Storage/StorageService.cs b/Engine/Source/Programs/Horde/Horde.Server/Storage/StorageService.cs
--- a/Engine/Source/Programs/Horde/Horde.Server/Storage/StorageService.cs
+++ b/Engine/Source/Programs/Horde/Horde.Server/Storage/StorageService.cs
@@ -50,13 +50,13 @@ namespace Horde.Server.Storage
 	/// </summary>
 	public sealed class StorageService : IHostedService, IStorageClientFactory, IAsyncDisposable
 	{
-		sealed class LeafBlobHandle : BlobHandle
+		sealed class LeafBlobHandle : IBlobHandle
 		{
 			readonly IStorageBackend _backend;
 			readonly string _path;
 			readonly Tracer _tracer;
 
-			public override BlobHandle? Outer => null;
+			public IBlobHandle? Outer => null;
 
 			public LeafBlobHandle(IStorageBackend backend, string path, Tracer tracer)
 			{
@@ -66,12 +66,12 @@ namespace Horde.Server.Storage
 			}
 
 			/// <inheritdoc/>
-			public override ValueTask<BlobType> GetTypeAsync(CancellationToken cancellationToken = default) => new ValueTask<BlobType>(BlobType.Leaf);
+			public ValueTask<BlobType> ReadTypeAsync(CancellationToken cancellationToken = default) => new ValueTask<BlobType>(BlobType.Leaf);
 
 			/// <inheritdoc/>
-			public override async Task<Stream> OpenAsync(int offset = 0, int? length = null, CancellationToken cancellationToken = default)
+			public async Task<Stream> OpenBodyAsync(int offset = 0, int? length = null, CancellationToken cancellationToken = default)
 			{
-				using TelemetrySpan span = _tracer.StartActiveSpan($"{nameof(StorageService)}.{nameof(LeafBlobHandle)}.{nameof(OpenAsync)}");
+				using TelemetrySpan span = _tracer.StartActiveSpan($"{nameof(StorageService)}.{nameof(LeafBlobHandle)}.{nameof(OpenBodyAsync)}");
 				span.SetAttribute("path", _path);
 				span.SetAttribute("offset", offset);
 				span.SetAttribute("length", length);
@@ -85,17 +85,17 @@ namespace Horde.Server.Storage
 			}
 
 			/// <inheritdoc/>
-			public override ValueTask<IReadOnlyList<BlobHandle>> GetRefsAsync(CancellationToken cancellationToken = default) => new ValueTask<IReadOnlyList<BlobHandle>>(Array.Empty<BlobHandle>());
+			public ValueTask<IReadOnlyList<IBlobHandle>> ReadImportsAsync(CancellationToken cancellationToken = default) => new ValueTask<IReadOnlyList<IBlobHandle>>(Array.Empty<IBlobHandle>());
 
 			/// <inheritdoc/>
-			public override async ValueTask<BlobData> ReadBlobDataAsync(CancellationToken cancellationToken = default)
+			public async ValueTask<BlobData> ReadAsync(CancellationToken cancellationToken = default)
 			{
 				IReadOnlyMemoryOwner<byte> obj = await _backend.ReadAsync(_path, cancellationToken);
-				return new ReadOnlyMemoryOwnerBlobData(BlobType.Leaf, obj, Array.Empty<BlobHandle>());
+				return new ReadOnlyMemoryOwnerBlobData(BlobType.Leaf, obj, Array.Empty<IBlobHandle>());
 			}
 
 			/// <inheritdoc/>
-			public override bool TryAppendIdentifier(Utf8StringBuilder builder)
+			public bool TryAppendIdentifier(Utf8StringBuilder builder)
 			{
 				builder.Append(_path);
 				return true;
@@ -151,7 +151,7 @@ namespace Horde.Server.Storage
 			#region Blobs
 
 			/// <inheritdoc/>
-			public BlobHandle CreateBlobHandle(BlobLocator locator)
+			public IBlobHandle CreateBlobHandle(BlobLocator locator)
 			{
 				if (locator.TryUnwrap(out BlobLocator baseLocator, out Utf8String fragment))
 				{
@@ -172,12 +172,12 @@ namespace Horde.Server.Storage
 			/// <inheritdoc/>
 			public async ValueTask<BlobData> ReadBlobAsync(BlobLocator locator, CancellationToken cancellationToken = default)
 			{
-				BlobHandle handle = CreateBlobHandle(locator);
-				return await handle.ReadBlobDataAsync(cancellationToken);
+				IBlobHandle handle = CreateBlobHandle(locator);
+				return await handle.ReadAsync(cancellationToken);
 			}
 
 			/// <inheritdoc/>
-			public async ValueTask<BlobHandle> WriteBlobAsync(BlobType type, Stream stream, IReadOnlyList<BlobHandle> references, string? basePath = null, CancellationToken cancellationToken = default)
+			public async ValueTask<IBlobHandle> WriteBlobAsync(BlobType type, Stream stream, IReadOnlyList<IBlobHandle> references, string? basePath = null, CancellationToken cancellationToken = default)
 			{
 				string path = await _backend.WriteAsync(stream, basePath, cancellationToken);
 
@@ -218,10 +218,10 @@ namespace Horde.Server.Storage
 			#region Aliases
 
 			/// <inheritdoc/>
-			public Task AddAliasAsync(string name, BlobHandle target, int rank = 0, ReadOnlyMemory<byte> data = default, CancellationToken cancellationToken = default) => _outer.AddAliasAsync(NamespaceId, name, target, rank, data, cancellationToken);
+			public Task AddAliasAsync(string name, IBlobHandle target, int rank = 0, ReadOnlyMemory<byte> data = default, CancellationToken cancellationToken = default) => _outer.AddAliasAsync(NamespaceId, name, target, rank, data, cancellationToken);
 
 			/// <inheritdoc/>
-			public Task RemoveAliasAsync(string name, BlobHandle target, CancellationToken cancellationToken = default) => _outer.RemoveAliasAsync(NamespaceId, name, target, cancellationToken);
+			public Task RemoveAliasAsync(string name, IBlobHandle target, CancellationToken cancellationToken = default) => _outer.RemoveAliasAsync(NamespaceId, name, target, cancellationToken);
 
 			/// <inheritdoc/>
 			public async Task<BlobAlias[]> FindAliasesAsync(string alias, int? maxResults, CancellationToken cancellationToken = default)
@@ -239,7 +239,7 @@ namespace Horde.Server.Storage
 			#region Refs
 
 			/// <inheritdoc/>
-			public async Task<BlobHandle?> TryReadRefAsync(RefName name, RefCacheTime cacheTime = default, CancellationToken cancellationToken = default)
+			public async Task<IBlobHandle?> TryReadRefAsync(RefName name, RefCacheTime cacheTime = default, CancellationToken cancellationToken = default)
 			{
 				RefInfo? result = await _outer.TryReadRefAsync(NamespaceId, name, cacheTime, cancellationToken);
 				if (result == null)
@@ -250,7 +250,7 @@ namespace Horde.Server.Storage
 			}
 
 			/// <inheritdoc/>
-			public async Task WriteRefAsync(RefName name, BlobHandle target, RefOptions? options = null, CancellationToken cancellationToken = default)
+			public async Task WriteRefAsync(RefName name, IBlobHandle target, RefOptions? options = null, CancellationToken cancellationToken = default)
 			{
 				await target.FlushAsync(cancellationToken);
 				await _outer.WriteRefAsync(NamespaceId, name, target.GetLocator(), options, cancellationToken);
@@ -689,11 +689,11 @@ namespace Horde.Server.Storage
 						NamespaceInfo? namespaceInfo;
 						if (scopedState.Value.Namespaces.TryGetValue(blobInfo.NamespaceId, out namespaceInfo))
 						{
-							BlobHandle handle = namespaceInfo.Client.CreateBlobHandle(blobInfo.Locator);
-							IReadOnlyList<BlobHandle> imports = await handle.GetRefsAsync(cancellationToken);
+							IBlobHandle handle = namespaceInfo.Client.CreateBlobHandle(blobInfo.Locator);
+							IReadOnlyList<IBlobHandle> imports = await handle.ReadImportsAsync(cancellationToken);
 
 							List<ObjectId> importInfoIds = new List<ObjectId>();
-							foreach (BlobHandle import in imports)
+							foreach (IBlobHandle import in imports)
 							{
 								BlobLocator importLocator = import.GetLocator();
 								string importPath = importLocator.BaseLocator.ToString();
@@ -730,7 +730,7 @@ namespace Horde.Server.Storage
 		/// <param name="data">Inline data to store with this alias</param>
 		/// <param name="cancellationToken">Cancellation token for the operation</param>
 		/// <returns>Sequence of thandles</returns>
-		async Task AddAliasAsync(NamespaceId namespaceId, string name, BlobHandle target, int rank, ReadOnlyMemory<byte> data, CancellationToken cancellationToken = default)
+		async Task AddAliasAsync(NamespaceId namespaceId, string name, IBlobHandle target, int rank, ReadOnlyMemory<byte> data, CancellationToken cancellationToken = default)
 		{
 			BlobLocator locator = target.GetLocator();
 
@@ -760,7 +760,7 @@ namespace Horde.Server.Storage
 		/// <param name="target">Target node for the alias</param>
 		/// <param name="cancellationToken">Cancellation token for the operation</param>
 		/// <returns>Sequence of thandles</returns>
-		async Task RemoveAliasAsync(NamespaceId namespaceId, string name, BlobHandle target, CancellationToken cancellationToken = default)
+		async Task RemoveAliasAsync(NamespaceId namespaceId, string name, IBlobHandle target, CancellationToken cancellationToken = default)
 		{
 			BlobLocator locator = target.GetLocator();
 
diff --git a/Engine/Source/Programs/Horde/Horde.Server/Tools/ToolCollection.cs b/Engine/Source/Programs/Horde/Horde.Server/Tools/ToolCollection.cs
--- a/Engine/Source/Programs/Horde/Horde.Server/Tools/ToolCollection.cs
+++ b/Engine/Source/Programs/Horde/Horde.Server/Tools/ToolCollection.cs
@@ -265,7 +265,7 @@ namespace Horde.Server.Tools
 				nodeRef = await writer.WriteHashedNodeAsync(directoryNode, cancellationToken);
 			}
 
-			BlobHandle target = nodeRef.Handle;
+			IBlobHandle target = nodeRef.Handle;
 			await client.WriteRefTargetAsync(refName, target, cancellationToken: cancellationToken);
 
 			return await CreateDeploymentAsync(tool, options, target.GetLocator(), globalConfig, cancellationToken);
@@ -288,7 +288,7 @@ namespace Horde.Server.Tools
 			RefName refName = new RefName($"{tool.Id}/{deploymentId}");
 
 			using IStorageClient client = _storageService.CreateClient(namespaceId);
-			BlobHandle targetHandle = client.CreateBlobHandle(target);
+			IBlobHandle targetHandle = client.CreateBlobHandle(target);
 			await client.WriteRefTargetAsync(refName, targetHandle, cancellationToken: cancellationToken);
 
 			return await CreateDeploymentInternalAsync(tool, deploymentId, options, namespaceId, refName, globalConfig, cancellationToken);
diff --git a/Engine/Source/Programs/Horde/Horde.Server/Tools/ToolsController.cs b/Engine/Source/Programs/Horde/Horde.Server/Tools/ToolsController.cs
--- a/Engine/Source/Programs/Horde/Horde.Server/Tools/ToolsController.cs
+++ b/Engine/Source/Programs/Horde/Horde.Server/Tools/ToolsController.cs
@@ -127,7 +127,7 @@ namespace Horde.Server.Tools
 		/// <summary>
 		/// Constructor
 		/// </summary>
-		public GetToolDeploymentResponse(IToolDeployment deployment, BlobHandle handle)
+		public GetToolDeploymentResponse(IToolDeployment deployment, IBlobHandle handle)
 		{
 			_deployment = deployment;
 			Locator = handle.GetLocator();
@@ -519,7 +519,7 @@ namespace Horde.Server.Tools
 		private async Task<GetToolDeploymentResponse> GetDeploymentInfoResponseAsync(ITool tool, IToolDeployment deployment, CancellationToken cancellationToken)
 		{
 			using IStorageClient client = _toolCollection.CreateStorageClient(tool);
-			BlobHandle rootHandle = await client.ReadRefTargetAsync(deployment.RefName, cancellationToken: cancellationToken);
+			IBlobHandle rootHandle = await client.ReadRefTargetAsync(deployment.RefName, cancellationToken: cancellationToken);
 
 			return new GetToolDeploymentResponse(deployment, rootHandle);
 		}
diff --git a/Engine/Source/Programs/Horde/Horde/Commands/Bundles/BundleExtract.cs b/Engine/Source/Programs/Horde/Horde/Commands/Bundles/BundleExtract.cs
--- a/Engine/Source/Programs/Horde/Horde/Commands/Bundles/BundleExtract.cs
+++ b/Engine/Source/Programs/Horde/Horde/Commands/Bundles/BundleExtract.cs
@@ -38,19 +38,19 @@ namespace Horde.Commands.Bundles
 			if (File != null)
 			{
 				using IStorageClient store = BundleStorageClient.CreateFromDirectory(File.Directory, BundleCache, logger);
-				BlobHandle handle = store.CreateBlobHandle(await FileStorageClient.ReadRefAsync(File));
+				IBlobHandle handle = store.CreateBlobHandle(await FileStorageClient.ReadRefAsync(File));
 				await ExecuteInternalAsync(store, handle, logger);
 			}
 			else if (Ref != null)
 			{
 				using IStorageClient store = CreateStorageClient();
-				BlobHandle handle = await store.ReadRefTargetAsync(new RefName(Ref));
+				IBlobHandle handle = await store.ReadRefTargetAsync(new RefName(Ref));
 				await ExecuteInternalAsync(store, handle, logger);
 			}
 			else if (Node != null)
 			{
 				using IStorageClient store = CreateStorageClient();
-				BlobHandle handle = store.CreateBlobHandle(new BlobLocator(Node));
+				IBlobHandle handle = store.CreateBlobHandle(new BlobLocator(Node));
 				await ExecuteInternalAsync(store, handle, logger);
 			}
 			else
@@ -61,7 +61,7 @@ namespace Horde.Commands.Bundles
 			return 0;
 		}
 
-		protected async Task ExecuteInternalAsync(IStorageClient store, BlobHandle handle, ILogger logger)
+		protected async Task ExecuteInternalAsync(IStorageClient store, IBlobHandle handle, ILogger logger)
 		{
 			Stopwatch timer = Stopwatch.StartNew();
 
diff --git a/Engine/Source/Programs/Horde/Horde/Commands/Compute/ComputeRun.cs b/Engine/Source/Programs/Horde/Horde/Commands/Compute/ComputeRun.cs
--- a/Engine/Source/Programs/Horde/Horde/Commands/Compute/ComputeRun.cs
+++ b/Engine/Source/Programs/Horde/Horde/Commands/Compute/ComputeRun.cs
@@ -136,7 +136,7 @@ namespace Horde.Commands.Compute
 			DirectoryNode sandbox = new DirectoryNode();
 			await sandbox.CopyFromDirectoryAsync(taskFile.Directory.ToDirectoryInfo(), new ChunkingOptions(), writer, null, cancellationToken);
 
-			BlobHandle handle = await writer.FlushAsync(sandbox, cancellationToken);
+			IBlobHandle handle = await writer.FlushAsync(sandbox, cancellationToken);
 			return handle.GetLocator();
 		}
 	}
diff --git a/Engine/Source/Programs/Horde/Horde/Commands/Workspace/WorkspaceSync.cs b/Engine/Source/Programs/Horde/Horde/Commands/Workspace/WorkspaceSync.cs
--- a/Engine/Source/Programs/Horde/Horde/Commands/Workspace/WorkspaceSync.cs
+++ b/Engine/Source/Programs/Horde/Horde/Commands/Workspace/WorkspaceSync.cs
@@ -43,19 +43,19 @@ namespace Horde.Commands.Workspace
 			if (File != null)
 			{
 				using IStorageClient store = BundleStorageClient.CreateFromDirectory(File.Directory, BundleCache, logger);
-				BlobHandle handle = store.CreateBlobHandle(await FileStorageClient.ReadRefAsync(File));
+				IBlobHandle handle = store.CreateBlobHandle(await FileStorageClient.ReadRefAsync(File));
 				return await ExecuteInternalAsync(store, handle, logger);
 			}
 			else if (Ref != null)
 			{
 				using IStorageClient store = CreateStorageClient();
-				BlobHandle handle = await store.ReadRefTargetAsync(new RefName(Ref));
+				IBlobHandle handle = await store.ReadRefTargetAsync(new RefName(Ref));
 				return await ExecuteInternalAsync(store, handle, logger);
 			}
 			else if (Node != null)
 			{
 				using IStorageClient store = CreateStorageClient();
-				BlobHandle handle = store.CreateBlobHandle(new BlobLocator(Node));
+				IBlobHandle handle = store.CreateBlobHandle(new BlobLocator(Node));
 				return await ExecuteInternalAsync(store, handle, logger);
 			}
 			else
@@ -64,7 +64,7 @@ namespace Horde.Commands.Workspace
 			}
 		}
 
-		async Task<int> ExecuteInternalAsync(IStorageClient store, BlobHandle handle, ILogger logger)
+		async Task<int> ExecuteInternalAsync(IStorageClient store, IBlobHandle handle, ILogger logger)
 		{
 			RootDir ??= DirectoryReference.GetCurrentDirectory();
 			CancellationToken cancellationToken = CancellationToken.None;
diff --git a/Engine/Source/Programs/Horde/Samples/RemoteClient/RemoteClientApp.cs b/Engine/Source/Programs/Horde/Samples/RemoteClient/RemoteClientApp.cs
--- a/Engine/Source/Programs/Horde/Samples/RemoteClient/RemoteClientApp.cs
+++ b/Engine/Source/Programs/Horde/Samples/RemoteClient/RemoteClientApp.cs
@@ -115,7 +115,7 @@ namespace RemoteClient
 				{
 					DirectoryNode sandbox = new DirectoryNode();
 					await sandbox.CopyFromDirectoryAsync(uploadDir.ToDirectoryInfo(), new ChunkingOptions(), treeWriter, null);
-					BlobHandle handle = await treeWriter.FlushAsync(sandbox);
+					IBlobHandle handle = await treeWriter.FlushAsync(sandbox);
 					await channel.UploadFilesAsync("", handle.GetLocator(), storage);
 				}
 
diff --git a/Engine/Source/Programs/Shared/EpicGames.Horde.Tests/BlobDataTests.cs b/Engine/Source/Programs/Shared/EpicGames.Horde.Tests/BlobDataTests.cs
--- a/Engine/Source/Programs/Shared/EpicGames.Horde.Tests/BlobDataTests.cs
+++ b/Engine/Source/Programs/Shared/EpicGames.Horde.Tests/BlobDataTests.cs
@@ -13,16 +13,16 @@ namespace EpicGames.Horde.Tests
 	[TestClass]
 	public class BlobDataTests
 	{
-		class DummyHandle : BlobHandle
+		class DummyHandle : IBlobHandle
 		{
 			readonly BlobLocator _locator;
 
-			public override BlobHandle? Outer => null;
+			public IBlobHandle? Outer => null;
 
 			public DummyHandle(string locator) => _locator = new BlobLocator(new Utf8String(locator));
-			public override ValueTask<BlobData> ReadBlobDataAsync(CancellationToken cancellationToken = default) => throw new NotImplementedException();
+			public ValueTask<BlobData> ReadAsync(CancellationToken cancellationToken = default) => throw new NotImplementedException();
 
-			public override bool TryAppendIdentifier(Utf8StringBuilder builder)
+			public bool TryAppendIdentifier(Utf8StringBuilder builder)
 			{
 				builder.Append(_locator.Path);
 				return true;
@@ -35,7 +35,7 @@ namespace EpicGames.Horde.Tests
 		[TestMethod]
 		public void HeaderSerialization()
 		{
-			List<BlobHandle> refs = new List<BlobHandle>();
+			List<IBlobHandle> refs = new List<IBlobHandle>();
 			refs.Add(new DummyHandle("hello"));
 			refs.Add(new DummyHandle("world"));
 
diff --git a/Engine/Source/Programs/Shared/EpicGames.Horde.Tests/BundleStoreTests.cs b/Engine/Source/Programs/Shared/EpicGames.Horde.Tests/BundleStoreTests.cs
--- a/Engine/Source/Programs/Shared/EpicGames.Horde.Tests/BundleStoreTests.cs
+++ b/Engine/Source/Programs/Shared/EpicGames.Horde.Tests/BundleStoreTests.cs
@@ -62,10 +62,10 @@ namespace EpicGames.Horde.Tests
 			await using IStorageWriter writer = store.CreateWriter(options: new BundleOptions { CompressionFormat = BundleCompressionFormat.None });
 
 			TextNode node = new TextNode("Hello world");
-			BlobHandle handle = await writer.FlushAsync(node, CancellationToken.None);
+			IBlobHandle handle = await writer.FlushAsync(node, CancellationToken.None);
 
-			BlobHandle bundleHandle = store.CreateBlobHandle(handle.GetLocator().BaseLocator);
-			using BlobData blobData = await bundleHandle.ReadBlobDataAsync();
+			IBlobHandle bundleHandle = store.CreateBlobHandle(handle.GetLocator().BaseLocator);
+			using BlobData blobData = await bundleHandle.ReadAsync();
 
 			return blobData.Data.ToArray();
 		}
diff --git a/Engine/Source/Programs/Shared/EpicGames.Horde.Tests/BundleTests.cs b/Engine/Source/Programs/Shared/EpicGames.Horde.Tests/BundleTests.cs
--- a/Engine/Source/Programs/Shared/EpicGames.Horde.Tests/BundleTests.cs
+++ b/Engine/Source/Programs/Shared/EpicGames.Horde.Tests/BundleTests.cs
@@ -219,9 +219,9 @@ namespace EpicGames.Horde.Tests
 			Assert.AreEqual(1, _memoryStore.Blobs.Count);
 
 			// Check the ref
-			BlobHandle refTarget =  await store.ReadRefTargetAsync(refName);
-			BlobHandle bundleTarget = store.CreateBlobHandle(refTarget.GetLocator().BaseLocator);
-			using BlobData bundleData = await bundleTarget.ReadBlobDataAsync();
+			IBlobHandle refTarget =  await store.ReadRefTargetAsync(refName);
+			IBlobHandle bundleTarget = store.CreateBlobHandle(refTarget.GetLocator().BaseLocator);
+			using BlobData bundleData = await bundleTarget.ReadAsync();
 
 			BundleHeader bundleHeader = BundleHeader.Read(bundleData.Data);
 			Assert.AreEqual(0, bundleHeader.Imports.Count);
diff --git a/Engine/Source/Programs/Shared/EpicGames.Horde.Tests/ComputeSocketTests.cs b/Engine/Source/Programs/Shared/EpicGames.Horde.Tests/ComputeSocketTests.cs
--- a/Engine/Source/Programs/Shared/EpicGames.Horde.Tests/ComputeSocketTests.cs
+++ b/Engine/Source/Programs/Shared/EpicGames.Horde.Tests/ComputeSocketTests.cs
@@ -163,7 +163,7 @@ namespace EpicGames.Horde.Tests
 						DirectoryNode root = new DirectoryNode();
 						root.AddDirectory(new DirectoryEntry("subdir", directory.Length, directoryRef));
 
-						BlobHandle handle = await treeWriter.FlushAsync(root);
+						IBlobHandle handle = await treeWriter.FlushAsync(root);
 						await channel.UploadFilesAsync("", handle.GetLocator(), storage);
 
 						Assert.IsTrue(FileReference.Exists(file));
diff --git a/Engine/Source/Programs/Shared/EpicGames.Horde/Compute/AgentMessage.cs b/Engine/Source/Programs/Shared/EpicGames.Horde/Compute/AgentMessage.cs
--- a/Engine/Source/Programs/Shared/EpicGames.Horde/Compute/AgentMessage.cs
+++ b/Engine/Source/Programs/Shared/EpicGames.Horde/Compute/AgentMessage.cs
@@ -652,8 +652,8 @@ namespace EpicGames.Horde.Compute
 		/// <param name="cancellationToken">Cancellation token for the operation</param>
 		public static async Task SendBlobDataAsync(this AgentMessageChannel channel, BlobLocator locator, int offset, int length, IStorageClient storage, CancellationToken cancellationToken = default)
 		{
-			BlobHandle handle = storage.CreateBlobHandle(locator);
-			using Stream stream = await handle.OpenAsync(offset, (length == 0)? null : length, cancellationToken);
+			IBlobHandle handle = storage.CreateBlobHandle(locator);
+			using Stream stream = await handle.OpenBodyAsync(offset, (length == 0)? null : length, cancellationToken);
 
 			const int MaxChunkSize = 512 * 1024;
 			for (int chunkOffset = 0; ;)
diff --git a/Engine/Source/Programs/Shared/EpicGames.Horde/Compute/AgentMessageHandler.cs b/Engine/Source/Programs/Shared/EpicGames.Horde/Compute/AgentMessageHandler.cs
--- a/Engine/Source/Programs/Shared/EpicGames.Horde/Compute/AgentMessageHandler.cs
+++ b/Engine/Source/Programs/Shared/EpicGames.Horde/Compute/AgentMessageHandler.cs
@@ -139,7 +139,7 @@ namespace EpicGames.Horde.Compute
 			using BundleCache cache = new BundleCache(new BundleCacheOptions { HeaderCacheSize = 10 * 1024 * 1024, PacketCacheSize = 128 * 1024 * 1024 });
 			using BundleStorageClient store = new BundleStorageClient(innerStore, cache, _logger);
 
-			BlobHandle handle = store.CreateBlobHandle(locator);
+			IBlobHandle handle = store.CreateBlobHandle(locator);
 			DirectoryNode directoryNode = await handle.ReadNodeAsync<DirectoryNode>(cancellationToken);
 
 			DirectoryReference outputDir = DirectoryReference.Combine(_sandboxDir, path);
diff --git a/Engine/Source/Programs/Shared/EpicGames.Horde/Compute/AgentStorageClient.cs b/Engine/Source/Programs/Shared/EpicGames.Horde/Compute/AgentStorageClient.cs
--- a/Engine/Source/Programs/Shared/EpicGames.Horde/Compute/AgentStorageClient.cs
+++ b/Engine/Source/Programs/Shared/EpicGames.Horde/Compute/AgentStorageClient.cs
@@ -52,7 +52,7 @@ namespace EpicGames.Horde.Compute
 			{
 				// TODO: Want to pass 0 for length to read here (meaning entire blob), but older streams misinterpret this as a 0 byte read.
 				ReadOnlyMemory<byte> data = await _channel.ReadBlobAsync(locator.ToString(), 0, 128 * 1024 * 1024, cancellationToken);
-				return new BlobData(BlobType.Leaf, data, Array.Empty<BlobHandle>());
+				return new BlobData(BlobType.Leaf, data, Array.Empty<IBlobHandle>());
 			}
 			finally
 			{
@@ -61,7 +61,7 @@ namespace EpicGames.Horde.Compute
 		}
 
 		/// <inheritdoc/>
-		public override ValueTask<BlobHandle> WriteBlobAsync(BlobType type, Stream stream, IReadOnlyList<BlobHandle> references, string? basePath = null, CancellationToken cancellationToken = default) => throw new NotSupportedException();
+		public override ValueTask<IBlobHandle> WriteBlobAsync(BlobType type, Stream stream, IReadOnlyList<IBlobHandle> references, string? basePath = null, CancellationToken cancellationToken = default) => throw new NotSupportedException();
 
 		/// <inheritdoc/>
 		public override ValueTask<Uri?> TryGetReadRedirectAsync(BlobLocator locator, CancellationToken cancellationToken = default) => default;
@@ -74,10 +74,10 @@ namespace EpicGames.Horde.Compute
 		#region Aliases
 
 		/// <inheritdoc/>
-		public override Task AddAliasAsync(string name, BlobHandle target, int rank, ReadOnlyMemory<byte> data, CancellationToken cancellationToken = default) => throw new NotSupportedException();
+		public override Task AddAliasAsync(string name, IBlobHandle target, int rank, ReadOnlyMemory<byte> data, CancellationToken cancellationToken = default) => throw new NotSupportedException();
 
 		/// <inheritdoc/>
-		public override Task RemoveAliasAsync(string name, BlobHandle target, CancellationToken cancellationToken = default) => throw new NotSupportedException();
+		public override Task RemoveAliasAsync(string name, IBlobHandle target, CancellationToken cancellationToken = default) => throw new NotSupportedException();
 
 		/// <inheritdoc/>
 		public override Task<BlobAlias[]> FindAliasesAsync(string name, int? maxResults, CancellationToken cancellationToken = default) => throw new NotSupportedException();
@@ -90,10 +90,10 @@ namespace EpicGames.Horde.Compute
 		public override Task<bool> DeleteRefAsync(RefName name, CancellationToken cancellationToken = default) => throw new NotSupportedException();
 
 		/// <inheritdoc/>
-		public override Task<BlobHandle?> TryReadRefAsync(RefName name, RefCacheTime cacheTime = default, CancellationToken cancellationToken = default) => throw new NotSupportedException();
+		public override Task<IBlobHandle?> TryReadRefAsync(RefName name, RefCacheTime cacheTime = default, CancellationToken cancellationToken = default) => throw new NotSupportedException();
 
 		/// <inheritdoc/>
-		public override Task WriteRefAsync(RefName name, BlobHandle target, RefOptions? options = null, CancellationToken cancellationToken = default) => throw new NotSupportedException();
+		public override Task WriteRefAsync(RefName name, IBlobHandle target, RefOptions? options = null, CancellationToken cancellationToken = default) => throw new NotSupportedException();
 
 		#endregion
 
diff --git a/Engine/Source/Programs/Shared/EpicGames.Horde/Storage/BlobAlias.cs b/Engine/Source/Programs/Shared/EpicGames.Horde/Storage/BlobAlias.cs
--- a/Engine/Source/Programs/Shared/EpicGames.Horde/Storage/BlobAlias.cs
+++ b/Engine/Source/Programs/Shared/EpicGames.Horde/Storage/BlobAlias.cs
@@ -10,5 +10,5 @@ namespace EpicGames.Horde.Storage
 	/// <param name="Target">Handle to the target blob for the alias</param>
 	/// <param name="Rank">Rank for the alias</param>
 	/// <param name="Data">Data stored inline with the alias</param>
-	public record class BlobAlias(BlobHandle Target, int Rank, ReadOnlyMemory<byte> Data);
+	public record class BlobAlias(IBlobHandle Target, int Rank, ReadOnlyMemory<byte> Data);
 }
diff --git a/Engine/Source/Programs/Shared/EpicGames.Horde/Storage/BlobData.cs b/Engine/Source/Programs/Shared/EpicGames.Horde/Storage/BlobData.cs
--- a/Engine/Source/Programs/Shared/EpicGames.Horde/Storage/BlobData.cs
+++ b/Engine/Source/Programs/Shared/EpicGames.Horde/Storage/BlobData.cs
@@ -24,12 +24,12 @@ namespace EpicGames.Horde.Storage
 		/// <summary>
 		/// Handles to referenced blobs
 		/// </summary>
-		public IReadOnlyList<BlobHandle> Refs { get; }
+		public IReadOnlyList<IBlobHandle> Refs { get; }
 
 		/// <summary>
 		/// Constructor
 		/// </summary>
-		public BlobData(BlobType type, ReadOnlyMemory<byte> data, IReadOnlyList<BlobHandle> refs)
+		public BlobData(BlobType type, ReadOnlyMemory<byte> data, IReadOnlyList<IBlobHandle> refs)
 		{
 			Type = type;
 			Data = data;
@@ -62,7 +62,7 @@ namespace EpicGames.Horde.Storage
 		/// <summary>
 		/// Constructor
 		/// </summary>
-		public ReadOnlyMemoryOwnerBlobData(BlobType type, IReadOnlyMemoryOwner<byte> owner, IReadOnlyList<BlobHandle> refs)
+		public ReadOnlyMemoryOwnerBlobData(BlobType type, IReadOnlyMemoryOwner<byte> owner, IReadOnlyList<IBlobHandle> refs)
 			: base(type, owner.Memory, refs)
 		{
 			_owner = owner;
diff --git a/Engine/Source/Programs/Shared/EpicGames.Horde/Storage/BlobHandle.cs b/Engine/Source/Programs/Shared/EpicGames.Horde/Storage/BlobHandle.cs
deleted file mode 100644
diff --git a/Engine/Source/Programs/Shared/EpicGames.Horde/Storage/Bundles/V1/BundleReader.cs b/Engine/Source/Programs/Shared/EpicGames.Horde/Storage/Bundles/V1/BundleReader.cs
--- a/Engine/Source/Programs/Shared/EpicGames.Horde/Storage/Bundles/V1/BundleReader.cs
+++ b/Engine/Source/Programs/Shared/EpicGames.Horde/Storage/Bundles/V1/BundleReader.cs
@@ -443,8 +443,8 @@ namespace EpicGames.Horde.Storage.Bundles.V1
 			int prefetchSize = _cache != null ? DefaultFetchSize : DefaultUncachedFetchSize;
 			for (; ; )
 			{
-				BlobHandle handle = _store.CreateBlobHandle(bundle.Locator);
-				await using (Stream stream = await handle.OpenAsync(0, prefetchSize, cancellationToken))
+				IBlobHandle handle = _store.CreateBlobHandle(bundle.Locator);
+				await using (Stream stream = await handle.OpenBodyAsync(0, prefetchSize, cancellationToken))
 				{
 					// Read the header data
 					byte[] prelude = new byte[Bundle.SignatureLength];
@@ -499,8 +499,8 @@ namespace EpicGames.Horde.Storage.Bundles.V1
 			Interlocked.Increment(ref _numPacketReads);
 			Interlocked.Add(ref _numBytesRead, maxOffset - minOffset);
 
-			BlobHandle handle = _store.CreateBlobHandle(bundleInfo.Locator);
-			await using (Stream stream = await handle.OpenAsync(bundleInfo.HeaderLength + minOffset, maxOffset - minOffset, cancellationToken))
+			IBlobHandle handle = _store.CreateBlobHandle(bundleInfo.Locator);
+			await using (Stream stream = await handle.OpenBodyAsync(bundleInfo.HeaderLength + minOffset, maxOffset - minOffset, cancellationToken))
 			{
 				// Copy all the packets that have been read into separate buffers, so we can cache them indidually.
 				for (int packetIdx = minPacketIdx; packetIdx <= maxPacketIdx; packetIdx++)
@@ -546,7 +546,7 @@ namespace EpicGames.Horde.Storage.Bundles.V1
 			BundleHeader header = await ReadHeaderAsync(bundleLocator, cancellationToken);
 			BundleExport export = header.Exports[exportIdx];
 
-			List<BlobHandle> refs = new List<BlobHandle>(export.References.Count);
+			List<IBlobHandle> refs = new List<IBlobHandle>(export.References.Count);
 			foreach (BundleExportRef reference in export.References)
 			{
 				BlobLocator importBlob;
diff --git a/Engine/Source/Programs/Shared/EpicGames.Horde/Storage/Bundles/V1/BundleWriter.cs b/Engine/Source/Programs/Shared/EpicGames.Horde/Storage/Bundles/V1/BundleWriter.cs
--- a/Engine/Source/Programs/Shared/EpicGames.Horde/Storage/Bundles/V1/BundleWriter.cs
+++ b/Engine/Source/Programs/Shared/EpicGames.Horde/Storage/Bundles/V1/BundleWriter.cs
@@ -16,9 +16,9 @@ using Microsoft.Extensions.Logging;
 namespace EpicGames.Horde.Storage.Bundles.V1
 {
 	/// <summary>
-	/// Implementation of <see cref="BlobHandle"/> for nodes which can be read from storage
+	/// Implementation of <see cref="IBlobHandle"/> for nodes which can be read from storage
 	/// </summary>
-	sealed class FlushedNodeHandle : BlobHandle
+	sealed class FlushedNodeHandle : IBlobHandle
 	{
 		readonly BundleReader _reader;
 
@@ -26,12 +26,12 @@ namespace EpicGames.Horde.Storage.Bundles.V1
 		public int ExportIdx { get; }
 
 		/// <inheritdoc/>
-		public override BlobHandle Outer { get; }
+		public IBlobHandle Outer { get; }
 
 		/// <summary>
 		/// Constructor
 		/// </summary>
-		public FlushedNodeHandle(BundleReader reader, BlobLocator bundleLocator, BlobHandle bundleHandle, int exportIdx)
+		public FlushedNodeHandle(BundleReader reader, BlobLocator bundleLocator, IBlobHandle bundleHandle, int exportIdx)
 		{
 			Debug.Assert(!bundleLocator.CanUnwrap());
 
@@ -41,7 +41,7 @@ namespace EpicGames.Horde.Storage.Bundles.V1
 			ExportIdx = exportIdx;
 		}
 
-		public FlushedNodeHandle(BundleReader reader, BlobLocator bundleLocator, BlobHandle bundleHandle, ReadOnlySpan<byte> fragment)
+		public FlushedNodeHandle(BundleReader reader, BlobLocator bundleLocator, IBlobHandle bundleHandle, ReadOnlySpan<byte> fragment)
 		{
 			Debug.Assert(!bundleLocator.CanUnwrap());
 
@@ -58,17 +58,17 @@ namespace EpicGames.Horde.Storage.Bundles.V1
 		}
 
 		/// <inheritdoc/>
-		public override bool TryAppendIdentifier(Utf8StringBuilder builder)
+		public bool TryAppendIdentifier(Utf8StringBuilder builder)
 		{
 			builder.Append(ExportIdx);
 			return true;
 		}
 
 		/// <inheritdoc/>
-		public override ValueTask<BlobData> ReadBlobDataAsync(CancellationToken cancellationToken = default) => _reader.ReadNodeDataAsync(BundleLocator, ExportIdx, cancellationToken);
+		public ValueTask<BlobData> ReadAsync(CancellationToken cancellationToken = default) => _reader.ReadNodeDataAsync(BundleLocator, ExportIdx, cancellationToken);
 
 		/// <inheritdoc/>
-		public override ValueTask FlushAsync(CancellationToken cancellationToken = default) => new ValueTask();
+		public ValueTask FlushAsync(CancellationToken cancellationToken = default) => new ValueTask();
 
 		/// <inheritdoc/>
 		public override bool Equals(object? obj) => obj is FlushedNodeHandle other && BundleLocator == other.BundleLocator && ExportIdx == other.ExportIdx;
@@ -84,7 +84,7 @@ namespace EpicGames.Horde.Storage.Bundles.V1
 	public sealed class BundleWriter : IStorageWriter
 	{
 		// Information about a unique output node. Note that multiple node refs may de-duplicate to the same output node.
-		internal class PendingNode : BlobHandle
+		internal class PendingNode : IBlobHandle
 		{
 			readonly BundleReader _reader;
 
@@ -97,15 +97,15 @@ namespace EpicGames.Horde.Storage.Bundles.V1
 			public readonly int Packet;
 			public readonly int Offset;
 			public readonly int Length;
-			public readonly BlobHandle[] Refs;
+			public readonly IBlobHandle[] Refs;
 			public readonly AliasInfo[] Aliases;
 
 			public PendingBundle? PendingBundle => _pendingBundle;
 			public FlushedNodeHandle? FlushedNodeHandle => _flushedHandle;
 
-			public override BlobHandle? Outer => (_flushedHandle != null) ? _flushedHandle.Outer : throw new NotSupportedException();
+			public IBlobHandle? Outer => (_flushedHandle != null) ? _flushedHandle.Outer : throw new NotSupportedException();
 
-			public PendingNode(BundleReader reader, BlobType blobType, int packet, int offset, int length, IReadOnlyList<BlobHandle> refs, IReadOnlyList<AliasInfo> aliases, PendingBundle pendingBundle)
+			public PendingNode(BundleReader reader, BlobType blobType, int packet, int offset, int length, IReadOnlyList<IBlobHandle> refs, IReadOnlyList<AliasInfo> aliases, PendingBundle pendingBundle)
 			{
 				_reader = reader;
 
@@ -120,7 +120,7 @@ namespace EpicGames.Horde.Storage.Bundles.V1
 			}
 
 			/// <inheritdoc/>
-			public override bool TryAppendIdentifier(Utf8StringBuilder builder)
+			public bool TryAppendIdentifier(Utf8StringBuilder builder)
 			{
 				return _flushedHandle?.TryAppendIdentifier(builder) ?? false;
 			}
@@ -136,7 +136,7 @@ namespace EpicGames.Horde.Storage.Bundles.V1
 			}
 
 			/// <inheritdoc/>
-			public override async ValueTask<BlobData> ReadBlobDataAsync(CancellationToken cancellationToken = default)
+			public async ValueTask<BlobData> ReadAsync(CancellationToken cancellationToken = default)
 			{
 				if (_flushedHandle == null)
 				{
@@ -150,18 +150,18 @@ namespace EpicGames.Horde.Storage.Bundles.V1
 					}
 				}
 
-				return await _flushedHandle!.ReadBlobDataAsync(cancellationToken);
+				return await _flushedHandle!.ReadAsync(cancellationToken);
 			}
 
 			/// <inheritdoc/>
-			public override async ValueTask FlushAsync(CancellationToken cancellationToken = default)
+			public async ValueTask FlushAsync(CancellationToken cancellationToken = default)
 			{
 				if (_flushedHandle != null)
 				{
 					return;
 				}
 
-				foreach (BlobHandle nodeRef in Refs)
+				foreach (IBlobHandle nodeRef in Refs)
 				{
 					await nodeRef.FlushAsync(cancellationToken);
 				}
@@ -298,7 +298,7 @@ namespace EpicGames.Horde.Storage.Bundles.V1
 			}
 
 			// Finish a node write
-			public PendingNode WriteNode(BlobType blobType, int size, IReadOnlyList<BlobHandle> refs, IReadOnlyList<AliasInfo> aliases)
+			public PendingNode WriteNode(BlobType blobType, int size, IReadOnlyList<IBlobHandle> refs, IReadOnlyList<AliasInfo> aliases)
 			{
 				PendingNode pendingNode = new PendingNode(_treeReader, blobType, _currentPacketIdx, _currentPacketLength, (int)size, refs, aliases, this);
 				_currentPacketLength += pendingNode.Length;
@@ -411,7 +411,7 @@ namespace EpicGames.Horde.Storage.Bundles.V1
 					(BundleHeader header, List<ReadOnlyMemory<byte>> packets) = CreateBundle();
 
 					// Write the bundle to storage
-					BlobHandle[] imports = new BlobHandle[header.Imports.Count];
+					IBlobHandle[] imports = new IBlobHandle[header.Imports.Count];
 					for (int idx = 0; idx < header.Imports.Count; idx++)
 					{
 						imports[idx] = store.CreateBlobHandle(new BlobLocator(header.Imports[idx].Path));
@@ -427,7 +427,7 @@ namespace EpicGames.Horde.Storage.Bundles.V1
 					}
 
 					// Write it
-					BlobHandle handle;
+					IBlobHandle handle;
 					using (ReadOnlySequenceStream stream = new ReadOnlySequenceStream(sequence.Construct()))
 					{
 						handle = await store.WriteBlobAsync(Bundle.BlobType, stream, imports, basePath);
@@ -480,10 +480,10 @@ namespace EpicGames.Horde.Storage.Bundles.V1
 				Dictionary<BlobType, int> typeToIndex = new Dictionary<BlobType, int>();
 
 				// Map of node handle to reference
-				Dictionary<BlobHandle, BundleExportRef> nodeHandleToExportRef = new Dictionary<BlobHandle, BundleExportRef>();
+				Dictionary<IBlobHandle, BundleExportRef> nodeHandleToExportRef = new Dictionary<IBlobHandle, BundleExportRef>();
 				for (int exportIdx = 0; exportIdx < _queue.Count; exportIdx++)
 				{
-					BlobHandle handle = _queue[exportIdx];
+					IBlobHandle handle = _queue[exportIdx];
 					nodeHandleToExportRef[handle] = new BundleExportRef(-1, exportIdx);
 				}
 
@@ -494,7 +494,7 @@ namespace EpicGames.Horde.Storage.Bundles.V1
 					int typeIdx = FindOrAddItemIndex(nodeInfo.BlobType, types, typeToIndex);
 
 					List<BundleExportRef> exportRefs = new List<BundleExportRef>();
-					foreach (BlobHandle handle in nodeInfo.Refs)
+					foreach (IBlobHandle handle in nodeInfo.Refs)
 					{
 						BundleExportRef exportRef;
 						if (!nodeHandleToExportRef.TryGetValue(handle, out exportRef))
@@ -744,7 +744,7 @@ namespace EpicGames.Horde.Storage.Bundles.V1
 		/// <param name="aliases">Aliases for the node</param>
 		/// <param name="cancellationToken">Cancellation token for the operation</param>
 		/// <returns>Handle to the written node</returns>
-		public async ValueTask<BlobHandle> WriteBlobAsync(BlobType type, int size, IReadOnlyList<BlobHandle> references, IReadOnlyList<AliasInfo> aliases, CancellationToken cancellationToken = default)
+		public async ValueTask<IBlobHandle> WriteBlobAsync(BlobType type, int size, IReadOnlyList<IBlobHandle> references, IReadOnlyList<AliasInfo> aliases, CancellationToken cancellationToken = default)
 		{
 			PendingBundle currentBundle = GetCurrentBundle();
 
@@ -753,7 +753,7 @@ namespace EpicGames.Horde.Storage.Bundles.V1
 			TraceLogger?.LogInformation("Added new node for {NodeKey} in bundle {BundleId}", pendingNode, currentBundle.BundleId);
 
 			// Add dependencies on all bundles containing a dependent node
-			foreach (BlobHandle reference in references)
+			foreach (IBlobHandle reference in references)
 			{
 				PendingNode? pendingReference = reference as PendingNode;
 				if (pendingReference?.PendingBundle != null)
@@ -803,7 +803,7 @@ namespace EpicGames.Horde.Storage.Bundles.V1
 		/// <param name="root">Root for the tree</param>
 		/// <param name="cancellationToken">Cancellation token for the operation</param>
 		/// <returns></returns>
-		public async Task<BlobHandle> FlushAsync(Node root, CancellationToken cancellationToken = default)
+		public async Task<IBlobHandle> FlushAsync(Node root, CancellationToken cancellationToken = default)
 		{
 			return await NodeRefExtensions.FlushAsync(this, root, cancellationToken);
 		}
diff --git a/Engine/Source/Programs/Shared/EpicGames.Horde/Storage/Clients/BundleStorageClient.cs b/Engine/Source/Programs/Shared/EpicGames.Horde/Storage/Clients/BundleStorageClient.cs
--- a/Engine/Source/Programs/Shared/EpicGames.Horde/Storage/Clients/BundleStorageClient.cs
+++ b/Engine/Source/Programs/Shared/EpicGames.Horde/Storage/Clients/BundleStorageClient.cs
@@ -19,14 +19,14 @@ namespace EpicGames.Horde.Storage.Clients
 		/// <summary>
 		/// Handle to a bundle object
 		/// </summary>
-		class BundleHandle : BlobHandle
+		class BundleHandle : IBlobHandle
 		{
 			readonly BundleStorageClient _storageClient;
 			readonly BlobLocator _locator;
-			List<BlobHandle>? _refs;
+			List<IBlobHandle>? _refs;
 
 			/// <inheritdoc/>
-			public override BlobHandle? Outer => null;
+			public IBlobHandle? Outer => null;
 
 			/// <summary>
 			/// Constructor
@@ -38,14 +38,14 @@ namespace EpicGames.Horde.Storage.Clients
 			}
 
 			/// <inheritdoc/>
-			public override ValueTask<BlobType> GetTypeAsync(CancellationToken cancellationToken = default) => new ValueTask<BlobType>(Bundle.BlobType);
+			public ValueTask<BlobType> ReadTypeAsync(CancellationToken cancellationToken = default) => new ValueTask<BlobType>(Bundle.BlobType);
 
 			/// <inheritdoc/>
-			public override async ValueTask<IReadOnlyList<BlobHandle>> GetRefsAsync(CancellationToken cancellationToken = default)
+			public async ValueTask<IReadOnlyList<IBlobHandle>> ReadImportsAsync(CancellationToken cancellationToken = default)
 			{
 				if (_refs == null)
 				{
-					List<BlobHandle> refs = new List<BlobHandle>();
+					List<IBlobHandle> refs = new List<IBlobHandle>();
 
 					Bundles.V1.BundleHeader header = await _storageClient.ReadHeaderAsync(_locator, cancellationToken);
 					foreach (BlobLocator import in header.Imports)
@@ -59,23 +59,23 @@ namespace EpicGames.Horde.Storage.Clients
 			}
 
 			/// <inheritdoc/>
-			public override Task<Stream> OpenAsync(int offset = 0, int? length = null, CancellationToken cancellationToken = default)
+			public Task<Stream> OpenBodyAsync(int offset = 0, int? length = null, CancellationToken cancellationToken = default)
 			{
 				return _storageClient.OpenBlobAsync(_locator, offset, length, cancellationToken);
 			}
 
 			/// <inheritdoc/>
-			public override async ValueTask<BlobData> ReadBlobDataAsync(CancellationToken cancellationToken = default)
+			public async ValueTask<BlobData> ReadAsync(CancellationToken cancellationToken = default)
 			{
 				using (Stream stream = await _storageClient.OpenBlobAsync(_locator, 0, cancellationToken: cancellationToken))
 				{
 					byte[] data = await stream.ReadAllBytesAsync(cancellationToken);
-					return new BlobData(Bundle.BlobType, data, await GetRefsAsync(cancellationToken));
+					return new BlobData(Bundle.BlobType, data, await ReadImportsAsync(cancellationToken));
 				}
 			}
 
 			/// <inheritdoc/>
-			public override bool TryAppendIdentifier(Utf8StringBuilder builder)
+			public bool TryAppendIdentifier(Utf8StringBuilder builder)
 			{
 				builder.Append(_locator.Path);
 				return true;
@@ -136,12 +136,12 @@ namespace EpicGames.Horde.Storage.Clients
 		/// <inheritdoc/>
 		public async ValueTask<BlobData> ReadBlobAsync(BlobLocator locator, CancellationToken cancellationToken = default)
 		{
-			BlobHandle handle = CreateBlobHandle(locator);
-			return await handle.ReadBlobDataAsync(cancellationToken);
+			IBlobHandle handle = CreateBlobHandle(locator);
+			return await handle.ReadAsync(cancellationToken);
 		}
 
 		/// <inheritdoc/>
-		public async ValueTask<BlobHandle> WriteBlobAsync(BlobType type, Stream stream, IReadOnlyList<BlobHandle> references, string? basePath = null, CancellationToken cancellationToken = default)
+		public async ValueTask<IBlobHandle> WriteBlobAsync(BlobType type, Stream stream, IReadOnlyList<IBlobHandle> references, string? basePath = null, CancellationToken cancellationToken = default)
 		{
 			if (type == Bundle.BlobType)
 			{
@@ -180,12 +180,12 @@ namespace EpicGames.Horde.Storage.Clients
 		/// Read a bundle from the underlying storage
 		/// </summary>
 		async Task<Stream> OpenBundleAsync(BlobLocator locator, int offset, int? length, CancellationToken cancellationToken)
-			=> await _inner.CreateBlobHandle(locator).OpenAsync(offset, length, cancellationToken);
+			=> await _inner.CreateBlobHandle(locator).OpenBodyAsync(offset, length, cancellationToken);
 
 		/// <summary>
 		/// Write a bundle to the underlying storage
 		/// </summary>
-		ValueTask<BlobHandle> WriteBundleAsync(Stream stream, IReadOnlyList<BlobHandle> references, string? basePath = null, CancellationToken cancellationToken = default)
+		ValueTask<IBlobHandle> WriteBundleAsync(Stream stream, IReadOnlyList<IBlobHandle> references, string? basePath = null, CancellationToken cancellationToken = default)
 			=> _inner.WriteBlobAsync(Bundle.BlobType, stream, references, basePath, cancellationToken);
 
 		/// <inheritdoc/>
@@ -196,7 +196,7 @@ namespace EpicGames.Horde.Storage.Clients
 		#region Nodes
 
 		/// <inheritdoc/>
-		public BlobHandle CreateBlobHandle(BlobLocator locator)
+		public IBlobHandle CreateBlobHandle(BlobLocator locator)
 		{
 			if (locator.TryUnwrap(out BlobLocator baseLocator, out Utf8String fragment))
 			{
@@ -219,11 +219,11 @@ namespace EpicGames.Horde.Storage.Clients
 		#region Aliases
 
 		/// <inheritdoc/>
-		public Task AddAliasAsync(string name, BlobHandle handle, int rank, ReadOnlyMemory<byte> data, CancellationToken cancellationToken)
+		public Task AddAliasAsync(string name, IBlobHandle handle, int rank, ReadOnlyMemory<byte> data, CancellationToken cancellationToken)
 			=> _inner.AddAliasAsync(name, handle, rank, data, cancellationToken);
 
 		/// <inheritdoc/>
-		public Task RemoveAliasAsync(string name, BlobHandle handle, CancellationToken cancellationToken)
+		public Task RemoveAliasAsync(string name, IBlobHandle handle, CancellationToken cancellationToken)
 			=> _inner.RemoveAliasAsync(name, handle, cancellationToken);
 
 		/// <inheritdoc/>
@@ -247,9 +247,9 @@ namespace EpicGames.Horde.Storage.Clients
 			=> _inner.DeleteRefAsync(name, cancellationToken);
 
 		/// <inheritdoc/>
-		public async Task<BlobHandle?> TryReadRefAsync(RefName name, RefCacheTime cacheTime = default, CancellationToken cancellationToken = default)
+		public async Task<IBlobHandle?> TryReadRefAsync(RefName name, RefCacheTime cacheTime = default, CancellationToken cancellationToken = default)
 		{
-			BlobHandle? target = await _inner.TryReadRefAsync(name, cacheTime, cancellationToken);
+			IBlobHandle? target = await _inner.TryReadRefAsync(name, cacheTime, cancellationToken);
 			if (target != null)
 			{
 				target = CreateBlobHandle(target.GetLocator());
@@ -258,7 +258,7 @@ namespace EpicGames.Horde.Storage.Clients
 		}
 
 		/// <inheritdoc/>
-		public Task WriteRefAsync(RefName name, BlobHandle target, RefOptions? options, CancellationToken cancellationToken)
+		public Task WriteRefAsync(RefName name, IBlobHandle target, RefOptions? options, CancellationToken cancellationToken)
 			=> _inner.WriteRefAsync(name, target, options, cancellationToken);
 
 		#endregion
diff --git a/Engine/Source/Programs/Shared/EpicGames.Horde/Storage/Clients/FileStorageClient.cs b/Engine/Source/Programs/Shared/EpicGames.Horde/Storage/Clients/FileStorageClient.cs
--- a/Engine/Source/Programs/Shared/EpicGames.Horde/Storage/Clients/FileStorageClient.cs
+++ b/Engine/Source/Programs/Shared/EpicGames.Horde/Storage/Clients/FileStorageClient.cs
@@ -65,11 +65,11 @@ namespace EpicGames.Horde.Storage.Clients
 		public override async ValueTask<BlobData> ReadBlobAsync(BlobLocator locator, CancellationToken cancellationToken = default)
 		{
 			IReadOnlyMemoryOwner<byte> owner = await _backend.ReadAsync(locator.ToString(), cancellationToken);
-			return new ReadOnlyMemoryOwnerBlobData(BlobType.Leaf, owner, Array.Empty<BlobHandle>());
+			return new ReadOnlyMemoryOwnerBlobData(BlobType.Leaf, owner, Array.Empty<IBlobHandle>());
 		}
 
 		/// <inheritdoc/>
-		public override async ValueTask<BlobHandle> WriteBlobAsync(BlobType type, Stream stream, IReadOnlyList<BlobHandle> references, string? basePath = null, CancellationToken cancellationToken = default)
+		public override async ValueTask<IBlobHandle> WriteBlobAsync(BlobType type, Stream stream, IReadOnlyList<IBlobHandle> references, string? basePath = null, CancellationToken cancellationToken = default)
 		{
 			string path = await _backend.WriteAsync(stream, basePath, cancellationToken);
 			return CreateBlobHandle(new BlobLocator(path));
@@ -86,13 +86,13 @@ namespace EpicGames.Horde.Storage.Clients
 		#region Aliases
 
 		/// <inheritdoc/>
-		public override Task AddAliasAsync(string name, BlobHandle handle, int rank = 0, ReadOnlyMemory<byte> data = default, CancellationToken cancellationToken = default)
+		public override Task AddAliasAsync(string name, IBlobHandle handle, int rank = 0, ReadOnlyMemory<byte> data = default, CancellationToken cancellationToken = default)
 		{
 			throw new NotSupportedException("File storage client does not currently support aliases.");
 		}
 
 		/// <inheritdoc/>
-		public override Task RemoveAliasAsync(string name, BlobHandle handle, CancellationToken cancellationToken = default)
+		public override Task RemoveAliasAsync(string name, IBlobHandle handle, CancellationToken cancellationToken = default)
 		{
 			throw new NotSupportedException("File storage client does not currently support aliases.");
 		}
@@ -120,7 +120,7 @@ namespace EpicGames.Horde.Storage.Clients
 		}
 
 		/// <inheritdoc/>
-		public override async Task<BlobHandle?> TryReadRefAsync(RefName name, RefCacheTime cacheTime = default, CancellationToken cancellationToken = default)
+		public override async Task<IBlobHandle?> TryReadRefAsync(RefName name, RefCacheTime cacheTime = default, CancellationToken cancellationToken = default)
 		{
 			FileReference file = GetRefFile(name);
 			if (!FileReference.Exists(file))
@@ -131,12 +131,12 @@ namespace EpicGames.Horde.Storage.Clients
 			_logger.LogInformation("Reading {File}", file);
 			string[] lines = await FileReference.ReadAllLinesAsync(file, cancellationToken);
 
-			BlobHandle handle = CreateBlobHandle(new BlobLocator(lines[0].Trim()));
+			IBlobHandle handle = CreateBlobHandle(new BlobLocator(lines[0].Trim()));
 			return handle;
 		}
 
 		/// <inheritdoc/>
-		public override async Task WriteRefAsync(RefName name, BlobHandle target, RefOptions? options = null, CancellationToken cancellationToken = default)
+		public override async Task WriteRefAsync(RefName name, IBlobHandle target, RefOptions? options = null, CancellationToken cancellationToken = default)
 		{
 			await target.FlushAsync(cancellationToken);
 			BlobLocator locator = target.GetLocator();
diff --git a/Engine/Source/Programs/Shared/EpicGames.Horde/Storage/Clients/HttpStorageClient.cs b/Engine/Source/Programs/Shared/EpicGames.Horde/Storage/Clients/HttpStorageClient.cs
--- a/Engine/Source/Programs/Shared/EpicGames.Horde/Storage/Clients/HttpStorageClient.cs
+++ b/Engine/Source/Programs/Shared/EpicGames.Horde/Storage/Clients/HttpStorageClient.cs
@@ -21,13 +21,13 @@ namespace EpicGames.Horde.Storage.Clients
 	/// </summary>
 	public sealed class HttpStorageClient : IStorageClient
 	{
-		class Handle : BlobHandle
+		class Handle : IBlobHandle
 		{
 			readonly HttpStorageClient _httpStorageClient;
 			readonly BlobLocator _locator;
 
 			/// <inheritdoc/>
-			public override BlobHandle? Outer => null;
+			public IBlobHandle? Outer => null;
 
 			public Handle(HttpStorageClient httpStorageClient, BlobLocator locator)
 			{
@@ -36,15 +36,15 @@ namespace EpicGames.Horde.Storage.Clients
 			}
 
 			/// <inheritdoc/>
-			public override Task<Stream> OpenAsync(int offset = 0, int? length = null, CancellationToken cancellationToken = default)
+			public Task<Stream> OpenBodyAsync(int offset = 0, int? length = null, CancellationToken cancellationToken = default)
 				=> _httpStorageClient._backend.OpenAsync(_locator.ToString(), offset, length, cancellationToken);
 
 			/// <inheritdoc/>
-			public override ValueTask<BlobData> ReadBlobDataAsync(CancellationToken cancellationToken = default)
+			public ValueTask<BlobData> ReadAsync(CancellationToken cancellationToken = default)
 				=> _httpStorageClient.ReadBlobAsync(_locator, cancellationToken);
 
 			/// <inheritdoc/>
-			public override bool TryAppendIdentifier(Utf8StringBuilder builder)
+			public bool TryAppendIdentifier(Utf8StringBuilder builder)
 			{
 				builder.Append(_locator.Path);
 				return true;
@@ -85,7 +85,7 @@ namespace EpicGames.Horde.Storage.Clients
 		#region Blobs
 
 		/// <inheritdoc/>
-		public BlobHandle CreateBlobHandle(BlobLocator locator)
+		public IBlobHandle CreateBlobHandle(BlobLocator locator)
 		{
 			if (locator.TryUnwrap(out BlobLocator outer, out Utf8String fragment))
 			{
@@ -105,11 +105,11 @@ namespace EpicGames.Horde.Storage.Clients
 		public async ValueTask<BlobData> ReadBlobAsync(BlobLocator locator, CancellationToken cancellationToken = default)
 		{
 			IReadOnlyMemoryOwner<byte> owner = await _backend.ReadAsync(locator.ToString(), cancellationToken);
-			return new ReadOnlyMemoryOwnerBlobData(BlobType.Leaf, owner, Array.Empty<BlobHandle>());
+			return new ReadOnlyMemoryOwnerBlobData(BlobType.Leaf, owner, Array.Empty<IBlobHandle>());
 		}
 
 		/// <inheritdoc/>
-		public async ValueTask<BlobHandle> WriteBlobAsync(BlobType type, Stream stream, IReadOnlyList<BlobHandle> references, string? basePath = null, CancellationToken cancellationToken = default)
+		public async ValueTask<IBlobHandle> WriteBlobAsync(BlobType type, Stream stream, IReadOnlyList<IBlobHandle> references, string? basePath = null, CancellationToken cancellationToken = default)
 		{
 			string path = await _backend.WriteAsync(stream, basePath, cancellationToken);
 			return new Handle(this, new BlobLocator(path));
@@ -135,13 +135,13 @@ namespace EpicGames.Horde.Storage.Clients
 		#region Aliases
 
 		/// <inheritdoc/>
-		public Task AddAliasAsync(string name, BlobHandle target, int rank = 0, ReadOnlyMemory<byte> data = default, CancellationToken cancellationToken = default)
+		public Task AddAliasAsync(string name, IBlobHandle target, int rank = 0, ReadOnlyMemory<byte> data = default, CancellationToken cancellationToken = default)
 		{
 			throw new NotSupportedException("Http storage client does not currently support aliases.");
 		}
 
 		/// <inheritdoc/>
-		public Task RemoveAliasAsync(string name, BlobHandle target, CancellationToken cancellationToken = default)
+		public Task RemoveAliasAsync(string name, IBlobHandle target, CancellationToken cancellationToken = default)
 		{
 			throw new NotSupportedException("Http storage client does not currently support aliases.");
 		}
@@ -170,7 +170,7 @@ namespace EpicGames.Horde.Storage.Clients
 						for (int idx = 0; idx < message.Nodes.Count; idx++)
 						{
 							FindNodeResponse node = message.Nodes[idx];
-							BlobHandle handle = CreateBlobHandle(node.Blob);
+							IBlobHandle handle = CreateBlobHandle(node.Blob);
 							aliases[idx] = new BlobAlias(handle, node.Rank, node.Data);
 						}
 
@@ -211,7 +211,7 @@ namespace EpicGames.Horde.Storage.Clients
 		}
 
 		/// <inheritdoc/>
-		public async Task<BlobHandle?> TryReadRefAsync(RefName name, RefCacheTime cacheTime = default, CancellationToken cancellationToken = default)
+		public async Task<IBlobHandle?> TryReadRefAsync(RefName name, RefCacheTime cacheTime = default, CancellationToken cancellationToken = default)
 		{
 			using (HttpClient httpClient = _createClient())
 			{
@@ -248,7 +248,7 @@ namespace EpicGames.Horde.Storage.Clients
 		}
 
 		/// <inheritdoc/>
-		public async Task WriteRefAsync(RefName name, BlobHandle target, RefOptions? options = null, CancellationToken cancellationToken = default)
+		public async Task WriteRefAsync(RefName name, IBlobHandle target, RefOptions? options = null, CancellationToken cancellationToken = default)
 		{
 			await target.FlushAsync(cancellationToken);
 			BlobLocator locator = target.GetLocator();
diff --git a/Engine/Source/Programs/Shared/EpicGames.Horde/Storage/Clients/KeyValueStorageClient.cs b/Engine/Source/Programs/Shared/EpicGames.Horde/Storage/Clients/KeyValueStorageClient.cs
--- a/Engine/Source/Programs/Shared/EpicGames.Horde/Storage/Clients/KeyValueStorageClient.cs
+++ b/Engine/Source/Programs/Shared/EpicGames.Horde/Storage/Clients/KeyValueStorageClient.cs
@@ -14,13 +14,13 @@ namespace EpicGames.Horde.Storage.Clients
 	/// </summary>
 	public abstract class KeyValueStorageClient : IStorageClient
 	{
-		class Handle : BlobHandle
+		class Handle : IBlobHandle
 		{
 			readonly KeyValueStorageClient _keyValueStorageClient;
 			readonly BlobLocator _locator;
 
 			/// <inheritdoc/>
-			public override BlobHandle? Outer => null;
+			public IBlobHandle? Outer => null;
 
 			/// <summary>
 			/// Constructor
@@ -32,10 +32,10 @@ namespace EpicGames.Horde.Storage.Clients
 			}
 
 			/// <inheritdoc/>
-			public override ValueTask<BlobData> ReadBlobDataAsync(CancellationToken cancellationToken = default) => _keyValueStorageClient.ReadBlobAsync(_locator, cancellationToken);
+			public ValueTask<BlobData> ReadAsync(CancellationToken cancellationToken = default) => _keyValueStorageClient.ReadBlobAsync(_locator, cancellationToken);
 
 			/// <inheritdoc/>
-			public override bool TryAppendIdentifier(Utf8StringBuilder builder)
+			public bool TryAppendIdentifier(Utf8StringBuilder builder)
 			{
 				builder.Append(_locator.Path);
 				return true;
@@ -68,7 +68,7 @@ namespace EpicGames.Horde.Storage.Clients
 		#region Blobs
 
 		/// <inheritdoc/>
-		public virtual BlobHandle CreateBlobHandle(BlobLocator locator)
+		public virtual IBlobHandle CreateBlobHandle(BlobLocator locator)
 		{
 			if (locator.TryUnwrap(out BlobLocator baseLocator, out Utf8String fragment))
 			{
@@ -87,7 +87,7 @@ namespace EpicGames.Horde.Storage.Clients
 		public abstract ValueTask<BlobData> ReadBlobAsync(BlobLocator locator, CancellationToken cancellationToken = default);
 
 		/// <inheritdoc/>
-		public abstract ValueTask<BlobHandle> WriteBlobAsync(BlobType type, Stream stream, IReadOnlyList<BlobHandle> references, string? basePath = null, CancellationToken cancellationToken = default);
+		public abstract ValueTask<IBlobHandle> WriteBlobAsync(BlobType type, Stream stream, IReadOnlyList<IBlobHandle> references, string? basePath = null, CancellationToken cancellationToken = default);
 
 		/// <inheritdoc/>
 		public abstract ValueTask<Uri?> TryGetReadRedirectAsync(BlobLocator locator, CancellationToken cancellationToken = default);
@@ -100,10 +100,10 @@ namespace EpicGames.Horde.Storage.Clients
 		#region Aliases
 
 		/// <inheritdoc/>
-		public abstract Task AddAliasAsync(string name, BlobHandle handle, int rank = 0, ReadOnlyMemory<byte> data = default, CancellationToken cancellationToken = default);
+		public abstract Task AddAliasAsync(string name, IBlobHandle handle, int rank = 0, ReadOnlyMemory<byte> data = default, CancellationToken cancellationToken = default);
 
 		/// <inheritdoc/>
-		public abstract Task RemoveAliasAsync(string name, BlobHandle handle, CancellationToken cancellationToken = default);
+		public abstract Task RemoveAliasAsync(string name, IBlobHandle handle, CancellationToken cancellationToken = default);
 
 		/// <inheritdoc/>
 		public abstract Task<BlobAlias[]> FindAliasesAsync(string name, int? maxResults = null, CancellationToken cancellationToken = default);
@@ -113,10 +113,10 @@ namespace EpicGames.Horde.Storage.Clients
 		#region Refs
 
 		/// <inheritdoc/>
-		public abstract Task<BlobHandle?> TryReadRefAsync(RefName name, RefCacheTime cacheTime = default, CancellationToken cancellationToken = default);
+		public abstract Task<IBlobHandle?> TryReadRefAsync(RefName name, RefCacheTime cacheTime = default, CancellationToken cancellationToken = default);
 
 		/// <inheritdoc/>
-		public abstract Task WriteRefAsync(RefName name, BlobHandle handle, RefOptions? options = null, CancellationToken cancellationToken = default);
+		public abstract Task WriteRefAsync(RefName name, IBlobHandle handle, RefOptions? options = null, CancellationToken cancellationToken = default);
 
 		/// <inheritdoc/>
 		public abstract Task<bool> DeleteRefAsync(RefName name, CancellationToken cancellationToken = default);
diff --git a/Engine/Source/Programs/Shared/EpicGames.Horde/Storage/Clients/MemoryStorageClient.cs b/Engine/Source/Programs/Shared/EpicGames.Horde/Storage/Clients/MemoryStorageClient.cs
--- a/Engine/Source/Programs/Shared/EpicGames.Horde/Storage/Clients/MemoryStorageClient.cs
+++ b/Engine/Source/Programs/Shared/EpicGames.Horde/Storage/Clients/MemoryStorageClient.cs
@@ -46,7 +46,7 @@ namespace EpicGames.Horde.Storage.Clients
 		}
 
 		/// <inheritdoc/>
-		public override async ValueTask<BlobHandle> WriteBlobAsync(BlobType type, Stream stream, IReadOnlyList<BlobHandle> references, string? basePath = null, CancellationToken cancellationToken = default)
+		public override async ValueTask<IBlobHandle> WriteBlobAsync(BlobType type, Stream stream, IReadOnlyList<IBlobHandle> references, string? basePath = null, CancellationToken cancellationToken = default)
 		{
 			BlobLocator locator = new BlobLocator($"mem-{s_sessionId}-{Interlocked.Increment(ref s_uniqueId)}");
 			_blobs[locator] = new BlobData(type, await stream.ReadAllBytesAsync(cancellationToken), references);
@@ -66,7 +66,7 @@ namespace EpicGames.Horde.Storage.Clients
 		#region Aliases
 
 		/// <inheritdoc/>
-		public override Task AddAliasAsync(string name, BlobHandle handle, int rank = 0, ReadOnlyMemory<byte> data = default, CancellationToken cancellationToken = default)
+		public override Task AddAliasAsync(string name, IBlobHandle handle, int rank = 0, ReadOnlyMemory<byte> data = default, CancellationToken cancellationToken = default)
 		{
 			BlobLocator locator = handle.GetLocator();
 			_aliases.AddOrUpdate(name, _ => new AliasListNode(locator, rank, data, null), (_, entry) => new AliasListNode(locator, rank, data, entry));
@@ -74,7 +74,7 @@ namespace EpicGames.Horde.Storage.Clients
 		}
 
 		/// <inheritdoc/>
-		public override Task RemoveAliasAsync(string name, BlobHandle handle, CancellationToken cancellationToken = default)
+		public override Task RemoveAliasAsync(string name, IBlobHandle handle, CancellationToken cancellationToken = default)
 		{
 			BlobLocator locator = handle.GetLocator();
 			for (; ; )
@@ -136,7 +136,7 @@ namespace EpicGames.Horde.Storage.Clients
 			{
 				for (; entry != null; entry = entry.Next)
 				{
-					BlobHandle handle = CreateBlobHandle(entry.Locator);
+					IBlobHandle handle = CreateBlobHandle(entry.Locator);
 					aliases.Add(new BlobAlias(handle, entry.Rank, entry.Data));
 				}
 			}
@@ -151,21 +151,21 @@ namespace EpicGames.Horde.Storage.Clients
 		public override Task<bool> DeleteRefAsync(RefName name, CancellationToken cancellationToken) => Task.FromResult(_refs.TryRemove(name, out _));
 
 		/// <inheritdoc/>
-		public override Task<BlobHandle?> TryReadRefAsync(RefName name, RefCacheTime cacheTime = default, CancellationToken cancellationToken = default)
+		public override Task<IBlobHandle?> TryReadRefAsync(RefName name, RefCacheTime cacheTime = default, CancellationToken cancellationToken = default)
 		{
 			BlobLocator locator;
 			if (_refs.TryGetValue(name, out locator))
 			{
-				return Task.FromResult<BlobHandle?>(CreateBlobHandle(locator));
+				return Task.FromResult<IBlobHandle?>(CreateBlobHandle(locator));
 			}
 			else
 			{
-				return Task.FromResult<BlobHandle?>(null);
+				return Task.FromResult<IBlobHandle?>(null);
 			}
 		}
 
 		/// <inheritdoc/>
-		public override async Task WriteRefAsync(RefName name, BlobHandle target, RefOptions? options = null, CancellationToken cancellationToken = default)
+		public override async Task WriteRefAsync(RefName name, IBlobHandle target, RefOptions? options = null, CancellationToken cancellationToken = default)
 		{
 			await target.FlushAsync(cancellationToken);
 			_refs[name] = target.GetLocator();
diff --git a/Engine/Source/Programs/Shared/EpicGames.Horde/Storage/Clients/SharedStorageClient.cs b/Engine/Source/Programs/Shared/EpicGames.Horde/Storage/Clients/SharedStorageClient.cs
--- a/Engine/Source/Programs/Shared/EpicGames.Horde/Storage/Clients/SharedStorageClient.cs
+++ b/Engine/Source/Programs/Shared/EpicGames.Horde/Storage/Clients/SharedStorageClient.cs
@@ -71,7 +71,7 @@ namespace EpicGames.Horde.Storage.Clients
 		#region Blobs
 
 		/// <inheritdoc/>
-		public BlobHandle CreateBlobHandle(BlobLocator locator) => _inner.CreateBlobHandle(locator);
+		public IBlobHandle CreateBlobHandle(BlobLocator locator) => _inner.CreateBlobHandle(locator);
 
 		/// <inheritdoc/>
 		public IStorageWriter CreateWriter(string? basePath = null) => _inner.CreateWriter(basePath);
@@ -80,7 +80,7 @@ namespace EpicGames.Horde.Storage.Clients
 		public ValueTask<BlobData> ReadBlobAsync(BlobLocator locator, CancellationToken cancellationToken = default) => _inner.ReadBlobAsync(locator, cancellationToken);
 
 		/// <inheritdoc/>
-		public ValueTask<BlobHandle> WriteBlobAsync(BlobType type, Stream stream, IReadOnlyList<BlobHandle> references, string? basePath = null, CancellationToken cancellationToken = default) => _inner.WriteBlobAsync(type, stream, references, basePath, cancellationToken);
+		public ValueTask<IBlobHandle> WriteBlobAsync(BlobType type, Stream stream, IReadOnlyList<IBlobHandle> references, string? basePath = null, CancellationToken cancellationToken = default) => _inner.WriteBlobAsync(type, stream, references, basePath, cancellationToken);
 
 		/// <inheritdoc/>
 		public ValueTask<Uri?> TryGetReadRedirectAsync(BlobLocator locator, CancellationToken cancellationToken = default) => _inner.TryGetReadRedirectAsync(locator, cancellationToken);
@@ -92,10 +92,10 @@ namespace EpicGames.Horde.Storage.Clients
 		#region Alias
 
 		/// <inheritdoc/>
-		public Task AddAliasAsync(string name, BlobHandle handle, int rank = 0, ReadOnlyMemory<byte> data = default, CancellationToken cancellationToken = default) => _inner.AddAliasAsync(name, handle, rank, data, cancellationToken);
+		public Task AddAliasAsync(string name, IBlobHandle handle, int rank = 0, ReadOnlyMemory<byte> data = default, CancellationToken cancellationToken = default) => _inner.AddAliasAsync(name, handle, rank, data, cancellationToken);
 
 		/// <inheritdoc/>
-		public Task RemoveAliasAsync(string name, BlobHandle handle, CancellationToken cancellationToken = default) => _inner.RemoveAliasAsync(name, handle, cancellationToken);
+		public Task RemoveAliasAsync(string name, IBlobHandle handle, CancellationToken cancellationToken = default) => _inner.RemoveAliasAsync(name, handle, cancellationToken);
 
 		/// <inheritdoc/>
 		public Task<BlobAlias[]> FindAliasesAsync(string name, int? maxResults = null, CancellationToken cancellationToken = default) => _inner.FindAliasesAsync(name, maxResults, cancellationToken);
@@ -104,10 +104,10 @@ namespace EpicGames.Horde.Storage.Clients
 		#region Refs
 
 		/// <inheritdoc/>
-		public Task<BlobHandle?> TryReadRefAsync(RefName name, RefCacheTime cacheTime = default, CancellationToken cancellationToken = default) => _inner.TryReadRefAsync(name, cacheTime, cancellationToken);
+		public Task<IBlobHandle?> TryReadRefAsync(RefName name, RefCacheTime cacheTime = default, CancellationToken cancellationToken = default) => _inner.TryReadRefAsync(name, cacheTime, cancellationToken);
 
 		/// <inheritdoc/>
-		public Task WriteRefAsync(RefName name, BlobHandle handle, RefOptions? options = null, CancellationToken cancellationToken = default) => _inner.WriteRefAsync(name, handle, options, cancellationToken);
+		public Task WriteRefAsync(RefName name, IBlobHandle handle, RefOptions? options = null, CancellationToken cancellationToken = default) => _inner.WriteRefAsync(name, handle, options, cancellationToken);
 
 		/// <inheritdoc/>
 		public Task<bool> DeleteRefAsync(RefName name, CancellationToken cancellationToken = default) => _inner.DeleteRefAsync(name, cancellationToken);
diff --git a/Engine/Source/Programs/Shared/EpicGames.Horde/Storage/IBlobHandle.cs b/Engine/Source/Programs/Shared/EpicGames.Horde/Storage/IBlobHandle.cs
new file mode 100644
--- /dev/null
+++ b/Engine/Source/Programs/Shared/EpicGames.Horde/Storage/IBlobHandle.cs
@@ -0,0 +1,280 @@
+// Copyright Epic Games, Inc. All Rights Reserved.
+
+using System;
+using System.Collections.Generic;
+using System.Diagnostics.CodeAnalysis;
+using System.IO;
+using System.Threading;
+using System.Threading.Tasks;
+using EpicGames.Core;
+
+namespace EpicGames.Horde.Storage
+{
+	/// <summary>
+	/// Handle to a node. Can be used to reference nodes that have not been flushed yet. Handles should have value equality semantics, and must override <see cref="Object.Equals(Object?)"/> and <see cref="Object.GetHashCode()"/>.
+	/// </summary>
+	public interface IBlobHandle
+	{
+		private class BlobDataFragment : IReadOnlyMemoryOwner<byte>
+		{
+			readonly BlobData _data;
+
+			public ReadOnlyMemory<byte> Memory { get; }
+
+			public BlobDataFragment(BlobData data, int offset, int? length)
+			{
+				_data = data;
+
+				if (length == null)
+				{
+					Memory = data.Data.Slice(offset);
+				}
+				else
+				{
+					Memory = data.Data.Slice(offset, Math.Min(length.Value, data.Data.Length - offset));
+				}
+			}
+
+			public void Dispose() => _data.Dispose();
+		}
+
+		private class BlobDataStream : ReadOnlyMemoryStream
+		{
+			readonly BlobData _blobData;
+
+			public BlobDataStream(BlobData blobData, ReadOnlyMemory<byte> data)
+				: base(data)
+			{
+				_blobData = blobData;
+			}
+
+			protected override void Dispose(bool disposing)
+			{
+				base.Dispose(disposing);
+				if (disposing)
+				{
+					_blobData.Dispose();
+				}
+			}
+		}
+
+		/// <summary>
+		/// For a blob nested within another blob, gets a handle to the containing blob (eg. For a bundle node, will return the packet. For a bundle packet, will return the bundle. For a bundle or other non-nested blob, returns null.)
+		/// </summary>
+		IBlobHandle? Outer { get; }
+
+		/// <summary>
+		/// Flush the referenced not to underlying storage
+		/// </summary>
+		ValueTask FlushAsync(CancellationToken cancellationToken = default) => new ValueTask();
+
+		/// <summary>
+		/// Open the blob's data stream
+		/// </summary>
+		/// <param name="offset">Start offset of the stream</param>
+		/// <param name="length">Length of the stream</param>
+		/// <param name="cancellationToken">Cancellation token for the operation</param>
+		async Task<Stream> OpenBodyAsync(int offset = 0, int? length = null, CancellationToken cancellationToken = default)
+		{
+			BlobData blobData = await ReadAsync(cancellationToken);
+			int maxLength = blobData.Data.Length - offset;
+			ReadOnlyMemory<byte> memory = blobData.Data.Slice(offset, length.HasValue ? Math.Min(length.Value, maxLength) : maxLength);
+			return new BlobDataStream(blobData, memory);
+		}
+
+		/// <summary>
+		/// Reads the blob's data
+		/// </summary>
+		/// <param name="cancellationToken">Cancellation token for the operation</param>
+		ValueTask<BlobData> ReadAsync(CancellationToken cancellationToken = default);
+
+		/// <summary>
+		/// Gets the type of this blob
+		/// </summary>
+		/// <param name="cancellationToken">Cancellation token for the operation</param>
+		async ValueTask<BlobType> ReadTypeAsync(CancellationToken cancellationToken = default)
+		{
+			using BlobData data = await ReadAsync(cancellationToken);
+			return data.Type;
+		}
+
+		/// <summary>
+		/// Gets the outward references from this blob
+		/// </summary>
+		/// <param name="cancellationToken">Cancellation token for the operation</param>
+		async ValueTask<IReadOnlyList<IBlobHandle>> ReadImportsAsync(CancellationToken cancellationToken = default)
+		{
+			using BlobData data = await ReadAsync(cancellationToken);
+			return data.Refs;
+		}
+
+		/// <summary>
+		/// Reads part of the blob, and returns a handle that can be used to access the data.
+		/// </summary>
+		/// <param name="cancellationToken">Cancellation token for the operation</param>
+		/// <returns></returns>
+		async ValueTask<IReadOnlyMemoryOwner<byte>> ReadBodyAsync(CancellationToken cancellationToken = default)
+		{
+			return await ReadBodyAsync(0, null, cancellationToken);
+		}
+
+		/// <summary>
+		/// Reads part of the blob, and returns a handle that can be used to access the data.
+		/// </summary>
+		/// <param name="offset">Offset of the data within the blob</param>
+		/// <param name="length">Length of the data to read</param>
+		/// <param name="cancellationToken">Cancellation token for the operation</param>
+		/// <returns></returns>
+		async ValueTask<IReadOnlyMemoryOwner<byte>> ReadBodyAsync(int offset, int? length, CancellationToken cancellationToken = default)
+		{
+			BlobData data = await ReadAsync(cancellationToken);
+			return new BlobDataFragment(data, offset, length);
+		}
+
+		/// <summary>
+		/// Gets an identifier for this blob, relative to its outer
+		/// </summary>
+		/// <param name="builder">Builder for appending the identifier to</param>
+		/// <returns>True if an identifier was returned, false otherwise</returns>
+		bool TryAppendIdentifier(Utf8StringBuilder builder);
+	}
+
+	/// <summary>
+	/// Instance of <see cref="IBlobHandle"/> which wraps an inner handle and a fragment
+	/// </summary>
+	public class BlobFragmentHandle : IBlobHandle
+	{
+		/// <summary>
+		/// Handle to the outer blob
+		/// </summary>
+		public IBlobHandle Outer { get; }
+
+		/// <summary>
+		/// The fragment portion of the handle
+		/// </summary>
+		public Utf8String Fragment { get; }
+
+		/// <summary>
+		/// Constructor
+		/// </summary>
+		public BlobFragmentHandle(IBlobHandle outer, Utf8String fragment)
+		{
+			Outer = outer;
+			Fragment = fragment;
+		}
+
+		/// <inheritdoc/>
+		public bool TryAppendIdentifier(Utf8StringBuilder builder)
+		{
+			builder.Append(Fragment);
+			return true;
+		}
+
+		/// <inheritdoc/>
+		public ValueTask<BlobData> ReadAsync(CancellationToken cancellationToken = default)
+		{
+			throw new NotSupportedException("Blob fragment handles cannot be read directly, and should be deconstructed into more specific types.");
+		}
+
+		/// <inheritdoc/>
+		public override bool Equals(object? obj) => obj is BlobFragmentHandle other && Outer.Equals(other.Outer) && Fragment == other.Fragment;
+
+		/// <inheritdoc/>
+		public override int GetHashCode() => HashCode.Combine(Outer, Fragment);
+	}
+
+	/// <summary>
+	/// Extension methods for blob handles
+	/// </summary>
+	public static class BlobHandleExtensions
+	{
+		/// <summary>
+		/// Gets a path to this blob that can be used to describe blob references over the wire.
+		/// </summary>
+		/// <param name="handle">Handle to query</param>
+		public static BlobLocator GetLocator(this IBlobHandle handle)
+		{
+			BlobLocator locator;
+			if (!TryGetLocator(handle, out locator))
+			{
+				throw new InvalidOperationException("Blob has not yet been written to storage");
+			}
+			return locator;
+		}
+
+		/// <summary>
+		/// Attempt to get a path for this blob.
+		/// </summary>
+		/// <param name="handle">Handle to query</param>
+		/// <param name="locator">Receives the blob path on success.</param>
+		/// <returns>True if a path was available, false if the blob has not yet been flushed to storage.</returns>
+		public static bool TryGetLocator(this IBlobHandle handle, [NotNullWhen(true)] out BlobLocator locator)
+		{
+			Utf8StringBuilder builder = new Utf8StringBuilder();
+			if (AppendLocator(handle, builder))
+			{
+				locator = new BlobLocator(builder.ToUtf8String());
+				return true;
+			}
+			else
+			{
+				locator = default;
+				return false;
+			}
+		}
+
+		/// <summary>
+		/// Builds a full locator for a blob by traversing the outer chain
+		/// </summary>
+		static bool AppendLocator(IBlobHandle handle, Utf8StringBuilder builder)
+		{
+			IBlobHandle? outer = handle.Outer;
+			if (outer != null)
+			{
+				if (!AppendLocator(outer, builder))
+				{
+					return false;
+				}
+				if (outer.Outer == null)
+				{
+					builder.Append('#');
+				}
+				else
+				{
+					builder.Append('&');
+				}
+			}
+			return handle.TryAppendIdentifier(builder);
+		}
+
+		/// <summary>
+		/// Gets an identifier for a blob handle
+		/// </summary>
+		/// <param name="handle">Handle to the blob</param>
+		/// <param name="fragment">On success, receives the fragment path</param>
+		/// <returns>True if the handle has an identifier</returns>
+		public static bool TryGetIdentifier(this IBlobHandle handle, out Utf8String fragment)
+		{
+			Utf8StringBuilder builder = new Utf8StringBuilder();
+			if (handle.TryAppendIdentifier(builder))
+			{
+				fragment = builder.ToUtf8String();
+				return true;
+			}
+			else
+			{
+				fragment = default;
+				return false;
+			}
+		}
+
+		/// <summary>
+		/// Helper method for awaiting a handle and returning its locator
+		/// </summary>
+		public static async ValueTask<BlobLocator> GetLocatorAsync(this Task<IBlobHandle> handleTask)
+		{
+			IBlobHandle handle = await handleTask;
+			return handle.GetLocator();
+		}
+	}
+}
diff --git a/Engine/Source/Programs/Shared/EpicGames.Horde/Storage/IStorageClient.cs b/Engine/Source/Programs/Shared/EpicGames.Horde/Storage/IStorageClient.cs
--- a/Engine/Source/Programs/Shared/EpicGames.Horde/Storage/IStorageClient.cs
+++ b/Engine/Source/Programs/Shared/EpicGames.Horde/Storage/IStorageClient.cs
@@ -75,7 +75,7 @@ namespace EpicGames.Horde.Storage
 	/// </summary>
 	/// <param name="Target">Target blob</param>
 	/// <param name="Data">Inline data stored with the ref</param>
-	public record class RefValue(BlobHandle Target, ReadOnlyMemory<byte> Data);
+	public record class RefValue(IBlobHandle Target, ReadOnlyMemory<byte> Data);
 
 	/// <summary>
 	/// Interface for the storage system.
@@ -94,7 +94,7 @@ namespace EpicGames.Horde.Storage
 		/// </summary>
 		/// <param name="locator">Path to the blob</param>
 		/// <returns>New handle to the blob</returns>
-		BlobHandle CreateBlobHandle(BlobLocator locator);
+		IBlobHandle CreateBlobHandle(BlobLocator locator);
 
 		/// <summary>
 		/// Creates a new writer for storage blobs
@@ -104,7 +104,7 @@ namespace EpicGames.Horde.Storage
 		IStorageWriter CreateWriter(string? basePath = null);
 
 		/// <summary>
-		/// Read a blob from the underlying storage system. Calling <see cref="BlobHandle.ReadBlobDataAsync(CancellationToken)"/> is more efficient than calling this method repeatedly for small blobs.
+		/// Read a blob from the underlying storage system. Calling <see cref="IBlobHandle.ReadAsync(CancellationToken)"/> is more efficient than calling this method repeatedly for small blobs.
 		/// </summary>
 		/// <param name="locator">Locator for the blob</param>
 		/// <param name="cancellationToken">Cancellation token for the operation</param>
@@ -118,7 +118,7 @@ namespace EpicGames.Horde.Storage
 		/// <param name="references">References to other blobs</param>
 		/// <param name="basePath">Base path for writes</param>
 		/// <param name="cancellationToken">Cancellation token for the operation</param>
-		ValueTask<BlobHandle> WriteBlobAsync(BlobType type, Stream stream, IReadOnlyList<BlobHandle> references, string? basePath = null, CancellationToken cancellationToken = default);
+		ValueTask<IBlobHandle> WriteBlobAsync(BlobType type, Stream stream, IReadOnlyList<IBlobHandle> references, string? basePath = null, CancellationToken cancellationToken = default);
 
 		/// <summary>
 		/// Attempts to get a redirect URL for the given blob
@@ -148,7 +148,7 @@ namespace EpicGames.Horde.Storage
 		/// <param name="rank">Rank for this alias. In situations where an alias has multiple mappings, the alias with the highest rank will be returned by default.</param>
 		/// <param name="data">Additional data to be stored inline with the alias</param>
 		/// <param name="cancellationToken">Cancellation token for the operation</param>
-		Task AddAliasAsync(string name, BlobHandle handle, int rank = 0, ReadOnlyMemory<byte> data = default, CancellationToken cancellationToken = default);
+		Task AddAliasAsync(string name, IBlobHandle handle, int rank = 0, ReadOnlyMemory<byte> data = default, CancellationToken cancellationToken = default);
 
 		/// <summary>
 		/// Removes an alias from a blob
@@ -156,7 +156,7 @@ namespace EpicGames.Horde.Storage
 		/// <param name="name">Name of the alias</param>
 		/// <param name="handle">Locator for the blob</param>
 		/// <param name="cancellationToken">Cancellation token for the operation</param>
-		Task RemoveAliasAsync(string name, BlobHandle handle, CancellationToken cancellationToken = default);
+		Task RemoveAliasAsync(string name, IBlobHandle handle, CancellationToken cancellationToken = default);
 
 		/// <summary>
 		/// Finds blobs with the given alias. Unlike refs, aliases do not serve as GC roots.
@@ -178,7 +178,7 @@ namespace EpicGames.Horde.Storage
 		/// <param name="cacheTime">Minimum coherency for any cached value to be returned</param>
 		/// <param name="cancellationToken">Cancellation token for the operation</param>
 		/// <returns>Blob pointed to by the ref</returns>
-		Task<BlobHandle?> TryReadRefAsync(RefName name, RefCacheTime cacheTime = default, CancellationToken cancellationToken = default);
+		Task<IBlobHandle?> TryReadRefAsync(RefName name, RefCacheTime cacheTime = default, CancellationToken cancellationToken = default);
 
 		/// <summary>
 		/// Writes a new ref to the store
@@ -188,7 +188,7 @@ namespace EpicGames.Horde.Storage
 		/// <param name="options">Options for the new ref</param>
 		/// <param name="cancellationToken">Cancellation token for the operation</param>
 		/// <returns>Unique identifier for the blob</returns>
-		Task WriteRefAsync(RefName name, BlobHandle handle, RefOptions? options = null, CancellationToken cancellationToken = default);
+		Task WriteRefAsync(RefName name, IBlobHandle handle, RefOptions? options = null, CancellationToken cancellationToken = default);
 
 		/// <summary>
 		/// Reads data for a ref from the store
@@ -343,7 +343,7 @@ namespace EpicGames.Horde.Storage
 		/// <param name="data">Data to be stored</param>
 		/// <param name="references">References to other blobs</param>
 		/// <param name="cancellationToken">Cancellation token for the operation</param>
-		public static ValueTask<BlobHandle> WriteBlobAsync(this IStorageClient store, BlobType type, ReadOnlyMemory<byte> data, IReadOnlyList<BlobHandle> references, CancellationToken cancellationToken = default)
+		public static ValueTask<IBlobHandle> WriteBlobAsync(this IStorageClient store, BlobType type, ReadOnlyMemory<byte> data, IReadOnlyList<IBlobHandle> references, CancellationToken cancellationToken = default)
 		{
 			return WriteBlobAsync(store, null, type, data, references, cancellationToken);
 		}
@@ -357,7 +357,7 @@ namespace EpicGames.Horde.Storage
 		/// <param name="data">Data to be stored</param>
 		/// <param name="references">References to other blobs</param>
 		/// <param name="cancellationToken">Cancellation token for the operation</param>
-		public static async ValueTask<BlobHandle> WriteBlobAsync(this IStorageClient store, string? basePath, BlobType type, ReadOnlyMemory<byte> data, IReadOnlyList<BlobHandle> references, CancellationToken cancellationToken = default)
+		public static async ValueTask<IBlobHandle> WriteBlobAsync(this IStorageClient store, string? basePath, BlobType type, ReadOnlyMemory<byte> data, IReadOnlyList<IBlobHandle> references, CancellationToken cancellationToken = default)
 		{
 			using ReadOnlyMemoryStream stream = new ReadOnlyMemoryStream(data);
 			return await store.WriteBlobAsync(type, stream, references, basePath, cancellationToken);
@@ -394,7 +394,7 @@ namespace EpicGames.Horde.Storage
 		/// <returns>True if the ref exists, false if it did not exist</returns>
 		public static async Task<bool> HasRefAsync(this IStorageClient store, RefName name, RefCacheTime cacheTime = default, CancellationToken cancellationToken = default)
 		{
-			BlobHandle? target = await store.TryReadRefAsync(name, cacheTime, cancellationToken);
+			IBlobHandle? target = await store.TryReadRefAsync(name, cacheTime, cancellationToken);
 			return target != null;
 		}
 
@@ -406,7 +406,7 @@ namespace EpicGames.Horde.Storage
 		/// <param name="cacheTime">Minimum coherency for any cached value to be returned</param>
 		/// <param name="cancellationToken">Cancellation token for the operation</param>
 		/// <returns>Blob pointed to by the ref</returns>
-		public static async Task<BlobHandle?> TryReadRefTargetAsync(this IStorageClient store, RefName name, RefCacheTime cacheTime = default, CancellationToken cancellationToken = default)
+		public static async Task<IBlobHandle?> TryReadRefTargetAsync(this IStorageClient store, RefName name, RefCacheTime cacheTime = default, CancellationToken cancellationToken = default)
 		{
 			return await store.TryReadRefAsync(name, cacheTime, cancellationToken);
 		}
@@ -419,9 +419,9 @@ namespace EpicGames.Horde.Storage
 		/// <param name="cacheTime">Minimum coherency of any cached result</param>
 		/// <param name="cancellationToken">Cancellation token for the operation</param>
 		/// <returns>The ref target</returns>
-		public static async Task<BlobHandle> ReadRefTargetAsync(this IStorageClient store, RefName name, RefCacheTime cacheTime = default, CancellationToken cancellationToken = default)
+		public static async Task<IBlobHandle> ReadRefTargetAsync(this IStorageClient store, RefName name, RefCacheTime cacheTime = default, CancellationToken cancellationToken = default)
 		{
-			BlobHandle? refTarget = await store.TryReadRefTargetAsync(name, cacheTime, cancellationToken);
+			IBlobHandle? refTarget = await store.TryReadRefTargetAsync(name, cacheTime, cancellationToken);
 			return refTarget ?? throw new RefNameNotFoundException(name);
 		}
 
@@ -434,7 +434,7 @@ namespace EpicGames.Horde.Storage
 		/// <param name="options">Options for the new ref</param>
 		/// <param name="cancellationToken">Cancellation token for the operation</param>
 		/// <returns>Unique identifier for the blob</returns>
-		public static Task WriteRefTargetAsync(this IStorageClient store, RefName name, BlobHandle handle, RefOptions? options = null, CancellationToken cancellationToken = default)
+		public static Task WriteRefTargetAsync(this IStorageClient store, RefName name, IBlobHandle handle, RefOptions? options = null, CancellationToken cancellationToken = default)
 		{
 			return store.WriteRefAsync(name, handle, options: options, cancellationToken: cancellationToken);
 		}
diff --git a/Engine/Source/Programs/Shared/EpicGames.Horde/Storage/IStorageWriter.cs b/Engine/Source/Programs/Shared/EpicGames.Horde/Storage/IStorageWriter.cs
--- a/Engine/Source/Programs/Shared/EpicGames.Horde/Storage/IStorageWriter.cs
+++ b/Engine/Source/Programs/Shared/EpicGames.Horde/Storage/IStorageWriter.cs
@@ -58,7 +58,7 @@ namespace EpicGames.Horde.Storage
 		/// <param name="aliases">Aliases for this node</param>
 		/// <param name="cancellationToken">Cancellation token for the operation</param>
 		/// <returns>Handle to the written node</returns>
-		ValueTask<BlobHandle> WriteBlobAsync(BlobType type, int size, IReadOnlyList<BlobHandle> references, IReadOnlyList<AliasInfo> aliases, CancellationToken cancellationToken = default);
+		ValueTask<IBlobHandle> WriteBlobAsync(BlobType type, int size, IReadOnlyList<IBlobHandle> references, IReadOnlyList<AliasInfo> aliases, CancellationToken cancellationToken = default);
 	}
 
 	/// <summary>
@@ -108,14 +108,14 @@ namespace EpicGames.Horde.Storage
 		}
 
 		/// <inheritdoc/>
-		public async ValueTask<BlobHandle> WriteBlobAsync(BlobType type, int size, IReadOnlyList<BlobHandle> references, IReadOnlyList<AliasInfo> aliases, CancellationToken cancellationToken = default)
+		public async ValueTask<IBlobHandle> WriteBlobAsync(BlobType type, int size, IReadOnlyList<IBlobHandle> references, IReadOnlyList<AliasInfo> aliases, CancellationToken cancellationToken = default)
 		{
 			ReadOnlyMemory<byte> data = _data.AsMemory(_offset, size);
 			_offset += size;
 
 			using ReadOnlyMemoryStream stream = new ReadOnlyMemoryStream(data);
 
-			BlobHandle handle = await _outer.WriteBlobAsync(type, stream, references, _basePath, cancellationToken);
+			IBlobHandle handle = await _outer.WriteBlobAsync(type, stream, references, _basePath, cancellationToken);
 			foreach (AliasInfo aliasInfo in aliases)
 			{
 				await _outer.AddAliasAsync(aliasInfo.Name, handle, aliasInfo.Rank, aliasInfo.Data, cancellationToken);
@@ -135,16 +135,16 @@ namespace EpicGames.Horde.Storage
 		{
 			readonly int _maxKeys;
 			readonly Queue<BlobKey> _blobKeys = new Queue<BlobKey>();
-			readonly Dictionary<BlobKey, BlobHandle> _blobKeyToHandle = new Dictionary<BlobKey, BlobHandle>();
+			readonly Dictionary<BlobKey, IBlobHandle> _blobKeyToHandle = new Dictionary<BlobKey, IBlobHandle>();
 
 			public DedupeCache(int maxKeys)
 			{
 				_maxKeys = maxKeys;
 				_blobKeys = new Queue<BlobKey>(maxKeys);
-				_blobKeyToHandle = new Dictionary<BlobKey, BlobHandle>(maxKeys);
+				_blobKeyToHandle = new Dictionary<BlobKey, IBlobHandle>(maxKeys);
 			}
 
-			internal void Add(BlobKey key, BlobHandle handle)
+			internal void Add(BlobKey key, IBlobHandle handle)
 			{
 				BlobKey? prevKey;
 				if (_blobKeys.Count == _maxKeys && _blobKeys.TryDequeue(out prevKey))
@@ -154,25 +154,25 @@ namespace EpicGames.Horde.Storage
 				_blobKeyToHandle.TryAdd(key, handle);
 			}
 
-			internal bool TryGetValue(BlobKey key, [NotNullWhen(true)] out BlobHandle? handle) => _blobKeyToHandle.TryGetValue(key, out handle);
+			internal bool TryGetValue(BlobKey key, [NotNullWhen(true)] out IBlobHandle? handle) => _blobKeyToHandle.TryGetValue(key, out handle);
 		}
 
-		class WrappedHandle : BlobHandle
+		class WrappedHandle : IBlobHandle
 		{
 			public object _lockObject = new object();
-			public BlobHandle? _inner;
+			public IBlobHandle? _inner;
 
 			/// <inheritdoc/>
-			public override BlobHandle? Outer => _inner?.Outer;
+			public IBlobHandle? Outer => _inner?.Outer;
 
 			/// <inheritdoc/>
-			public override bool TryAppendIdentifier(Utf8StringBuilder builder)
+			public bool TryAppendIdentifier(Utf8StringBuilder builder)
 			{
 				return _inner?.TryAppendIdentifier(builder) ?? false;
 			}
 
 			/// <inheritdoc/>
-			public override ValueTask FlushAsync(CancellationToken cancellationToken)
+			public ValueTask FlushAsync(CancellationToken cancellationToken)
 			{
 				if (_inner == null)
 				{
@@ -185,7 +185,7 @@ namespace EpicGames.Horde.Storage
 			}
 
 			/// <inheritdoc/>
-			public override ValueTask<BlobData> ReadBlobDataAsync(CancellationToken cancellationToken = default)
+			public ValueTask<BlobData> ReadAsync(CancellationToken cancellationToken = default)
 			{
 				if (_inner == null)
 				{
@@ -193,7 +193,7 @@ namespace EpicGames.Horde.Storage
 				}
 				else
 				{
-					return _inner.ReadBlobDataAsync(cancellationToken);
+					return _inner.ReadAsync(cancellationToken);
 				}
 			}
 
@@ -242,7 +242,7 @@ namespace EpicGames.Horde.Storage
 		public Memory<byte> GetOutputBuffer(int usedSize, int desiredSize) => _inner.GetOutputBuffer(usedSize, desiredSize);
 
 		/// <inheritdoc/>
-		public async ValueTask<BlobHandle> WriteBlobAsync(BlobType type, int size, IReadOnlyList<BlobHandle> references, IReadOnlyList<AliasInfo> aliases, CancellationToken cancellationToken = default)
+		public async ValueTask<IBlobHandle> WriteBlobAsync(BlobType type, int size, IReadOnlyList<IBlobHandle> references, IReadOnlyList<AliasInfo> aliases, CancellationToken cancellationToken = default)
 		{
 			ReadOnlyMemory<byte> data = _inner.GetOutputBuffer(size, size).Slice(0, size);
 			IoHash hash = IoHash.Compute(data.Span);
@@ -251,7 +251,7 @@ namespace EpicGames.Horde.Storage
 			WrappedHandle? wrappedHandle;
 			lock (_cache)
 			{
-				BlobHandle? handle;
+				IBlobHandle? handle;
 				if (_cache.TryGetValue(key, out handle))
 				{
 					return handle;
@@ -285,7 +285,7 @@ namespace EpicGames.Horde.Storage
 		/// <param name="references">References to other nodes</param>
 		/// <param name="cancellationToken">Cancellation token for the operation</param>
 		/// <returns>Handle to the written node</returns>
-		public static ValueTask<BlobHandle> WriteBlobAsync(this IStorageWriter writer, BlobType type, int size, IReadOnlyList<BlobHandle> references, CancellationToken cancellationToken = default)
+		public static ValueTask<IBlobHandle> WriteBlobAsync(this IStorageWriter writer, BlobType type, int size, IReadOnlyList<IBlobHandle> references, CancellationToken cancellationToken = default)
 		{
 			return writer.WriteBlobAsync(type, size, references, Array.Empty<AliasInfo>(), cancellationToken);
 		}
@@ -299,10 +299,10 @@ namespace EpicGames.Horde.Storage
 		/// <param name="references">References to other nodes</param>
 		/// <param name="cancellationToken">Cancellation token for the operation</param>
 		/// <returns>Handle to the written node</returns>
-		public static async ValueTask<HashedNodeRef<T>> WriteHashedNodeRefAsync<T>(this IStorageWriter writer, BlobType type, int size, IReadOnlyList<BlobHandle> references, CancellationToken cancellationToken = default) where T : Node
+		public static async ValueTask<HashedNodeRef<T>> WriteHashedNodeRefAsync<T>(this IStorageWriter writer, BlobType type, int size, IReadOnlyList<IBlobHandle> references, CancellationToken cancellationToken = default) where T : Node
 		{
 			IoHash hash = IoHash.Compute(writer.GetOutputBuffer(size, size).Span.Slice(0, size));
-			BlobHandle blobHandle = await WriteBlobAsync(writer, type, size, references, cancellationToken);
+			IBlobHandle blobHandle = await WriteBlobAsync(writer, type, size, references, cancellationToken);
 			return new HashedNodeRef<T>(hash, blobHandle);
 		}
 	}
diff --git a/Engine/Source/Programs/Shared/EpicGames.Horde/Storage/Node.cs b/Engine/Source/Programs/Shared/EpicGames.Horde/Storage/Node.cs
--- a/Engine/Source/Programs/Shared/EpicGames.Horde/Storage/Node.cs
+++ b/Engine/Source/Programs/Shared/EpicGames.Horde/Storage/Node.cs
@@ -257,12 +257,12 @@ namespace EpicGames.Horde.Storage
 		/// <summary>
 		/// Locations of all referenced nodes.
 		/// </summary>
-		IReadOnlyList<BlobHandle> References { get; }
+		IReadOnlyList<IBlobHandle> References { get; }
 
 		/// <summary>
 		/// Gets the next serialized blob handle
 		/// </summary>
-		BlobHandle ReadBlobReference();
+		IBlobHandle ReadBlobReference();
 	}
 
 	/// <summary>
@@ -298,7 +298,7 @@ namespace EpicGames.Horde.Storage
 		/// <summary>
 		/// Locations of all referenced nodes.
 		/// </summary>
-		public IReadOnlyList<BlobHandle> References => _blobData.Refs;
+		public IReadOnlyList<IBlobHandle> References => _blobData.Refs;
 
 		readonly BlobData _blobData;
 		int _refIdx;
@@ -315,7 +315,7 @@ namespace EpicGames.Horde.Storage
 		/// <summary>
 		/// Gets the next serialized blob handle
 		/// </summary>
-		public BlobHandle ReadBlobReference() => _blobData.Refs[_refIdx++];
+		public IBlobHandle ReadBlobReference() => _blobData.Refs[_refIdx++];
 	}
 
 	/// <summary>
@@ -327,7 +327,7 @@ namespace EpicGames.Horde.Storage
 		/// Adds a reference to another blob. This reference is stored out of band, and will not result in any bytes written to the output.
 		/// </summary>
 		/// <param name="reference">Referenced blob</param>
-		void WriteBlobReference(BlobHandle reference);
+		void WriteBlobReference(IBlobHandle reference);
 	}
 
 	/// <summary>
@@ -338,13 +338,13 @@ namespace EpicGames.Horde.Storage
 		readonly IStorageWriter _treeWriter;
 
 		Memory<byte> _memory;
-		readonly List<BlobHandle> _refs = new List<BlobHandle>();
+		readonly List<IBlobHandle> _refs = new List<IBlobHandle>();
 		int _length;
 
 		/// <summary>
 		/// List of serialized references
 		/// </summary>
-		public IReadOnlyList<BlobHandle> References => _refs;
+		public IReadOnlyList<IBlobHandle> References => _refs;
 
 		/// <inheritdoc/>
 		public int Length => _length;
@@ -363,7 +363,7 @@ namespace EpicGames.Horde.Storage
 		/// Adds a reference to another blob. This reference is stored out of band, and will not result in any bytes written to the output.
 		/// </summary>
 		/// <param name="reference">Referenced blob</param>
-		public void WriteBlobReference(BlobHandle reference) => _refs.Add(reference);
+		public void WriteBlobReference(IBlobHandle reference) => _refs.Add(reference);
 
 		/// <summary>
 		/// Computes the hash of the written data
@@ -373,7 +373,7 @@ namespace EpicGames.Horde.Storage
 		/// <summary>
 		/// Writes a handle to another node
 		/// </summary>
-		public void WriteHashedBlobHandle(IoHash hash, BlobHandle target)
+		public void WriteHashedBlobHandle(IoHash hash, IBlobHandle target)
 		{
 			this.WriteIoHash(hash);
 			WriteBlobReference(target);
@@ -410,9 +410,9 @@ namespace EpicGames.Horde.Storage
 		/// <param name="handle"></param>
 		/// <param name="cancellationToken"></param>
 		/// <returns></returns>
-		public static async ValueTask<TNode> ReadNodeAsync<TNode>(this BlobHandle handle, CancellationToken cancellationToken = default) where TNode : Node
+		public static async ValueTask<TNode> ReadNodeAsync<TNode>(this IBlobHandle handle, CancellationToken cancellationToken = default) where TNode : Node
 		{
-			using BlobData blobData = await handle.ReadBlobDataAsync(cancellationToken);
+			using BlobData blobData = await handle.ReadAsync(cancellationToken);
 			return Node.Deserialize<TNode>(blobData);
 		}
 
@@ -425,7 +425,7 @@ namespace EpicGames.Horde.Storage
 		/// <param name="refOptions">Options for the ref</param>
 		/// <param name="cancellationToken">Cancellation token for the operation</param>
 		/// <returns>Location of node targetted by the ref</returns>
-		public static async Task<BlobHandle> WriteRefAsync(this IStorageClient store, RefName name, Node node, RefOptions? refOptions = null, CancellationToken cancellationToken = default)
+		public static async Task<IBlobHandle> WriteRefAsync(this IStorageClient store, RefName name, Node node, RefOptions? refOptions = null, CancellationToken cancellationToken = default)
 		{
 			await using IStorageWriter writer = store.CreateWriter(name);
 			HashedNodeRef<Node> nodeRef = await writer.WriteHashedNodeAsync(node, cancellationToken);
@@ -443,13 +443,13 @@ namespace EpicGames.Horde.Storage
 		/// <returns>Node for the given ref, or null if it does not exist</returns>
 		public static async Task<TNode?> TryReadRefAsync<TNode>(this IStorageClient store, RefName name, DateTime cacheTime = default, CancellationToken cancellationToken = default) where TNode : Node
 		{
-			BlobHandle? refTarget = await store.TryReadRefTargetAsync(name, cacheTime, cancellationToken);
+			IBlobHandle? refTarget = await store.TryReadRefTargetAsync(name, cacheTime, cancellationToken);
 			if (refTarget == null)
 			{
 				return null;
 			}
 
-			using BlobData blobData = await refTarget.ReadBlobDataAsync(cancellationToken);
+			using BlobData blobData = await refTarget.ReadAsync(cancellationToken);
 			return Node.Deserialize<TNode>(blobData);
 		}
 
diff --git a/Engine/Source/Programs/Shared/EpicGames.Horde/Storage/NodeRef.cs b/Engine/Source/Programs/Shared/EpicGames.Horde/Storage/NodeRef.cs
--- a/Engine/Source/Programs/Shared/EpicGames.Horde/Storage/NodeRef.cs
+++ b/Engine/Source/Programs/Shared/EpicGames.Horde/Storage/NodeRef.cs
@@ -15,13 +15,13 @@ namespace EpicGames.Horde.Storage
 		/// <summary>
 		/// Handle to the node if in storage (or pending write to storage)
 		/// </summary>
-		public BlobHandle Handle { get; }
+		public IBlobHandle Handle { get; }
 
 		/// <summary>
 		/// Creates a reference to a node in storage.
 		/// </summary>
 		/// <param name="handle">Handle to the referenced node</param>
-		public NodeRef(BlobHandle handle) => Handle = handle;
+		public NodeRef(IBlobHandle handle) => Handle = handle;
 
 		/// <summary>
 		/// Creates a reference to a node in storage.
@@ -48,7 +48,7 @@ namespace EpicGames.Horde.Storage
 		/// <returns></returns>
 		public async ValueTask<Node> ExpandAsync(CancellationToken cancellationToken = default)
 		{
-			using BlobData blobData = await Handle!.ReadBlobDataAsync(cancellationToken);
+			using BlobData blobData = await Handle!.ReadAsync(cancellationToken);
 			return Node.Deserialize(blobData);
 		}
 	}
@@ -62,7 +62,7 @@ namespace EpicGames.Horde.Storage
 		/// <summary>
 		/// Constructor
 		/// </summary>
-		public NodeRef(BlobHandle handle) : base(handle)
+		public NodeRef(IBlobHandle handle) : base(handle)
 		{ }
 
 		/// <summary>
@@ -100,7 +100,7 @@ namespace EpicGames.Horde.Storage
 		/// </summary>
 		/// <param name="hash">Hash of the target node</param>
 		/// <param name="handle">Handle to the referenced node</param>
-		public HashedNodeRef(IoHash hash, BlobHandle handle) : base(handle) => Hash = hash;
+		public HashedNodeRef(IoHash hash, IBlobHandle handle) : base(handle) => Hash = hash;
 
 		/// <summary>
 		/// Deserialization constructor
@@ -129,7 +129,7 @@ namespace EpicGames.Horde.Storage
 		/// <summary>
 		/// Constructor
 		/// </summary>
-		public HashedNodeRef(IoHash hash, BlobHandle handle) : base(hash, handle)
+		public HashedNodeRef(IoHash hash, IBlobHandle handle) : base(hash, handle)
 		{
 		}
 
@@ -292,7 +292,7 @@ namespace EpicGames.Horde.Storage
 			IoHash hash = nodeWriter.ComputeHash();
 
 			// Write the final data
-			BlobHandle handle = await writer.WriteBlobAsync(node.NodeType, nodeWriter.Length, nodeWriter.References, cancellationToken);
+			IBlobHandle handle = await writer.WriteBlobAsync(node.NodeType, nodeWriter.Length, nodeWriter.References, cancellationToken);
 			return new HashedNodeRef<TNode>(hash, handle);
 		}
 
@@ -322,7 +322,7 @@ namespace EpicGames.Horde.Storage
 		/// <param name="cancellationToken"></param>
 		public static async ValueTask<NodeRef<TNode>?> TryReadRefTargetAsync<TNode>(this IStorageClient storageClient, RefName refName, DateTime cacheTime = default, CancellationToken cancellationToken = default) where TNode : Node
 		{
-			BlobHandle? handle = await storageClient.TryReadRefTargetAsync(refName, cacheTime, cancellationToken);
+			IBlobHandle? handle = await storageClient.TryReadRefTargetAsync(refName, cacheTime, cancellationToken);
 			if (handle == null)
 			{
 				return null;
@@ -350,7 +350,7 @@ namespace EpicGames.Horde.Storage
 		/// <param name="root">Root for the tree</param>
 		/// <param name="cancellationToken">Cancellation token for the operation</param>
 		/// <returns></returns>
-		public static async Task<BlobHandle> FlushAsync(this IStorageWriter writer, Node root, CancellationToken cancellationToken = default)
+		public static async Task<IBlobHandle> FlushAsync(this IStorageWriter writer, Node root, CancellationToken cancellationToken = default)
 		{
 			HashedNodeRef<Node> rootRef = await writer.WriteHashedNodeAsync(root, cancellationToken);
 			await writer.FlushAsync(cancellationToken);
diff --git a/Engine/Source/Programs/Shared/EpicGames.Horde/Storage/Nodes/CbNode.cs b/Engine/Source/Programs/Shared/EpicGames.Horde/Storage/Nodes/CbNode.cs
--- a/Engine/Source/Programs/Shared/EpicGames.Horde/Storage/Nodes/CbNode.cs
+++ b/Engine/Source/Programs/Shared/EpicGames.Horde/Storage/Nodes/CbNode.cs
@@ -27,7 +27,7 @@ namespace EpicGames.Horde.Storage.Nodes
 			{
 				if (field.IsAttachment())
 				{
-					BlobHandle handle = _reader.ReadBlobReference();
+					IBlobHandle handle = _reader.ReadBlobReference();
 					_refs.Add(new HashedNodeRef(field.AsAttachment(), handle));
 				}
 				else if (field.IsArray())
diff --git a/Engine/Source/Programs/Shared/EpicGames.Horde/Storage/Nodes/ChunkedDataNode.cs b/Engine/Source/Programs/Shared/EpicGames.Horde/Storage/Nodes/ChunkedDataNode.cs
--- a/Engine/Source/Programs/Shared/EpicGames.Horde/Storage/Nodes/ChunkedDataNode.cs
+++ b/Engine/Source/Programs/Shared/EpicGames.Horde/Storage/Nodes/ChunkedDataNode.cs
@@ -35,9 +35,9 @@ namespace EpicGames.Horde.Storage.Nodes
 		/// <param name="handle">Handle to the data to read</param>
 		/// <param name="outputStream">The output stream to receive the data</param>
 		/// <param name="cancellationToken">Cancellation token for the operation</param>
-		public static async Task CopyToStreamAsync(BlobHandle handle, Stream outputStream, CancellationToken cancellationToken)
+		public static async Task CopyToStreamAsync(IBlobHandle handle, Stream outputStream, CancellationToken cancellationToken)
 		{
-			using BlobData blobData = await handle.ReadBlobDataAsync(cancellationToken);
+			using BlobData blobData = await handle.ReadAsync(cancellationToken);
 			if (blobData.Type.Guid == s_leafNodeGuid)
 			{
 				await LeafChunkedDataNode.CopyToStreamAsync(blobData, outputStream, cancellationToken);
@@ -299,7 +299,7 @@ namespace EpicGames.Horde.Storage.Nodes
 					sizeSinceProgressUpdate = 0;
 				}
 
-				HashedNodeRef<ChunkedDataNode> nodeRef = await writer.WriteHashedNodeRefAsync<ChunkedDataNode>(GetNodeType<LeafChunkedDataNode>(), nextLength, Array.Empty<BlobHandle>(), cancellationToken);
+				HashedNodeRef<ChunkedDataNode> nodeRef = await writer.WriteHashedNodeRefAsync<ChunkedDataNode>(GetNodeType<LeafChunkedDataNode>(), nextLength, Array.Empty<IBlobHandle>(), cancellationToken);
 				leafNodeRefs.Add(new ChunkedDataNodeRef(ChunkedDataNodeType.Leaf, nodeRef));
 
 				readBuffer.Memory.Slice(nextLength, size - nextLength).CopyTo(readBuffer.Memory);
@@ -538,7 +538,7 @@ namespace EpicGames.Horde.Storage.Nodes
 			while (nodeReader.GetMemory(0).Length > 0)
 			{
 				_ = nodeReader.ReadIoHash();
-				BlobHandle handle = nodeReader.ReadBlobReference();
+				IBlobHandle handle = nodeReader.ReadBlobReference();
 				if (nodeReader.Version >= 2)
 				{
 					_ = nodeReader.ReadUnsignedVarInt();
diff --git a/Engine/Source/Programs/Shared/EpicGames.Horde/Storage/Nodes/ChunkedDataWriter.cs b/Engine/Source/Programs/Shared/EpicGames.Horde/Storage/Nodes/ChunkedDataWriter.cs
--- a/Engine/Source/Programs/Shared/EpicGames.Horde/Storage/Nodes/ChunkedDataWriter.cs
+++ b/Engine/Source/Programs/Shared/EpicGames.Horde/Storage/Nodes/ChunkedDataWriter.cs
@@ -345,7 +345,7 @@ namespace EpicGames.Horde.Storage.Nodes
 		/// <returns>Handle to the written leaf node</returns>
 		async ValueTask FlushLeafNodeAsync(CancellationToken cancellationToken)
 		{
-			HashedNodeRef<ChunkedDataNode> nodeRef  = await _writer.WriteHashedNodeRefAsync<ChunkedDataNode>(s_leafNodeType, _leafLength, Array.Empty<BlobHandle>(), cancellationToken);
+			HashedNodeRef<ChunkedDataNode> nodeRef  = await _writer.WriteHashedNodeRefAsync<ChunkedDataNode>(s_leafNodeType, _leafLength, Array.Empty<IBlobHandle>(), cancellationToken);
 			_leafHandles.Add(new ChunkedDataNodeRef(ChunkedDataNodeType.Leaf, nodeRef));
 			ResetLeafState();
 		}
diff --git a/Engine/Source/Programs/Shared/EpicGames.Horde/Storage/Nodes/DdcRefNode.cs b/Engine/Source/Programs/Shared/EpicGames.Horde/Storage/Nodes/DdcRefNode.cs
--- a/Engine/Source/Programs/Shared/EpicGames.Horde/Storage/Nodes/DdcRefNode.cs
+++ b/Engine/Source/Programs/Shared/EpicGames.Horde/Storage/Nodes/DdcRefNode.cs
@@ -20,7 +20,7 @@ namespace Horde.Server.Ddc
 		/// <summary>
 		/// References to attachments. We embed this in the ref node to ensure any aliased blobs have a hard reference from the root.
 		/// </summary>
-		public List<(IoHash Hash, BlobHandle Handle)> References { get; }
+		public List<(IoHash Hash, IBlobHandle Handle)> References { get; }
 
 		/// <summary>
 		/// Constructor
@@ -28,7 +28,7 @@ namespace Horde.Server.Ddc
 		public DdcRefNode(IoHash rootHash)
 		{
 			RootHash = rootHash;
-			References = new List<(IoHash, BlobHandle)>();
+			References = new List<(IoHash, IBlobHandle)>();
 		}
 
 		/// <summary>
@@ -47,7 +47,7 @@ namespace Horde.Server.Ddc
 			writer.WriteList(References, x => WriteReference(writer, x.Hash, x.Handle));
 		}
 
-		static void WriteReference(INodeWriter writer, IoHash hash, BlobHandle handle)
+		static void WriteReference(INodeWriter writer, IoHash hash, IBlobHandle handle)
 		{
 			writer.WriteIoHash(hash);
 			writer.WriteBlobReference(handle);
diff --git a/Engine/Source/Programs/Shared/EpicGames.Horde/Storage/Nodes/RedirectNode.cs b/Engine/Source/Programs/Shared/EpicGames.Horde/Storage/Nodes/RedirectNode.cs
--- a/Engine/Source/Programs/Shared/EpicGames.Horde/Storage/Nodes/RedirectNode.cs
+++ b/Engine/Source/Programs/Shared/EpicGames.Horde/Storage/Nodes/RedirectNode.cs
@@ -11,7 +11,7 @@ namespace EpicGames.Horde.Storage.Nodes
 		/// <summary>
 		/// The target handle
 		/// </summary>
-		public BlobHandle Handle { get; }
+		public IBlobHandle Handle { get; }
 
 		/// <summary>
 		/// Constructor
@@ -23,7 +23,7 @@ namespace EpicGames.Horde.Storage.Nodes
 		/// Constructor
 		/// </summary>
 		/// <param name="handle">Target node for the redirect</param>
-		public RedirectNode(BlobHandle handle) => Handle = handle;
+		public RedirectNode(IBlobHandle handle) => Handle = handle;
 
 		/// <summary>
 		/// Deserialization constructor
diff --git a/Engine/Source/Programs/Shared/EpicGames.Horde/Storage/Workspace.cs b/Engine/Source/Programs/Shared/EpicGames.Horde/Storage/Workspace.cs
--- a/Engine/Source/Programs/Shared/EpicGames.Horde/Storage/Workspace.cs
+++ b/Engine/Source/Programs/Shared/EpicGames.Horde/Storage/Workspace.cs
@@ -723,7 +723,7 @@ namespace EpicGames.Horde.Storage
 
 		async Task ExtractDataAsync(HashedNodeRef<ChunkedDataNode> nodeRef, Stream outputStream, CancellationToken cancellationToken)
 		{
-			BlobType blobType = await nodeRef.Handle.GetTypeAsync(cancellationToken);
+			BlobType blobType = await nodeRef.Handle.ReadTypeAsync(cancellationToken);
 			if (blobType.Guid == LeafChunkedDataNode.BlobType.Guid)
 			{
 				await ExtractLeafDataAsync(nodeRef, outputStream, cancellationToken);
@@ -755,7 +755,7 @@ namespace EpicGames.Horde.Storage
 			}
 
 			// Otherwise 
-			using BlobData blobData = await nodeRef.Handle.ReadBlobDataAsync(cancellationToken);
+			using BlobData blobData = await nodeRef.Handle.ReadAsync(cancellationToken);
 			await LeafChunkedDataNode.CopyToStreamAsync(blobData, outputStream, cancellationToken);
 		}
 
diff --git a/Engine/Source/Programs/UnrealBuildTool/Artifacts/HordeStorageArtifactCache.cs b/Engine/Source/Programs/UnrealBuildTool/Artifacts/HordeStorageArtifactCache.cs
--- a/Engine/Source/Programs/UnrealBuildTool/Artifacts/HordeStorageArtifactCache.cs
+++ b/Engine/Source/Programs/UnrealBuildTool/Artifacts/HordeStorageArtifactCache.cs
@@ -497,7 +497,7 @@ namespace UnrealBuildTool.Artifacts
 					await hordeArtifactAction.WriteFilesAsync(writer, cancellationToken);
 
 					// Save the collection
-					BlobHandle _ = await _store.WriteRefAsync(refName, node, cancellationToken: cancellationToken);
+					IBlobHandle _ = await _store.WriteRefAsync(refName, node, cancellationToken: cancellationToken);
 				}, cancellationToken));
 			}
 			return tasks;
