EpicGames.Core: Add a deterministic memory allocator to EpicGames.Core for use with large allocations, reducing pressure on the LOH.

* Currently supports allocation from the global heap (GlobalHeapAllocator), through memory mapped files (VirtualMemoryAllocator), or through the standard .NET shared MemoryPool (PoolAllocator).
* ChunkedMemoryWriter now accepts an optional IMemoryAllocator argument for configuring where to obtain chunks from, defaulting to the previous behavior of using the shared .NET memory pool.
* Added a RefCountedMemoryWriter class which can implement reference counting on handles returned by IMemoryAllocator instances.

diff --git a/Engine/Source/Programs/Shared/EpicGames.Core.Tests/MemoryAllocatorTests.cs b/Engine/Source/Programs/Shared/EpicGames.Core.Tests/MemoryAllocatorTests.cs
new file mode 100644
--- /dev/null
+++ b/Engine/Source/Programs/Shared/EpicGames.Core.Tests/MemoryAllocatorTests.cs
@@ -0,0 +1,29 @@
+// Copyright Epic Games, Inc. All Rights Reserved.
+
+using System.Buffers;
+using Microsoft.VisualStudio.TestTools.UnitTesting;
+
+namespace EpicGames.Core.Tests
+{
+	[TestClass]
+	public class MemoryAllocatorTests
+	{
+		[TestMethod]
+		public void TestGlobalAllocator()
+		{
+			GlobalHeapAllocator allocator = new GlobalHeapAllocator();
+
+			using IMemoryOwner<byte> handle = allocator.Alloc(1024);
+			handle.Memory.Span[0] = 123;
+		}
+
+		[TestMethod]
+		public void TestVirtualAllocator()
+		{
+			VirtualMemoryAllocator allocator = new VirtualMemoryAllocator();
+
+			using IMemoryOwner<byte> handle = allocator.Alloc(1024);
+			handle.Memory.Span[0] = 123;
+		}
+	}
+}
diff --git a/Engine/Source/Programs/Shared/EpicGames.Core.Tests/MemoryWriterTests.cs b/Engine/Source/Programs/Shared/EpicGames.Core.Tests/MemoryWriterTests.cs
new file mode 100644
--- /dev/null
+++ b/Engine/Source/Programs/Shared/EpicGames.Core.Tests/MemoryWriterTests.cs
@@ -0,0 +1,150 @@
+// Copyright Epic Games, Inc. All Rights Reserved.
+
+using System;
+using System.Buffers;
+using System.Text;
+using System.Threading;
+using Microsoft.VisualStudio.TestTools.UnitTesting;
+
+namespace EpicGames.Core.Tests
+{
+	[TestClass]
+	public class MemoryWriterTests
+	{
+		class TestAllocator : IMemoryAllocator<byte>
+		{
+			class Allocation : IMemoryOwner<byte>
+			{
+				TestAllocator? _allocator;
+				readonly Memory<byte> _memory;
+
+				public Allocation(TestAllocator allocator, Memory<byte> memory)
+				{
+					_allocator = allocator;
+					_memory = memory;
+				}
+
+				public Memory<byte> Memory => _memory;
+
+				public void Dispose()
+				{
+					if (_allocator != null)
+					{
+						Interlocked.Decrement(ref _allocator._allocatedCount);
+						Interlocked.Add(ref _allocator._allocatedSize, -_memory.Length);
+						_allocator = null;
+					}
+				}
+			}
+
+			int _allocatedSize;
+			int _allocatedCount;
+
+			public int AllocatedSize => _allocatedSize;
+			public int AllocatedCount => _allocatedCount;
+
+			public IMemoryOwner<byte> Alloc(int size)
+			{
+				Interlocked.Increment(ref _allocatedCount);
+				Interlocked.Add(ref _allocatedSize, size);
+				return new Allocation(this, new byte[size]);
+			}
+		}
+
+		[TestMethod]
+		public void RefCountedMemoryWriter()
+		{
+			TestAllocator allocator = new TestAllocator();
+			Assert.AreEqual(0, allocator.AllocatedSize);
+			Assert.AreEqual(0, allocator.AllocatedCount);
+
+			using (RefCountedMemoryWriter writer = new RefCountedMemoryWriter(allocator, 4096))
+			{
+				writer.WriteString("Hello world");
+				Assert.IsTrue(allocator.AllocatedSize > 0);
+				Assert.AreEqual(1, allocator.AllocatedCount);
+			}
+
+			Assert.AreEqual(0, allocator.AllocatedSize);
+			Assert.AreEqual(0, allocator.AllocatedCount);
+		}
+
+		[TestMethod]
+		public void RefCountedMemoryWriterSequence()
+		{
+			TestAllocator allocator = new TestAllocator();
+			Assert.AreEqual(0, allocator.AllocatedSize);
+			Assert.AreEqual(0, allocator.AllocatedCount);
+
+			IRefCountedHandle<ReadOnlyMemory<byte>> memory;
+			IRefCountedHandle<ReadOnlySequence<byte>> sequence;
+
+			using (RefCountedMemoryWriter writer = new RefCountedMemoryWriter(allocator, 10))
+			{
+				writer.WriteString("This should be written to a single buffer and the original buffer should be freed");
+				int allocatedSize1 = allocator.AllocatedSize;
+				Assert.IsTrue(allocatedSize1 > 0);
+				Assert.AreEqual(1, allocator.AllocatedCount);
+
+				writer.WriteString("This should be written to a separate buffer");
+				int allocatedSize2 = allocator.AllocatedSize;
+				Assert.IsTrue(allocatedSize2 > allocatedSize1);
+				Assert.AreEqual(2, allocator.AllocatedCount);
+
+				sequence = writer.AsRefCountedSequence();
+				Assert.AreEqual(allocatedSize2, allocator.AllocatedSize);
+				Assert.AreEqual(2, allocator.AllocatedCount);
+
+				memory = writer.AsRefCountedMemory();
+				Assert.IsTrue(allocator.AllocatedSize > allocatedSize2);
+				Assert.AreEqual(3, allocator.AllocatedCount);
+			}
+
+			Assert.AreEqual(3, allocator.AllocatedCount);
+			memory.Dispose();
+
+			Assert.AreEqual(2, allocator.AllocatedCount);
+			sequence.Dispose();
+
+			Assert.AreEqual(0, allocator.AllocatedSize);
+			Assert.AreEqual(0, allocator.AllocatedCount);
+		}
+
+		[TestMethod]
+		public void RefCountedMemoryWriterSequence2()
+		{
+			TestAllocator allocator = new TestAllocator();
+			Assert.AreEqual(0, allocator.AllocatedSize);
+			Assert.AreEqual(0, allocator.AllocatedCount);
+
+			int allocatedSize1;
+			IRefCountedHandle<ReadOnlySequence<byte>> sequence;
+
+			using (RefCountedMemoryWriter writer = new RefCountedMemoryWriter(allocator, 10))
+			{
+				writer.WriteString("This should be written to a single buffer and the original buffer should be freed");
+				allocatedSize1 = allocator.AllocatedSize;
+				Assert.IsTrue(allocatedSize1 > 0);
+				Assert.AreEqual(1, allocator.AllocatedCount);
+
+				writer.WriteString("This should be written to a separate buffer");
+				int allocatedSize2 = allocator.AllocatedSize;
+				Assert.IsTrue(allocatedSize2 > allocatedSize1);
+				Assert.AreEqual(2, allocator.AllocatedCount);
+
+				sequence = writer.AsRefCountedSequence(6, 6); // one byte size field, followed by string above
+				Assert.AreEqual(allocatedSize2, allocator.AllocatedSize);
+				Assert.AreEqual(2, allocator.AllocatedCount);
+			}
+
+			Assert.AreEqual(allocatedSize1, allocator.AllocatedSize);
+			Assert.AreEqual(1, allocator.AllocatedCount);
+			Assert.AreEqual("should", Encoding.UTF8.GetString(sequence.Target.FirstSpan));
+
+			sequence.Dispose();
+
+			Assert.AreEqual(0, allocator.AllocatedSize);
+			Assert.AreEqual(0, allocator.AllocatedCount);
+		}
+	}
+}
diff --git a/Engine/Source/Programs/Shared/EpicGames.Core/IMemoryAllocator.cs b/Engine/Source/Programs/Shared/EpicGames.Core/IMemoryAllocator.cs
new file mode 100644
--- /dev/null
+++ b/Engine/Source/Programs/Shared/EpicGames.Core/IMemoryAllocator.cs
@@ -0,0 +1,146 @@
+// Copyright Epic Games, Inc. All Rights Reserved.
+
+using System;
+using System.Buffers;
+using System.IO.MemoryMappedFiles;
+using System.Runtime.InteropServices;
+
+namespace EpicGames.Core
+{
+	/// <summary>
+	/// Interface for an explicit memory allocator. Meant for use with large memory blocks, where objects would typically be allocated on the LOH and waiting for GC can cause
+	/// excessive memory usage.
+	/// </summary>
+	public interface IMemoryAllocator<T>
+	{
+		/// <summary>
+		/// Allocate a block of memory of the given minimum size
+		/// </summary>
+		/// <param name="minSize">Minimum size for the allocated data</param>
+		/// <returns>Reference counted block of memory</returns>
+		IMemoryOwner<T> Alloc(int minSize);
+	}
+
+	/// <summary>
+	/// Implementation of <see cref="IMemoryAllocator{Byte}"/> which rents blocks from a memory pool.
+	/// </summary>
+	public class PoolAllocator : IMemoryAllocator<byte>
+	{
+		readonly MemoryPool<byte> _pool;
+
+		/// <summary>
+		/// Shared allocator instance
+		/// </summary>
+		public static PoolAllocator Shared { get; } = new PoolAllocator(MemoryPool<byte>.Shared);
+
+		/// <summary>
+		/// Constructor
+		/// </summary>
+		public PoolAllocator(MemoryPool<byte> pool) => _pool = pool;
+
+		/// <inheritdoc/>
+		public IMemoryOwner<byte> Alloc(int minSize) => _pool.Rent(minSize);
+	}
+
+	/// <summary>
+	/// Implementation of <see cref="IMemoryAllocator{Byte}"/> which returns blocks from the global heap
+	/// </summary>
+	public class GlobalHeapAllocator : IMemoryAllocator<byte>
+	{
+		unsafe class Allocation : MemoryManager<byte>
+		{
+			IntPtr _handle;
+			byte* _pointer;
+			int _length;
+
+			public Allocation(int size)
+			{
+				_handle = Marshal.AllocHGlobal(size);
+				_pointer = (byte*)_handle.ToPointer();
+				_length = size;
+			}
+
+			/// <inheritdoc/>
+			public override Span<byte> GetSpan() => new Span<byte>(_pointer, _length);
+
+			/// <inheritdoc/>
+			public override MemoryHandle Pin(int elementIndex) => new MemoryHandle(_pointer + elementIndex);
+
+			/// <inheritdoc/>
+			public override void Unpin() { }
+
+			/// <inheritdoc/>
+			protected override void Dispose(bool disposing)
+			{
+				if (_handle != IntPtr.Zero)
+				{
+					Marshal.FreeHGlobal(_handle);
+					_handle = IntPtr.Zero;
+				}
+
+				_pointer = null;
+				_length = 0;
+			}
+		}
+
+		/// <inheritdoc/>
+		public IMemoryOwner<byte> Alloc(int minSize) => new Allocation(minSize);
+	}
+	
+	/// <summary>
+	/// Implementation of <see cref="IMemoryAllocator{Byte}"/> which returns memory backed by memory mapped files
+	/// </summary>
+	public class VirtualMemoryAllocator : IMemoryAllocator<byte>
+	{
+		unsafe class Allocation : MemoryManager<byte>
+		{
+			MemoryMappedFile? _file;
+			MemoryMappedViewAccessor? _viewAccessor;
+			byte* _pointer;
+			readonly int _length;
+
+			public Allocation(int length)
+			{
+				_file = MemoryMappedFile.CreateNew(null, length, MemoryMappedFileAccess.ReadWrite);
+				_viewAccessor = _file.CreateViewAccessor();
+				_viewAccessor.SafeMemoryMappedViewHandle.AcquirePointer(ref _pointer);
+				_length = length;
+			}
+
+			/// <inheritdoc/>
+			public override Span<byte> GetSpan() => new Span<byte>(_pointer, _length);
+
+			/// <inheritdoc/>
+			public override MemoryHandle Pin(int elementIndex) => new MemoryHandle(_pointer + elementIndex);
+
+			/// <inheritdoc/>
+			public override void Unpin() { }
+
+			/// <inheritdoc/>
+			protected override void Dispose(bool disposing)
+			{
+				if (disposing)
+				{
+					if (_pointer != null)
+					{
+						_viewAccessor!.SafeMemoryMappedViewHandle.ReleasePointer();
+						_pointer = null;
+					}
+					if (_viewAccessor != null)
+					{
+						_viewAccessor.Dispose();
+						_viewAccessor = null;
+					}
+					if (_file != null)
+					{
+						_file.Dispose();
+						_file = null;
+					}
+				}
+			}
+		}
+
+		/// <inheritdoc/>
+		public IMemoryOwner<byte> Alloc(int minSize) => new Allocation(minSize);
+	}
+}
diff --git a/Engine/Source/Programs/Shared/EpicGames.Core/MemoryWriter.cs b/Engine/Source/Programs/Shared/EpicGames.Core/MemoryWriter.cs
--- a/Engine/Source/Programs/Shared/EpicGames.Core/MemoryWriter.cs
+++ b/Engine/Source/Programs/Shared/EpicGames.Core/MemoryWriter.cs
@@ -306,6 +306,11 @@ namespace EpicGames.Core
 		readonly int _chunkSize;
 		Chunk _currentChunk;
 
+		/// <summary>
+		/// Accessor for the allocated chunks
+		/// </summary>
+		protected IReadOnlyList<Chunk> Chunks => _chunks;
+
 		/// <summary>
 		/// Length of the current sequence
 		/// </summary>
@@ -349,7 +354,14 @@ namespace EpicGames.Core
 			int requiredSize = _currentChunk.Length + Math.Max(sizeHint, 1);
 			if (requiredSize > _currentChunk.Data.Length)
 			{
-				_currentChunk = CreateChunk(_currentChunk.RunningIndex + _currentChunk.Length, Math.Max(sizeHint, _chunkSize));
+				int runningIndex = _currentChunk.RunningIndex + _currentChunk.Length;
+				if (_currentChunk.Length == 0)
+				{
+					_currentChunk.Release();
+					_chunks.RemoveAt(_chunks.Count - 1);
+				}
+
+				_currentChunk = CreateChunk(runningIndex, Math.Max(sizeHint, _chunkSize));
 				_chunks.Add(_currentChunk);
 			}
 			return _currentChunk.Data.Slice(_currentChunk.Length);
@@ -486,12 +498,7 @@ namespace EpicGames.Core
 		{
 			public readonly IMemoryOwner<byte> Owner;
 
-			public PooledChunk(int runningIndex, int size)
-				: this(runningIndex, MemoryPool<byte>.Shared.Rent(size))
-			{
-			}
-
-			private PooledChunk(int runningIndex, IMemoryOwner<byte> owner)
+			public PooledChunk(int runningIndex, IMemoryOwner<byte> owner)
 				: base(runningIndex, owner.Memory)
 			{
 				Owner = owner;
@@ -500,6 +507,8 @@ namespace EpicGames.Core
 			public override void Release() => Owner.Dispose();
 		}
 
+		readonly IMemoryAllocator<byte> _allocator;
+
 		/// <summary>
 		/// Constructor
 		/// </summary>
@@ -514,18 +523,156 @@ namespace EpicGames.Core
 		/// <param name="initialSize">Size of the initial chunk</param>
 		/// <param name="chunkSize">Default size for subsequent chunks</param>
 		public ChunkedMemoryWriter(int initialSize = 4096, int chunkSize = 4096)
-			: base(new PooledChunk(0, initialSize), chunkSize)
+			: this(PoolAllocator.Shared, initialSize, chunkSize)
+		{
+		}
+
+		/// <summary>
+		/// Constructor
+		/// </summary>
+		/// <param name="allocator">Allocator to draw from</param>
+		/// <param name="initialSize">Size of the initial chunk</param>
+		/// <param name="chunkSize">Default size for subsequent chunks</param>
+		public ChunkedMemoryWriter(IMemoryAllocator<byte> allocator, int initialSize = 4096, int chunkSize = 4096)
+			: base(new PooledChunk(0, allocator.Alloc(initialSize)), chunkSize)
+		{
+			_allocator = allocator;
+		}
+
+		/// <inheritdoc/>
+		protected override Chunk CreateChunk(int runningIndex, int size) => new PooledChunk(runningIndex, _allocator.Alloc(size));
+
+		/// <inheritdoc/>
+		public void Dispose()
+		{
+			Clear();
+		}
+	}
+
+	/// <summary>
+	/// Implementation of <see cref="ChunkedMemoryWriterBase"/> which takes pages from a <see cref="IMemoryAllocator{Byte}"/>
+	/// </summary>
+	public sealed class RefCountedMemoryWriter : ChunkedMemoryWriterBase, IDisposable
+	{
+		class BufferChunk : Chunk
+		{
+			public readonly IRefCountedHandle<Memory<byte>> Handle;
+
+			public BufferChunk(int runningIndex, IRefCountedHandle<Memory<byte>> handle)
+				: base(runningIndex, handle.Target)
+			{
+				Handle = handle;
+			}
+
+			public override void Release()
+			{
+				Handle.Dispose();
+			}
+		}
+
+		class ArrayDisposer : IDisposable
+		{
+			readonly IDisposable[] _handles;
+
+			public ArrayDisposer(IDisposable[] handles) => _handles = handles;
+
+			public void Dispose()
+			{
+				for (int idx = 0; idx < _handles.Length; idx++)
+				{
+					_handles[idx].Dispose();
+				}
+			}
+		}
+
+		readonly IMemoryAllocator<byte> _allocator;
+
+		/// <summary>
+		/// Constructor
+		/// </summary>
+		/// <param name="allocator">Allocator for buffers</param>
+		/// <param name="chunkSize"></param>
+		public RefCountedMemoryWriter(IMemoryAllocator<byte> allocator, int chunkSize)
+			: this(allocator, chunkSize, chunkSize)
+		{ }
+
+		/// <summary>
+		/// Constructor
+		/// </summary>
+		/// <param name="allocator">Allocator for buffers</param>
+		/// <param name="initialSize">Size of the initial chunk</param>
+		/// <param name="chunkSize">Default size for subsequent chunks</param>
+		public RefCountedMemoryWriter(IMemoryAllocator<byte> allocator, int initialSize, int chunkSize)
+			: base(new BufferChunk(0, RefCountedHandle.Create(allocator.Alloc(initialSize))), chunkSize)
 		{
+			_allocator = allocator;
 		}
 
 		/// <inheritdoc/>
-		protected override Chunk CreateChunk(int runningIndex, int size) => new PooledChunk(runningIndex, size);
+		protected override Chunk CreateChunk(int runningIndex, int size) => new BufferChunk(runningIndex, RefCountedHandle.Create(_allocator.Alloc(size)));
 
 		/// <inheritdoc/>
 		public void Dispose()
 		{
 			Clear();
 		}
+
+		/// <summary>
+		/// Returns a contiguous block of memory containing the written data
+		/// </summary>
+		public IRefCountedHandle<ReadOnlyMemory<byte>> AsRefCountedMemory()
+		{
+			if (Chunks.Count == 1)
+			{
+				BufferChunk buffer = (BufferChunk)Chunks[0];
+				return RefCountedHandle.Create<ReadOnlyMemory<byte>>(buffer.WrittenMemory, buffer.Handle.AddRef());
+			}
+			else
+			{
+				int length = Length;
+
+				IRefCountedHandle<Memory<byte>> allocation = RefCountedHandle.Create(_allocator.Alloc(length));
+				CopyTo(allocation.Target.Span);
+
+				return RefCountedHandle.Create<ReadOnlyMemory<byte>>(allocation.Target.Slice(0, length), allocation);
+			}
+		}
+
+		/// <summary>
+		/// Creates a reference counted sequence from the allocated data
+		/// </summary>
+		public IRefCountedHandle<ReadOnlySequence<byte>> AsRefCountedSequence() => AsRefCountedSequence(0, Length);
+
+		/// <summary>
+		/// Creates a reference counted sequence from the allocated data
+		/// </summary>
+		public IRefCountedHandle<ReadOnlySequence<byte>> AsRefCountedSequence(int offset) => AsRefCountedSequence(offset, Length - offset);
+
+		/// <summary>
+		/// Creates a reference counted sequence from the allocated data
+		/// </summary>
+		public IRefCountedHandle<ReadOnlySequence<byte>> AsRefCountedSequence(int offset, int length)
+		{
+			int minChunkIdx = 0;
+			while (minChunkIdx + 1 < Chunks.Count && offset > Chunks[minChunkIdx + 1].RunningIndex)
+			{
+				minChunkIdx++;
+			}
+
+			int maxChunkIdx = minChunkIdx + 1;
+			while (maxChunkIdx < Chunks.Count && offset + length > Chunks[maxChunkIdx].RunningIndex)
+			{
+				maxChunkIdx++;
+			}
+
+			IDisposable[] handles = new IDisposable[maxChunkIdx - minChunkIdx];
+			for (int chunkIdx = minChunkIdx; chunkIdx < maxChunkIdx; chunkIdx++)
+			{
+				handles[chunkIdx - minChunkIdx] = ((BufferChunk)Chunks[chunkIdx]).Handle.AddRef();
+			}
+
+			return new RefCountedHandle<ReadOnlySequence<byte>>(AsSequence(offset, length), new ArrayDisposer(handles));
+		}
 	}
 
 	/// <summary>
diff --git a/Engine/Source/Programs/Shared/EpicGames.Core/RefCountedHandle.cs b/Engine/Source/Programs/Shared/EpicGames.Core/RefCountedHandle.cs
--- a/Engine/Source/Programs/Shared/EpicGames.Core/RefCountedHandle.cs
+++ b/Engine/Source/Programs/Shared/EpicGames.Core/RefCountedHandle.cs
@@ -1,6 +1,7 @@
 // Copyright Epic Games, Inc. All Rights Reserved.
 
 using System;
+using System.Buffers;
 using System.Threading;
 
 namespace EpicGames.Core
@@ -54,6 +55,14 @@ namespace EpicGames.Core
 		{
 			return new RefCountedHandle<T>(target, owner);
 		}
+
+		/// <summary>
+		/// Creates a reference counted handle to a memory owner.
+		/// </summary>
+		public static IRefCountedHandle<Memory<T>> Create<T>(IMemoryOwner<T> target) where T : struct
+		{
+			return new RefCountedHandle<Memory<T>>(target.Memory, target);
+		}
 	}
 
 	/// <summary>
