Horde: Rename INodeReader -> IBlobReader and INodeWriter -> BlobWriter.

diff --git a/Engine/Source/Programs/Horde/Horde.Server/Perforce/PerforceReplicator.cs b/Engine/Source/Programs/Horde/Horde.Server/Perforce/PerforceReplicator.cs
--- a/Engine/Source/Programs/Horde/Horde.Server/Perforce/PerforceReplicator.cs
+++ b/Engine/Source/Programs/Horde/Horde.Server/Perforce/PerforceReplicator.cs
@@ -53,7 +53,7 @@ namespace Horde.Server.Perforce
 				Paths = new List<string>();
 			}
 
-			public SyncNode(INodeReader reader)
+			public SyncNode(IBlobReader reader)
 			{
 				Change = (int)reader.ReadUnsignedVarInt();
 				ParentChange = (int)reader.ReadUnsignedVarInt();
@@ -62,7 +62,7 @@ namespace Horde.Server.Perforce
 			}
 
 			/// <inheritdoc/>
-			public override void Serialize(INodeWriter writer)
+			public override void Serialize(IBlobWriter writer)
 			{
 				writer.WriteUnsignedVarInt(Change);
 				writer.WriteUnsignedVarInt(ParentChange);
diff --git a/Engine/Source/Programs/Shared/EpicGames.Horde.Tests/BundleStoreTests.cs b/Engine/Source/Programs/Shared/EpicGames.Horde.Tests/BundleStoreTests.cs
--- a/Engine/Source/Programs/Shared/EpicGames.Horde.Tests/BundleStoreTests.cs
+++ b/Engine/Source/Programs/Shared/EpicGames.Horde.Tests/BundleStoreTests.cs
@@ -39,12 +39,12 @@ namespace EpicGames.Horde.Tests
 
 			public TextNode(string text) => Text = text;
 
-			public TextNode(INodeReader reader)
+			public TextNode(IBlobReader reader)
 			{
 				Text = reader.ReadString();
 			}
 
-			public override void Serialize(INodeWriter writer)
+			public override void Serialize(IBlobWriter writer)
 			{
 				writer.WriteString(Text);
 			}
@@ -133,13 +133,13 @@ namespace EpicGames.Horde.Tests
 				Refs = refs;
 			}
 
-			public SimpleNode(INodeReader reader)
+			public SimpleNode(IBlobReader reader)
 			{
 				Data = new ReadOnlySequence<byte>(reader.ReadVariableLengthBytes());
 				Refs = reader.ReadVariableLengthArray(() => reader.ReadHashedNodeRef<SimpleNode>());
 			}
 
-			public override void Serialize(INodeWriter writer)
+			public override void Serialize(IBlobWriter writer)
 			{
 				writer.WriteVariableLengthBytes(Data);
 				writer.WriteVariableLengthArray(Refs, x => writer.WriteHashedNodeRef(x));
diff --git a/Engine/Source/Programs/Shared/EpicGames.Horde/Logs/LogChunkNode.cs b/Engine/Source/Programs/Shared/EpicGames.Horde/Logs/LogChunkNode.cs
--- a/Engine/Source/Programs/Shared/EpicGames.Horde/Logs/LogChunkNode.cs
+++ b/Engine/Source/Programs/Shared/EpicGames.Horde/Logs/LogChunkNode.cs
@@ -111,13 +111,13 @@ namespace EpicGames.Horde.Logs
 		/// Deserializing constructor
 		/// </summary>
 		/// <param name="reader">Reader to pull data from</param>
-		public LogChunkNode(INodeReader reader)
+		public LogChunkNode(IBlobReader reader)
 			: this(reader.ReadVariableLengthBytes())
 		{
 		}
 
 		/// <inheritdoc/>
-		public override void Serialize(INodeWriter writer)
+		public override void Serialize(IBlobWriter writer)
 		{
 			writer.WriteVariableLengthBytes(Data.Span);
 		}
@@ -228,7 +228,7 @@ namespace EpicGames.Horde.Logs
 		/// Deserializing constructor
 		/// </summary>
 		/// <param name="reader"></param>
-		public LogChunkRef(INodeReader reader)
+		public LogChunkRef(IBlobReader reader)
 			: base(reader)
 		{
 			LineIndex = (int)reader.ReadUnsignedVarInt();
@@ -238,7 +238,7 @@ namespace EpicGames.Horde.Logs
 		}
 
 		/// <inheritdoc/>
-		public override void Serialize(INodeWriter writer)
+		public override void Serialize(IBlobWriter writer)
 		{
 			base.Serialize(writer);
 
diff --git a/Engine/Source/Programs/Shared/EpicGames.Horde/Logs/LogIndexNode.cs b/Engine/Source/Programs/Shared/EpicGames.Horde/Logs/LogIndexNode.cs
--- a/Engine/Source/Programs/Shared/EpicGames.Horde/Logs/LogIndexNode.cs
+++ b/Engine/Source/Programs/Shared/EpicGames.Horde/Logs/LogIndexNode.cs
@@ -88,7 +88,7 @@ namespace EpicGames.Horde.Logs
 		/// Deserialization constructor
 		/// </summary>
 		/// <param name="reader">Reader for data</param>
-		public LogIndexNode(INodeReader reader)
+		public LogIndexNode(IBlobReader reader)
 		{
 			int version = (int)reader.ReadUnsignedVarInt();
 			if (version != CurrentVersion)
@@ -102,7 +102,7 @@ namespace EpicGames.Horde.Logs
 		}
 
 		/// <inheritdoc/>
-		public override void Serialize(INodeWriter writer)
+		public override void Serialize(IBlobWriter writer)
 		{
 			writer.WriteUnsignedVarInt(CurrentVersion);
 			writer.WriteNgramSet(_ngramSet);
diff --git a/Engine/Source/Programs/Shared/EpicGames.Horde/Logs/LogNode.cs b/Engine/Source/Programs/Shared/EpicGames.Horde/Logs/LogNode.cs
--- a/Engine/Source/Programs/Shared/EpicGames.Horde/Logs/LogNode.cs
+++ b/Engine/Source/Programs/Shared/EpicGames.Horde/Logs/LogNode.cs
@@ -82,7 +82,7 @@ namespace EpicGames.Horde.Logs
 		/// Deserializing constructor
 		/// </summary>
 		/// <param name="reader">Reader to draw data from</param>
-		public LogNode(INodeReader reader)
+		public LogNode(IBlobReader reader)
 		{
 			Format = (LogFormat)reader.ReadUInt8();
 			LineCount = (int)reader.ReadUnsignedVarInt();
@@ -93,7 +93,7 @@ namespace EpicGames.Horde.Logs
 		}
 
 		/// <inheritdoc/>
-		public override void Serialize(INodeWriter writer)
+		public override void Serialize(IBlobWriter writer)
 		{
 			writer.WriteUInt8((byte)Format);
 			writer.WriteUnsignedVarInt(LineCount);
diff --git a/Engine/Source/Programs/Shared/EpicGames.Horde/Storage/IBlobReader.cs b/Engine/Source/Programs/Shared/EpicGames.Horde/Storage/IBlobReader.cs
new file mode 100644
--- /dev/null
+++ b/Engine/Source/Programs/Shared/EpicGames.Horde/Storage/IBlobReader.cs
@@ -0,0 +1,87 @@
+// Copyright Epic Games, Inc. All Rights Reserved.
+
+using System;
+using System.Collections.Generic;
+using EpicGames.Core;
+
+namespace EpicGames.Horde.Storage
+{
+	/// <summary>
+	/// Interface for reading nodes from storage
+	/// </summary>
+	public interface IBlobReader : IMemoryReader
+	{
+		/// <summary>
+		/// Type to deserialize
+		/// </summary>
+		BlobType Type { get; }
+
+		/// <summary>
+		/// Version of the current node, as specified via <see cref="NodeTypeAttribute"/>
+		/// </summary>
+		int Version { get; }
+
+		/// <summary>
+		/// Locations of all referenced nodes.
+		/// </summary>
+		IReadOnlyList<IBlobHandle> References { get; }
+
+		/// <summary>
+		/// Gets the next serialized blob handle
+		/// </summary>
+		IBlobHandle ReadBlobReference();
+	}
+
+	/// <summary>
+	/// Reader for blob objects
+	/// </summary>
+	sealed class BlobReader : MemoryReader, IBlobReader
+	{
+		/// <summary>
+		/// Type to deserialize
+		/// </summary>
+		public BlobType Type => _blobData.Type;
+
+		/// <summary>
+		/// Version of the current node, as specified via <see cref="NodeTypeAttribute"/>
+		/// </summary>
+		public int Version => Type.Version;
+
+		/// <summary>
+		/// Total length of the data in this node
+		/// </summary>
+		public int Length => _blobData.Data.Length;
+
+		/// <summary>
+		/// Amount of data remaining to be read
+		/// </summary>
+		public int RemainingLength => RemainingMemory.Length;
+
+		/// <summary>
+		/// Raw data for this blob
+		/// </summary>
+		public ReadOnlyMemory<byte> Data => _blobData.Data;
+
+		/// <summary>
+		/// Locations of all referenced nodes.
+		/// </summary>
+		public IReadOnlyList<IBlobHandle> References => _blobData.Refs;
+
+		readonly BlobData _blobData;
+		int _refIdx;
+
+		/// <summary>
+		/// Constructor
+		/// </summary>
+		public BlobReader(BlobData blobData)
+			: base(blobData.Data)
+		{
+			_blobData = blobData;
+		}
+
+		/// <summary>
+		/// Gets the next serialized blob handle
+		/// </summary>
+		public IBlobHandle ReadBlobReference() => _blobData.Refs[_refIdx++];
+	}
+}
diff --git a/Engine/Source/Programs/Shared/EpicGames.Horde/Storage/IBlobWriter.cs b/Engine/Source/Programs/Shared/EpicGames.Horde/Storage/IBlobWriter.cs
new file mode 100644
--- /dev/null
+++ b/Engine/Source/Programs/Shared/EpicGames.Horde/Storage/IBlobWriter.cs
@@ -0,0 +1,88 @@
+// Copyright Epic Games, Inc. All Rights Reserved.
+
+using System;
+using System.Collections.Generic;
+using EpicGames.Core;
+
+namespace EpicGames.Horde.Storage
+{
+	/// <summary>
+	/// Interface for a writer of node objects
+	/// </summary>
+	public interface IBlobWriter : IMemoryWriter
+	{
+		/// <summary>
+		/// Adds a reference to another blob. This reference is stored out of band, and will not result in any bytes written to the output.
+		/// </summary>
+		/// <param name="reference">Referenced blob</param>
+		void WriteBlobReference(IBlobHandle reference);
+	}
+
+	/// <summary>
+	/// Writer for node objects, which tracks references to other nodes
+	/// </summary>
+	sealed class BlobWriter : IBlobWriter
+	{
+		readonly IStorageWriter _treeWriter;
+
+		Memory<byte> _memory;
+		readonly List<IBlobHandle> _refs = new List<IBlobHandle>();
+		int _length;
+
+		/// <summary>
+		/// List of serialized references
+		/// </summary>
+		public IReadOnlyList<IBlobHandle> References => _refs;
+
+		/// <inheritdoc/>
+		public int Length => _length;
+
+		/// <summary>
+		/// Constructor
+		/// </summary>
+		/// <param name="treeWriter"></param>
+		public BlobWriter(IStorageWriter treeWriter)
+		{
+			_treeWriter = treeWriter;
+			_memory = treeWriter.GetOutputBuffer(0, 2 * 1024);
+		}
+
+		/// <summary>
+		/// Adds a reference to another blob. This reference is stored out of band, and will not result in any bytes written to the output.
+		/// </summary>
+		/// <param name="reference">Referenced blob</param>
+		public void WriteBlobReference(IBlobHandle reference) => _refs.Add(reference);
+
+		/// <summary>
+		/// Computes the hash of the written data
+		/// </summary>
+		public IoHash ComputeHash() => IoHash.Compute(_memory.Span.Slice(0, _length));
+
+		/// <summary>
+		/// Writes a handle to another node
+		/// </summary>
+		public void WriteHashedBlobHandle(IoHash hash, IBlobHandle target)
+		{
+			this.WriteIoHash(hash);
+			WriteBlobReference(target);
+		}
+
+		/// <inheritdoc/>
+		public Span<byte> GetSpan(int sizeHint = 0) => GetMemory(sizeHint).Span;
+
+		/// <inheritdoc/>
+		public Memory<byte> GetMemory(int sizeHint = 0)
+		{
+			int newLength = _length + Math.Max(sizeHint, 1);
+			if (newLength > _memory.Length)
+			{
+				newLength = _length + Math.Max(sizeHint, 1024);
+				_memory = _treeWriter.GetOutputBuffer(_length, Math.Max(_memory.Length * 2, newLength));
+			}
+			return _memory.Slice(_length);
+		}
+
+		/// <inheritdoc/>
+		public void Advance(int length) => _length += length;
+	}
+}
diff --git a/Engine/Source/Programs/Shared/EpicGames.Horde/Storage/Node.cs b/Engine/Source/Programs/Shared/EpicGames.Horde/Storage/Node.cs
--- a/Engine/Source/Programs/Shared/EpicGames.Horde/Storage/Node.cs
+++ b/Engine/Source/Programs/Shared/EpicGames.Horde/Storage/Node.cs
@@ -2,13 +2,11 @@
 
 using System;
 using System.Collections.Concurrent;
-using System.Collections.Generic;
 using System.Diagnostics.CodeAnalysis;
 using System.Reflection;
 using System.Reflection.Emit;
 using System.Threading;
 using System.Threading.Tasks;
-using EpicGames.Core;
 
 namespace EpicGames.Horde.Storage
 {
@@ -66,7 +64,7 @@ namespace EpicGames.Horde.Storage
 		/// Serialize the contents of this node
 		/// </summary>
 		/// <returns>Data for the node</returns>
-		public abstract void Serialize(INodeWriter writer);
+		public abstract void Serialize(IBlobWriter writer);
 
 		#region Static methods
 
@@ -200,7 +198,7 @@ namespace EpicGames.Horde.Storage
 			s_guidToDeserializer.TryAdd(nodeType.Guid, deserializer);
 		}
 
-		static readonly ConstructorInfo? s_nodeReaderTypeCtor = typeof(NodeReader).GetConstructor(new[] { typeof(BlobData) });
+		static readonly ConstructorInfo? s_nodeReaderTypeCtor = typeof(BlobReader).GetConstructor(new[] { typeof(BlobData) });
 
 		static Func<BlobData, Node> CreateDeserializer(Type type)
 		{
@@ -219,7 +217,7 @@ namespace EpicGames.Horde.Storage
 				return (Func<BlobData, Node>)method.CreateDelegate(typeof(Func<BlobData, Node>));
 			}
 
-			ConstructorInfo? nodeReaderCtor = type.GetConstructor(new[] { typeof(NodeReader) });
+			ConstructorInfo? nodeReaderCtor = type.GetConstructor(new[] { typeof(BlobReader) });
 			if (nodeReaderCtor != null)
 			{
 				DynamicMethod method = new DynamicMethod($"Create_{type.Name}", type, signature, true);
@@ -233,171 +231,12 @@ namespace EpicGames.Horde.Storage
 				return (Func<BlobData, Node>)method.CreateDelegate(typeof(Func<BlobData, Node>));
 			}
 
-			throw new InvalidOperationException($"Type {type.Name} does not have a constructor taking a {typeof(BlobData).Name} or {typeof(NodeReader).Name} instance as parameter.");
+			throw new InvalidOperationException($"Type {type.Name} does not have a constructor taking a {typeof(BlobData).Name} or {typeof(BlobReader).Name} instance as parameter.");
 		}
 
 		#endregion
 	}
 
-	/// <summary>
-	/// Interface for reading nodes from storage
-	/// </summary>
-	public interface INodeReader : IMemoryReader
-	{
-		/// <summary>
-		/// Type to deserialize
-		/// </summary>
-		BlobType Type { get; }
-
-		/// <summary>
-		/// Version of the current node, as specified via <see cref="NodeTypeAttribute"/>
-		/// </summary>
-		int Version { get; }
-
-		/// <summary>
-		/// Locations of all referenced nodes.
-		/// </summary>
-		IReadOnlyList<IBlobHandle> References { get; }
-
-		/// <summary>
-		/// Gets the next serialized blob handle
-		/// </summary>
-		IBlobHandle ReadBlobReference();
-	}
-
-	/// <summary>
-	/// Reader for tree nodes
-	/// </summary>
-	sealed class NodeReader : MemoryReader, INodeReader
-	{
-		/// <summary>
-		/// Type to deserialize
-		/// </summary>
-		public BlobType Type => _blobData.Type;
-
-		/// <summary>
-		/// Version of the current node, as specified via <see cref="NodeTypeAttribute"/>
-		/// </summary>
-		public int Version => Type.Version;
-
-		/// <summary>
-		/// Total length of the data in this node
-		/// </summary>
-		public int Length => _blobData.Data.Length;
-
-		/// <summary>
-		/// Amount of data remaining to be read
-		/// </summary>
-		public int RemainingLength => RemainingMemory.Length;
-
-		/// <summary>
-		/// Raw data for this blob
-		/// </summary>
-		public ReadOnlyMemory<byte> Data => _blobData.Data;
-
-		/// <summary>
-		/// Locations of all referenced nodes.
-		/// </summary>
-		public IReadOnlyList<IBlobHandle> References => _blobData.Refs;
-
-		readonly BlobData _blobData;
-		int _refIdx;
-
-		/// <summary>
-		/// Constructor
-		/// </summary>
-		public NodeReader(BlobData blobData)
-			: base(blobData.Data)
-		{
-			_blobData = blobData;
-		}
-
-		/// <summary>
-		/// Gets the next serialized blob handle
-		/// </summary>
-		public IBlobHandle ReadBlobReference() => _blobData.Refs[_refIdx++];
-	}
-
-	/// <summary>
-	/// Interface for a writer of node objects
-	/// </summary>
-	public interface INodeWriter : IMemoryWriter
-	{
-		/// <summary>
-		/// Adds a reference to another blob. This reference is stored out of band, and will not result in any bytes written to the output.
-		/// </summary>
-		/// <param name="reference">Referenced blob</param>
-		void WriteBlobReference(IBlobHandle reference);
-	}
-
-	/// <summary>
-	/// Writer for node objects, which tracks references to other nodes
-	/// </summary>
-	sealed class NodeWriter : INodeWriter
-	{
-		readonly IStorageWriter _treeWriter;
-
-		Memory<byte> _memory;
-		readonly List<IBlobHandle> _refs = new List<IBlobHandle>();
-		int _length;
-
-		/// <summary>
-		/// List of serialized references
-		/// </summary>
-		public IReadOnlyList<IBlobHandle> References => _refs;
-
-		/// <inheritdoc/>
-		public int Length => _length;
-
-		/// <summary>
-		/// Constructor
-		/// </summary>
-		/// <param name="treeWriter"></param>
-		public NodeWriter(IStorageWriter treeWriter)
-		{
-			_treeWriter = treeWriter;
-			_memory = treeWriter.GetOutputBuffer(0, 2 * 1024);
-		}
-
-		/// <summary>
-		/// Adds a reference to another blob. This reference is stored out of band, and will not result in any bytes written to the output.
-		/// </summary>
-		/// <param name="reference">Referenced blob</param>
-		public void WriteBlobReference(IBlobHandle reference) => _refs.Add(reference);
-
-		/// <summary>
-		/// Computes the hash of the written data
-		/// </summary>
-		public IoHash ComputeHash() => IoHash.Compute(_memory.Span.Slice(0, _length));
-
-		/// <summary>
-		/// Writes a handle to another node
-		/// </summary>
-		public void WriteHashedBlobHandle(IoHash hash, IBlobHandle target)
-		{
-			this.WriteIoHash(hash);
-			WriteBlobReference(target);
-		}
-
-		/// <inheritdoc/>
-		public Span<byte> GetSpan(int sizeHint = 0) => GetMemory(sizeHint).Span;
-
-		/// <inheritdoc/>
-		public Memory<byte> GetMemory(int sizeHint = 0)
-		{
-			int newLength = _length + Math.Max(sizeHint, 1);
-			if (newLength > _memory.Length)
-			{
-				newLength = _length + Math.Max(sizeHint, 1024);
-				_memory = _treeWriter.GetOutputBuffer(_length, Math.Max(_memory.Length * 2, newLength));
-			}
-			return _memory.Slice(_length);
-		}
-
-		/// <inheritdoc/>
-		public void Advance(int length) => _length += length;
-	}
-
 	/// <summary>
 	/// Extension methods for serializing <see cref="Node"/> objects
 	/// </summary>
diff --git a/Engine/Source/Programs/Shared/EpicGames.Horde/Storage/NodeRef.cs b/Engine/Source/Programs/Shared/EpicGames.Horde/Storage/NodeRef.cs
--- a/Engine/Source/Programs/Shared/EpicGames.Horde/Storage/NodeRef.cs
+++ b/Engine/Source/Programs/Shared/EpicGames.Horde/Storage/NodeRef.cs
@@ -33,13 +33,13 @@ namespace EpicGames.Horde.Storage
 		/// Deserialization constructor
 		/// </summary>
 		/// <param name="reader"></param>
-		public NodeRef(INodeReader reader) => Handle = reader.ReadBlobReference();
+		public NodeRef(IBlobReader reader) => Handle = reader.ReadBlobReference();
 
 		/// <summary>
 		/// Serialize the node to the given writer
 		/// </summary>
 		/// <param name="writer"></param>
-		public virtual void Serialize(INodeWriter writer) => writer.WriteBlobReference(Handle);
+		public virtual void Serialize(IBlobWriter writer) => writer.WriteBlobReference(Handle);
 
 		/// <summary>
 		/// Resolve this reference to a concrete node
@@ -74,7 +74,7 @@ namespace EpicGames.Horde.Storage
 		/// <summary>
 		/// Constructor
 		/// </summary>
-		public NodeRef(INodeReader reader) : base(reader)
+		public NodeRef(IBlobReader reader) : base(reader)
 		{ }
 
 		/// <summary>
@@ -106,14 +106,14 @@ namespace EpicGames.Horde.Storage
 		/// Deserialization constructor
 		/// </summary>
 		/// <param name="reader"></param>
-		public HashedNodeRef(INodeReader reader) : this(reader.ReadIoHash(), reader.ReadBlobReference())
+		public HashedNodeRef(IBlobReader reader) : this(reader.ReadIoHash(), reader.ReadBlobReference())
 		{ }
 
 		/// <summary>
 		/// Serialize the node to the given writer
 		/// </summary>
 		/// <param name="writer"></param>
-		public override void Serialize(INodeWriter writer)
+		public override void Serialize(IBlobWriter writer)
 		{
 			base.Serialize(writer);
 			writer.WriteIoHash(Hash);
@@ -143,7 +143,7 @@ namespace EpicGames.Horde.Storage
 		/// <summary>
 		/// Constructor
 		/// </summary>
-		public HashedNodeRef(INodeReader reader) : base(reader)
+		public HashedNodeRef(IBlobReader reader) : base(reader)
 		{
 		}
 
@@ -168,7 +168,7 @@ namespace EpicGames.Horde.Storage
 		/// </summary>
 		/// <param name="reader">Reader to deserialize from</param>
 		/// <returns>New untyped ref</returns>
-		public static NodeRef ReadNodeRef(this INodeReader reader)
+		public static NodeRef ReadNodeRef(this IBlobReader reader)
 		{
 			return new NodeRef(reader);
 		}
@@ -179,7 +179,7 @@ namespace EpicGames.Horde.Storage
 		/// <typeparam name="T">Type of the referenced node</typeparam>
 		/// <param name="reader">Reader to deserialize from</param>
 		/// <returns>New strongly typed ref</returns>
-		public static NodeRef<T> ReadNodeRef<T>(this INodeReader reader) where T : Node
+		public static NodeRef<T> ReadNodeRef<T>(this IBlobReader reader) where T : Node
 		{
 			return new NodeRef<T>(reader);
 		}
@@ -189,7 +189,7 @@ namespace EpicGames.Horde.Storage
 		/// </summary>
 		/// <param name="reader">Reader to deserialize from</param>
 		/// <returns>New untyped ref</returns>
-		public static HashedNodeRef ReadHashedNodeRef(this INodeReader reader)
+		public static HashedNodeRef ReadHashedNodeRef(this IBlobReader reader)
 		{
 			return new HashedNodeRef(reader);
 		}
@@ -200,7 +200,7 @@ namespace EpicGames.Horde.Storage
 		/// <typeparam name="T">Type of the referenced node</typeparam>
 		/// <param name="reader">Reader to deserialize from</param>
 		/// <returns>New strongly typed ref</returns>
-		public static HashedNodeRef<T> ReadHashedNodeRef<T>(this INodeReader reader) where T : Node
+		public static HashedNodeRef<T> ReadHashedNodeRef<T>(this IBlobReader reader) where T : Node
 		{
 			return new HashedNodeRef<T>(reader);
 		}
@@ -210,7 +210,7 @@ namespace EpicGames.Horde.Storage
 		/// </summary>
 		/// <param name="reader">Reader to deserialize from</param>
 		/// <returns>New untyped ref</returns>
-		public static NodeRef? ReadOptionalNodeRef(this INodeReader reader)
+		public static NodeRef? ReadOptionalNodeRef(this IBlobReader reader)
 		{
 			if (reader.ReadBoolean())
 			{
@@ -227,7 +227,7 @@ namespace EpicGames.Horde.Storage
 		/// </summary>
 		/// <param name="reader">Reader to deserialize from</param>
 		/// <returns>New strongly typed ref</returns>
-		public static NodeRef<T>? ReadOptionalNodeRef<T>(this INodeReader reader) where T : Node
+		public static NodeRef<T>? ReadOptionalNodeRef<T>(this IBlobReader reader) where T : Node
 		{
 			if (reader.ReadBoolean())
 			{
@@ -244,7 +244,7 @@ namespace EpicGames.Horde.Storage
 		/// </summary>
 		/// <param name="writer">Writer to serialize to</param>
 		/// <param name="nodeRef">Value to write</param>
-		public static void WriteNodeRef(this INodeWriter writer, NodeRef nodeRef)
+		public static void WriteNodeRef(this IBlobWriter writer, NodeRef nodeRef)
 		{
 			nodeRef.Serialize(writer);
 		}
@@ -254,7 +254,7 @@ namespace EpicGames.Horde.Storage
 		/// </summary>
 		/// <param name="writer">Writer to serialize to</param>
 		/// <param name="nodeRef">Value to write</param>
-		public static void WriteHashedNodeRef(this INodeWriter writer, HashedNodeRef nodeRef)
+		public static void WriteHashedNodeRef(this IBlobWriter writer, HashedNodeRef nodeRef)
 		{
 			nodeRef.Serialize(writer);
 		}
@@ -264,7 +264,7 @@ namespace EpicGames.Horde.Storage
 		/// </summary>
 		/// <param name="writer">Writer to serialize to</param>
 		/// <param name="value">Value to write</param>
-		public static void WriteOptionalNodeRef(this INodeWriter writer, NodeRef? value)
+		public static void WriteOptionalNodeRef(this IBlobWriter writer, NodeRef? value)
 		{
 			if (value == null)
 			{
@@ -287,7 +287,7 @@ namespace EpicGames.Horde.Storage
 		public static async ValueTask<HashedNodeRef<TNode>> WriteHashedNodeAsync<TNode>(this IStorageWriter writer, TNode node, CancellationToken cancellationToken = default) where TNode : Node
 		{
 			// Serialize the node
-			NodeWriter nodeWriter = new NodeWriter(writer);
+			BlobWriter nodeWriter = new BlobWriter(writer);
 			node.Serialize(nodeWriter);
 			IoHash hash = nodeWriter.ComputeHash();
 
diff --git a/Engine/Source/Programs/Shared/EpicGames.Horde/Storage/Nodes/CbNode.cs b/Engine/Source/Programs/Shared/EpicGames.Horde/Storage/Nodes/CbNode.cs
--- a/Engine/Source/Programs/Shared/EpicGames.Horde/Storage/Nodes/CbNode.cs
+++ b/Engine/Source/Programs/Shared/EpicGames.Horde/Storage/Nodes/CbNode.cs
@@ -14,10 +14,10 @@ namespace EpicGames.Horde.Storage.Nodes
 	{
 		class HandleMapper
 		{
-			readonly INodeReader _reader;
+			readonly IBlobReader _reader;
 			readonly List<HashedNodeRef> _refs;
 
-			public HandleMapper(INodeReader reader, List<HashedNodeRef> refs)
+			public HandleMapper(IBlobReader reader, List<HashedNodeRef> refs)
 			{
 				_reader = reader;
 				_refs = refs;
@@ -68,7 +68,7 @@ namespace EpicGames.Horde.Storage.Nodes
 		/// Deserialization constructor
 		/// </summary>
 		/// <param name="reader">Reader to deserialize from</param>
-		public CbNode(INodeReader reader)
+		public CbNode(IBlobReader reader)
 		{
 			Object = new CbObject(reader.GetMemory().ToArray());
 
@@ -79,7 +79,7 @@ namespace EpicGames.Horde.Storage.Nodes
 		}
 
 		/// <inheritdoc/>
-		public override void Serialize(INodeWriter writer)
+		public override void Serialize(IBlobWriter writer)
 		{
 			writer.WriteFixedLengthBytes(Object.GetView().Span);
 		}
diff --git a/Engine/Source/Programs/Shared/EpicGames.Horde/Storage/Nodes/ChunkedDataNode.cs b/Engine/Source/Programs/Shared/EpicGames.Horde/Storage/Nodes/ChunkedDataNode.cs
--- a/Engine/Source/Programs/Shared/EpicGames.Horde/Storage/Nodes/ChunkedDataNode.cs
+++ b/Engine/Source/Programs/Shared/EpicGames.Horde/Storage/Nodes/ChunkedDataNode.cs
@@ -119,7 +119,7 @@ namespace EpicGames.Horde.Storage.Nodes
 		/// <summary>
 		/// Constructor
 		/// </summary>
-		public ChunkedDataNodeRef(INodeReader reader)
+		public ChunkedDataNodeRef(IBlobReader reader)
 			: base(reader)
 		{
 			Type = (ChunkedDataNodeType)reader.ReadUnsignedVarInt();
@@ -128,7 +128,7 @@ namespace EpicGames.Horde.Storage.Nodes
 		/// <summary>
 		/// Constructor
 		/// </summary>
-		public ChunkedDataNodeRef(ChunkedDataNodeType type, INodeReader reader)
+		public ChunkedDataNodeRef(ChunkedDataNodeType type, IBlobReader reader)
 			: base(reader)
 		{
 			Type = type;
@@ -144,7 +144,7 @@ namespace EpicGames.Horde.Storage.Nodes
 		}
 
 		/// <inheritdoc/>
-		public override void Serialize(INodeWriter writer)
+		public override void Serialize(IBlobWriter writer)
 		{
 			base.Serialize(writer);
 
@@ -185,14 +185,14 @@ namespace EpicGames.Horde.Storage.Nodes
 		/// <summary>
 		/// Create a leaf node from the given serialized data
 		/// </summary>
-		public LeafChunkedDataNode(INodeReader reader)
+		public LeafChunkedDataNode(IBlobReader reader)
 		{
 			// Keep this code in sync with CopyToStreamAsync
 			Data = reader.GetMemory();
 		}
 
 		/// <inheritdoc/>
-		public override void Serialize(INodeWriter writer)
+		public override void Serialize(IBlobWriter writer)
 		{
 			writer.WriteFixedLengthBytes(Data.Span);
 		}
@@ -410,7 +410,7 @@ namespace EpicGames.Horde.Storage.Nodes
 		/// <summary>
 		/// Constructor
 		/// </summary>
-		public InteriorChunkedDataNode(INodeReader reader)
+		public InteriorChunkedDataNode(IBlobReader reader)
 		{
 			// Keep this code in sync with CopyToStreamAsync
 			List<ChunkedDataNodeRef> children = new List<ChunkedDataNodeRef>();
@@ -429,7 +429,7 @@ namespace EpicGames.Horde.Storage.Nodes
 		}
 
 		/// <inheritdoc/>
-		public override void Serialize(INodeWriter writer)
+		public override void Serialize(IBlobWriter writer)
 		{
 			foreach (ChunkedDataNodeRef child in Children)
 			{
@@ -534,7 +534,7 @@ namespace EpicGames.Horde.Storage.Nodes
 		/// <param name="cancellationToken">Cancellation token for the operation</param>
 		public static async Task CopyToStreamAsync(BlobData nodeData, Stream outputStream, CancellationToken cancellationToken)
 		{
-			NodeReader nodeReader = new NodeReader(nodeData);
+			BlobReader nodeReader = new BlobReader(nodeData);
 			while (nodeReader.GetMemory(0).Length > 0)
 			{
 				_ = nodeReader.ReadIoHash();
diff --git a/Engine/Source/Programs/Shared/EpicGames.Horde/Storage/Nodes/CommitNode.cs b/Engine/Source/Programs/Shared/EpicGames.Horde/Storage/Nodes/CommitNode.cs
--- a/Engine/Source/Programs/Shared/EpicGames.Horde/Storage/Nodes/CommitNode.cs
+++ b/Engine/Source/Programs/Shared/EpicGames.Horde/Storage/Nodes/CommitNode.cs
@@ -85,7 +85,7 @@ namespace EpicGames.Horde.Storage.Nodes
 		/// Deserializing constructor
 		/// </summary>
 		/// <param name="reader"></param>
-		public CommitNode(INodeReader reader)
+		public CommitNode(IBlobReader reader)
 		{
 			Number = (int)reader.ReadUnsignedVarInt();
 			Parent = reader.ReadOptionalNodeRef<CommitNode>();
@@ -100,7 +100,7 @@ namespace EpicGames.Horde.Storage.Nodes
 		}
 
 		/// <inheritdoc/>
-		public override void Serialize(INodeWriter writer)
+		public override void Serialize(IBlobWriter writer)
 		{
 			writer.WriteUnsignedVarInt(Number);
 			writer.WriteOptionalNodeRef(Parent);
diff --git a/Engine/Source/Programs/Shared/EpicGames.Horde/Storage/Nodes/DdcRefNode.cs b/Engine/Source/Programs/Shared/EpicGames.Horde/Storage/Nodes/DdcRefNode.cs
--- a/Engine/Source/Programs/Shared/EpicGames.Horde/Storage/Nodes/DdcRefNode.cs
+++ b/Engine/Source/Programs/Shared/EpicGames.Horde/Storage/Nodes/DdcRefNode.cs
@@ -34,20 +34,20 @@ namespace Horde.Server.Ddc
 		/// <summary>
 		/// Constructor
 		/// </summary>
-		public DdcRefNode(INodeReader reader)
+		public DdcRefNode(IBlobReader reader)
 		{
 			RootHash = reader.ReadIoHash();
 			References = reader.ReadList(x => (reader.ReadIoHash(), reader.ReadBlobReference()));
 		}
 
 		/// <inheritdoc/>
-		public override void Serialize(INodeWriter writer)
+		public override void Serialize(IBlobWriter writer)
 		{
 			writer.WriteIoHash(RootHash);
 			writer.WriteList(References, x => WriteReference(writer, x.Hash, x.Handle));
 		}
 
-		static void WriteReference(INodeWriter writer, IoHash hash, IBlobHandle handle)
+		static void WriteReference(IBlobWriter writer, IoHash hash, IBlobHandle handle)
 		{
 			writer.WriteIoHash(hash);
 			writer.WriteBlobReference(handle);
diff --git a/Engine/Source/Programs/Shared/EpicGames.Horde/Storage/Nodes/DirectoryEntry.cs b/Engine/Source/Programs/Shared/EpicGames.Horde/Storage/Nodes/DirectoryEntry.cs
--- a/Engine/Source/Programs/Shared/EpicGames.Horde/Storage/Nodes/DirectoryEntry.cs
+++ b/Engine/Source/Programs/Shared/EpicGames.Horde/Storage/Nodes/DirectoryEntry.cs
@@ -27,7 +27,7 @@ namespace EpicGames.Horde.Storage.Nodes
 		/// Deserializing constructor
 		/// </summary>
 		/// <param name="reader"></param>
-		public DirectoryEntry(INodeReader reader)
+		public DirectoryEntry(IBlobReader reader)
 			: base(reader)
 		{
 			Name = reader.ReadString();
@@ -37,7 +37,7 @@ namespace EpicGames.Horde.Storage.Nodes
 		/// Serialize this directory entry to disk
 		/// </summary>
 		/// <param name="writer"></param>
-		public override void Serialize(INodeWriter writer)
+		public override void Serialize(IBlobWriter writer)
 		{
 			base.Serialize(writer);
 
diff --git a/Engine/Source/Programs/Shared/EpicGames.Horde/Storage/Nodes/DirectoryNode.cs b/Engine/Source/Programs/Shared/EpicGames.Horde/Storage/Nodes/DirectoryNode.cs
--- a/Engine/Source/Programs/Shared/EpicGames.Horde/Storage/Nodes/DirectoryNode.cs
+++ b/Engine/Source/Programs/Shared/EpicGames.Horde/Storage/Nodes/DirectoryNode.cs
@@ -187,7 +187,7 @@ namespace EpicGames.Horde.Storage.Nodes
 		/// Deserialization constructor
 		/// </summary>
 		/// <param name="reader">Reader to deserialize from</param>
-		public DirectoryNode(INodeReader reader)
+		public DirectoryNode(IBlobReader reader)
 		{
 			Flags = (DirectoryFlags)reader.ReadUnsignedVarInt();
 
@@ -207,7 +207,7 @@ namespace EpicGames.Horde.Storage.Nodes
 		}
 
 		/// <inheritdoc/>
-		public override void Serialize(INodeWriter writer)
+		public override void Serialize(IBlobWriter writer)
 		{
 			writer.WriteUnsignedVarInt((ulong)Flags);
 
diff --git a/Engine/Source/Programs/Shared/EpicGames.Horde/Storage/Nodes/DirectoryNodeRef.cs b/Engine/Source/Programs/Shared/EpicGames.Horde/Storage/Nodes/DirectoryNodeRef.cs
--- a/Engine/Source/Programs/Shared/EpicGames.Horde/Storage/Nodes/DirectoryNodeRef.cs
+++ b/Engine/Source/Programs/Shared/EpicGames.Horde/Storage/Nodes/DirectoryNodeRef.cs
@@ -26,7 +26,7 @@ namespace EpicGames.Horde.Storage.Nodes
 		/// <summary>
 		/// Constructor
 		/// </summary>
-		public DirectoryNodeRef(INodeReader reader)
+		public DirectoryNodeRef(IBlobReader reader)
 			: base(reader)
 		{
 			Length = (long)reader.ReadUnsignedVarInt();
@@ -36,7 +36,7 @@ namespace EpicGames.Horde.Storage.Nodes
 		/// Serialize this directory entry to disk
 		/// </summary>
 		/// <param name="writer"></param>
-		public override void Serialize(INodeWriter writer)
+		public override void Serialize(IBlobWriter writer)
 		{
 			base.Serialize(writer);
 
diff --git a/Engine/Source/Programs/Shared/EpicGames.Horde/Storage/Nodes/FileEntry.cs b/Engine/Source/Programs/Shared/EpicGames.Horde/Storage/Nodes/FileEntry.cs
--- a/Engine/Source/Programs/Shared/EpicGames.Horde/Storage/Nodes/FileEntry.cs
+++ b/Engine/Source/Programs/Shared/EpicGames.Horde/Storage/Nodes/FileEntry.cs
@@ -93,7 +93,7 @@ namespace EpicGames.Horde.Storage.Nodes
 		/// Deserialize from a buffer
 		/// </summary>
 		/// <param name="reader"></param>
-		public FileEntry(INodeReader reader)
+		public FileEntry(IBlobReader reader)
 			: base(reader)
 		{
 			Name = reader.ReadString();
@@ -112,7 +112,7 @@ namespace EpicGames.Horde.Storage.Nodes
 		/// Serialize this entry
 		/// </summary>
 		/// <param name="writer"></param>
-		public override void Serialize(INodeWriter writer)
+		public override void Serialize(IBlobWriter writer)
 		{
 			base.Serialize(writer);
 
diff --git a/Engine/Source/Programs/Shared/EpicGames.Horde/Storage/Nodes/RedirectNode.cs b/Engine/Source/Programs/Shared/EpicGames.Horde/Storage/Nodes/RedirectNode.cs
--- a/Engine/Source/Programs/Shared/EpicGames.Horde/Storage/Nodes/RedirectNode.cs
+++ b/Engine/Source/Programs/Shared/EpicGames.Horde/Storage/Nodes/RedirectNode.cs
@@ -29,7 +29,7 @@ namespace EpicGames.Horde.Storage.Nodes
 		/// Deserialization constructor
 		/// </summary>
 		/// <param name="reader">Reader to deserialize from</param>
-		public RedirectNode(INodeReader reader) => Handle = reader.ReadBlobReference();
+		public RedirectNode(IBlobReader reader) => Handle = reader.ReadBlobReference();
 
 		/// <summary>
 		/// Gets a typed reference to the target node
@@ -42,6 +42,6 @@ namespace EpicGames.Horde.Storage.Nodes
 		public NodeRef<TTarget> ToNodeRef<TTarget>() where TTarget : Node => new NodeRef<TTarget>(Handle);
 
 		/// <inheritdoc/>
-		public override void Serialize(INodeWriter writer) => writer.WriteBlobReference(Handle);
+		public override void Serialize(IBlobWriter writer) => writer.WriteBlobReference(Handle);
 	}
 }
diff --git a/Engine/Source/Programs/UnrealBuildTool/Artifacts/HordeStorageArtifactCache.cs b/Engine/Source/Programs/UnrealBuildTool/Artifacts/HordeStorageArtifactCache.cs
--- a/Engine/Source/Programs/UnrealBuildTool/Artifacts/HordeStorageArtifactCache.cs
+++ b/Engine/Source/Programs/UnrealBuildTool/Artifacts/HordeStorageArtifactCache.cs
@@ -47,7 +47,7 @@ namespace UnrealBuildTool.Artifacts
 		/// Construct a new artifact action from the reader
 		/// </summary>
 		/// <param name="reader">Source reader</param>
-		public HordeArtifactAction(INodeReader reader)
+		public HordeArtifactAction(IBlobReader reader)
 		{
 			ArtifactAction = reader.ReadArtifactAction();
 			OutputRefs = reader.ReadVariableLengthArray(() => new HashedNodeRef<ChunkedDataNode>(reader));
@@ -57,7 +57,7 @@ namespace UnrealBuildTool.Artifacts
 		/// Serialize the artifact action 
 		/// </summary>
 		/// <param name="writer">Destination writer</param>
-		public void Serialize(INodeWriter writer)
+		public void Serialize(IBlobWriter writer)
 		{
 			writer.WriteArtifactAction(ArtifactAction);
 			writer.WriteVariableLengthArray(OutputRefs, x => x.Serialize(writer));
@@ -104,7 +104,7 @@ namespace UnrealBuildTool.Artifacts
 		/// </summary>
 		/// <param name="reader">Source reader</param>
 		/// <returns>Created artifact action</returns>
-		public static HordeArtifactAction ReadHordeArtifactAction(this INodeReader reader)
+		public static HordeArtifactAction ReadHordeArtifactAction(this IBlobReader reader)
 		{
 			return new HordeArtifactAction(reader);
 		}
@@ -114,7 +114,7 @@ namespace UnrealBuildTool.Artifacts
 		/// </summary>
 		/// <param name="writer">Destination writer</param>
 		/// <param name="artifactAction">Artifact action to write</param>
-		public static void WriteHordeArtifactAction(this INodeWriter writer, HordeArtifactAction artifactAction)
+		public static void WriteHordeArtifactAction(this IBlobWriter writer, HordeArtifactAction artifactAction)
 		{
 			artifactAction.Serialize(writer);
 		}
@@ -143,13 +143,13 @@ namespace UnrealBuildTool.Artifacts
 		/// Construct a new artifact action collection from the source reader
 		/// </summary>
 		/// <param name="reader">Source reader</param>
-		public ArtifactActionCollectionNode(INodeReader reader)
+		public ArtifactActionCollectionNode(IBlobReader reader)
 		{
 			ArtifactActions = reader.ReadDictionary<IoHash, HordeArtifactAction>(() => reader.ReadIoHash(), () => reader.ReadHordeArtifactAction());
 		}
 
 		/// <inheritdoc/>
-		public override void Serialize(INodeWriter writer)
+		public override void Serialize(IBlobWriter writer)
 		{
 			writer.WriteDictionary<IoHash, HordeArtifactAction>(ArtifactActions, (x) => writer.WriteIoHash(x), (x) => writer.WriteHordeArtifactAction(x));
 		}
