Horde: Improved cache for reading bundles.

* Now supports caching of generic async operations, with resource accounting done by explicit calls to acquire disposable objects.
* When a resource is requested that would exceed the hard limit specified in an option, cache items are purged until the resource can be satisfied.
* Items in the cache are reference counted, and will not be purged while referenced.

diff --git a/Engine/Source/Programs/Horde/Horde.Agent.Tests/JsonRpcLoggerTests.cs b/Engine/Source/Programs/Horde/Horde.Agent.Tests/JsonRpcLoggerTests.cs
--- a/Engine/Source/Programs/Horde/Horde.Agent.Tests/JsonRpcLoggerTests.cs
+++ b/Engine/Source/Programs/Horde/Horde.Agent.Tests/JsonRpcLoggerTests.cs
@@ -138,7 +138,7 @@ namespace Horde.Agent.Tests
 		[TestMethod]
 		public async Task StorageLoggerTestAsync()
 		{
-			using BundleCache cache = new BundleCache();
+			await using BundleCache cache = new BundleCache();
 			using MemoryStorageClient store = new MemoryStorageClient();
 
 			BundleReader reader = new BundleReader(store, cache, NullLogger.Instance);
diff --git a/Engine/Source/Programs/Horde/Horde.Agent.Tests/WorkerServiceTest.cs b/Engine/Source/Programs/Horde/Horde.Agent.Tests/WorkerServiceTest.cs
--- a/Engine/Source/Programs/Horde/Horde.Agent.Tests/WorkerServiceTest.cs
+++ b/Engine/Source/Programs/Horde/Horde.Agent.Tests/WorkerServiceTest.cs
@@ -163,7 +163,7 @@ namespace Horde.Agent.Tests
 			});
 
 			_serviceCollection.AddSingleton<IJobExecutorFactory>(x => new SimpleTestExecutorFactory(executor));
-			using ServiceProvider serviceProvider = _serviceCollection.BuildServiceProvider();
+			await using ServiceProvider serviceProvider = _serviceCollection.BuildServiceProvider();
 
 			JobHandler jobHandler = serviceProvider.GetRequiredService<JobHandler>();
 			jobHandler._stepAbortPollInterval = TimeSpan.FromMilliseconds(1);
@@ -191,7 +191,7 @@ namespace Horde.Agent.Tests
 			});
 
 			_serviceCollection.AddSingleton<IJobExecutorFactory>(x => new SimpleTestExecutorFactory(executor));
-			using ServiceProvider serviceProvider = _serviceCollection.BuildServiceProvider();
+			await using ServiceProvider serviceProvider = _serviceCollection.BuildServiceProvider();
 
 			JobHandler jobHandler = serviceProvider.GetRequiredService<JobHandler>();
 			jobHandler._stepAbortPollInterval = TimeSpan.FromMilliseconds(5);
diff --git a/Engine/Source/Programs/Horde/Horde.Server/Commands/Install/SetupCommand.cs b/Engine/Source/Programs/Horde/Horde.Server/Commands/Install/SetupCommand.cs
--- a/Engine/Source/Programs/Horde/Horde.Server/Commands/Install/SetupCommand.cs
+++ b/Engine/Source/Programs/Horde/Horde.Server/Commands/Install/SetupCommand.cs
@@ -58,7 +58,7 @@ namespace Horde.Server.Commands.Install
 
 			RefName refName = new RefName("latest");
 
-			using (BundleCache bundleCache = new BundleCache())
+			await using (BundleCache bundleCache = new BundleCache())
 			{
 				using (IStorageClient client = BundleStorageClient.CreateFromDirectory(bundleDir, bundleCache, logger))
 				{
@@ -75,7 +75,7 @@ namespace Horde.Server.Commands.Install
 			DirectoryReference looseAgentInstallerDir = DirectoryReference.Combine(serverDir, "Tools", "horde-agent-installer-loose");
 			DirectoryReference installerBundleDir = DirectoryReference.Combine(serverDir, "Tools", "horde-agent-installer");
 			DirectoryReference.CreateDirectory(installerBundleDir);
-			using (BundleCache bundleCache = new BundleCache())
+			await using (BundleCache bundleCache = new BundleCache())
 			{
 				using (IStorageClient client = BundleStorageClient.CreateFromDirectory(installerBundleDir, bundleCache, logger))
 				{
diff --git a/Engine/Source/Programs/Shared/EpicGames.Horde/Compute/AgentMessageHandler.cs b/Engine/Source/Programs/Shared/EpicGames.Horde/Compute/AgentMessageHandler.cs
--- a/Engine/Source/Programs/Shared/EpicGames.Horde/Compute/AgentMessageHandler.cs
+++ b/Engine/Source/Programs/Shared/EpicGames.Horde/Compute/AgentMessageHandler.cs
@@ -136,7 +136,7 @@ namespace EpicGames.Horde.Compute
 		async Task WriteFilesAsync(AgentMessageChannel channel, string path, BlobLocator locator, CancellationToken cancellationToken)
 		{
 			using AgentStorageClient innerStore = new AgentStorageClient(channel);
-			using BundleCache cache = new BundleCache(new BundleCacheOptions { HeaderCacheSize = 10 * 1024 * 1024, PacketCacheSize = 128 * 1024 * 1024 });
+			await using BundleCache cache = new BundleCache(new BundleCacheOptions { HeaderCacheSize = 10 * 1024 * 1024, PacketCacheSize = 128 * 1024 * 1024 });
 			using BundleStorageClient store = new BundleStorageClient(innerStore, cache, _logger);
 
 			IBlobHandle handle = store.CreateBlobHandle(locator);
diff --git a/Engine/Source/Programs/Shared/EpicGames.Horde/Storage/Bundles/BundleCache.cs b/Engine/Source/Programs/Shared/EpicGames.Horde/Storage/Bundles/BundleCache.cs
--- a/Engine/Source/Programs/Shared/EpicGames.Horde/Storage/Bundles/BundleCache.cs
+++ b/Engine/Source/Programs/Shared/EpicGames.Horde/Storage/Bundles/BundleCache.cs
@@ -1,7 +1,11 @@
 // Copyright Epic Games, Inc. All Rights Reserved.
 
 using System;
+using System.Collections.Generic;
 using System.Diagnostics.CodeAnalysis;
+using System.Threading;
+using System.Threading.Tasks;
+using EpicGames.Core;
 using Microsoft.Extensions.Caching.Memory;
 
 namespace EpicGames.Horde.Storage
@@ -11,6 +15,11 @@ namespace EpicGames.Horde.Storage
 	/// </summary>
 	public class BundleCacheOptions
 	{
+		/// <summary>
+		/// Maximum size of the cache
+		/// </summary>
+		public long MaxSize { get; set; } = 128 * 1024 * 1024;
+
 		/// <summary>
 		/// Size of the header cache
 		/// </summary>
@@ -23,27 +32,132 @@ namespace EpicGames.Horde.Storage
 	}
 
 	/// <summary>
-	/// Implements caching functionality for reading from a storage client
+	/// Caches items with user-defined keys, where initialization of each value is performed asynchronously. 
+	/// Memory allocations are tracked via calls to Reserve(), and a hard upper limit is enforced as long as data can be freed.
 	/// </summary>
-	public sealed class BundleCache : IDisposable
+	public sealed class BundleCache : IAsyncDisposable
 	{
+		class CacheValue : IDisposable
+		{
+			public LinkedListNode<CacheValue> Node { get; }
+			public Task<IDisposable> InitTask { get; }
+			int _lockCount;
+
+			public CacheValue(IDisposable value)
+			{
+				Node = new LinkedListNode<CacheValue>(this);
+				Lock();
+				InitTask = Task.FromResult<IDisposable>(value);
+			}
+
+			public CacheValue(Func<Task<IDisposable>> initTask)
+			{
+				Node = new LinkedListNode<CacheValue>(this);
+				Lock();
+				InitTask = Task.Run(() => RunAndUnlockAsync(initTask));
+			}
+
+			public void Dispose()
+			{
+				InitTask.Result.Dispose();
+			}
+
+			async Task<IDisposable> RunAndUnlockAsync(Func<Task<IDisposable>> initTask)
+			{
+				try
+				{
+					return await initTask();
+				}
+				finally
+				{
+					Unlock();
+				}
+			}
+
+			public bool IsLocked() => Interlocked.CompareExchange(ref _lockCount, 0, 0) > 0;
+
+			public void Lock() => Interlocked.Increment(ref _lockCount);
+			public void Unlock() => Interlocked.Decrement(ref _lockCount);
+		}
+
+		class CacheValueHandle<T> : IRefCountedHandle<T> where T : class, IDisposable
+		{
+			CacheValue? _item;
+			T? _target;
+
+			public T Target => _target ?? throw new ObjectDisposedException(nameof(CacheValueHandle<T>));
+
+			public CacheValueHandle(CacheValue? item, T? value)
+			{
+				_item = item;
+				_target = value;
+			}
+
+			public void Dispose()
+			{
+				if (_item != null)
+				{
+					_item.Unlock();
+					_item = null;
+					_target = null;
+				}
+			}
+
+			public IRefCountedHandle<T> AddRef()
+			{
+				if (_item == null)
+				{
+					throw new ObjectDisposedException(nameof(CacheValueHandle<T>));
+				}
+
+				_item.Lock();
+				return new CacheValueHandle<T>(_item, _target);
+			}
+		}
+
+		class MemoryReservation : IDisposable
+		{
+			readonly BundleCache _outer;
+			long _size;
+
+			public MemoryReservation(BundleCache outer, long size)
+			{
+				_outer = outer;
+				_size = size;
+			}
+
+			public void Dispose()
+			{
+				Interlocked.Add(ref _outer._currentSize, -_size);
+				_size = 0;
+			}
+		}
+
+		readonly object _lockObject = new object();
+		readonly BundleCacheOptions _options;
+		readonly LinkedList<CacheValue> _items = new LinkedList<CacheValue>();
+		readonly Dictionary<object, CacheValue> _itemLookup = new Dictionary<object, CacheValue>();
+		readonly CancellationTokenSource _cancellationSource = new CancellationTokenSource();
+
+		long _currentSize;
+
 		readonly MemoryCache? _headerCache;
 		readonly MemoryCache? _packetCache;
 
 		/// <summary>
 		/// Instance of an empty cache
 		/// </summary>
-		public static BundleCache None { get; } = new BundleCache(new BundleCacheOptions { HeaderCacheSize = 0, PacketCacheSize = 0 });
+		public static BundleCache None { get; } = new BundleCache(new BundleCacheOptions { MaxSize = 0, HeaderCacheSize = 0, PacketCacheSize = 0 });
 
 		/// <summary>
 		/// Size of the configured header cache
 		/// </summary>
-		public long HeaderCacheSize { get; }
+		public long HeaderCacheSize => _options.HeaderCacheSize;
 
 		/// <summary>
 		/// Size of the configured packet cache
 		/// </summary>
-		public long PacketCacheSize { get; }
+		public long PacketCacheSize => _options.PacketCacheSize;
 
 		/// <summary>
 		/// Whether there is a packet cache present
@@ -62,25 +176,151 @@ namespace EpicGames.Horde.Storage
 		/// </summary>
 		public BundleCache(BundleCacheOptions options)
 		{
+			_options = options;
 			if (options.HeaderCacheSize > 0)
 			{
-				HeaderCacheSize = options.HeaderCacheSize;
 				_headerCache = new MemoryCache(new MemoryCacheOptions { SizeLimit = options.HeaderCacheSize });
 			}
 			if (options.PacketCacheSize > 0)
 			{
-				PacketCacheSize = options.PacketCacheSize;
 				_packetCache = new MemoryCache(new MemoryCacheOptions { SizeLimit = options.PacketCacheSize });
 			}
 		}
 
 		/// <inheritdoc/>
-		public void Dispose()
+		public async ValueTask DisposeAsync()
 		{
 			_headerCache?.Dispose();
 			_packetCache?.Dispose();
+
+			if (_items.Count > 0)
+			{
+				_cancellationSource.Cancel();
+
+				foreach (CacheValue item in _items)
+				{
+					try
+					{
+						await item.InitTask;
+					}
+					catch (OperationCanceledException)
+					{
+					}
+				}
+
+				_items.Clear();
+			}
+
+			_cancellationSource.Dispose();
+		}
+
+		/// <summary>
+		/// Find or add a new cached value to the cache
+		/// </summary>
+		/// <param name="key">Key for the lookup</param>
+		/// <returns></returns>
+		public IRefCountedHandle<TValue>? Find<TKey, TValue>(TKey key)
+			where TKey : notnull
+			where TValue : class, IDisposable
+		{
+			lock (_lockObject)
+			{
+				CacheValue? item;
+				if (_itemLookup.TryGetValue(key, out item) && item.InitTask.TryGetResult(out IDisposable result))
+				{
+					item.Lock();
+					return new CacheValueHandle<TValue>(item, (TValue)result);
+				}
+			}
+			return null;
+		}
+
+		/// <summary>
+		/// Attempts to add a value to the cache
+		/// </summary>
+		/// <param name="key"></param>
+		/// <param name="createFunc"></param>
+		/// <returns>True if the value was added</returns>
+		public bool TryAdd<TKey, TValue>(TKey key, Func<TValue> createFunc)
+			where TKey : notnull
+			where TValue : class, IDisposable
+		{
+			lock (_lockObject)
+			{
+				if (!_itemLookup.ContainsKey(key))
+				{
+					CacheValue item = new CacheValue(createFunc());
+					_itemLookup.Add(key, item);
+					return true;
+				}
+			}
+			return false;
+		}
+
+		/// <summary>
+		/// Find or add a new cached value to the cache
+		/// </summary>
+		/// <param name="key"></param>
+		/// <param name="createAsync"></param>
+		/// <param name="cancellationToken"></param>
+		/// <returns>Handle to the item that was read. Must be disposed by the caller.</returns>
+		public async Task<IRefCountedHandle<TValue>> FindOrAddAsync<TKey, TValue>(TKey key, Func<TKey, CancellationToken, Task<TValue>> createAsync, CancellationToken cancellationToken)
+			where TKey : notnull
+			where TValue : class, IDisposable
+		{
+			CacheValue? item = null;
+			try
+			{
+				lock (_lockObject)
+				{
+					CacheValue? untypedItem;
+					if (_itemLookup.TryGetValue(key, out untypedItem))
+					{
+						item = untypedItem;
+						item.Lock();
+					}
+					else
+					{
+						item = new CacheValue(async () => await createAsync(key, _cancellationSource.Token));
+						_items.AddFirst(item);
+						_itemLookup.Add(key, item);
+					}
+				}
+
+				TValue value = (TValue)await item.InitTask.WaitAsync(cancellationToken);
+				return new CacheValueHandle<TValue>(item, value);
+			}
+			finally
+			{
+				item?.Unlock();
+			}
 		}
 
+		/// <summary>
+		/// Add a reservation against the maximum allowed size of cached data. 
+		/// </summary>
+		public IDisposable Reserve(long size)
+		{
+			lock (_lockObject)
+			{
+				Interlocked.Add(ref _currentSize, size);
+
+				for (LinkedListNode<CacheValue>? lastNode = _items.Last; lastNode != null && Interlocked.CompareExchange(ref _currentSize, 0, 0) > _options.MaxSize; lastNode = lastNode.Previous)
+				{
+					CacheValue lastItem = lastNode.Value;
+					if (!lastItem.IsLocked())
+					{
+						_items.Remove(lastItem);
+						lastItem.Dispose();
+					}
+				}
+
+				return new MemoryReservation(this, size);
+			}
+		}
+
+		#region V1
+
 		static void AddCachedValue(IMemoryCache? cache, string cacheKey, object value, int size)
 		{
 			if (cache != null)
@@ -154,5 +394,7 @@ namespace EpicGames.Horde.Storage
 		{
 			return TryGetCachedValue(_packetCache, GetDecodedPacketCacheKey(locator, packetIdx), out data);
 		}
+
+		#endregion
 	}
 }
