Horde: Allow storing arbitrary IDisposable data along with BlobData objects, rather than requiring an IReadOnlyMemoryOwner.

diff --git a/Engine/Source/Programs/Horde/Horde.Server/Storage/StorageService.cs b/Engine/Source/Programs/Horde/Horde.Server/Storage/StorageService.cs
--- a/Engine/Source/Programs/Horde/Horde.Server/Storage/StorageService.cs
+++ b/Engine/Source/Programs/Horde/Horde.Server/Storage/StorageService.cs
@@ -91,7 +91,7 @@ namespace Horde.Server.Storage
 			public async ValueTask<BlobData> ReadAsync(CancellationToken cancellationToken = default)
 			{
 				IReadOnlyMemoryOwner<byte> obj = await _backend.ReadAsync(_path, cancellationToken);
-				return new ReadOnlyMemoryOwnerBlobData(BlobType.Leaf, obj, Array.Empty<IBlobHandle>());
+				return new BlobDataWithOwner(BlobType.Leaf, obj.Memory, Array.Empty<IBlobHandle>(), obj);
 			}
 
 			/// <inheritdoc/>
diff --git a/Engine/Source/Programs/Shared/EpicGames.Horde/Storage/BlobData.cs b/Engine/Source/Programs/Shared/EpicGames.Horde/Storage/BlobData.cs
--- a/Engine/Source/Programs/Shared/EpicGames.Horde/Storage/BlobData.cs
+++ b/Engine/Source/Programs/Shared/EpicGames.Horde/Storage/BlobData.cs
@@ -55,15 +55,15 @@ namespace EpicGames.Horde.Storage
 	/// <summary>
 	/// Implementation of <see cref="BlobData"/> for <see cref="IReadOnlyMemoryOwner{Byte}"/> instances.
 	/// </summary>
-	public class ReadOnlyMemoryOwnerBlobData : BlobData
+	public class BlobDataWithOwner : BlobData
 	{
-		readonly IReadOnlyMemoryOwner<byte> _owner;
+		readonly IDisposable _owner;
 
 		/// <summary>
 		/// Constructor
 		/// </summary>
-		public ReadOnlyMemoryOwnerBlobData(BlobType type, IReadOnlyMemoryOwner<byte> owner, IReadOnlyList<IBlobHandle> refs)
-			: base(type, owner.Memory, refs)
+		public BlobDataWithOwner(BlobType type, ReadOnlyMemory<byte> data, IReadOnlyList<IBlobHandle> refs, IDisposable owner)
+			: base(type, data, refs)
 		{
 			_owner = owner;
 		}
diff --git a/Engine/Source/Programs/Shared/EpicGames.Horde/Storage/Clients/BundleStorageClient.cs b/Engine/Source/Programs/Shared/EpicGames.Horde/Storage/Clients/BundleStorageClient.cs
--- a/Engine/Source/Programs/Shared/EpicGames.Horde/Storage/Clients/BundleStorageClient.cs
+++ b/Engine/Source/Programs/Shared/EpicGames.Horde/Storage/Clients/BundleStorageClient.cs
@@ -258,7 +258,7 @@ namespace EpicGames.Horde.Storage.Clients
 		}
 
 		/// <inheritdoc/>
-		public Task WriteRefAsync(RefName name, IBlobHandle target, RefOptions? options, CancellationToken cancellationToken)
+		public Task WriteRefAsync(RefName name, IBlobHandle target, RefOptions? options = null, CancellationToken cancellationToken = default)
 			=> _inner.WriteRefAsync(name, target, options, cancellationToken);
 
 		#endregion
diff --git a/Engine/Source/Programs/Shared/EpicGames.Horde/Storage/Clients/FileStorageClient.cs b/Engine/Source/Programs/Shared/EpicGames.Horde/Storage/Clients/FileStorageClient.cs
--- a/Engine/Source/Programs/Shared/EpicGames.Horde/Storage/Clients/FileStorageClient.cs
+++ b/Engine/Source/Programs/Shared/EpicGames.Horde/Storage/Clients/FileStorageClient.cs
@@ -65,7 +65,7 @@ namespace EpicGames.Horde.Storage.Clients
 		public override async ValueTask<BlobData> ReadBlobAsync(BlobLocator locator, CancellationToken cancellationToken = default)
 		{
 			IReadOnlyMemoryOwner<byte> owner = await _backend.ReadAsync(locator.ToString(), cancellationToken);
-			return new ReadOnlyMemoryOwnerBlobData(BlobType.Leaf, owner, Array.Empty<IBlobHandle>());
+			return new BlobDataWithOwner(BlobType.Leaf, owner.Memory, Array.Empty<IBlobHandle>(), owner);
 		}
 
 		/// <inheritdoc/>
diff --git a/Engine/Source/Programs/Shared/EpicGames.Horde/Storage/Clients/HttpStorageClient.cs b/Engine/Source/Programs/Shared/EpicGames.Horde/Storage/Clients/HttpStorageClient.cs
--- a/Engine/Source/Programs/Shared/EpicGames.Horde/Storage/Clients/HttpStorageClient.cs
+++ b/Engine/Source/Programs/Shared/EpicGames.Horde/Storage/Clients/HttpStorageClient.cs
@@ -105,7 +105,7 @@ namespace EpicGames.Horde.Storage.Clients
 		public async ValueTask<BlobData> ReadBlobAsync(BlobLocator locator, CancellationToken cancellationToken = default)
 		{
 			IReadOnlyMemoryOwner<byte> owner = await _backend.ReadAsync(locator.ToString(), cancellationToken);
-			return new ReadOnlyMemoryOwnerBlobData(BlobType.Leaf, owner, Array.Empty<IBlobHandle>());
+			return new BlobDataWithOwner(BlobType.Leaf, owner.Memory, Array.Empty<IBlobHandle>(), owner);
 		}
 
 		/// <inheritdoc/>
diff --git a/Engine/Source/Programs/Shared/EpicGames.Horde/Storage/NodeRef.cs b/Engine/Source/Programs/Shared/EpicGames.Horde/Storage/NodeRef.cs
--- a/Engine/Source/Programs/Shared/EpicGames.Horde/Storage/NodeRef.cs
+++ b/Engine/Source/Programs/Shared/EpicGames.Horde/Storage/NodeRef.cs
@@ -277,6 +277,24 @@ namespace EpicGames.Horde.Storage
 			}
 		}
 
+		/// <summary>
+		/// Writes an individual node to storage
+		/// </summary>
+		/// <param name="writer">Writer to serialize nodes to</param>
+		/// <param name="node">Node to write</param>
+		/// <param name="cancellationToken">Cancellation token for the operation</param>
+		/// <returns>A flag indicating whether the node is dirty, and if it is, an optional bundle that contains it</returns>
+		public static async ValueTask<NodeRef<TNode>> WriteNodeAsync<TNode>(this IStorageWriter writer, TNode node, CancellationToken cancellationToken = default) where TNode : Node
+		{
+			// Serialize the node
+			BlobWriter nodeWriter = new BlobWriter(writer);
+			node.Serialize(nodeWriter);
+
+			// Write the final data
+			IBlobHandle handle = await writer.WriteBlobAsync(node.NodeType, nodeWriter.Length, nodeWriter.References, cancellationToken);
+			return new NodeRef<TNode>(handle);
+		}
+
 		/// <summary>
 		/// Writes an individual node to storage
 		/// </summary>
