Horde: Remove default implementation of IBlobHandle.FlushAsync.

diff --git a/Engine/Source/Programs/Horde/Horde.Server/Storage/StorageService.cs b/Engine/Source/Programs/Horde/Horde.Server/Storage/StorageService.cs
--- a/Engine/Source/Programs/Horde/Horde.Server/Storage/StorageService.cs
+++ b/Engine/Source/Programs/Horde/Horde.Server/Storage/StorageService.cs
@@ -65,6 +65,9 @@ namespace Horde.Server.Storage
 				_tracer = tracer;
 			}
 
+			/// <inheritdoc/>
+			public ValueTask FlushAsync(CancellationToken cancellationToken = default) => default;
+
 			/// <inheritdoc/>
 			public ValueTask<BlobType> ReadTypeAsync(CancellationToken cancellationToken = default) => new ValueTask<BlobType>(BlobType.Leaf);
 
diff --git a/Engine/Source/Programs/Shared/EpicGames.Horde.Tests/BlobDataTests.cs b/Engine/Source/Programs/Shared/EpicGames.Horde.Tests/BlobDataTests.cs
--- a/Engine/Source/Programs/Shared/EpicGames.Horde.Tests/BlobDataTests.cs
+++ b/Engine/Source/Programs/Shared/EpicGames.Horde.Tests/BlobDataTests.cs
@@ -20,6 +20,7 @@ namespace EpicGames.Horde.Tests
 			public IBlobHandle? Outer => null;
 
 			public DummyHandle(string locator) => _locator = new BlobLocator(new Utf8String(locator));
+			public ValueTask FlushAsync(CancellationToken cancellationToken = default) => default;
 			public ValueTask<BlobData> ReadAsync(CancellationToken cancellationToken = default) => throw new NotImplementedException();
 
 			public bool TryAppendIdentifier(Utf8StringBuilder builder)
diff --git a/Engine/Source/Programs/Shared/EpicGames.Horde/Storage/Clients/BundleStorageClient.cs b/Engine/Source/Programs/Shared/EpicGames.Horde/Storage/Clients/BundleStorageClient.cs
--- a/Engine/Source/Programs/Shared/EpicGames.Horde/Storage/Clients/BundleStorageClient.cs
+++ b/Engine/Source/Programs/Shared/EpicGames.Horde/Storage/Clients/BundleStorageClient.cs
@@ -37,6 +37,9 @@ namespace EpicGames.Horde.Storage.Clients
 				_locator = locator;
 			}
 
+			/// <inheritdoc/>
+			public ValueTask FlushAsync(CancellationToken cancellationToken = default) => default;
+
 			/// <inheritdoc/>
 			public ValueTask<BlobType> ReadTypeAsync(CancellationToken cancellationToken = default) => new ValueTask<BlobType>(Bundle.BlobType);
 
diff --git a/Engine/Source/Programs/Shared/EpicGames.Horde/Storage/Clients/HttpStorageClient.cs b/Engine/Source/Programs/Shared/EpicGames.Horde/Storage/Clients/HttpStorageClient.cs
--- a/Engine/Source/Programs/Shared/EpicGames.Horde/Storage/Clients/HttpStorageClient.cs
+++ b/Engine/Source/Programs/Shared/EpicGames.Horde/Storage/Clients/HttpStorageClient.cs
@@ -35,6 +35,9 @@ namespace EpicGames.Horde.Storage.Clients
 				_locator = locator;
 			}
 
+			/// <inheritdoc/>
+			public ValueTask FlushAsync(CancellationToken cancellationToken = default) => default;
+
 			/// <inheritdoc/>
 			public Task<Stream> OpenBodyAsync(int offset = 0, int? length = null, CancellationToken cancellationToken = default)
 				=> _httpStorageClient._backend.OpenAsync(_locator.ToString(), offset, length, cancellationToken);
diff --git a/Engine/Source/Programs/Shared/EpicGames.Horde/Storage/Clients/KeyValueStorageClient.cs b/Engine/Source/Programs/Shared/EpicGames.Horde/Storage/Clients/KeyValueStorageClient.cs
--- a/Engine/Source/Programs/Shared/EpicGames.Horde/Storage/Clients/KeyValueStorageClient.cs
+++ b/Engine/Source/Programs/Shared/EpicGames.Horde/Storage/Clients/KeyValueStorageClient.cs
@@ -31,6 +31,9 @@ namespace EpicGames.Horde.Storage.Clients
 				_locator = locator;
 			}
 
+			/// <inheritdoc/>
+			public ValueTask FlushAsync(CancellationToken cancellationToken = default) => default;
+
 			/// <inheritdoc/>
 			public ValueTask<BlobData> ReadAsync(CancellationToken cancellationToken = default) => _keyValueStorageClient.ReadBlobAsync(_locator, cancellationToken);
 
diff --git a/Engine/Source/Programs/Shared/EpicGames.Horde/Storage/IBlobHandle.cs b/Engine/Source/Programs/Shared/EpicGames.Horde/Storage/IBlobHandle.cs
--- a/Engine/Source/Programs/Shared/EpicGames.Horde/Storage/IBlobHandle.cs
+++ b/Engine/Source/Programs/Shared/EpicGames.Horde/Storage/IBlobHandle.cs
@@ -66,7 +66,7 @@ namespace EpicGames.Horde.Storage
 		/// <summary>
 		/// Flush the referenced not to underlying storage
 		/// </summary>
-		ValueTask FlushAsync(CancellationToken cancellationToken = default) => new ValueTask();
+		ValueTask FlushAsync(CancellationToken cancellationToken = default);
 
 		/// <summary>
 		/// Open the blob's data stream
@@ -163,6 +163,9 @@ namespace EpicGames.Horde.Storage
 			Fragment = fragment;
 		}
 
+		/// <inheritdoc/>
+		public ValueTask FlushAsync(CancellationToken cancellationToken = default) => default;
+
 		/// <inheritdoc/>
 		public bool TryAppendIdentifier(Utf8StringBuilder builder)
 		{
