Horde: Updated bundle file format and addressing scheme. Not yet enabled by default.

File format:
* Bundles are now just a stream of packets, but packets follow the stanard bundle signature format, allowing interop with older file formats.
* All import/export information is stored in packets, allowing them to be decoded without reading the a monolithic bundle header. Previous implementation did not factor in the size of header metadata with bundles, nor the size of supporting data structures, which could dwarf the size of payload data. Splitting headers into separate packets allows us to separate the packet size from read size, and still allows "over-reading" of bundles to prefetch additional coherent data for low-latency connections.

Locators:
* Blob locators now have a standard URI syntax; the #fragment portion of the URI identifies an embedded resource within the bundle, with parameters encoded using media-fragment syntax.
* Blob locators now encode the packet offset and length to support reading of individual packets.
* Packets and exports can be referenced directly with the same scheme, eg. bundle: "foo/bar", packet: "foo/bar#pkt=100,0", export: "foo/bar#pkt=100,0&exp=4".
* The packet import table now encodes arbitary string locators for each individual export in a prefix tree, with each node being a cachable handle to an individual blob object.
* Since blob locators are now effectively standard URIs, we can use them to link to content on other servers, or to link between separate endpoints on the same server.

Interop:
* Bundle writer options now include a setting for the maximum version number to use when writing bundles, allowing both codebases to exist side by side.
* Added separate LatestV1 and LatestV2 version number constants to identify encoding scheme to use before we migrate fully to V2 encoding.

Resource usage:
* Writes between layers in the write pipeline - with the exception of compression - are now done into final buffers, reducing the churn on memory allocations causing requiring frequent GC runs.
* Buffers allocated in the write process reserve memory against a fixed memory limit tracked in BundleCache; buffers allocated for asynchronous writes will now cause items to be purged from the cache to stay below the fixed memory limit.
* Buffers are reference counted, and can transfer ownership to the cache layer once unlocked.
* Buffers are now manually allocated and freed on the unmanaged heap to avoid needing to trigger the garbage collector.

Misc:
* The cache supports async initiatlization of values, fixing previous issue where multiple writes are executed.
* Handles now support equality comparisons to avoid having to construct temporary string objects for cache keys.

diff --git a/Engine/Source/Programs/Horde/Horde.Server.Tests/GcServiceTests.cs b/Engine/Source/Programs/Horde/Horde.Server.Tests/GcServiceTests.cs
--- a/Engine/Source/Programs/Horde/Horde.Server.Tests/GcServiceTests.cs
+++ b/Engine/Source/Programs/Horde/Horde.Server.Tests/GcServiceTests.cs
@@ -49,9 +49,11 @@ namespace Horde.Server.Tests
 
 			IStorageBackend backend = ServiceProvider.GetRequiredService<IStorageBackendProvider>().CreateBackend(globalConfig.Storage.Backends[0]);
 
-			string[] remaining = await backend.EnumerateAsync().Select(x => $"{x}#0").ToArrayAsync();
+			string[] remaining = await backend.EnumerateAsync().ToArrayAsync();
 			Assert.AreEqual(nodes.Count, remaining.Length);
-			Assert.IsTrue(remaining.All(x => nodes.Contains(new BlobLocator(x))));
+
+			HashSet<string> nodePaths = new HashSet<string>(nodes.Select(x => x.BaseLocator.ToString()), StringComparer.Ordinal);
+			Assert.IsTrue(remaining.All(x => nodePaths.Contains(x)));
 		}
 
 		static async Task<HashSet<BlobLocator>> FindNodesAsync(IStorageClient store, IEnumerable<BlobLocator> roots)
diff --git a/Engine/Source/Programs/Shared/EpicGames.Core/IMemoryAllocator.cs b/Engine/Source/Programs/Shared/EpicGames.Core/IMemoryAllocator.cs
--- a/Engine/Source/Programs/Shared/EpicGames.Core/IMemoryAllocator.cs
+++ b/Engine/Source/Programs/Shared/EpicGames.Core/IMemoryAllocator.cs
@@ -21,6 +21,29 @@ namespace EpicGames.Core
 		IMemoryOwner<T> Alloc(int minSize);
 	}
 
+	/// <summary>
+	/// Implementation of <see cref="IMemoryAllocator{Byte}"/> which creates regular arrays on the managed heap. Note that these
+	/// allocations are subject to GC, and will not be freed immediately.
+	/// </summary>
+	public class ManagedHeapAllocator : IMemoryAllocator<byte>
+	{
+		class MemoryOwner : IMemoryOwner<byte>
+		{
+			public Memory<byte> Memory { get; }
+
+			public MemoryOwner(int size) => Memory = new byte[size];
+			public void Dispose() { }
+		}
+
+		/// <summary>
+		/// Default shared instance
+		/// </summary>
+		public static ManagedHeapAllocator Instance { get; } = new ManagedHeapAllocator();
+
+		/// <inheritdoc/>
+		public IMemoryOwner<byte> Alloc(int minSize) => new MemoryOwner(minSize);
+	}
+
 	/// <summary>
 	/// Implementation of <see cref="IMemoryAllocator{Byte}"/> which rents blocks from a memory pool.
 	/// </summary>
diff --git a/Engine/Source/Programs/Shared/EpicGames.Horde.Tests/BundlePacketTests.cs b/Engine/Source/Programs/Shared/EpicGames.Horde.Tests/BundlePacketTests.cs
new file mode 100644
--- /dev/null
+++ b/Engine/Source/Programs/Shared/EpicGames.Horde.Tests/BundlePacketTests.cs
@@ -0,0 +1,137 @@
+// Copyright Epic Games, Inc. All Rights Reserved.
+
+using System;
+using System.Linq;
+using System.Text;
+using System.Threading;
+using System.Threading.Tasks;
+using EpicGames.Core;
+using EpicGames.Horde.Storage;
+using EpicGames.Horde.Storage.Bundles;
+using EpicGames.Horde.Storage.Bundles.V2;
+using Microsoft.VisualStudio.TestTools.UnitTesting;
+
+namespace EpicGames.Horde.Tests
+{
+	[TestClass]
+	public class BundlePacketTests
+	{
+		class FakeHandle : IBlobHandle
+		{
+			readonly IBlobHandle? _outer;
+			readonly string? _fragment;
+
+			public IBlobHandle? Outer => _outer;
+
+			public FakeHandle(IBlobHandle? outer, string? fragment)
+			{
+				_outer = outer;
+				_fragment = fragment;
+			}
+
+			public ValueTask<BlobData> ReadBlobDataAsync(CancellationToken cancellationToken = default) => throw new NotImplementedException();
+			public ValueTask FlushAsync(CancellationToken cancellationToken = default) => throw new NotImplementedException();
+
+			public ValueTask<BlobData> ReadAsync(CancellationToken cancellationToken = default)
+			{
+				throw new NotImplementedException();
+			}
+
+			public bool TryAppendIdentifier(Utf8StringBuilder builder)
+			{
+				if (_fragment == null)
+				{
+					return false;
+				}
+
+				builder.Append(_fragment);
+				return true;
+			}
+
+			public override bool Equals(object? obj)
+				=> obj is FakeHandle other && Object.Equals(_outer, other._outer) && String.Equals(_fragment, other._fragment, StringComparison.Ordinal);
+
+			public override int GetHashCode()
+				=> HashCode.Combine(_outer, _fragment);
+		}
+
+		[TestMethod]
+		public void TestTypes()
+		{
+			Guid guid = Guid.Parse("{30F2CA83-B5D4-494C-8802-0661454CCD58}");
+
+			// Write out some data
+			BlobType blobType1 = new BlobType(guid, 1);
+			BlobType blobType2 = new BlobType(guid, 2);
+
+			FakeHandle bundleHandle = new FakeHandle(null, null);
+			FakeHandle packetHandle = new FakeHandle(bundleHandle, null);
+
+			using PacketWriter writer = new PacketWriter(bundleHandle, packetHandle, ManagedHeapAllocator.Instance, new object());
+			int type1 = writer.FindOrAddType(blobType1);
+			int type2 = writer.FindOrAddType(blobType2);
+			int type1b = writer.FindOrAddType(blobType1); // Check it's deduped
+			Assert.AreEqual(type1, type1b);
+
+			FakeHandle fooBundleHandle = new FakeHandle(null, "foo");
+			FakeHandle barBundleHandle = new FakeHandle(null, "bar");
+
+			int import1 = writer.FindOrAddImport(new FakeHandle(fooBundleHandle, "123"));
+			int import2 = writer.FindOrAddImport(new FakeHandle(fooBundleHandle, "456"));
+			int import3 = writer.FindOrAddImport(new FakeHandle(fooBundleHandle, "789"));
+			int import4 = writer.FindOrAddImport(new FakeHandle(barBundleHandle, "123"));
+			int import4b = writer.FindOrAddImport(new FakeHandle(barBundleHandle, "123"));
+			Assert.AreEqual(import4, import4b);
+
+			Encoding.UTF8.GetBytes("hello", writer.GetOutputBuffer(0, 5).Span);
+			int export1 = writer.CompleteExport(5, type1, new[] { import1, import2 });
+			Assert.AreEqual(export1, 0);
+
+			Encoding.UTF8.GetBytes(" ", writer.GetOutputBuffer(0, 1).Span);
+			int export2 = writer.CompleteExport(1, type2, Array.Empty<int>());
+			Assert.AreEqual(export2, 1);
+
+			Encoding.UTF8.GetBytes("world!", writer.GetOutputBuffer(0, 6).Span);
+			int export3 = writer.CompleteExport(6, type2, new[] { import2, import3, import4 });
+			Assert.AreEqual(export3, 2);
+
+			Packet packet = writer.CompletePacket();
+
+			ArrayMemoryWriter memoryWriter = new ArrayMemoryWriter(4096);
+			packet.Encode(BundleCompressionFormat.None, memoryWriter);
+
+			// Check we can read it back in
+			IRefCountedHandle<Packet> handle = Packet.Decode(memoryWriter.WrittenMemory, ManagedHeapAllocator.Instance);
+			packet = handle.Target;
+
+			Assert.AreEqual(2, packet.GetTypeCount());
+			Assert.AreEqual(blobType1, packet.GetType(0));
+			Assert.AreEqual(blobType2, packet.GetType(1));
+
+			Assert.AreEqual(6, packet.GetImportCount());
+			Assert.AreEqual(new PacketImport(-1, new Utf8String("foo")), packet.GetImport(0));
+			Assert.AreEqual(new PacketImport(0, new Utf8String("123")), packet.GetImport(1));
+			Assert.AreEqual(new PacketImport(0, new Utf8String("456")), packet.GetImport(2));
+			Assert.AreEqual(new PacketImport(0, new Utf8String("789")), packet.GetImport(3));
+			Assert.AreEqual(new PacketImport(-1, new Utf8String("bar")), packet.GetImport(4));
+			Assert.AreEqual(new PacketImport(4, new Utf8String("123")), packet.GetImport(5));
+
+			Assert.AreEqual(3, packet.GetExportCount());
+
+			PacketExportHeader header1 = packet.GetExport(0).GetHeader();
+			Assert.AreEqual(0, header1.TypeIdx);
+			Assert.IsTrue(header1.Imports.SequenceEqual(new int[] { 1, 2 }));
+			Assert.AreEqual(new Utf8String("hello"), new Utf8String(packet.GetExport(0).GetPayload()));
+
+			PacketExportHeader header2 = packet.GetExport(1).GetHeader();
+			Assert.AreEqual(1, header2.TypeIdx);
+			Assert.AreEqual(0, header2.Imports.Length);
+			Assert.AreEqual(new Utf8String(" "), new Utf8String(packet.GetExport(1).GetPayload()));
+
+			PacketExportHeader header3 = packet.GetExport(2).GetHeader();
+			Assert.AreEqual(1, header3.TypeIdx);
+			Assert.IsTrue(header3.Imports.SequenceEqual(new int[] { 2, 3, 5 }));
+			Assert.AreEqual(new Utf8String("world!"), new Utf8String(packet.GetExport(2).GetPayload()));
+		}
+	}
+}
diff --git a/Engine/Source/Programs/Shared/EpicGames.Horde.Tests/BundleStoreTests.cs b/Engine/Source/Programs/Shared/EpicGames.Horde.Tests/BundleStoreTests.cs
--- a/Engine/Source/Programs/Shared/EpicGames.Horde.Tests/BundleStoreTests.cs
+++ b/Engine/Source/Programs/Shared/EpicGames.Horde.Tests/BundleStoreTests.cs
@@ -25,7 +25,7 @@ namespace EpicGames.Horde.Tests
 	public class BundleStoreTests
 	{
 		[TestMethod]
-		public async Task CreateBundlesManuallyAsync()
+		public async Task CreateBundlesManuallyV1Async()
 		{
 			byte[] a = CreateBundleManually();
 			byte[] b = await CreateBundleNormalAsync();
@@ -59,7 +59,7 @@ namespace EpicGames.Horde.Tests
 		{
 			using MemoryStorageClient memoryStore = new MemoryStorageClient();
 			using BundleStorageClient store = new BundleStorageClient(memoryStore, BundleCache.None, NullLogger.Instance);
-			await using IStorageWriter writer = store.CreateWriter(options: new BundleOptions { CompressionFormat = BundleCompressionFormat.None });
+			await using IStorageWriter writer = store.CreateWriter(options: new BundleOptions { MaxVersion = BundleVersion.ImportHashes, CompressionFormat = BundleCompressionFormat.None });
 
 			TextNode node = new TextNode("Hello world");
 			IBlobHandle handle = await writer.FlushAsync(node, CancellationToken.None);
@@ -400,11 +400,25 @@ namespace EpicGames.Horde.Tests
 
 				HashedNodeRef<ChunkedDataNode> file = root.GetFileEntry("test");
 
-				long uniqueSize = memoryStore.Blobs.Values.Select(x => BundleHeader.Read(x.Data)).SelectMany(x => x.Packets).Sum(x => x.DecodedLength);
+				Dictionary<BlobLocator, long> locatorToSize = new Dictionary<BlobLocator, long>();
+				await GetUniqueBlobsAsync(file.Handle, locatorToSize);
+
+				long uniqueSize = locatorToSize.Sum(x => x.Value);
 				Assert.IsTrue(uniqueSize < data.Length / 3); // random fraction meaning "lots of dedupe happened"
 			}
 		}
 
+		static async Task GetUniqueBlobsAsync(IBlobHandle handle, Dictionary<BlobLocator, long> locatorToSize)
+		{
+			using BlobData data = await handle.ReadAsync();
+			locatorToSize[handle.GetLocator()] = data.Data.Length;
+
+			foreach (IBlobHandle reference in data.Refs)
+			{
+				await GetUniqueBlobsAsync(reference, locatorToSize);
+			}
+		}
+
 		static async Task CompareTreesAsync(DirectoryNode oldNode, DirectoryNode newNode)
 		{
 			Assert.AreEqual(oldNode.Length, newNode.Length);
diff --git a/Engine/Source/Programs/Shared/EpicGames.Horde.Tests/BundleTests.cs b/Engine/Source/Programs/Shared/EpicGames.Horde.Tests/BundleTests.cs
--- a/Engine/Source/Programs/Shared/EpicGames.Horde.Tests/BundleTests.cs
+++ b/Engine/Source/Programs/Shared/EpicGames.Horde.Tests/BundleTests.cs
@@ -223,9 +223,10 @@ namespace EpicGames.Horde.Tests
 			IBlobHandle bundleTarget = store.CreateBlobHandle(refTarget.GetLocator().BaseLocator);
 			using BlobData bundleData = await bundleTarget.ReadAsync();
 
-			BundleHeader bundleHeader = BundleHeader.Read(bundleData.Data);
-			Assert.AreEqual(0, bundleHeader.Imports.Count);
-			Assert.AreEqual(3, bundleHeader.Exports.Count);
+// This is specific to V1 data
+//			BundleHeader bundleHeader = BundleHeader.Read(bundleData.Data);
+//			Assert.AreEqual(0, bundleHeader.Imports.Count);
+//			Assert.AreEqual(3, bundleHeader.Exports.Count);
 
 			// Create a new bundle and read it back in again
 			DirectoryNode newRoot = await store.ReadRefAsync<DirectoryNode>(refName);
diff --git a/Engine/Source/Programs/Shared/EpicGames.Horde.Tests/StorageClientTests.cs b/Engine/Source/Programs/Shared/EpicGames.Horde.Tests/StorageClientTests.cs
new file mode 100644
--- /dev/null
+++ b/Engine/Source/Programs/Shared/EpicGames.Horde.Tests/StorageClientTests.cs
@@ -0,0 +1,99 @@
+// Copyright Epic Games, Inc. All Rights Reserved.
+
+using System.Threading.Tasks;
+using EpicGames.Core;
+using EpicGames.Horde.Storage;
+using EpicGames.Horde.Storage.Clients;
+using Microsoft.Extensions.Logging.Abstractions;
+using Microsoft.VisualStudio.TestTools.UnitTesting;
+
+namespace EpicGames.Horde.Tests
+{
+	[TestClass]
+	public class StorageClientTests
+	{
+		[NodeType("{99601905-A089-4F6E-87F1-D603EEAF1B71}", 1)]
+		class TestNode : Node
+		{
+			public int Value { get; }
+			public NodeRef<TestNode>[] Refs { get; }
+
+			public TestNode(int value, params NodeRef<TestNode>[] refs)
+			{
+				Value = value;
+				Refs = refs;
+			}
+
+			public TestNode(IBlobReader reader)
+			{
+				Value = reader.ReadInt32();
+				Refs = reader.ReadVariableLengthArray(() => reader.ReadNodeRef<TestNode>());
+			}
+
+			public override void Serialize(IBlobWriter writer)
+			{
+				writer.WriteInt32(Value);
+				writer.WriteVariableLengthArray(Refs, x => writer.WriteNodeRef(x));
+			}
+		}
+
+		static StorageClientTests()
+		{
+			Node.RegisterType<TestNode>();
+		}
+
+		[TestMethod]
+		public async Task TestBasicAsync()
+		{
+			using MemoryStorageClient store = new MemoryStorageClient();
+			await TestBasicAsync(store);
+		}
+
+		[TestMethod]
+		public async Task TestBasicBundleV2Async()
+		{
+			using MemoryStorageClient store = new MemoryStorageClient();
+			await using BundleCache cache = new BundleCache();
+			using BundleStorageClient storeV2 = new BundleStorageClient(store, cache, NullLogger.Instance);
+			await TestBasicAsync(storeV2);
+		}
+
+		[TestMethod]
+		public async Task PendingRefsAsync()
+		{
+			await using BundleCache cache = new BundleCache();
+
+			using MemoryStorageClient memoryStore = new MemoryStorageClient();
+			using BundleStorageClient store = new BundleStorageClient(memoryStore, cache, NullLogger.Instance);
+
+			NodeRef<TestNode> nodeRef2;
+			await using (IStorageWriter writer = store.CreateWriter())
+			{
+				NodeRef<TestNode> nodeRef1 = await writer.WriteNodeAsync(new TestNode(123));
+				nodeRef2 = await writer.WriteNodeAsync(new TestNode(456, nodeRef1));
+			}
+			await store.WriteRefAsync("hello", nodeRef2.Handle);
+
+			TestNode output2 = await store.ReadRefAsync<TestNode>("hello");
+			Assert.AreEqual(456, output2.Value);
+			Assert.AreEqual(1, output2.Refs.Length);
+
+			TestNode output1 = await output2.Refs[0].ExpandAsync();
+			Assert.AreEqual(0, output1.Refs.Length);
+			Assert.AreEqual(123, output1.Value);
+		}
+
+		static async Task TestBasicAsync(IStorageClient store)
+		{
+			NodeRef nodeRef;
+			await using (IStorageWriter writer = store.CreateWriter())
+			{
+				nodeRef = await writer.WriteHashedNodeAsync(new TestNode(123));
+			}
+			await store.WriteRefAsync("hello", nodeRef.Handle);
+
+			TestNode output = await store.ReadRefAsync<TestNode>("hello");
+			Assert.AreEqual(123, output.Value);
+		}
+	}
+}
diff --git a/Engine/Source/Programs/Shared/EpicGames.Horde/Storage/Bundles/Bundle.cs b/Engine/Source/Programs/Shared/EpicGames.Horde/Storage/Bundles/Bundle.cs
--- a/Engine/Source/Programs/Shared/EpicGames.Horde/Storage/Bundles/Bundle.cs
+++ b/Engine/Source/Programs/Shared/EpicGames.Horde/Storage/Bundles/Bundle.cs
@@ -37,6 +37,11 @@ namespace EpicGames.Horde.Storage.Bundles
 		/// </summary>
 		ImportHashes = 4,
 
+		/// <summary>
+		/// Structure bundles as a sequence of self-contained packets (uses V2 code)
+		/// </summary>
+		PacketSequence = 5,
+
 		/// <summary>
 		/// Last item in the enum. Used for <see cref="Latest"/>
 		/// </summary>
@@ -47,6 +52,16 @@ namespace EpicGames.Horde.Storage.Bundles
 		/// The current version number
 		/// </summary>
 		Latest = (int)LatestPlusOne - 1,
+
+		/// <summary>
+		/// Last version using the V1 pipeline
+		/// </summary>
+		LatestV1 = ImportHashes,
+
+		/// <summary>
+		/// Last version using the V2 pipeline
+		/// </summary>
+		LatestV2 = Latest,
 #pragma warning restore CA1069 // Enums values should not be duplicated
 	}
 
diff --git a/Engine/Source/Programs/Shared/EpicGames.Horde/Storage/Bundles/BundleCache.cs b/Engine/Source/Programs/Shared/EpicGames.Horde/Storage/Bundles/BundleCache.cs
--- a/Engine/Source/Programs/Shared/EpicGames.Horde/Storage/Bundles/BundleCache.cs
+++ b/Engine/Source/Programs/Shared/EpicGames.Horde/Storage/Bundles/BundleCache.cs
@@ -1,6 +1,7 @@
 // Copyright Epic Games, Inc. All Rights Reserved.
 
 using System;
+using System.Buffers;
 using System.Collections.Generic;
 using System.Diagnostics.CodeAnalysis;
 using System.Threading;
@@ -126,18 +127,53 @@ namespace EpicGames.Horde.Storage
 				_size = size;
 			}
 
-			public void Dispose()
+			public virtual void Dispose()
 			{
 				Interlocked.Add(ref _outer._currentSize, -_size);
 				_size = 0;
 			}
 		}
 
+		sealed class MemoryAllocation : MemoryReservation, IMemoryOwner<byte>
+		{
+			readonly IMemoryOwner<byte> _owner;
+
+			public Memory<byte> Memory => _owner.Memory;
+
+			public MemoryAllocation(BundleCache outer, IMemoryOwner<byte> owner)
+				: base(outer, owner.Memory.Length)
+			{
+				_owner = owner;
+			}
+
+			public override void Dispose()
+			{
+				base.Dispose();
+
+				_owner.Dispose();
+			}
+		}
+
+		class MemoryAllocator : IMemoryAllocator<byte>
+		{
+			readonly BundleCache _outer;
+			readonly IMemoryAllocator<byte> _inner;
+
+			public MemoryAllocator(BundleCache outer, IMemoryAllocator<byte> inner)
+			{
+				_outer = outer;
+				_inner = inner;
+			}
+
+			public IMemoryOwner<byte> Alloc(int minSize) => new MemoryAllocation(_outer, _inner.Alloc(minSize));
+		}
+
 		readonly object _lockObject = new object();
 		readonly BundleCacheOptions _options;
 		readonly LinkedList<CacheValue> _items = new LinkedList<CacheValue>();
 		readonly Dictionary<object, CacheValue> _itemLookup = new Dictionary<object, CacheValue>();
 		readonly CancellationTokenSource _cancellationSource = new CancellationTokenSource();
+		readonly IMemoryAllocator<byte> _allocator;
 
 		long _currentSize;
 
@@ -149,6 +185,11 @@ namespace EpicGames.Horde.Storage
 		/// </summary>
 		public static BundleCache None { get; } = new BundleCache(new BundleCacheOptions { MaxSize = 0, HeaderCacheSize = 0, PacketCacheSize = 0 });
 
+		/// <summary>
+		/// Accessor for the default allocator
+		/// </summary>
+		public IMemoryAllocator<byte> Allocator => _allocator;
+
 		/// <summary>
 		/// Size of the configured header cache
 		/// </summary>
@@ -177,6 +218,8 @@ namespace EpicGames.Horde.Storage
 		public BundleCache(BundleCacheOptions options)
 		{
 			_options = options;
+			_allocator = new MemoryAllocator(this, PoolAllocator.Shared);
+
 			if (options.HeaderCacheSize > 0)
 			{
 				_headerCache = new MemoryCache(new MemoryCacheOptions { SizeLimit = options.HeaderCacheSize });
@@ -296,6 +339,15 @@ namespace EpicGames.Horde.Storage
 			}
 		}
 
+		/// <summary>
+		/// Allocate a block of memory from the shared heap
+		/// </summary>
+		/// <param name="size">Size of the allocation. The returned memory may be larger than this.</param>
+		public IMemoryOwner<byte> Allocate(int size)
+		{
+			return _allocator.Alloc(size);
+		}
+
 		/// <summary>
 		/// Add a reservation against the maximum allowed size of cached data. 
 		/// </summary>
diff --git a/Engine/Source/Programs/Shared/EpicGames.Horde/Storage/Bundles/BundleOptions.cs b/Engine/Source/Programs/Shared/EpicGames.Horde/Storage/Bundles/BundleOptions.cs
--- a/Engine/Source/Programs/Shared/EpicGames.Horde/Storage/Bundles/BundleOptions.cs
+++ b/Engine/Source/Programs/Shared/EpicGames.Horde/Storage/Bundles/BundleOptions.cs
@@ -7,6 +7,16 @@ namespace EpicGames.Horde.Storage.Bundles
 	/// </summary>
 	public class BundleOptions
 	{
+		/// <summary>
+		/// Default options value
+		/// </summary>
+		public static BundleOptions Default { get; } = new BundleOptions();
+
+		/// <summary>
+		/// Maximum version number of bundles to write
+		/// </summary>
+		public BundleVersion MaxVersion { get; set; } = BundleVersion.LatestV1;
+
 		/// <summary>
 		/// Maximum payload size fo a blob
 		/// </summary>
diff --git a/Engine/Source/Programs/Shared/EpicGames.Horde/Storage/Bundles/V1/BundleHeader.cs b/Engine/Source/Programs/Shared/EpicGames.Horde/Storage/Bundles/V1/BundleHeader.cs
--- a/Engine/Source/Programs/Shared/EpicGames.Horde/Storage/Bundles/V1/BundleHeader.cs
+++ b/Engine/Source/Programs/Shared/EpicGames.Horde/Storage/Bundles/V1/BundleHeader.cs
@@ -146,7 +146,7 @@ namespace EpicGames.Horde.Storage.Bundles.V1
 			Packets.Write(writer);
 
 			// Fill in the length of the header, and append it to the builder
-			BundleSignature signature = new BundleSignature(BundleVersion.Latest, writer.Length);
+			BundleSignature signature = new BundleSignature(BundleVersion.LatestV1, writer.Length);
 			Bundle.WriteSignature(signatureData, signature);
 
 			writer.AppendTo(builder);
@@ -209,7 +209,11 @@ namespace EpicGames.Horde.Storage.Bundles.V1
 		/// <returns>New header object</returns>
 		public static async Task<BundleHeader> ReadAsync(BundleSignature signature, Stream stream, CancellationToken cancellationToken = default)
 		{
-			if (signature.Version >= BundleVersion.InPlace)
+			if (signature.Version > BundleVersion.LatestV1)
+			{
+				return new BundleHeader(Array.Empty<BlobType>(), Array.Empty<BlobLocator>(), Array.Empty<BundleExport>(), Array.Empty<BundlePacket>());
+			}
+			else if (signature.Version >= BundleVersion.InPlace)
 			{
 				return await ReadLatestAsync(signature.Version, stream, signature.HeaderLength - Bundle.SignatureLength, cancellationToken);
 			}
diff --git a/Engine/Source/Programs/Shared/EpicGames.Horde/Storage/Bundles/V2/BundleWriter.cs b/Engine/Source/Programs/Shared/EpicGames.Horde/Storage/Bundles/V2/BundleWriter.cs
new file mode 100644
--- /dev/null
+++ b/Engine/Source/Programs/Shared/EpicGames.Horde/Storage/Bundles/V2/BundleWriter.cs
@@ -0,0 +1,374 @@
+// Copyright Epic Games, Inc. All Rights Reserved.
+
+using System;
+using System.Collections.Generic;
+using System.Diagnostics;
+using System.Linq;
+using System.Threading;
+using System.Threading.Tasks;
+using EpicGames.Core;
+
+namespace EpicGames.Horde.Storage.Bundles.V2
+{
+	/// <summary>
+	/// Implements the primary storage writer interface for V2 bundles. Writes exports into packets, and flushes them to storage in bundles.
+	/// </summary>
+	public sealed class BundleWriter : IStorageWriter
+	{
+		// An export that has been written and is waiting to be flushed to disk
+		class PendingExportHandle : IBlobHandle
+		{
+			readonly PendingPacketHandle _packet;
+			readonly int _exportIdx;
+
+			readonly BlobType _type;
+			readonly IBlobHandle[] _imports;
+			readonly AliasInfo[] _aliases;
+		
+			public IBlobHandle? Outer => _packet;
+
+			public IReadOnlyList<AliasInfo> Aliases => _aliases;
+
+			public PendingExportHandle(PendingPacketHandle packet, int exportIdx, BlobType type, IBlobHandle[] imports, AliasInfo[] aliases)
+			{
+				_packet = packet;
+				_exportIdx = exportIdx;
+				_type = type;
+				_imports = imports;
+				_aliases = aliases;
+			}
+
+			public ValueTask FlushAsync(CancellationToken cancellationToken = default) => _packet.FlushAsync(cancellationToken);
+
+			public async ValueTask<BlobData> ReadAsync(CancellationToken cancellationToken = default)
+			{
+				IReadOnlyMemoryOwner<byte> body = await _packet.ReadExportBodyAsync(_exportIdx, cancellationToken);
+				return new BlobDataWithOwner(_type, body.Memory, _imports, body);
+			}
+
+			public bool TryAppendIdentifier(Utf8StringBuilder builder)
+			{
+				ExportHandle.AppendIdentifier(builder, _exportIdx);
+				return true;
+			}
+		}
+
+		// Packet that is still being built, but may be redirected to a flushed packet
+		class PendingPacketHandle : IBlobHandle, IDisposable
+		{
+			readonly object _lockObject = new object();
+			readonly PendingBundleHandle _bundle;
+
+			PacketHandle? _flushedHandle;
+			readonly List<PendingExportHandle> _pendingExports = new List<PendingExportHandle>();
+
+			Packet? _packet;
+			PacketWriter? _packetWriter;
+
+			public IBlobHandle? Outer => _bundle;
+
+			public Packet? Packet => _packet;
+
+			public PendingPacketHandle(PendingBundleHandle bundle, IMemoryAllocator<byte> allocator)
+			{
+				_bundle = bundle;
+				_packetWriter = new PacketWriter(_bundle, this, allocator, _lockObject);
+			}
+
+			public void Dispose()
+			{
+				if (_packetWriter != null)
+				{
+					_packetWriter.Dispose();
+					_packetWriter = null;
+				}
+			}
+
+			public ValueTask FlushAsync(CancellationToken cancellationToken = default) => _bundle.FlushAsync(cancellationToken);
+
+			public bool IsEmpty() 
+				=> _pendingExports.Count == 0;
+
+			public int GetLength()
+				=> _packetWriter!.Length;
+
+			public Memory<byte> GetOutputBuffer(int usedSize, int desiredSize)
+				=> _packetWriter!.GetOutputBuffer(usedSize, desiredSize);
+
+			public PendingExportHandle CompleteExport(BlobType type, int size, IReadOnlyList<IBlobHandle> references, IReadOnlyList<AliasInfo> aliases)
+			{
+				int exportIdx = _packetWriter!.CompleteExport(size, type, references);
+
+				PendingExportHandle exportHandle = new PendingExportHandle(this, exportIdx, type, references.ToArray(), aliases.ToArray());
+				_pendingExports.Add(exportHandle);
+
+				return exportHandle;
+			}
+
+			public Packet CompletePacket()
+			{
+				_packet ??= _packetWriter!.CompletePacket();
+				return _packet;
+			}
+
+			public async ValueTask SetFlushedHandleAsync(IStorageClient storageClient, PacketHandle flushedHandle, CancellationToken cancellationToken)
+			{
+				Debug.Assert(_packetWriter != null);
+
+				lock (_lockObject)
+				{
+					_flushedHandle = flushedHandle;
+
+					_packetWriter!.Dispose();
+					_packetWriter = null;
+				}
+
+				for(int exportIdx = 0; exportIdx < _pendingExports.Count; exportIdx++)
+				{
+					PendingExportHandle pendingExport = _pendingExports[exportIdx];
+					foreach (AliasInfo alias in pendingExport.Aliases)
+					{
+						ExportHandle exportHandle = new ExportHandle(_flushedHandle, exportIdx);
+						await storageClient.AddAliasAsync(alias.Name, exportHandle, alias.Rank, alias.Data, cancellationToken);
+					}
+				}
+			}
+
+			public ValueTask<BlobData> ReadAsync(CancellationToken cancellationToken = default)
+			{
+				lock (_lockObject)
+				{
+					if (_packetWriter != null)
+					{
+						throw new NotSupportedException("Reading pending packets is not currently supported.");
+					}
+				}
+				return _flushedHandle!.ReadAsync(cancellationToken);
+			}
+
+			public async ValueTask<IReadOnlyMemoryOwner<byte>> ReadExportBodyAsync(int exportIdx, CancellationToken cancellationToken = default)
+			{
+				lock (_lockObject)
+				{
+					if (_packetWriter != null)
+					{
+						return _packetWriter.GetExportData(exportIdx);
+					}
+				}
+
+				return await _flushedHandle!.ReadExportBodyAsync(exportIdx, cancellationToken);
+			}
+
+			public bool TryAppendIdentifier(Utf8StringBuilder builder)
+				=> _flushedHandle?.TryAppendIdentifier(builder) ?? false;
+		}
+
+		// Fragment of a bundle that needs to be written to storage.
+		class PendingBundleHandle : IBlobHandle, IDisposable
+		{
+			readonly IStorageClient _storageClient;
+			readonly string? _basePath;
+			readonly BundleCache _cache;
+			readonly BundleOptions _options;
+
+			IBlobHandle? _flushedHandle;
+#pragma warning disable CA2213
+			PendingPacketHandle? _currentPacket;
+#pragma warning restore CA2213
+			List<PendingPacketHandle>? _pendingPackets = new List<PendingPacketHandle>();
+			int _length;
+
+			/// <inheritdoc/>
+			public IBlobHandle? Outer => null;
+
+			/// <summary>
+			/// Uncompressed length of this bundle
+			/// </summary>
+			public int Length => _length;
+
+			public PendingBundleHandle(IStorageClient storageClient, string? basePath, BundleCache cache, BundleOptions options)
+			{
+				_storageClient = storageClient;
+				_basePath = basePath;
+				_cache = cache;
+				_options = options;
+
+				_currentPacket = new PendingPacketHandle(this, cache.Allocator);
+			}
+
+			/// <inheritdoc/>
+			public void Dispose()
+			{
+				if (_currentPacket != null)
+				{
+					_currentPacket.Dispose();
+					_currentPacket = null;
+				}
+				if (_pendingPackets != null)
+				{
+					foreach (PendingPacketHandle pendingPacket in _pendingPackets)
+					{
+						pendingPacket.Dispose();
+					}
+					_pendingPackets = null;
+				}
+			}
+
+			public Memory<byte> GetOutputBuffer(int usedSize, int desiredSize)
+				=> _currentPacket!.GetOutputBuffer(usedSize, desiredSize);
+
+			public PendingExportHandle CompleteExport(BlobType type, int size, IReadOnlyList<IBlobHandle> references, IReadOnlyList<AliasInfo> aliases)
+			{
+				PendingExportHandle exportHandle = _currentPacket!.CompleteExport(type, size, references, aliases);
+				if (_currentPacket.GetLength() > Math.Min(_options.MinCompressionPacketSize, _options.MaxBlobSize))
+				{
+					CompletePacket();
+					_currentPacket = new PendingPacketHandle(this, _cache.Allocator);
+				}
+				return exportHandle;
+			}
+
+			void CompletePacket()
+			{
+				_currentPacket!.CompletePacket();
+
+				_length += _currentPacket.Packet!.Length;
+				_pendingPackets!.Add(_currentPacket);
+
+				_currentPacket = null;
+			}
+
+			// Write this bundle to storage
+			public async ValueTask FlushAsync(CancellationToken cancellationToken = default)
+			{
+				if (_currentPacket == null)
+				{
+					return;
+				}
+				if (!_currentPacket!.IsEmpty())
+				{
+					CompletePacket();
+				}
+				if (_pendingPackets!.Count == 0)
+				{
+					return;
+				}
+
+				// Write the signature bytes
+				using RefCountedMemoryWriter memoryWriter = new RefCountedMemoryWriter(_cache.Allocator, 65536);
+
+				// Compress all the packets
+				List<int> encodedOffsets = new List<int>();
+				foreach (PendingPacketHandle pendingPacket in _pendingPackets)
+				{
+					encodedOffsets.Add(memoryWriter.Length);
+					pendingPacket.Packet!.Encode(_options.CompressionFormat, memoryWriter);
+				}
+				encodedOffsets.Add(memoryWriter.Length);
+
+				// Find all the other bundles that are referenced
+				HashSet<BlobLocator> references = new HashSet<BlobLocator>();
+				foreach (PendingPacketHandle pendingPacket in _pendingPackets)
+				{
+					for (int importIdx = 0; importIdx < pendingPacket.Packet!.GetImportCount(); importIdx++)
+					{
+						PacketImport import = pendingPacket.Packet.GetImport(importIdx);
+						if (import.BaseIdx == -1)
+						{
+							references.Add(new BlobLocator(import.Fragment));
+						}
+					}
+				}
+
+				// TODO: put all the encoded packets into the cache
+
+				// Write the bundle data
+				using (ReadOnlySequenceStream stream = new ReadOnlySequenceStream(memoryWriter.AsSequence()))
+				{
+					IBlobHandle[] referencedHandles = references.Select(x => _storageClient.CreateBlobHandle(x)).ToArray(); 
+					_flushedHandle = await _storageClient.WriteBlobAsync(Bundle.BlobType, stream, referencedHandles, _basePath, cancellationToken);
+				}
+
+				// Update all the packets to point to the flushed bundle
+				for (int packetIdx = 0; packetIdx < _pendingPackets.Count; packetIdx++)
+				{
+					int packetLength = encodedOffsets[packetIdx + 1] - encodedOffsets[packetIdx];
+					PacketHandle packetHandle = new PacketHandle(_storageClient, _flushedHandle, encodedOffsets[packetIdx], packetLength, _cache);
+					await _pendingPackets[packetIdx].SetFlushedHandleAsync(_storageClient, packetHandle, cancellationToken);
+				}
+
+				// Clear out all the pending packets
+				_pendingPackets = null;
+			}
+
+			/// <inheritdoc/>
+			public ValueTask<BlobData> ReadAsync(CancellationToken cancellationToken)
+				=> GetFlushedHandle().ReadAsync(cancellationToken);
+
+			/// <inheritdoc/>
+			public ValueTask<IReadOnlyMemoryOwner<byte>> ReadBodyAsync(int offset, int? length, CancellationToken cancellationToken = default)
+				=> GetFlushedHandle().ReadBodyAsync(offset, length, cancellationToken);
+
+			/// <inheritdoc/>
+			public bool TryAppendIdentifier(Utf8StringBuilder builder)
+				=> _flushedHandle?.TryAppendIdentifier(builder) ?? false;
+
+			IBlobHandle GetFlushedHandle()
+				=> _flushedHandle ?? throw new NotSupportedException("Bundle has not yet been flushed");
+		}
+
+		readonly IStorageClient _storageClient;
+		readonly string? _basePath;
+		readonly BundleOptions _options;
+		readonly BundleCache _bundleCache;
+
+		PendingBundleHandle _currentBundle;
+
+		/// <summary>
+		/// 
+		/// </summary>
+		public BundleWriter(IStorageClient storageClient, string? basePath, BundleCache bundleCache, BundleOptions options)
+		{
+			_storageClient = storageClient;
+			_basePath = basePath;
+			_options = options;
+			_bundleCache = bundleCache;
+
+			_currentBundle = new PendingBundleHandle(storageClient, basePath, bundleCache, _options);
+		}
+
+		/// <inheritdoc/>
+		public async ValueTask DisposeAsync()
+		{
+			await FlushAsync();
+			_currentBundle.Dispose();
+		}
+
+		/// <inheritdoc/>
+		public async Task FlushAsync(CancellationToken cancellationToken = default)
+		{
+			await _currentBundle.FlushAsync(cancellationToken);
+			_currentBundle.Dispose();
+			_currentBundle = new PendingBundleHandle(_storageClient, _basePath, _bundleCache, _options);
+		}
+
+		/// <inheritdoc/>
+		public IStorageWriter Fork()
+			=> new BundleWriter(_storageClient, _basePath, _bundleCache, _options);
+
+		/// <inheritdoc/>
+		public Memory<byte> GetOutputBuffer(int usedSize, int desiredSize)
+			=> _currentBundle.GetOutputBuffer(usedSize, desiredSize);
+
+		/// <inheritdoc/>
+		public async ValueTask<IBlobHandle> WriteBlobAsync(BlobType type, int size, IReadOnlyList<IBlobHandle> references, IReadOnlyList<AliasInfo> aliases, CancellationToken cancellationToken = default)
+		{
+			PendingExportHandle exportHandle = _currentBundle.CompleteExport(type, size, references, aliases);
+			if (_currentBundle.Length > _options.MaxBlobSize)
+			{
+				await FlushAsync(cancellationToken);
+			}
+			return exportHandle;
+		}
+	}
+}
diff --git a/Engine/Source/Programs/Shared/EpicGames.Horde/Storage/Bundles/V2/ExportHandle.cs b/Engine/Source/Programs/Shared/EpicGames.Horde/Storage/Bundles/V2/ExportHandle.cs
new file mode 100644
--- /dev/null
+++ b/Engine/Source/Programs/Shared/EpicGames.Horde/Storage/Bundles/V2/ExportHandle.cs
@@ -0,0 +1,98 @@
+// Copyright Epic Games, Inc. All Rights Reserved.
+
+using System;
+using System.Buffers.Text;
+using System.Text;
+using System.Threading;
+using System.Threading.Tasks;
+using EpicGames.Core;
+
+namespace EpicGames.Horde.Storage.Bundles.V2
+{
+	/// <summary>
+	/// Handle to an export within a packet. Same implementation is used for flushed and pending exports.
+	/// </summary>
+	public class ExportHandle : IBlobHandle
+	{
+		static readonly Utf8String s_fragmentPrefix = new Utf8String("exp=");
+
+		readonly PacketHandle _packet;
+		readonly int _exportIdx;
+
+		/// <inheritdoc/>
+		public IBlobHandle? Outer => _packet;
+
+		/// <summary>
+		/// Constructor
+		/// </summary>
+		public ExportHandle(PacketHandle packet, int exportIdx)
+		{
+			_packet = packet;
+			_exportIdx = exportIdx;
+		}
+
+		/// <summary>
+		/// Constructor
+		/// </summary>
+		public ExportHandle(PacketHandle packet, ReadOnlySpan<byte> fragment)
+		{
+			_packet = packet;
+			if (!TryParse(fragment, out _exportIdx))
+			{
+				throw new FormatException($"Invalid fragment {Encoding.UTF8.GetString(fragment)} relative to {packet}");
+			}
+		}
+
+		/// <inheritdoc/>
+		public ValueTask FlushAsync(CancellationToken cancellationToken = default) => default;
+
+		/// <summary>
+		/// Attempt to parse an export index from the given fragment
+		/// </summary>
+		static bool TryParse(ReadOnlySpan<byte> fragment, out int exportIdx)
+		{
+			if (!fragment.StartsWith(s_fragmentPrefix))
+			{
+				exportIdx = 0;
+				return false;
+			}
+
+			fragment = fragment.Slice(s_fragmentPrefix.Length);
+			if (!Utf8Parser.TryParse(fragment, out exportIdx, out int numBytesRead) || numBytesRead != fragment.Length)
+			{
+				exportIdx = 0;
+				return false;
+			}
+
+			return true;
+		}
+
+		/// <inheritdoc/>
+		public ValueTask<BlobData> ReadAsync(CancellationToken cancellationToken = default)
+			=> _packet.ReadExportAsync(_exportIdx, cancellationToken);
+
+		/// <inheritdoc/>
+		public bool TryAppendIdentifier(Utf8StringBuilder builder)
+		{
+			AppendIdentifier(builder, _exportIdx);
+			return true;
+		}
+
+		/// <summary>
+		/// Appends an export identifier to the given string builder
+		/// </summary>
+		public static void AppendIdentifier(Utf8StringBuilder builder, int exportIdx)
+		{
+			builder.Append(s_fragmentPrefix);
+			builder.Append(exportIdx);
+		}
+
+		/// <inheritdoc/>
+		public override bool Equals(object? obj)
+			=> obj is ExportHandle other && _packet.Equals(other._packet) && _exportIdx == other._exportIdx;
+
+		/// <inheritdoc/>
+		public override int GetHashCode()
+			=> HashCode.Combine(_packet, _exportIdx);
+	}
+}
diff --git a/Engine/Source/Programs/Shared/EpicGames.Horde/Storage/Bundles/V2/Packet.cs b/Engine/Source/Programs/Shared/EpicGames.Horde/Storage/Bundles/V2/Packet.cs
new file mode 100644
--- /dev/null
+++ b/Engine/Source/Programs/Shared/EpicGames.Horde/Storage/Bundles/V2/Packet.cs
@@ -0,0 +1,343 @@
+// Copyright Epic Games, Inc. All Rights Reserved.
+
+using System;
+using System.Buffers;
+using System.Buffers.Binary;
+using System.Diagnostics;
+using System.Linq;
+using EpicGames.Core;
+
+namespace EpicGames.Horde.Storage.Bundles.V2
+{
+	/// <summary>
+	/// Accessor for data structures stored into a serialized bundle packet.
+	/// </summary>
+	/// <remarks>
+	/// <code>
+	/// Each packet starts with a header:
+	///  - 4 bytes: Signature/version field ('U', 'E', 'B', [version number])
+	///  - 4 bytes: Packet length
+	///  - 4 bytes: offset of type table from the start of the packet
+	///  - 4 bytes: offset of import table from the start of the packet
+	///  - 4 bytes: offset of export table from the start of the packet
+	///	
+	/// The type table is constructed as:
+	///  - 4 bytes: Number of entries
+	///  - 20 bytes * Number of entries: BlobType data
+	///	
+	/// The import table is constructed as:
+	///  - 4 bytes: Number of entries
+	///  - 4 bytes * (Number of entries + 1): Offset of each entry from the start of the packet, with a sentinel value for the end of the last entry. Length of each entry is implicit by next entry - this entry.
+	///	
+	/// The export index is constructed as:
+	///  - 4 bytes: Number of entries
+	///  - 4 bytes * (Number of entries + 1): Offset of each entry from the start of the packet, with a sentinel value for the end of the last entry. Length of each entry is implicit by next entry - this entry.
+	///	
+	/// Each import is written as:
+	///  - VarInt: Base import index, with a +1 bias, or zero for 'none'
+	///  - Utf8 string containing fragment on top of base import, without a null terminator.
+	///	
+	/// Each export is written as:
+	///  - 4 bytes: Length of payload
+	///  - ? bytes: Payload data
+	///  - VarInt: Type index
+	///  - [VarInt]: Import indexes
+	/// </code>
+	/// </remarks>
+	[DebuggerTypeProxy(typeof(Packet.DebugProxy))]
+	public sealed class Packet
+	{
+		class DebugProxy
+		{
+			public BlobType[] Types { get; }
+			public PacketImport[] Imports { get; }
+			public PacketExport[] Exports { get; }
+
+			public DebugProxy(Packet packet)
+			{
+				Types = Enumerable.Range(0, packet.GetTypeCount()).Select(packet.GetType).ToArray();
+				Imports = Enumerable.Range(0, packet.GetImportCount()).Select(packet.GetImport).ToArray();
+				Exports = Enumerable.Range(0, packet.GetExportCount()).Select(packet.GetExport).ToArray();
+			}
+		}
+
+		readonly ReadOnlyMemory<byte> _data;
+
+		/// <summary>
+		/// Constructor
+		/// </summary>
+		/// <param name="data">Data for the packet</param>
+		public Packet(ReadOnlyMemory<byte> data) => _data = data;
+
+		/// <summary>
+		/// Length of this packet
+		/// </summary>
+		public int Length => _data.Length;
+
+		/// <summary>
+		/// Gets the number of types in this packet
+		/// </summary>
+		public int GetTypeCount()
+		{
+			ReadOnlySpan<byte> span = _data.Span;
+			int typeTableOffset = BinaryPrimitives.ReadInt32LittleEndian(span[8..]);
+			return BinaryPrimitives.ReadInt32LittleEndian(span[typeTableOffset..]);
+		}
+
+		/// <summary>
+		/// Gets a type from the packet
+		/// </summary>
+		public BlobType GetType(int typeIdx)
+		{
+			ReadOnlySpan<byte> span = _data.Span;
+			int typeTableOffset = BinaryPrimitives.ReadInt32LittleEndian(span[8..]);
+
+			int typeCount = BinaryPrimitives.ReadInt32LittleEndian(span[typeTableOffset..]);
+			if (typeIdx < 0 || typeIdx >= typeCount)
+			{
+				throw new IndexOutOfRangeException();
+			}
+
+			int typeOffset = typeTableOffset + sizeof(int) + typeIdx * BlobType.NumBytes;
+			return BlobType.Read(span[typeOffset..]);
+		}
+
+		/// <summary>
+		/// Gets the number of imports in this packet
+		/// </summary>
+		public int GetImportCount()
+		{
+			ReadOnlySpan<byte> span = _data.Span;
+			int importTableOffset = BinaryPrimitives.ReadInt32LittleEndian(span[12..]);
+			return BinaryPrimitives.ReadInt32LittleEndian(span[importTableOffset..]);
+		}
+
+		/// <summary>
+		/// Gets the locator for a particular import
+		/// </summary>
+		/// <param name="importIdx">Index of the import to retrieve</param>
+		public PacketImport GetImport(int importIdx)
+		{
+			ReadOnlySpan<byte> span = _data.Span;
+			int importTableOffset = BinaryPrimitives.ReadInt32LittleEndian(span[12..]);
+
+			int importCount = BinaryPrimitives.ReadInt32LittleEndian(span[importTableOffset..]);
+			if (importIdx < 0 || importIdx >= importCount)
+			{
+				throw new IndexOutOfRangeException();
+			}
+
+			int importEntryOffset = importTableOffset + sizeof(int) + importIdx * sizeof(int);
+
+			int importOffset = BinaryPrimitives.ReadInt32LittleEndian(span[importEntryOffset..]);
+			int importLength = BinaryPrimitives.ReadInt32LittleEndian(span[(importEntryOffset + 4)..]) - importOffset;
+
+			return PacketImport.Read(_data.Slice(importOffset, importLength));
+		}
+
+		/// <summary>
+		/// Gets the number of exports in this packet
+		/// </summary>
+		public int GetExportCount()
+		{
+			ReadOnlySpan<byte> span = _data.Span;
+			int exportTableOffset = BinaryPrimitives.ReadInt32LittleEndian(span[16..]);
+			return BinaryPrimitives.ReadInt32LittleEndian(span[exportTableOffset..]);
+		}
+
+		/// <summary>
+		/// Gets the bulk data for a particular export
+		/// </summary>
+		/// <param name="exportIdx"></param>
+		public PacketExport GetExport(int exportIdx)
+		{
+			ReadOnlySpan<byte> span = _data.Span;
+			int exportTableOffset = BinaryPrimitives.ReadInt32LittleEndian(span[16..]);
+
+			int exportCount = BinaryPrimitives.ReadInt32LittleEndian(span[exportTableOffset..]);
+			if (exportIdx < 0 || exportIdx >= exportCount)
+			{
+				throw new IndexOutOfRangeException();
+			}
+
+			int exportEntryOffset = exportTableOffset + sizeof(int) + exportIdx * sizeof(int);
+
+			int exportOffset = BinaryPrimitives.ReadInt32LittleEndian(span[exportEntryOffset..]);
+			int exportLength = BinaryPrimitives.ReadInt32LittleEndian(span[(exportEntryOffset + 4)..]) - exportOffset;
+
+			return new PacketExport(_data.Slice(exportOffset, exportLength));
+		}
+
+		/// <summary>
+		/// Encodes a packet
+		/// </summary>
+		/// <param name="format">Compression format for the encoded data</param>
+		/// <param name="writer">Writer for the encoded data</param>
+		public void Encode(BundleCompressionFormat format, IMemoryWriter writer)
+		{
+			Span<byte> signatureSpan = writer.GetSpanAndAdvance(Bundle.SignatureLength);
+			writer.WriteInt32(_data.Length);
+			writer.WriteUInt8((byte)format);
+
+			int encodedLength = BundleData.Compress(format, _data, writer);
+			Bundle.WriteSignature(signatureSpan, new BundleSignature(BundleVersion.LatestV2, encodedLength + sizeof(int) + 1));
+		}
+
+		/// <summary>
+		/// Decodes this packet
+		/// </summary>
+		/// <returns></returns>
+		public static IRefCountedHandle<Packet> Decode(ReadOnlyMemory<byte> data, IMemoryAllocator<byte> allocator)
+		{
+			ReadOnlySpan<byte> span = data.Span;
+
+			BundleSignature signature = Bundle.ReadSignature(span);
+			span = span[Bundle.SignatureLength..];
+
+			if (signature.Version <= BundleVersion.LatestV1 || signature.Version > BundleVersion.LatestV2)
+			{
+				throw new InvalidOperationException($"Cannot read bundle packet; unsupported version {(int)signature.Version}");
+			}
+
+			int decodedLength = BinaryPrimitives.ReadInt32LittleEndian(span);
+			span = span[4..];
+
+			BundleCompressionFormat format = (BundleCompressionFormat)span[0];
+			span = span[1..];
+
+			IMemoryOwner<byte> owner = allocator.Alloc(decodedLength);
+			BundleData.Decompress(format, data.Slice(data.Length - span.Length), owner.Memory);
+
+			Packet packet = new Packet(owner.Memory);
+			return RefCountedHandle.Create(packet, owner);
+		}
+	}
+
+	/// <summary>
+	/// Raw data for an imported node
+	/// </summary>
+	/// <param name="BaseIdx">Base index for this locator</param>
+	/// <param name="Fragment">The utf8 fragment appended to the base index</param>
+	public record struct PacketImport(int BaseIdx, Utf8String Fragment)
+	{
+		/// <summary>
+		/// Bias for indexes into the import table
+		/// </summary>
+		public const int Bias = 3;
+
+		/// <summary>
+		/// 
+		/// </summary>
+		public const int InvalidBaseIdx = -1;
+
+		/// <summary>
+		/// 
+		/// </summary>
+		public const int CurrentPacketBaseIdx = -2;
+
+		/// <summary>
+		/// 
+		/// </summary>
+		public const int CurrentBundleBaseIdx = -3;
+
+		/// <summary>
+		/// Reads an import from a block of memory
+		/// </summary>
+		public static PacketImport Read(ReadOnlyMemory<byte> data)
+		{
+			ReadOnlySpan<byte> span = data.Span;
+			int baseIdx = (int)VarInt.ReadUnsigned(span, out int bytesRead) - Bias;
+			Utf8String fragment = new Utf8String(data.Slice(bytesRead));
+			return new PacketImport(baseIdx, fragment);
+		}
+	}
+
+	/// <summary>
+	/// Data for an exported node in a packet
+	/// </summary>
+	[DebuggerTypeProxy(typeof(PacketExport.DebugProxy))]
+	public struct PacketExport
+	{
+		class DebugProxy
+		{
+			public int TypeIdx { get; }
+			public int[] Imports { get; }
+			public ReadOnlyMemory<byte> Payload { get; }
+
+			public DebugProxy(PacketExport export)
+			{
+				PacketExportHeader header = export.GetHeader();
+				TypeIdx = header.TypeIdx;
+				Imports = header.Imports;
+				Payload = export.GetPayload();
+			}
+		}
+
+		/// <summary>
+		/// Data for this export
+		/// </summary>
+		public ReadOnlyMemory<byte> Data { get; }
+
+		/// <summary>
+		/// Constructor
+		/// </summary>
+		public PacketExport(ReadOnlyMemory<byte> data) => Data = data;
+
+		/// <summary>
+		/// Gets the header for this export
+		/// </summary>
+		public PacketExportHeader GetHeader()
+		{
+			ReadOnlySpan<byte> span = Data.Span;
+			int payloadLength = BinaryPrimitives.ReadInt32LittleEndian(span);
+			return new PacketExportHeader(span.Slice(sizeof(int) + payloadLength));
+		}
+
+		/// <summary>
+		/// Gets the payload for this export
+		/// </summary>
+		public ReadOnlyMemory<byte> GetPayload()
+		{
+			ReadOnlySpan<byte> span = Data.Span;
+			int payloadLength = BinaryPrimitives.ReadInt32LittleEndian(span);
+			return Data.Slice(sizeof(int), payloadLength);
+		}
+	}
+
+	/// <summary>
+	/// Data for an exported node in a bundle
+	/// </summary>
+	public struct PacketExportHeader
+	{
+		/// <summary>
+		/// Index of the type for this export
+		/// </summary>
+		public int TypeIdx { get; }
+
+		/// <summary>
+		/// Index of imports for this export
+		/// </summary>
+		public int[] Imports { get; }
+
+		/// <summary>
+		/// Constructor
+		/// </summary>
+		public PacketExportHeader(ReadOnlySpan<byte> span)
+		{
+			int bytesRead;
+
+			TypeIdx = (int)VarInt.ReadUnsigned(span, out bytesRead);
+			span = span[bytesRead..];
+
+			int numImports = (int)VarInt.ReadUnsigned(span, out bytesRead);
+			span = span[bytesRead..];
+
+			Imports = new int[numImports];
+			for (int idx = 0; idx < numImports; idx++)
+			{
+				Imports[idx] = (int)VarInt.ReadUnsigned(span, out bytesRead);
+				span = span[bytesRead..];
+			}
+		}
+	}
+}
diff --git a/Engine/Source/Programs/Shared/EpicGames.Horde/Storage/Bundles/V2/PacketHandle.cs b/Engine/Source/Programs/Shared/EpicGames.Horde/Storage/Bundles/V2/PacketHandle.cs
new file mode 100644
--- /dev/null
+++ b/Engine/Source/Programs/Shared/EpicGames.Horde/Storage/Bundles/V2/PacketHandle.cs
@@ -0,0 +1,183 @@
+// Copyright Epic Games, Inc. All Rights Reserved.
+
+using System;
+using System.Buffers.Text;
+using System.Text;
+using System.Threading;
+using System.Threading.Tasks;
+using EpicGames.Core;
+
+namespace EpicGames.Horde.Storage.Bundles.V2
+{
+	/// <summary>
+	/// Handle to an packet within a bundle. 
+	/// </summary>
+	public class PacketHandle : IBlobHandle
+	{
+		/// <summary>
+		/// Type for packet blobs
+		/// </summary>
+		public static BlobType BlobType { get; } = new BlobType(Guid.Parse("{CD9A04EF-CAC1-47D3-A605-2A498130E651}"), 1);
+
+		static readonly Utf8String s_fragmentPrefix = new Utf8String("pkt=");
+
+		readonly IStorageClient _storageClient;
+		readonly IBlobHandle _outer;
+		readonly int _packetOffset;
+		readonly int _packetLength;
+		readonly BundleCache _cache;
+
+		/// <inheritdoc/>
+		public IBlobHandle? Outer => _outer;
+
+		/// <summary>
+		/// Constructor
+		/// </summary>
+		public PacketHandle(IStorageClient storageClient, IBlobHandle outer, int packetOffset, int packetLength, BundleCache cache)
+		{
+			_storageClient = storageClient;
+
+			_outer = outer;
+			_packetOffset = packetOffset;
+			_packetLength = packetLength;
+			_cache = cache;
+		}
+
+		/// <summary>
+		/// Constructor
+		/// </summary>
+		public PacketHandle(IStorageClient storageClient, IBlobHandle outer, ReadOnlySpan<byte> fragment, BundleCache cache)
+		{
+			_storageClient = storageClient;
+
+			_outer = outer;
+			_cache = cache;
+
+			if (!TryParse(fragment, out _packetOffset, out _packetLength))
+			{
+				throw new FormatException($"Cannot parse fragment {Encoding.UTF8.GetString(fragment)} relative to {outer}");
+			}
+		}
+
+		/// <summary>
+		/// Parse a fragment containing an offset and length
+		/// </summary>
+		static bool TryParse(ReadOnlySpan<byte> fragment, out int packetOffset, out int packetLength)
+		{
+			if (!fragment.StartsWith(s_fragmentPrefix.Span))
+			{
+				packetOffset = packetLength = 0;
+				return false;
+			}
+
+			fragment = fragment.Slice(s_fragmentPrefix.Length);
+			if (!Utf8Parser.TryParse(fragment, out packetOffset, out int numBytesRead))
+			{
+				packetOffset = packetLength = 0;
+				return false;
+			}
+
+			fragment = fragment[numBytesRead..];
+			if (fragment.Length == 0 || fragment[0] != (byte)',')
+			{
+				packetOffset = packetLength = 0;
+				return false;
+			}
+
+			fragment = fragment[1..];
+			if (!Utf8Parser.TryParse(fragment, out packetLength, out numBytesRead) || numBytesRead != fragment.Length)
+			{
+				packetOffset = packetLength = 0;
+				return false;
+			}
+
+			return true;
+		}
+
+		/// <inheritdoc/>
+		public ValueTask FlushAsync(CancellationToken cancellationToken = default) => default;
+
+		/// <inheritdoc/>
+		public ValueTask<BlobData> ReadAsync(CancellationToken cancellationToken = default)
+		{
+			throw new NotImplementedException();
+		}
+
+		/// <summary>
+		/// Reads an export from this packet
+		/// </summary>
+		/// <param name="exportIdx">Index of the export</param>
+		/// <param name="cancellationToken">Cancellation token for the operation</param>
+		public async ValueTask<BlobData> ReadExportAsync(int exportIdx, CancellationToken cancellationToken = default)
+		{
+			using IRefCountedHandle<PacketReader> packetReaderHandle = await GetPacketReaderAsync(cancellationToken);
+			return packetReaderHandle.Target.ReadExport(exportIdx);
+		}
+
+		/// <summary>
+		/// Reads an export body from this packet
+		/// </summary>
+		/// <param name="exportIdx">Index of the export</param>
+		/// <param name="cancellationToken">Cancellation token for the operation</param>
+		public async ValueTask<IReadOnlyMemoryOwner<byte>> ReadExportBodyAsync(int exportIdx, CancellationToken cancellationToken = default)
+		{
+			using IRefCountedHandle<PacketReader> packetReaderHandle = await GetPacketReaderAsync(cancellationToken);
+			return packetReaderHandle.Target.ReadExportBody(exportIdx);
+		}
+
+		/// <inheritdoc/>
+		public bool TryAppendIdentifier(Utf8StringBuilder builder)
+		{
+			builder.Append("pkt=");
+			builder.Append(_packetOffset);
+			builder.Append(",");
+			builder.Append(_packetLength);
+			return true;
+		}
+
+		/// <inheritdoc/>
+		public IBlobHandle GetFragmentHandle(ReadOnlySpan<byte> fragment)
+			=> new ExportHandle(this, fragment);
+
+		/// <inheritdoc/>
+		public override bool Equals(object? obj)
+			=> obj is PacketHandle other && _outer.Equals(other._outer) && _packetOffset == other._packetOffset;
+
+		/// <inheritdoc/>
+		public override int GetHashCode()
+			=> HashCode.Combine(_outer, _packetOffset);
+
+		#region Packet reader access
+
+		record struct PacketReaderCacheKey(IBlobHandle Bundle, int Offset);
+		record struct EncodedPacketCacheKey(IBlobHandle Bundle, int Offset);
+
+		async ValueTask<IRefCountedHandle<PacketReader>> GetPacketReaderAsync(CancellationToken cancellationToken = default)
+		{
+			PacketReaderCacheKey cacheKey = new PacketReaderCacheKey(_outer, _packetOffset);
+			return await _cache.FindOrAddAsync(cacheKey, CreatePacketReaderAsync, cancellationToken);
+		}
+
+		async Task<PacketReader> CreatePacketReaderAsync(PacketReaderCacheKey cacheKey, CancellationToken cancellationToken)
+		{
+			using IRefCountedHandle<IReadOnlyMemoryOwner<byte>> encodedData = await ReadEncodedPacketAsync(cancellationToken);
+			IRefCountedHandle<Packet> packet = Packet.Decode(encodedData.Target.Memory, _cache.Allocator);
+			return new PacketReader(_storageClient, _cache, _outer, this, packet.Target, packet);
+		}
+
+		async ValueTask<IRefCountedHandle<IReadOnlyMemoryOwner<byte>>> ReadEncodedPacketAsync(CancellationToken cancellationToken)
+		{
+			EncodedPacketCacheKey encodedPacketCacheKey = new EncodedPacketCacheKey(_outer, _packetOffset);
+			return await _cache.FindOrAddAsync(encodedPacketCacheKey, ReadEncodedPacketInternalAsync, cancellationToken);
+		}
+
+		async Task<IReadOnlyMemoryOwner<byte>> ReadEncodedPacketInternalAsync(EncodedPacketCacheKey key, CancellationToken cancellationToken)
+		{
+			// TODO: If this is more than one packet, slice it up into separate things and add them to the cache.
+			// Use a BundleSignature or BundleFrame
+			return await _outer.ReadBodyAsync(_packetOffset, _packetLength, cancellationToken);
+		}
+
+		#endregion
+	}
+}
diff --git a/Engine/Source/Programs/Shared/EpicGames.Horde/Storage/Bundles/V2/PacketReader.cs b/Engine/Source/Programs/Shared/EpicGames.Horde/Storage/Bundles/V2/PacketReader.cs
new file mode 100644
--- /dev/null
+++ b/Engine/Source/Programs/Shared/EpicGames.Horde/Storage/Bundles/V2/PacketReader.cs
@@ -0,0 +1,112 @@
+// Copyright Epic Games, Inc. All Rights Reserved.
+
+using System;
+using EpicGames.Core;
+
+namespace EpicGames.Horde.Storage.Bundles.V2
+{
+	/// <summary>
+	/// Utility class for constructing BlobData objects from a packet, caching any computed handles to other blobs.
+	/// </summary>
+	public sealed class PacketReader : IDisposable
+	{
+		readonly IStorageClient _storageClient;
+		readonly BundleCache _cache;
+		readonly IBlobHandle _bundleHandle;
+		readonly PacketHandle _packetHandle;
+		readonly Packet _decodedPacket;
+		readonly IRefCountedHandle? _memoryOwner;
+		readonly IBlobHandle?[] _cachedImportHandles;
+
+		/// <summary>
+		/// Constructor
+		/// </summary>
+		/// <param name="storageClient"></param>
+		/// <param name="cache"></param>
+		/// <param name="bundleHandle"></param>
+		/// <param name="packetHandle"></param>
+		/// <param name="decodedPacket">Data for the packet</param>
+		/// <param name="memoryOwner">Owner for the packet data</param>
+		public PacketReader(IStorageClient storageClient, BundleCache cache, IBlobHandle bundleHandle, PacketHandle packetHandle, Packet decodedPacket, IRefCountedHandle? memoryOwner)
+		{
+			_storageClient = storageClient;
+			_cache = cache;
+			_bundleHandle = bundleHandle;
+			_packetHandle = packetHandle;
+			_decodedPacket = decodedPacket;
+			_memoryOwner = memoryOwner;
+			_cachedImportHandles = new IBlobHandle?[_decodedPacket.GetImportCount()];
+		}
+
+		/// <inheritdoc/>
+		public void Dispose() => _memoryOwner?.Dispose();
+
+		/// <summary>
+		/// Reads an export from this packet
+		/// </summary>
+		/// <param name="exportIdx">Index of the export</param>
+		public BlobData ReadExport(int exportIdx)
+		{
+			PacketExport export = _decodedPacket.GetExport(exportIdx);
+			PacketExportHeader exportHeader = export.GetHeader();
+
+			BlobType type = _decodedPacket.GetType(exportHeader.TypeIdx);
+
+			IBlobHandle[] imports = new IBlobHandle[exportHeader.Imports.Length];
+			for (int idx = 0; idx < exportHeader.Imports.Length; idx++)
+			{
+				imports[idx] = GetImportHandle(exportHeader.Imports[idx]);
+			}
+
+			if (_memoryOwner == null)
+			{
+				return new BlobData(type, export.GetPayload(), imports);
+			}
+			else
+			{
+				return new BlobDataWithOwner(type, export.GetPayload(), imports, _memoryOwner.AddRef());
+			}
+		}
+
+		/// <summary>
+		/// Reads an export from this packet
+		/// </summary>
+		/// <param name="exportIdx">Index of the export</param>
+		public IReadOnlyMemoryOwner<byte> ReadExportBody(int exportIdx)
+		{
+			PacketExport export = _decodedPacket.GetExport(exportIdx);
+			return ReadOnlyMemoryOwner.Create(export.GetPayload(), _memoryOwner?.AddRef());
+		}
+
+		/// <summary>
+		/// Gets an import handle for the packet
+		/// </summary>
+		IBlobHandle GetImportHandle(int index)
+		{
+			IBlobHandle? importHandle = _cachedImportHandles[index];
+			if (importHandle is null)
+			{
+				PacketImport import = _decodedPacket.GetImport(index);
+
+				switch (import.BaseIdx)
+				{
+					case PacketImport.InvalidBaseIdx:
+						importHandle = _storageClient.CreateBlobHandle(new BlobLocator(import.Fragment));
+						break;
+					case PacketImport.CurrentBundleBaseIdx:
+						importHandle = new PacketHandle(_storageClient, _bundleHandle, import.Fragment, _cache);
+						break;
+					case PacketImport.CurrentPacketBaseIdx:
+						importHandle = new ExportHandle(_packetHandle, import.Fragment);
+						break;
+					default:
+						importHandle = GetImportHandle(import.BaseIdx).GetFragmentHandle(import.Fragment);
+						break;
+				}
+
+				_cachedImportHandles[index] = importHandle;
+			}
+			return importHandle;
+		}
+	}
+}
diff --git a/Engine/Source/Programs/Shared/EpicGames.Horde/Storage/Bundles/V2/PacketWriter.cs b/Engine/Source/Programs/Shared/EpicGames.Horde/Storage/Bundles/V2/PacketWriter.cs
new file mode 100644
--- /dev/null
+++ b/Engine/Source/Programs/Shared/EpicGames.Horde/Storage/Bundles/V2/PacketWriter.cs
@@ -0,0 +1,329 @@
+// Copyright Epic Games, Inc. All Rights Reserved.
+
+using System;
+using System.Buffers.Binary;
+using System.Collections.Generic;
+using System.Runtime.InteropServices;
+using EpicGames.Core;
+
+namespace EpicGames.Horde.Storage.Bundles.V2
+{
+	/// <summary>
+	/// Writes exports into a new bundle packet
+	/// </summary>
+	public sealed class PacketWriter : IDisposable
+	{
+		readonly IBlobHandle _bundleHandle;
+		readonly IBlobHandle _packetHandle;
+		readonly IMemoryAllocator<byte> _allocator;
+		readonly object _lockObject;
+
+		IRefCountedHandle<Memory<byte>> _bufferHandle;
+		Memory<byte> _buffer;
+
+		int _length;
+
+		readonly List<BlobType> _types = new List<BlobType>();
+		readonly List<PacketImport> _imports = new List<PacketImport>();
+		readonly Dictionary<IBlobHandle, int> _importMap = new Dictionary<IBlobHandle, int>();
+		readonly List<int> _exportOffsets = new List<int>();
+
+		/// <summary>
+		/// Current length of the packet
+		/// </summary>
+		public int Length => _length;
+
+		/// <summary>
+		/// Constructor
+		/// </summary>
+		public PacketWriter(IBlobHandle bundleHandle, IBlobHandle packetHandle, IMemoryAllocator<byte> allocator, object lockObject)
+		{
+			_bundleHandle = bundleHandle;
+			_packetHandle = packetHandle;
+			_allocator = allocator;
+			_lockObject = lockObject;
+
+			_importMap[_bundleHandle] = PacketImport.CurrentBundleBaseIdx;
+			_importMap[_packetHandle] = PacketImport.CurrentPacketBaseIdx;
+
+			_bufferHandle = RefCountedHandle.Create(_allocator.Alloc(1024));
+			_buffer = _bufferHandle.Target;
+
+			Bundle.WriteSignature(_buffer.Span, new BundleSignature(BundleVersion.LatestV2, 0));
+
+			_length = Bundle.SignatureLength + (sizeof(int) * 3);
+			_exportOffsets.Add(_length);
+		}
+
+		/// <inheritdoc/>
+		public void Dispose()
+		{
+			_buffer = Memory<byte>.Empty;
+
+			if (_bufferHandle != null)
+			{
+				_bufferHandle.Dispose();
+				_bufferHandle = null!;
+			}
+		}
+
+		/// <summary>
+		/// Writes a new blob to this packet
+		/// </summary>
+		/// <param name="size"></param>
+		/// <param name="type"></param>
+		/// <param name="references"></param>
+		public int CompleteExport(int size, BlobType type, IReadOnlyList<IBlobHandle> references)
+		{
+			int[] importIndices = new int[references.Count];
+			for (int idx = 0; idx < references.Count; idx++)
+			{
+				importIndices[idx] = FindOrAddImport(references[idx]);
+			}
+
+			int typeIdx = FindOrAddType(type);
+			return CompleteExport(size, typeIdx, importIndices);
+		}
+
+		/// <summary>
+		/// Complete the current export
+		/// </summary>
+		/// <param name="size">Size of data written to the export buffer</param>
+		/// <param name="typeIdx">Index of the blob type</param>
+		/// <param name="imports">Indexes of the imported blobs</param>
+		internal int CompleteExport(int size, int typeIdx, int[] imports)
+		{
+			// Finalize the written export data, inserting the length at the start.
+			BinaryPrimitives.WriteInt32LittleEndian(_buffer.Slice(_length).Span, size);
+			_length += sizeof(int) + size;
+
+			// Write the type and import list
+			Span<byte> span = GetOutputBufferInternal(0, (imports.Length + 2) * 5).Span;
+			int offset = 0;
+
+			offset += VarInt.WriteUnsigned(span[offset..], typeIdx);
+			offset += VarInt.WriteUnsigned(span[offset..], imports.Length);
+
+			for (int idx = 0; idx < imports.Length; idx++)
+			{
+				offset += VarInt.WriteUnsigned(span[offset..], imports[idx]);
+			}
+
+			_length = Align(_length + offset, 4);
+
+			// Write the next export offset to the buffer
+			_exportOffsets.Add(_length);
+			return _exportOffsets.Count - 2;
+		}
+
+		/// <summary>
+		/// Reads data for a blob written to storage
+		/// </summary>
+		/// <param name="exportIdx"></param>
+		/// <returns></returns>
+		public IReadOnlyMemoryOwner<byte> GetExportData(int exportIdx)
+		{
+			lock (_lockObject)
+			{
+				int offset = _exportOffsets[exportIdx];
+				int length = _exportOffsets[exportIdx + 1] - offset;
+
+				PacketExport export = new PacketExport(_buffer.Slice(offset, length));
+				return ReadOnlyMemoryOwner.Create(export.GetPayload(), _bufferHandle.AddRef());
+			}
+		}
+
+		/// <summary>
+		/// Add a new blob type to be written
+		/// </summary>
+		/// <param name="type">Type to add</param>
+		/// <returns>Index of the type</returns>
+		public int FindOrAddType(BlobType type)
+		{
+			int idx = _types.IndexOf(type);
+			if (idx == -1)
+			{
+				idx = _types.Count;
+				_types.Add(type);
+			}
+			return idx;
+		}
+
+		/// <summary>
+		/// Adds a new imported blob locator
+		/// </summary>
+		/// <param name="handle">Handle to add</param>
+		/// <returns>Index of the import</returns>
+		public int FindOrAddImport(IBlobHandle handle)
+		{
+			int idx;
+			if (!_importMap.TryGetValue(handle, out idx))
+			{
+				Utf8String fragment;
+				if (!handle.TryGetIdentifier(out fragment))
+				{
+					throw new InvalidOperationException("Handle does not have an existing import entry");
+				}
+
+				IBlobHandle? outer = handle.Outer;
+
+				int baseIdx;
+				if (outer != null)
+				{
+					baseIdx = FindOrAddImport(outer);
+				}
+				else
+				{
+					baseIdx = PacketImport.InvalidBaseIdx;
+				}
+
+				PacketImport import = new PacketImport(baseIdx, fragment);
+				idx = _imports.Count;
+				_imports.Add(import);
+				_importMap.Add(handle, idx);
+			}
+			return idx;
+		}
+
+		/// <summary>
+		/// Gets data to write new export
+		/// </summary>
+		/// <param name="usedSize">Size of data in the current buffer that has been written to</param>
+		/// <param name="desiredSize"></param>
+		/// <returns></returns>
+		public Memory<byte> GetOutputBuffer(int usedSize, int desiredSize)
+		{
+			int totalUsedSize = 0;
+			if (usedSize > 0)
+			{
+				totalUsedSize = sizeof(int) + usedSize;
+			}
+
+			Memory<byte> memory = GetOutputBufferInternal(totalUsedSize, sizeof(int) + desiredSize);
+			return memory[4..];
+		}
+
+		/// <summary>
+		/// Aligns an offset to a power-of-2 boundary
+		/// </summary>
+		static int Align(int value, int alignment) => value + (alignment - 1) & ~(alignment - 1);
+
+		/// <summary>
+		/// Gets data to write new export
+		/// </summary>
+		Memory<byte> GetOutputBufferInternal(int usedSize, int desiredSize)
+		{
+			if (_length + desiredSize > _buffer.Length)
+			{
+				lock (_lockObject)
+				{
+					int newSize = _length + desiredSize + 4096 + 16384 & ~16384;
+
+					IRefCountedHandle<Memory<byte>> newBufferHandle = RefCountedHandle.Create(_allocator.Alloc(newSize));
+					_buffer.Slice(0, _length + usedSize).CopyTo(newBufferHandle.Target);
+					_buffer = newBufferHandle.Target;
+
+					_bufferHandle = newBufferHandle;
+				}
+			}
+			return _buffer.Slice(_length);
+		}
+
+		/// <summary>
+		/// Gets an output span and updates the length accordingly
+		/// </summary>
+		Span<byte> GetOutputSpanAndAdvance(int length)
+		{
+			Span<byte> span = GetOutputBufferInternal(0, length).Span;
+			_length += length;
+			return span.Slice(0, length);
+		}
+
+		/// <summary>
+		/// Mark the current packet as complete
+		/// </summary>
+		public Packet CompletePacket()
+		{
+			_length = Align(_length, 4);
+
+			// Write the various lookup tables
+			WriteTypeTable();
+			WriteImportTable();
+			WriteExportTable();
+
+			// Write the final packet length
+			BinaryPrimitives.WriteInt32LittleEndian(_buffer.Span.Slice(4), _length);
+			return new Packet(_buffer.Slice(0, _length));
+		}
+
+		void WriteTypeTable()
+		{
+			// Write the type table offset in the header
+			BinaryPrimitives.WriteInt32LittleEndian(_buffer.Span.Slice(8), _length);
+
+			// Write the types to the end of the buffer
+			Span<byte> span = GetOutputSpanAndAdvance(sizeof(int) + _types.Count * BlobType.NumBytes);
+
+			BinaryPrimitives.WriteInt32LittleEndian(span, _types.Count);
+			span = span[4..];
+
+			foreach (BlobType type in _types)
+			{
+				type.Write(span);
+				span = span[IoHash.NumBytes..];
+			}
+		}
+
+		void WriteImportTable()
+		{
+			// Write the the import table offset in the header
+			BinaryPrimitives.WriteInt32LittleEndian(_buffer.Span.Slice(12), _length);
+
+			// Measure the length of all the imports
+			int importDataLength = 0;
+			foreach (PacketImport import in _imports)
+			{
+				importDataLength += VarInt.MeasureUnsigned(import.BaseIdx + PacketImport.Bias) + import.Fragment.Length;
+			}
+			importDataLength = Align(importDataLength, 4);
+
+			// Allocate the buffers
+			int importHeaderLength = (_imports.Count + 2) * sizeof(int);
+			int importDataOffset = _length + importHeaderLength;
+			Span<byte> span = GetOutputSpanAndAdvance(importHeaderLength + importDataLength);
+
+			// Write all the import data
+			BinaryPrimitives.WriteInt32LittleEndian(span, _imports.Count);
+			span = span[4..];
+
+			foreach (PacketImport import in _imports)
+			{
+				BinaryPrimitives.WriteInt32LittleEndian(span, importDataOffset);
+				span = span[4..];
+
+				Span<byte> baseSpan = _buffer.Span.Slice(importDataOffset);
+				importDataOffset += VarInt.WriteUnsigned(baseSpan, import.BaseIdx + PacketImport.Bias);
+
+				Span<byte> fragmentSpan = _buffer.Span.Slice(importDataOffset);
+				import.Fragment.Span.CopyTo(fragmentSpan);
+				importDataOffset += import.Fragment.Length;
+			}
+			BinaryPrimitives.WriteInt32LittleEndian(span, importDataOffset);
+		}
+
+		void WriteExportTable()
+		{
+			// Write the export table offset in the header
+			BinaryPrimitives.WriteInt32LittleEndian(_buffer.Span.Slice(16), _length);
+
+			// Write the exports
+			Span<int> span = MemoryMarshal.Cast<byte, int>(GetOutputSpanAndAdvance((_exportOffsets.Count + 1) * sizeof(int)));
+
+			span[0] = _exportOffsets.Count - 1;
+			for (int idx = 0; idx < _exportOffsets.Count; idx++)
+			{
+				span[idx + 1] = _exportOffsets[idx];
+			}
+		}
+	}
+}
diff --git a/Engine/Source/Programs/Shared/EpicGames.Horde/Storage/Clients/BundleStorageClient.cs b/Engine/Source/Programs/Shared/EpicGames.Horde/Storage/Clients/BundleStorageClient.cs
--- a/Engine/Source/Programs/Shared/EpicGames.Horde/Storage/Clients/BundleStorageClient.cs
+++ b/Engine/Source/Programs/Shared/EpicGames.Horde/Storage/Clients/BundleStorageClient.cs
@@ -1,12 +1,16 @@
 // Copyright Epic Games, Inc. All Rights Reserved.
 
 using System;
+using System.Buffers.Text;
 using System.Collections.Generic;
 using System.IO;
+using System.Linq;
 using System.Threading;
 using System.Threading.Tasks;
 using EpicGames.Core;
 using EpicGames.Horde.Storage.Bundles;
+using EpicGames.Horde.Storage.Bundles.V1;
+using EpicGames.Horde.Storage.Bundles.V2;
 using Microsoft.Extensions.Logging;
 
 namespace EpicGames.Horde.Storage.Clients
@@ -22,23 +26,23 @@ namespace EpicGames.Horde.Storage.Clients
 		class BundleHandle : IBlobHandle
 		{
 			readonly BundleStorageClient _storageClient;
-			readonly BlobLocator _locator;
+			readonly IBlobHandle _inner;
 			List<IBlobHandle>? _refs;
 
 			/// <inheritdoc/>
-			public IBlobHandle? Outer => null;
+			public IBlobHandle? Outer => _inner.Outer;
 
 			/// <summary>
 			/// Constructor
 			/// </summary>
-			public BundleHandle(BundleStorageClient storageClient, BlobLocator locator)
+			public BundleHandle(BundleStorageClient storageClient, IBlobHandle inner)
 			{
 				_storageClient = storageClient;
-				_locator = locator;
+				_inner = inner;
 			}
 
 			/// <inheritdoc/>
-			public ValueTask FlushAsync(CancellationToken cancellationToken = default) => default;
+			public ValueTask FlushAsync(CancellationToken cancellationToken = default) => _inner.FlushAsync(cancellationToken);
 
 			/// <inheritdoc/>
 			public ValueTask<BlobType> ReadTypeAsync(CancellationToken cancellationToken = default) => new ValueTask<BlobType>(Bundle.BlobType);
@@ -50,11 +54,19 @@ namespace EpicGames.Horde.Storage.Clients
 				{
 					List<IBlobHandle> refs = new List<IBlobHandle>();
 
-					Bundles.V1.BundleHeader header = await _storageClient.ReadHeaderAsync(_locator, cancellationToken);
+					Bundles.V1.BundleHeader header = await _storageClient.ReadHeaderAsync(_inner.GetLocator(), cancellationToken);
+					if (header.Exports.Count > 0)
+					{
 					foreach (BlobLocator import in header.Imports)
 					{
 						refs.Add(_storageClient.CreateBlobHandle(new BlobLocator(import.Path)));
 					}
+					}
+					else
+					{
+						using BlobData blobData = await ReadAsync(cancellationToken);
+						refs.AddRange(blobData.Refs);
+					}
 
 					_refs = refs;
 				}
@@ -63,37 +75,104 @@ namespace EpicGames.Horde.Storage.Clients
 
 			/// <inheritdoc/>
 			public Task<Stream> OpenBodyAsync(int offset = 0, int? length = null, CancellationToken cancellationToken = default)
-			{
-				return _storageClient.OpenBlobAsync(_locator, offset, length, cancellationToken);
-			}
+				=> _inner.OpenBodyAsync(offset, length, cancellationToken);
 
 			/// <inheritdoc/>
 			public async ValueTask<BlobData> ReadAsync(CancellationToken cancellationToken = default)
 			{
-				using (Stream stream = await _storageClient.OpenBlobAsync(_locator, 0, cancellationToken: cancellationToken))
+				IReadOnlyMemoryOwner<byte> data = await _inner.ReadBodyAsync(cancellationToken);
+
+				List<BlobLocator> importLocators = ReadImports(data.Memory).ToList();
+				List<IBlobHandle> importHandles = importLocators.ConvertAll(x => _storageClient.CreateBlobHandle(x));
+
+				return new BlobDataWithOwner(Bundle.BlobType, data.Memory, importHandles, data);
+			}
+
+			IEnumerable<BlobLocator> ReadImports(ReadOnlyMemory<byte> data)
+			{
+				BundleSignature signature = Bundle.ReadSignature(data.Span);
+				if (signature.Version <= BundleVersion.LatestV1)
+				{
+					return ReadImportsV1(data);
+				}
+				else if (signature.Version <= BundleVersion.LatestV2)
+				{
+					return ReadImportsV2(data);
+				}
+				else
 				{
-					byte[] data = await stream.ReadAllBytesAsync(cancellationToken);
-					return new BlobData(Bundle.BlobType, data, await ReadImportsAsync(cancellationToken));
+					throw new InvalidOperationException($"Unsupported bundle version {(int)signature.Version}");
 				}
 			}
 
+			static IEnumerable<BlobLocator> ReadImportsV1(ReadOnlyMemory<byte> data)
+			{
+				BundleHeader header = BundleHeader.Read(data);
+				return header.Imports.Select(x => x.BaseLocator);
+			}
+
+			IEnumerable<BlobLocator> ReadImportsV2(ReadOnlyMemory<byte> data)
+			{
+				HashSet<BlobLocator> locators = new HashSet<BlobLocator>();
+				while (data.Length > 0)
+				{
+					BundleSignature signature = Bundle.ReadSignature(data.Span);
+
+					using IRefCountedHandle<Bundles.V2.Packet> packet = Bundles.V2.Packet.Decode(data, _storageClient._cache.Allocator);
+					for (int idx = 0; idx < packet.Target.GetImportCount(); idx++)
+					{
+						PacketImport import = packet.Target.GetImport(idx);
+						if (import.BaseIdx == -1)
+						{
+							locators.Add(new BlobLocator(import.Fragment));
+						}
+					}
+
+					data = data.Slice(Bundle.SignatureLength + signature.HeaderLength);
+				}
+				return locators;
+			}
+
 			/// <inheritdoc/>
 			public bool TryAppendIdentifier(Utf8StringBuilder builder)
+				=> _inner.TryAppendIdentifier(builder);
+
+			/// <inheritdoc/>
+			public IBlobHandle GetFragmentHandle(ReadOnlySpan<byte> fragment)
 			{
-				builder.Append(_locator.Path);
-				return true;
+				int exportIdx;
+				if (Utf8Parser.TryParse(fragment, out exportIdx, out int numBytesRead) && numBytesRead == fragment.Length)
+				{
+					return new Bundles.V1.FlushedNodeHandle(_storageClient._bundleReader, _inner.GetLocator(), this, exportIdx);
+				}
+
+				int ampIdx = fragment.IndexOf((byte)'&');
+				if (ampIdx == -1)
+				{
+					return new Bundles.V2.PacketHandle(_storageClient, this, fragment, _storageClient._cache);
+				}
+				else
+				{
+					return new Bundles.V2.PacketHandle(_storageClient, this, fragment.Slice(0, ampIdx), _storageClient._cache).GetFragmentHandle(fragment.Slice(ampIdx + 1));
+				}
 			}
 
 			/// <inheritdoc/>
-			public override bool Equals(object? obj) => obj is BundleHandle other && _locator == other._locator;
+			public override bool Equals(object? obj) => obj is BundleHandle other && _inner.Equals(other._inner);
 
 			/// <inheritdoc/>
-			public override int GetHashCode() => _locator.GetHashCode();
+			public override int GetHashCode() => _inner.GetHashCode();
 		}
 
 		readonly IStorageClient _inner;
+		readonly BundleCache _cache;
 		readonly Bundles.V1.BundleReader _bundleReader;
 
+		/// <summary>
+		/// Allocator which trims the cache to keep below a maximum size
+		/// </summary>
+		public IMemoryAllocator<byte> Allocator => _cache.Allocator;
+
 		/// <inheritdoc/>
 		public bool SupportsRedirects { get; } = false;
 
@@ -103,6 +182,7 @@ namespace EpicGames.Horde.Storage.Clients
 		public BundleStorageClient(IStorageClient inner, BundleCache cache, ILogger logger)
 		{
 			_inner = inner;
+			_cache = cache;
 			_bundleReader = new Bundles.V1.BundleReader(this, cache, logger);
 		}
 
@@ -203,16 +283,48 @@ namespace EpicGames.Horde.Storage.Clients
 		{
 			if (locator.TryUnwrap(out BlobLocator baseLocator, out Utf8String fragment))
 			{
-				return new Bundles.V1.FlushedNodeHandle(_bundleReader, baseLocator, new BundleHandle(this, baseLocator), fragment.Span);
+				BundleHandle bundleHandle = new BundleHandle(this, CreateBlobHandle(baseLocator));
+
+				int exportIdx;
+				if (Utf8Parser.TryParse(fragment.Span, out exportIdx, out int numBytesRead) && numBytesRead == fragment.Length)
+				{
+					return new Bundles.V1.FlushedNodeHandle(_bundleReader, baseLocator, bundleHandle, exportIdx);
+				}
+
+				int ampIdx = fragment.IndexOf('&');
+				if (ampIdx == -1)
+				{
+					return new Bundles.V2.PacketHandle(this, bundleHandle, fragment.Span, _cache);
+				}
+				else
+				{
+					return new Bundles.V2.ExportHandle(new Bundles.V2.PacketHandle(this, bundleHandle, fragment.Slice(0, ampIdx).Span, _cache), fragment.Span.Slice(ampIdx + 1));
+				}
 			}
 			else
 			{
-				return new BundleHandle(this, locator);
+				return new BundleHandle(this, _inner.CreateBlobHandle(baseLocator));
 			}
 		}
 
 		/// <inheritdoc/>
-		public IStorageWriter CreateWriter(string? basePath = null, BundleOptions? options = null) => new Bundles.V1.BundleWriter(this, _bundleReader, basePath, options);
+		public IStorageWriter CreateWriter(string? basePath = null, BundleOptions? options = null)
+		{
+			options ??= BundleOptions.Default;
+
+			if (options.MaxVersion == BundleVersion.LatestV1)
+			{
+				return new Bundles.V1.BundleWriter(this, _bundleReader, basePath, options);
+			}
+			else if(options.MaxVersion == BundleVersion.LatestV2)
+			{
+				return new Bundles.V2.BundleWriter(this, basePath, _cache, options);
+			}
+			else
+			{
+				throw new InvalidOperationException($"Unsupported bundle version: {(int)options.MaxVersion}");
+			}
+		}
 
 		/// <inheritdoc/>
 		IStorageWriter IStorageClient.CreateWriter(string? basePath) => CreateWriter(basePath);
diff --git a/Engine/Source/Programs/Shared/EpicGames.Horde/Storage/IBlobHandle.cs b/Engine/Source/Programs/Shared/EpicGames.Horde/Storage/IBlobHandle.cs
--- a/Engine/Source/Programs/Shared/EpicGames.Horde/Storage/IBlobHandle.cs
+++ b/Engine/Source/Programs/Shared/EpicGames.Horde/Storage/IBlobHandle.cs
@@ -137,6 +137,13 @@ namespace EpicGames.Horde.Storage
 		/// <param name="builder">Builder for appending the identifier to</param>
 		/// <returns>True if an identifier was returned, false otherwise</returns>
 		bool TryAppendIdentifier(Utf8StringBuilder builder);
+
+		/// <summary>
+		/// Gets a handle to a nested blob object
+		/// </summary>
+		/// <param name="fragment">Name of the blob fragment</param>
+		public IBlobHandle GetFragmentHandle(ReadOnlySpan<byte> fragment)
+			=> throw new InvalidOperationException("Not supported for this handle type.");
 	}
 
 	/// <summary>
