Horde: Fix a few issues creating bundles from directory trees using the new storage system.

diff --git a/Engine/Source/Programs/Horde/Horde/Commands/Bundles/BundleCreate.cs b/Engine/Source/Programs/Horde/Horde/Commands/Bundles/BundleCreate.cs
--- a/Engine/Source/Programs/Horde/Horde/Commands/Bundles/BundleCreate.cs
+++ b/Engine/Source/Programs/Horde/Horde/Commands/Bundles/BundleCreate.cs
@@ -87,7 +87,7 @@ namespace Horde.Commands.Bundles
 				if (File != null)
 				{
 					logger.LogInformation("Writing {File}", File);
-					await FileReference.WriteAllTextAsync(File, nodeRef.Handle.ToString()!);
+					await FileReference.WriteAllTextAsync(File, nodeRef.Handle.GetLocator().ToString()!);
 				}
 				else
 				{
diff --git a/Engine/Source/Programs/Shared/EpicGames.Horde/Storage/Bundles/V2/BundleWriter.cs b/Engine/Source/Programs/Shared/EpicGames.Horde/Storage/Bundles/V2/BundleWriter.cs
--- a/Engine/Source/Programs/Shared/EpicGames.Horde/Storage/Bundles/V2/BundleWriter.cs
+++ b/Engine/Source/Programs/Shared/EpicGames.Horde/Storage/Bundles/V2/BundleWriter.cs
@@ -64,6 +64,8 @@ namespace EpicGames.Horde.Storage.Bundles.V2
 
 			Packet? _packet;
 			PacketWriter? _packetWriter;
+			int _packetOffset;
+			int _packetLength;
 
 			public IBlobHandle? Outer => _bundle;
 
@@ -105,19 +107,23 @@ namespace EpicGames.Horde.Storage.Bundles.V2
 				return exportHandle;
 			}
 
-			public Packet CompletePacket()
+			public void CompletePacket(BundleCompressionFormat compressionFormat, IMemoryWriter writer)
 			{
-				_packet ??= _packetWriter!.CompletePacket();
-				return _packet;
+				Debug.Assert(_packet == null);
+				_packet = _packetWriter!.CompletePacket();
+
+				_packetOffset = writer.Length;
+				_packet.Encode(compressionFormat, writer);
+				_packetLength = writer.Length - _packetOffset;
 			}
 
-			public async ValueTask SetFlushedHandleAsync(IStorageClient storageClient, PacketHandle flushedHandle, CancellationToken cancellationToken)
+			public async ValueTask CompleteBundleAsync(IStorageClient storageClient, IBlobHandle bundleHandle, BundleCache cache, CancellationToken cancellationToken)
 			{
 				Debug.Assert(_packetWriter != null);
 
 				lock (_lockObject)
 				{
-					_flushedHandle = flushedHandle;
+					_flushedHandle = new PacketHandle(storageClient, bundleHandle, _packetOffset, _packetLength, cache);
 
 					_packetWriter!.Dispose();
 					_packetWriter = null;
@@ -160,7 +166,21 @@ namespace EpicGames.Horde.Storage.Bundles.V2
 			}
 
 			public bool TryAppendIdentifier(Utf8StringBuilder builder)
-				=> _flushedHandle?.TryAppendIdentifier(builder) ?? false;
+			{
+				if (_flushedHandle != null)
+				{
+					return _flushedHandle.TryAppendIdentifier(builder);
+				}
+				else if (_packet != null)
+				{
+					PacketHandle.AppendIdentifier(builder, _packetOffset, _packetLength);
+					return true;
+				}
+				else
+				{
+					return false;
+				}
+			}
 		}
 
 		// Fragment of a bundle that needs to be written to storage.
@@ -177,6 +197,7 @@ namespace EpicGames.Horde.Storage.Bundles.V2
 #pragma warning restore CA2213
 			List<PendingPacketHandle>? _pendingPackets = new List<PendingPacketHandle>();
 			int _length;
+			RefCountedMemoryWriter _writer;
 
 			/// <inheritdoc/>
 			public IBlobHandle? Outer => null;
@@ -194,6 +215,7 @@ namespace EpicGames.Horde.Storage.Bundles.V2
 				_options = options;
 
 				_currentPacket = new PendingPacketHandle(this, cache.Allocator);
+				_writer = new RefCountedMemoryWriter(_cache.Allocator, 65536);
 			}
 
 			/// <inheritdoc/>
@@ -230,7 +252,7 @@ namespace EpicGames.Horde.Storage.Bundles.V2
 
 			void CompletePacket()
 			{
-				_currentPacket!.CompletePacket();
+				_currentPacket!.CompletePacket(_options.CompressionFormat, _writer);
 
 				_length += _currentPacket.Packet!.Length;
 				_pendingPackets!.Add(_currentPacket);
@@ -254,18 +276,6 @@ namespace EpicGames.Horde.Storage.Bundles.V2
 					return;
 				}
 
-				// Write the signature bytes
-				using RefCountedMemoryWriter memoryWriter = new RefCountedMemoryWriter(_cache.Allocator, 65536);
-
-				// Compress all the packets
-				List<int> encodedOffsets = new List<int>();
-				foreach (PendingPacketHandle pendingPacket in _pendingPackets)
-				{
-					encodedOffsets.Add(memoryWriter.Length);
-					pendingPacket.Packet!.Encode(_options.CompressionFormat, memoryWriter);
-				}
-				encodedOffsets.Add(memoryWriter.Length);
-
 				// Find all the other bundles that are referenced
 				HashSet<BlobLocator> references = new HashSet<BlobLocator>();
 				foreach (PendingPacketHandle pendingPacket in _pendingPackets)
@@ -283,7 +293,7 @@ namespace EpicGames.Horde.Storage.Bundles.V2
 				// TODO: put all the encoded packets into the cache
 
 				// Write the bundle data
-				using (ReadOnlySequenceStream stream = new ReadOnlySequenceStream(memoryWriter.AsSequence()))
+				using (ReadOnlySequenceStream stream = new ReadOnlySequenceStream(_writer.AsSequence()))
 				{
 					IBlobHandle[] referencedHandles = references.Select(x => _storageClient.CreateBlobHandle(x)).ToArray(); 
 					_flushedHandle = await _storageClient.WriteBlobAsync(Bundle.BlobType, stream, referencedHandles, _basePath, cancellationToken);
@@ -292,9 +302,7 @@ namespace EpicGames.Horde.Storage.Bundles.V2
 				// Update all the packets to point to the flushed bundle
 				for (int packetIdx = 0; packetIdx < _pendingPackets.Count; packetIdx++)
 				{
-					int packetLength = encodedOffsets[packetIdx + 1] - encodedOffsets[packetIdx];
-					PacketHandle packetHandle = new PacketHandle(_storageClient, _flushedHandle, encodedOffsets[packetIdx], packetLength, _cache);
-					await _pendingPackets[packetIdx].SetFlushedHandleAsync(_storageClient, packetHandle, cancellationToken);
+					await _pendingPackets[packetIdx].CompleteBundleAsync(_storageClient, _flushedHandle, _cache, cancellationToken);
 				}
 
 				// Clear out all the pending packets
diff --git a/Engine/Source/Programs/Shared/EpicGames.Horde/Storage/Bundles/V2/PacketHandle.cs b/Engine/Source/Programs/Shared/EpicGames.Horde/Storage/Bundles/V2/PacketHandle.cs
--- a/Engine/Source/Programs/Shared/EpicGames.Horde/Storage/Bundles/V2/PacketHandle.cs
+++ b/Engine/Source/Programs/Shared/EpicGames.Horde/Storage/Bundles/V2/PacketHandle.cs
@@ -128,13 +128,21 @@ namespace EpicGames.Horde.Storage.Bundles.V2
 		/// <inheritdoc/>
 		public bool TryAppendIdentifier(Utf8StringBuilder builder)
 		{
-			builder.Append("pkt=");
-			builder.Append(_packetOffset);
-			builder.Append(",");
-			builder.Append(_packetLength);
+			AppendIdentifier(builder, _packetOffset, _packetLength);
 			return true;
 		}
 
+		/// <summary>
+		/// Appends an identifier for a packet to the given buffer
+		/// </summary>
+		public static void AppendIdentifier(Utf8StringBuilder builder, int packetOffset, int packetLength)
+		{
+			builder.Append(s_fragmentPrefix);
+			builder.Append(packetOffset);
+			builder.Append((byte)',');
+			builder.Append(packetLength);
+		}
+
 		/// <inheritdoc/>
 		public IBlobHandle GetFragmentHandle(ReadOnlySpan<byte> fragment)
 			=> new ExportHandle(this, fragment);
diff --git a/Engine/Source/Programs/Shared/EpicGames.Horde/Storage/Clients/BundleStorageClient.cs b/Engine/Source/Programs/Shared/EpicGames.Horde/Storage/Clients/BundleStorageClient.cs
--- a/Engine/Source/Programs/Shared/EpicGames.Horde/Storage/Clients/BundleStorageClient.cs
+++ b/Engine/Source/Programs/Shared/EpicGames.Horde/Storage/Clients/BundleStorageClient.cs
@@ -124,7 +124,7 @@ namespace EpicGames.Horde.Storage.Clients
 						PacketImport import = packet.Target.GetImport(idx);
 						if (import.BaseIdx == -1)
 						{
-							locators.Add(new BlobLocator(import.Fragment));
+							locators.Add(new BlobLocator(import.Fragment.Clone()));
 						}
 					}
 
