Horde: Fix a couple of memory leaks when serializing bundles.

diff --git a/Engine/Source/Programs/Horde/Horde/Commands/Bundles/BundleCreate.cs b/Engine/Source/Programs/Horde/Horde/Commands/Bundles/BundleCreate.cs
--- a/Engine/Source/Programs/Horde/Horde/Commands/Bundles/BundleCreate.cs
+++ b/Engine/Source/Programs/Horde/Horde/Commands/Bundles/BundleCreate.cs
@@ -79,7 +79,7 @@ namespace Horde.Commands.Bundles
 			{
 				ChunkingOptions options = new ChunkingOptions();
 
-				DirectoryNode node = await DirectoryNode.CreateAsync(baseDir, files.ConvertAll(x => x.ToFileInfo()), options, writer, null, CancellationToken.None);
+				DirectoryNode node = await DirectoryNode.CreateAsync(baseDir, files.ConvertAll(x => x.ToFileInfo()), options, writer, new CopyStatsLogger(logger), CancellationToken.None);
 				HashedNodeRef<DirectoryNode> nodeRef = await writer.WriteHashedNodeAsync(node, CancellationToken.None);
 
 				await writer.FlushAsync();
diff --git a/Engine/Source/Programs/Shared/EpicGames.Horde/Storage/Bundles/V2/BundleWriter.cs b/Engine/Source/Programs/Shared/EpicGames.Horde/Storage/Bundles/V2/BundleWriter.cs
--- a/Engine/Source/Programs/Shared/EpicGames.Horde/Storage/Bundles/V2/BundleWriter.cs
+++ b/Engine/Source/Programs/Shared/EpicGames.Horde/Storage/Bundles/V2/BundleWriter.cs
@@ -194,8 +194,8 @@ namespace EpicGames.Horde.Storage.Bundles.V2
 			IBlobHandle? _flushedHandle;
 #pragma warning disable CA2213
 			PendingPacketHandle? _currentPacket;
-#pragma warning restore CA2213
 			List<PendingPacketHandle>? _pendingPackets = new List<PendingPacketHandle>();
+#pragma warning restore CA2213
 			int _length;
 			RefCountedMemoryWriter _writer;
 
@@ -234,6 +234,11 @@ namespace EpicGames.Horde.Storage.Bundles.V2
 					}
 					_pendingPackets = null;
 				}
+				if (_writer != null)
+				{
+					_writer.Dispose();
+					_writer = null!;
+				}
 			}
 
 			public Memory<byte> GetOutputBuffer(int usedSize, int desiredSize)
diff --git a/Engine/Source/Programs/Shared/EpicGames.Horde/Storage/Bundles/V2/PacketWriter.cs b/Engine/Source/Programs/Shared/EpicGames.Horde/Storage/Bundles/V2/PacketWriter.cs
--- a/Engine/Source/Programs/Shared/EpicGames.Horde/Storage/Bundles/V2/PacketWriter.cs
+++ b/Engine/Source/Programs/Shared/EpicGames.Horde/Storage/Bundles/V2/PacketWriter.cs
@@ -222,6 +222,7 @@ namespace EpicGames.Horde.Storage.Bundles.V2
 					IRefCountedHandle<Memory<byte>> newBufferHandle = RefCountedHandle.Create(_allocator.Alloc(newSize));
 					_buffer.Slice(0, _length + usedSize).CopyTo(newBufferHandle.Target);
 					_buffer = newBufferHandle.Target;
+					_bufferHandle.Dispose();
 
 					_bufferHandle = newBufferHandle;
 				}
