Horde: Null out some fields after flushing a bundle to storage to allow the referenced memory to be GC'd.

diff --git a/Engine/Source/Programs/Shared/EpicGames.Horde/Storage/Bundles/V2/BundleWriter.cs b/Engine/Source/Programs/Shared/EpicGames.Horde/Storage/Bundles/V2/BundleWriter.cs
--- a/Engine/Source/Programs/Shared/EpicGames.Horde/Storage/Bundles/V2/BundleWriter.cs
+++ b/Engine/Source/Programs/Shared/EpicGames.Horde/Storage/Bundles/V2/BundleWriter.cs
@@ -60,7 +60,7 @@ namespace EpicGames.Horde.Storage.Bundles.V2
 			readonly PendingBundleHandle _bundle;
 
 			PacketHandle? _flushedHandle;
-			readonly List<PendingExportHandle> _pendingExports = new List<PendingExportHandle>();
+			List<PendingExportHandle>? _pendingExports;
 
 			Packet? _packet;
 			PacketWriter? _packetWriter;
@@ -89,7 +89,7 @@ namespace EpicGames.Horde.Storage.Bundles.V2
 			public ValueTask FlushAsync(CancellationToken cancellationToken = default) => _bundle.FlushAsync(cancellationToken);
 
 			public bool IsEmpty() 
-				=> _pendingExports.Count == 0;
+				=> _pendingExports!.Count == 0;
 
 			public int GetLength()
 				=> _packetWriter!.Length;
@@ -102,7 +102,7 @@ namespace EpicGames.Horde.Storage.Bundles.V2
 				int exportIdx = _packetWriter!.CompleteExport(size, type, references);
 
 				PendingExportHandle exportHandle = new PendingExportHandle(this, exportIdx, type, references.ToArray(), aliases.ToArray());
-				_pendingExports.Add(exportHandle);
+				_pendingExports!.Add(exportHandle);
 
 				return exportHandle;
 			}
@@ -124,12 +124,13 @@ namespace EpicGames.Horde.Storage.Bundles.V2
 				lock (_lockObject)
 				{
 					_flushedHandle = new PacketHandle(storageClient, bundleHandle, _packetOffset, _packetLength, cache);
+					_packet = null;
 
 					_packetWriter!.Dispose();
 					_packetWriter = null;
 				}
 
-				for(int exportIdx = 0; exportIdx < _pendingExports.Count; exportIdx++)
+				for(int exportIdx = 0; exportIdx < _pendingExports!.Count; exportIdx++)
 				{
 					PendingExportHandle pendingExport = _pendingExports[exportIdx];
 					foreach (AliasInfo alias in pendingExport.Aliases)
@@ -138,6 +139,8 @@ namespace EpicGames.Horde.Storage.Bundles.V2
 						await storageClient.AddAliasAsync(alias.Name, exportHandle, alias.Rank, alias.Data, cancellationToken);
 					}
 				}
+
+				_pendingExports = null;
 			}
 
 			public ValueTask<BlobData> ReadAsync(CancellationToken cancellationToken = default)
