Horde: Set the length of explicitly allocated memory blocks to -1 on dispose, triggering a runtime exception if they're used again.

diff --git a/Engine/Source/Programs/Shared/EpicGames.Core/IMemoryAllocator.cs b/Engine/Source/Programs/Shared/EpicGames.Core/IMemoryAllocator.cs
--- a/Engine/Source/Programs/Shared/EpicGames.Core/IMemoryAllocator.cs
+++ b/Engine/Source/Programs/Shared/EpicGames.Core/IMemoryAllocator.cs
@@ -73,21 +73,21 @@ namespace EpicGames.Core
 		unsafe class Allocation : MemoryManager<byte>
 		{
 			IntPtr _handle;
-			byte* _pointer;
 			int _length;
 
 			public Allocation(int size)
 			{
 				_handle = Marshal.AllocHGlobal(size);
-				_pointer = (byte*)_handle.ToPointer();
 				_length = size;
 			}
 
+			byte* GetPointer() => (byte*)_handle.ToPointer();
+
 			/// <inheritdoc/>
-			public override Span<byte> GetSpan() => new Span<byte>(_pointer, _length);
+			public override Span<byte> GetSpan() => new Span<byte>(GetPointer(), _length);
 
 			/// <inheritdoc/>
-			public override MemoryHandle Pin(int elementIndex) => new MemoryHandle(_pointer + elementIndex);
+			public override MemoryHandle Pin(int elementIndex) => new MemoryHandle(GetPointer() + elementIndex);
 
 			/// <inheritdoc/>
 			public override void Unpin() { }
@@ -101,8 +101,7 @@ namespace EpicGames.Core
 					_handle = IntPtr.Zero;
 				}
 
-				_pointer = null;
-				_length = 0;
+				_length = -1;
 			}
 		}
 
