Horde: Reduce objects retained by flushed storage handles.

* Pending aliases are now stored on the bundle rather than on individual exports.
* Several fields are nulled out after flushing to allow referenced objects to be GC'd.

diff --git a/Engine/Source/Programs/Shared/EpicGames.Horde/Storage/Bundles/V2/BundleWriter.cs b/Engine/Source/Programs/Shared/EpicGames.Horde/Storage/Bundles/V2/BundleWriter.cs
--- a/Engine/Source/Programs/Shared/EpicGames.Horde/Storage/Bundles/V2/BundleWriter.cs
+++ b/Engine/Source/Programs/Shared/EpicGames.Horde/Storage/Bundles/V2/BundleWriter.cs
@@ -23,19 +23,15 @@ namespace EpicGames.Horde.Storage.Bundles.V2
 
 			readonly BlobType _type;
 			readonly IBlobHandle[] _imports;
-			readonly AliasInfo[] _aliases;
 		
 			public IBlobHandle? Outer => _packet;
 
-			public IReadOnlyList<AliasInfo> Aliases => _aliases;
-
-			public PendingExportHandle(PendingPacketHandle packet, int exportIdx, BlobType type, IBlobHandle[] imports, AliasInfo[] aliases)
+			public PendingExportHandle(PendingPacketHandle packet, int exportIdx, BlobType type, IBlobHandle[] imports)
 			{
 				_packet = packet;
 				_exportIdx = exportIdx;
 				_type = type;
 				_imports = imports;
-				_aliases = aliases;
 			}
 
 			public ValueTask FlushAsync(CancellationToken cancellationToken = default) => _packet.FlushAsync(cancellationToken);
@@ -78,18 +74,26 @@ namespace EpicGames.Horde.Storage.Bundles.V2
 			}
 
 			public void Dispose()
+			{
+				ReleaseResources();
+			}
+
+			void ReleaseResources()
 			{
 				if (_packetWriter != null)
 				{
 					_packetWriter.Dispose();
 					_packetWriter = null;
 				}
+
+				_pendingExports = null;
+				_packet = null;
 			}
 
 			public ValueTask FlushAsync(CancellationToken cancellationToken = default) => _bundle.FlushAsync(cancellationToken);
 
-			public bool IsEmpty() 
-				=> _pendingExports!.Count == 0;
+			public bool IsEmpty()
+				=> _pendingExports == null;
 
 			public int GetLength()
 				=> _packetWriter!.Length;
@@ -97,12 +101,13 @@ namespace EpicGames.Horde.Storage.Bundles.V2
 			public Memory<byte> GetOutputBuffer(int usedSize, int desiredSize)
 				=> _packetWriter!.GetOutputBuffer(usedSize, desiredSize);
 
-			public PendingExportHandle CompleteExport(BlobType type, int size, IReadOnlyList<IBlobHandle> references, IReadOnlyList<AliasInfo> aliases)
+			public PendingExportHandle CompleteExport(BlobType type, int size, IReadOnlyList<IBlobHandle> references)
 			{
 				int exportIdx = _packetWriter!.CompleteExport(size, type, references);
 
-				PendingExportHandle exportHandle = new PendingExportHandle(this, exportIdx, type, references.ToArray(), aliases.ToArray());
-				_pendingExports!.Add(exportHandle);
+				PendingExportHandle exportHandle = new PendingExportHandle(this, exportIdx, type, references.ToArray());
+				_pendingExports ??= new List<PendingExportHandle>();
+				_pendingExports.Add(exportHandle);
 
 				return exportHandle;
 			}
@@ -117,30 +122,14 @@ namespace EpicGames.Horde.Storage.Bundles.V2
 				_packetLength = writer.Length - _packetOffset;
 			}
 
-			public async ValueTask CompleteBundleAsync(IStorageClient storageClient, IBlobHandle bundleHandle, BundleCache cache, CancellationToken cancellationToken)
+			public void CompleteBundle(IStorageClient storageClient, IBlobHandle bundleHandle, BundleCache cache, CancellationToken cancellationToken)
 			{
 				Debug.Assert(_packetWriter != null);
-
 				lock (_lockObject)
 				{
 					_flushedHandle = new PacketHandle(storageClient, bundleHandle, _packetOffset, _packetLength, cache);
-					_packet = null;
-
-					_packetWriter!.Dispose();
-					_packetWriter = null;
-				}
-
-				for(int exportIdx = 0; exportIdx < _pendingExports!.Count; exportIdx++)
-				{
-					PendingExportHandle pendingExport = _pendingExports[exportIdx];
-					foreach (AliasInfo alias in pendingExport.Aliases)
-					{
-						ExportHandle exportHandle = new ExportHandle(_flushedHandle, exportIdx);
-						await storageClient.AddAliasAsync(alias.Name, exportHandle, alias.Rank, alias.Data, cancellationToken);
-					}
+					ReleaseResources();
 				}
-
-				_pendingExports = null;
 			}
 
 			public ValueTask<BlobData> ReadAsync(CancellationToken cancellationToken = default)
@@ -164,7 +153,6 @@ namespace EpicGames.Horde.Storage.Bundles.V2
 						return _packetWriter.GetExportData(exportIdx);
 					}
 				}
-
 				return await _flushedHandle!.ReadExportBodyAsync(exportIdx, cancellationToken);
 			}
 
@@ -198,9 +186,10 @@ namespace EpicGames.Horde.Storage.Bundles.V2
 #pragma warning disable CA2213
 			PendingPacketHandle? _currentPacket;
 			List<PendingPacketHandle>? _pendingPackets = new List<PendingPacketHandle>();
+			List<(PendingExportHandle, AliasInfo)>? _pendingExportAliases;
 #pragma warning restore CA2213
 			int _length;
-			RefCountedMemoryWriter _writer;
+			RefCountedMemoryWriter? _writer;
 
 			/// <inheritdoc/>
 			public IBlobHandle? Outer => null;
@@ -222,7 +211,9 @@ namespace EpicGames.Horde.Storage.Bundles.V2
 			}
 
 			/// <inheritdoc/>
-			public void Dispose()
+			public void Dispose() => ReleaseResources();
+
+			void ReleaseResources()
 			{
 				if (_currentPacket != null)
 				{
@@ -240,8 +231,12 @@ namespace EpicGames.Horde.Storage.Bundles.V2
 				if (_writer != null)
 				{
 					_writer.Dispose();
-					_writer = null!;
+					_writer = null;
 				}
+
+				// Also clear out any arrays that can be GC'd
+				_pendingPackets = null;
+				_pendingExportAliases = null;
 			}
 
 			public Memory<byte> GetOutputBuffer(int usedSize, int desiredSize)
@@ -249,7 +244,12 @@ namespace EpicGames.Horde.Storage.Bundles.V2
 
 			public PendingExportHandle CompleteExport(BlobType type, int size, IReadOnlyList<IBlobHandle> references, IReadOnlyList<AliasInfo> aliases)
 			{
-				PendingExportHandle exportHandle = _currentPacket!.CompleteExport(type, size, references, aliases);
+				PendingExportHandle exportHandle = _currentPacket!.CompleteExport(type, size, references);
+				if (aliases.Count > 0)
+				{
+					_pendingExportAliases ??= new List<(PendingExportHandle, AliasInfo)>();
+					_pendingExportAliases.AddRange(aliases.Select(x => (exportHandle, x)));
+				}
 				if (_currentPacket.GetLength() > Math.Min(_options.MinCompressionPacketSize, _options.MaxBlobSize))
 				{
 					CompletePacket();
@@ -260,6 +260,8 @@ namespace EpicGames.Horde.Storage.Bundles.V2
 
 			void CompletePacket()
 			{
+				Debug.Assert(_writer != null);
+
 				_currentPacket!.CompletePacket(_options.CompressionFormat, _writer);
 
 				_length += _currentPacket.Packet!.Length;
@@ -284,6 +286,8 @@ namespace EpicGames.Horde.Storage.Bundles.V2
 					return;
 				}
 
+				Debug.Assert(_writer != null);
+
 				// Find all the other bundles that are referenced
 				HashSet<BlobLocator> references = new HashSet<BlobLocator>();
 				foreach (PendingPacketHandle pendingPacket in _pendingPackets)
@@ -308,13 +312,19 @@ namespace EpicGames.Horde.Storage.Bundles.V2
 				}
 
 				// Update all the packets to point to the flushed bundle
-				for (int packetIdx = 0; packetIdx < _pendingPackets.Count; packetIdx++)
+				foreach(PendingPacketHandle pendingPacket in _pendingPackets)
 				{
-					await _pendingPackets[packetIdx].CompleteBundleAsync(_storageClient, _flushedHandle, _cache, cancellationToken);
+					pendingPacket.CompleteBundle(_storageClient, _flushedHandle, _cache, cancellationToken);
 				}
 
-				// Clear out all the pending packets
-				_pendingPackets = null;
+				// Add all the aliases
+				if (_pendingExportAliases != null)
+				{
+					foreach ((PendingExportHandle exportHandle, AliasInfo aliasInfo) in _pendingExportAliases)
+					{
+						await _storageClient.AddAliasAsync(aliasInfo.Name, exportHandle, aliasInfo.Rank, aliasInfo.Data, cancellationToken);
+					}
+				}
 			}
 
 			/// <inheritdoc/>
