Horde: Allow decoding exports from partially written packets, reducing the amount of memory retained for each handle.

diff --git a/Engine/Source/Programs/Shared/EpicGames.Horde/Storage/Bundles/V2/BundleWriter.cs b/Engine/Source/Programs/Shared/EpicGames.Horde/Storage/Bundles/V2/BundleWriter.cs
--- a/Engine/Source/Programs/Shared/EpicGames.Horde/Storage/Bundles/V2/BundleWriter.cs
+++ b/Engine/Source/Programs/Shared/EpicGames.Horde/Storage/Bundles/V2/BundleWriter.cs
@@ -21,26 +21,19 @@ namespace EpicGames.Horde.Storage.Bundles.V2
 			readonly PendingPacketHandle _packet;
 			readonly int _exportIdx;
 
-			readonly BlobType _type;
-			readonly IBlobHandle[] _imports;
-		
 			public IBlobHandle? Outer => _packet;
 
-			public PendingExportHandle(PendingPacketHandle packet, int exportIdx, BlobType type, IBlobHandle[] imports)
+			public PendingExportHandle(PendingPacketHandle packet, int exportIdx)
 			{
 				_packet = packet;
 				_exportIdx = exportIdx;
-				_type = type;
-				_imports = imports;
 			}
 
-			public ValueTask FlushAsync(CancellationToken cancellationToken = default) => _packet.FlushAsync(cancellationToken);
+			public ValueTask FlushAsync(CancellationToken cancellationToken = default) 
+				=> _packet.FlushAsync(cancellationToken);
 
-			public async ValueTask<BlobData> ReadAsync(CancellationToken cancellationToken = default)
-			{
-				IReadOnlyMemoryOwner<byte> body = await _packet.ReadExportBodyAsync(_exportIdx, cancellationToken);
-				return new BlobDataWithOwner(_type, body.Memory, _imports, body);
-			}
+			public ValueTask<BlobData> ReadAsync(CancellationToken cancellationToken = default)
+				=> _packet.ReadExportAsync(_exportIdx, cancellationToken);
 
 			public bool TryAppendIdentifier(Utf8StringBuilder builder)
 			{
@@ -105,7 +98,7 @@ namespace EpicGames.Horde.Storage.Bundles.V2
 			{
 				int exportIdx = _packetWriter!.CompleteExport(size, type, references);
 
-				PendingExportHandle exportHandle = new PendingExportHandle(this, exportIdx, type, references.ToArray());
+				PendingExportHandle exportHandle = new PendingExportHandle(this, exportIdx);
 				_pendingExports ??= new List<PendingExportHandle>();
 				_pendingExports.Add(exportHandle);
 
@@ -144,16 +137,16 @@ namespace EpicGames.Horde.Storage.Bundles.V2
 				return _flushedHandle!.ReadAsync(cancellationToken);
 			}
 
-			public async ValueTask<IReadOnlyMemoryOwner<byte>> ReadExportBodyAsync(int exportIdx, CancellationToken cancellationToken = default)
+			public async ValueTask<BlobData> ReadExportAsync(int exportIdx, CancellationToken cancellationToken = default)
 			{
 				lock (_lockObject)
 				{
 					if (_packetWriter != null)
 					{
-						return _packetWriter.GetExportData(exportIdx);
+						return _packetWriter.GetExport(exportIdx);
 					}
 				}
-				return await _flushedHandle!.ReadExportBodyAsync(exportIdx, cancellationToken);
+				return await _flushedHandle!.ReadExportAsync(exportIdx, cancellationToken);
 			}
 
 			public bool TryAppendIdentifier(Utf8StringBuilder builder)
diff --git a/Engine/Source/Programs/Shared/EpicGames.Horde/Storage/Bundles/V2/PacketWriter.cs b/Engine/Source/Programs/Shared/EpicGames.Horde/Storage/Bundles/V2/PacketWriter.cs
--- a/Engine/Source/Programs/Shared/EpicGames.Horde/Storage/Bundles/V2/PacketWriter.cs
+++ b/Engine/Source/Programs/Shared/EpicGames.Horde/Storage/Bundles/V2/PacketWriter.cs
@@ -3,6 +3,7 @@
 using System;
 using System.Buffers.Binary;
 using System.Collections.Generic;
+using System.Diagnostics;
 using System.Runtime.InteropServices;
 using EpicGames.Core;
 
@@ -25,6 +26,7 @@ namespace EpicGames.Horde.Storage.Bundles.V2
 
 		readonly List<BlobType> _types = new List<BlobType>();
 		readonly List<PacketImport> _imports = new List<PacketImport>();
+		readonly List<IBlobHandle> _importHandles = new List<IBlobHandle>();
 		readonly Dictionary<IBlobHandle, int> _importMap = new Dictionary<IBlobHandle, int>();
 		readonly List<int> _exportOffsets = new List<int>();
 
@@ -43,6 +45,14 @@ namespace EpicGames.Horde.Storage.Bundles.V2
 			_allocator = allocator;
 			_lockObject = lockObject;
 
+			for (int idx = 0; idx < PacketImport.Bias; idx++)
+			{
+				_importHandles.Add(null!);
+			}
+
+			_importHandles[PacketImport.CurrentBundleBaseIdx + PacketImport.Bias] = _bundleHandle;
+			_importHandles[PacketImport.CurrentPacketBaseIdx + PacketImport.Bias] = _packetHandle;
+
 			_importMap[_bundleHandle] = PacketImport.CurrentBundleBaseIdx;
 			_importMap[_packetHandle] = PacketImport.CurrentPacketBaseIdx;
 
@@ -121,7 +131,7 @@ namespace EpicGames.Horde.Storage.Bundles.V2
 		/// </summary>
 		/// <param name="exportIdx"></param>
 		/// <returns></returns>
-		public IReadOnlyMemoryOwner<byte> GetExportData(int exportIdx)
+		public BlobData GetExport(int exportIdx)
 		{
 			lock (_lockObject)
 			{
@@ -129,7 +139,19 @@ namespace EpicGames.Horde.Storage.Bundles.V2
 				int length = _exportOffsets[exportIdx + 1] - offset;
 
 				PacketExport export = new PacketExport(_buffer.Slice(offset, length));
-				return ReadOnlyMemoryOwner.Create(export.GetPayload(), _bufferHandle.AddRef());
+				PacketExportHeader header = export.GetHeader();
+
+				BlobType type = _types[header.TypeIdx];
+
+				IBlobHandle[] imports = new IBlobHandle[header.Imports.Length];
+				for (int idx = 0; idx < header.Imports.Length; idx++)
+				{
+					int importIdx = header.Imports[idx];
+					imports[idx] = _importHandles[importIdx + PacketImport.Bias];
+				}
+
+				IReadOnlyMemoryOwner<byte> body = ReadOnlyMemoryOwner.Create(export.GetPayload(), _bufferHandle.AddRef());
+				return new BlobDataWithOwner(_types[header.TypeIdx], body.Memory, imports, body);
 			}
 		}
 
@@ -181,6 +203,9 @@ namespace EpicGames.Horde.Storage.Bundles.V2
 				idx = _imports.Count;
 				_imports.Add(import);
 				_importMap.Add(handle, idx);
+
+				Debug.Assert(idx + PacketImport.Bias == _importHandles.Count);
+				_importHandles.Add(handle);
 			}
 			return idx;
 		}
