Horde: Remove allocation of empty object used for locking access to PendingExportHandle.

diff --git a/Engine/Source/Programs/Shared/EpicGames.Horde/Storage/Bundles/V2/BundleWriter.cs b/Engine/Source/Programs/Shared/EpicGames.Horde/Storage/Bundles/V2/BundleWriter.cs
--- a/Engine/Source/Programs/Shared/EpicGames.Horde/Storage/Bundles/V2/BundleWriter.cs
+++ b/Engine/Source/Programs/Shared/EpicGames.Horde/Storage/Bundles/V2/BundleWriter.cs
@@ -16,7 +16,7 @@ namespace EpicGames.Horde.Storage.Bundles.V2
 	public sealed class BundleWriter : IStorageWriter
 	{
 		// An export that has been written and is waiting to be flushed to disk
-		class PendingExportHandle : IBlobHandle
+		sealed class PendingExportHandle : IBlobHandle
 		{
 			readonly PendingPacketHandle _packet;
 			readonly int _exportIdx;
@@ -43,9 +43,8 @@ namespace EpicGames.Horde.Storage.Bundles.V2
 		}
 
 		// Packet that is still being built, but may be redirected to a flushed packet
-		class PendingPacketHandle : IBlobHandle, IDisposable
+		sealed class PendingPacketHandle : IBlobHandle, IDisposable
 		{
-			readonly object _lockObject = new object();
 			readonly PendingBundleHandle _bundle;
 
 			PacketHandle? _flushedHandle;
@@ -63,7 +62,7 @@ namespace EpicGames.Horde.Storage.Bundles.V2
 			public PendingPacketHandle(PendingBundleHandle bundle, IMemoryAllocator<byte> allocator)
 			{
 				_bundle = bundle;
-				_packetWriter = new PacketWriter(_bundle, this, allocator, _lockObject);
+				_packetWriter = new PacketWriter(_bundle, this, allocator, _bundle);
 			}
 
 			public void Dispose()
@@ -83,7 +82,8 @@ namespace EpicGames.Horde.Storage.Bundles.V2
 				_packet = null;
 			}
 
-			public ValueTask FlushAsync(CancellationToken cancellationToken = default) => _bundle.FlushAsync(cancellationToken);
+			public ValueTask FlushAsync(CancellationToken cancellationToken = default) 
+				=> _bundle.FlushAsync(cancellationToken);
 
 			public bool IsEmpty()
 				=> _pendingExports == null;
@@ -115,10 +115,10 @@ namespace EpicGames.Horde.Storage.Bundles.V2
 				_packetLength = writer.Length - _packetOffset;
 			}
 
-			public void CompleteBundle(IStorageClient storageClient, IBlobHandle bundleHandle, BundleCache cache, CancellationToken cancellationToken)
+			public void CompleteBundle(IStorageClient storageClient, IBlobHandle bundleHandle, BundleCache cache)
 			{
 				Debug.Assert(_packetWriter != null);
-				lock (_lockObject)
+				lock (_bundle)
 				{
 					_flushedHandle = new PacketHandle(storageClient, bundleHandle, _packetOffset, _packetLength, cache);
 					ReleaseResources();
@@ -127,7 +127,7 @@ namespace EpicGames.Horde.Storage.Bundles.V2
 
 			public ValueTask<BlobData> ReadAsync(CancellationToken cancellationToken = default)
 			{
-				lock (_lockObject)
+				lock (_bundle)
 				{
 					if (_packetWriter != null)
 					{
@@ -139,7 +139,7 @@ namespace EpicGames.Horde.Storage.Bundles.V2
 
 			public async ValueTask<BlobData> ReadExportAsync(int exportIdx, CancellationToken cancellationToken = default)
 			{
-				lock (_lockObject)
+				lock (_bundle)
 				{
 					if (_packetWriter != null)
 					{
@@ -168,7 +168,7 @@ namespace EpicGames.Horde.Storage.Bundles.V2
 		}
 
 		// Fragment of a bundle that needs to be written to storage.
-		class PendingBundleHandle : IBlobHandle, IDisposable
+		sealed class PendingBundleHandle : IBlobHandle, IDisposable
 		{
 			readonly IStorageClient _storageClient;
 			readonly string? _basePath;
@@ -307,7 +307,7 @@ namespace EpicGames.Horde.Storage.Bundles.V2
 				// Update all the packets to point to the flushed bundle
 				foreach(PendingPacketHandle pendingPacket in _pendingPackets)
 				{
-					pendingPacket.CompleteBundle(_storageClient, _flushedHandle, _cache, cancellationToken);
+					pendingPacket.CompleteBundle(_storageClient, _flushedHandle, _cache);
 				}
 
 				// Add all the aliases
