Horde: Move management of temporary buffers for writing bundle packets into the pending buffer handle, since only one may be active at a time, and use the encoded packet data for reading any exports once it exists - allowing the write buffer to be freed sooner.

diff --git a/Engine/Source/Programs/Shared/EpicGames.Core/MemoryWriter.cs b/Engine/Source/Programs/Shared/EpicGames.Core/MemoryWriter.cs
--- a/Engine/Source/Programs/Shared/EpicGames.Core/MemoryWriter.cs
+++ b/Engine/Source/Programs/Shared/EpicGames.Core/MemoryWriter.cs
@@ -620,22 +620,28 @@ namespace EpicGames.Core
 		/// <summary>
 		/// Returns a contiguous block of memory containing the written data
 		/// </summary>
-		public IRefCountedHandle<ReadOnlyMemory<byte>> AsRefCountedMemory()
+		public IRefCountedHandle<ReadOnlyMemory<byte>> AsRefCountedMemory() => AsRefCountedMemory(0, Length);
+
+		/// <summary>
+		/// Returns a contiguous block of memory containing the written data
+		/// </summary>
+		public IRefCountedHandle<ReadOnlyMemory<byte>> AsRefCountedMemory(int offset) => AsRefCountedMemory(offset, Length - offset);
+
+		/// <summary>
+		/// Returns a contiguous block of memory containing the written data
+		/// </summary>
+		public IRefCountedHandle<ReadOnlyMemory<byte>> AsRefCountedMemory(int offset, int length)
 		{
-			if (Chunks.Count == 1)
+			IRefCountedHandle<ReadOnlySequence<byte>> sequence = AsRefCountedSequence(offset, length);
+			if (sequence.Target.IsSingleSegment)
 			{
-				BufferChunk buffer = (BufferChunk)Chunks[0];
-				return RefCountedHandle.Create<ReadOnlyMemory<byte>>(buffer.WrittenMemory, buffer.Handle.AddRef());
+				return RefCountedHandle.Create(sequence.Target.First, sequence);
 			}
-			else
-			{
-				int length = Length;
 
-				IRefCountedHandle<Memory<byte>> allocation = RefCountedHandle.Create(_allocator.Alloc(length));
-				CopyTo(allocation.Target.Span);
+			IRefCountedHandle<Memory<byte>> allocation = RefCountedHandle.Create(_allocator.Alloc(length));
+			sequence.Target.CopyTo(allocation.Target.Span);
 
-				return RefCountedHandle.Create<ReadOnlyMemory<byte>>(allocation.Target.Slice(0, length), allocation);
-			}
+			return RefCountedHandle.Create<ReadOnlyMemory<byte>>(allocation.Target.Slice(0, length), allocation);
 		}
 
 		/// <summary>
diff --git a/Engine/Source/Programs/Shared/EpicGames.Horde.Tests/StorageClientTests.cs b/Engine/Source/Programs/Shared/EpicGames.Horde.Tests/StorageClientTests.cs
--- a/Engine/Source/Programs/Shared/EpicGames.Horde.Tests/StorageClientTests.cs
+++ b/Engine/Source/Programs/Shared/EpicGames.Horde.Tests/StorageClientTests.cs
@@ -1,8 +1,11 @@
 // Copyright Epic Games, Inc. All Rights Reserved.
 
+using System;
 using System.Threading.Tasks;
 using EpicGames.Core;
 using EpicGames.Horde.Storage;
+using EpicGames.Horde.Storage.Bundles;
+using EpicGames.Horde.Storage.Bundles.V2;
 using EpicGames.Horde.Storage.Clients;
 using Microsoft.Extensions.Logging.Abstractions;
 using Microsoft.VisualStudio.TestTools.UnitTesting;
@@ -16,6 +19,7 @@ namespace EpicGames.Horde.Tests
 		class TestNode : Node
 		{
 			public int Value { get; }
+			public byte[] Padding { get; set; } = Array.Empty<byte>();
 			public NodeRef<TestNode>[] Refs { get; }
 
 			public TestNode(int value, params NodeRef<TestNode>[] refs)
@@ -27,12 +31,14 @@ namespace EpicGames.Horde.Tests
 			public TestNode(IBlobReader reader)
 			{
 				Value = reader.ReadInt32();
+				Padding = reader.ReadVariableLengthBytes().ToArray();
 				Refs = reader.ReadVariableLengthArray(() => reader.ReadNodeRef<TestNode>());
 			}
 
 			public override void Serialize(IBlobWriter writer)
 			{
 				writer.WriteInt32(Value);
+				writer.WriteVariableLengthBytes(Padding);
 				writer.WriteVariableLengthArray(Refs, x => writer.WriteNodeRef(x));
 			}
 		}
@@ -58,6 +64,19 @@ namespace EpicGames.Horde.Tests
 			await TestBasicAsync(storeV2);
 		}
 
+		static async Task TestBasicAsync(IStorageClient store)
+		{
+			NodeRef nodeRef;
+			await using (IStorageWriter writer = store.CreateWriter())
+			{
+				nodeRef = await writer.WriteHashedNodeAsync(new TestNode(123));
+			}
+			await store.WriteRefAsync("hello", nodeRef.Handle);
+
+			TestNode output = await store.ReadRefAsync<TestNode>("hello");
+			Assert.AreEqual(123, output.Value);
+		}
+
 		[TestMethod]
 		public async Task PendingRefsAsync()
 		{
@@ -83,17 +102,56 @@ namespace EpicGames.Horde.Tests
 			Assert.AreEqual(123, output1.Value);
 		}
 
-		static async Task TestBasicAsync(IStorageClient store)
+		[TestMethod]
+		public async Task PacketFlushAsync()
 		{
-			NodeRef nodeRef;
-			await using (IStorageWriter writer = store.CreateWriter())
-			{
-				nodeRef = await writer.WriteHashedNodeAsync(new TestNode(123));
-			}
-			await store.WriteRefAsync("hello", nodeRef.Handle);
+			await using BundleCache cache = new BundleCache();
 
-			TestNode output = await store.ReadRefAsync<TestNode>("hello");
-			Assert.AreEqual(123, output.Value);
+			using MemoryStorageClient memoryStore = new MemoryStorageClient();
+			using BundleStorageClient store = new BundleStorageClient(memoryStore, cache, NullLogger.Instance);
+
+			await using IStorageWriter writer = store.CreateWriter(options: new BundleOptions { MinCompressionPacketSize = 100, MaxBlobSize = 1024 * 1024 });
+			NodeRef<TestNode> nodeRef1 = await writer.WriteNodeAsync(new TestNode(123) { Padding = new byte[1024] });
+			await writer.FlushAsync();
+			NodeRef<TestNode> nodeRef2 = await writer.WriteNodeAsync(new TestNode(456, nodeRef1) { Padding = new byte[1024] });
+			NodeRef<TestNode> nodeRef3 = await writer.WriteNodeAsync(new TestNode(789, nodeRef2));
+
+			// nodeRef1 is in a flushed bundle
+			BundleWriter.PendingExportHandle export1 = (BundleWriter.PendingExportHandle)nodeRef1.Handle;
+			BundleWriter.PendingPacketHandle packet1 = (BundleWriter.PendingPacketHandle)export1.Outer!;
+			BundleWriter.PendingBundleHandle bundle1 = (BundleWriter.PendingBundleHandle)packet1.Outer!;
+			Assert.IsNotNull(packet1.FlushedHandle);
+			Assert.IsNotNull(bundle1.FlushedHandle);
+
+			TestNode node1 = await export1.ReadNodeAsync<TestNode>();
+			Assert.AreEqual(123, node1.Value);
+
+			// nodeRef2 is in a flushed packet, unflushed bundle
+			BundleWriter.PendingExportHandle export2 = (BundleWriter.PendingExportHandle)nodeRef2.Handle;
+
+			BundleWriter.PendingPacketHandle packet2 = (BundleWriter.PendingPacketHandle)export2.Outer!;
+			Assert.AreNotEqual(packet1, packet2);
+			Assert.IsNotNull(packet2.FlushedHandle);
+
+			BundleWriter.PendingBundleHandle bundle2 = (BundleWriter.PendingBundleHandle)packet2.Outer!;
+			Assert.AreNotEqual(bundle1, bundle2);
+			Assert.IsNull(bundle2.FlushedHandle);
+
+			TestNode node2 = await export2.ReadNodeAsync<TestNode>();
+			Assert.AreEqual(456, node2.Value);
+
+			// nodeRef3 is in an unflushed packet, unflushed bundle
+			BundleWriter.PendingExportHandle export3 = (BundleWriter.PendingExportHandle)nodeRef3.Handle;
+
+			BundleWriter.PendingPacketHandle packet3 = (BundleWriter.PendingPacketHandle)export3.Outer!;
+			Assert.AreNotEqual(packet2, packet3);
+			Assert.IsNull(packet3.FlushedHandle);
+
+			BundleWriter.PendingBundleHandle bundle3 = (BundleWriter.PendingBundleHandle)packet3.Outer!;
+			Assert.AreEqual(bundle2, bundle3);
+
+			TestNode node3 = await export3.ReadNodeAsync<TestNode>();
+			Assert.AreEqual(789, node3.Value);
 		}
 	}
 }
diff --git a/Engine/Source/Programs/Shared/EpicGames.Horde/Storage/Bundles/V2/BundleWriter.cs b/Engine/Source/Programs/Shared/EpicGames.Horde/Storage/Bundles/V2/BundleWriter.cs
--- a/Engine/Source/Programs/Shared/EpicGames.Horde/Storage/Bundles/V2/BundleWriter.cs
+++ b/Engine/Source/Programs/Shared/EpicGames.Horde/Storage/Bundles/V2/BundleWriter.cs
@@ -16,7 +16,7 @@ namespace EpicGames.Horde.Storage.Bundles.V2
 	public sealed class BundleWriter : IStorageWriter
 	{
 		// An export that has been written and is waiting to be flushed to disk
-		sealed class PendingExportHandle : IBlobHandle
+		internal sealed class PendingExportHandle : IBlobHandle
 		{
 			readonly PendingPacketHandle _packet;
 			readonly int _exportIdx;
@@ -43,93 +43,27 @@ namespace EpicGames.Horde.Storage.Bundles.V2
 		}
 
 		// Packet that is still being built, but may be redirected to a flushed packet
-		sealed class PendingPacketHandle : IBlobHandle, IDisposable
+		internal sealed class PendingPacketHandle : IBlobHandle
 		{
 			readonly PendingBundleHandle _bundle;
-
 			PacketHandle? _flushedHandle;
-			List<PendingExportHandle>? _pendingExports;
-
-			Packet? _packet;
-			PacketWriter? _packetWriter;
-			int _packetOffset;
-			int _packetLength;
-
-			public IBlobHandle? Outer => _bundle;
 
-			public Packet? Packet => _packet;
+			public IBlobHandle? Outer => _flushedHandle?.Outer ?? _bundle;
+			public IBlobHandle? FlushedHandle => _flushedHandle;
 
-			public PendingPacketHandle(PendingBundleHandle bundle, IMemoryAllocator<byte> allocator)
-			{
-				_bundle = bundle;
-				_packetWriter = new PacketWriter(_bundle, this, allocator, _bundle);
-			}
-
-			public void Dispose()
-			{
-				ReleaseResources();
-			}
-
-			void ReleaseResources()
-			{
-				if (_packetWriter != null)
-				{
-					_packetWriter.Dispose();
-					_packetWriter = null;
-				}
-
-				_pendingExports = null;
-				_packet = null;
-			}
+			public PendingPacketHandle(PendingBundleHandle bundle) => _bundle = bundle;
 
 			public ValueTask FlushAsync(CancellationToken cancellationToken = default) 
 				=> _bundle.FlushAsync(cancellationToken);
 
-			public bool IsEmpty()
-				=> _pendingExports == null;
-
-			public int GetLength()
-				=> _packetWriter!.Length;
-
-			public Memory<byte> GetOutputBuffer(int usedSize, int desiredSize)
-				=> _packetWriter!.GetOutputBuffer(usedSize, desiredSize);
-
-			public PendingExportHandle CompleteExport(BlobType type, int size, IReadOnlyList<IBlobHandle> references)
-			{
-				int exportIdx = _packetWriter!.CompleteExport(size, type, references);
-
-				PendingExportHandle exportHandle = new PendingExportHandle(this, exportIdx);
-				_pendingExports ??= new List<PendingExportHandle>();
-				_pendingExports.Add(exportHandle);
-
-				return exportHandle;
-			}
-
-			public void CompletePacket(BundleCompressionFormat compressionFormat, IMemoryWriter writer)
-			{
-				Debug.Assert(_packet == null);
-				_packet = _packetWriter!.CompletePacket();
-
-				_packetOffset = writer.Length;
-				_packet.Encode(compressionFormat, writer);
-				_packetLength = writer.Length - _packetOffset;
-			}
-
-			public void CompleteBundle(IStorageClient storageClient, IBlobHandle bundleHandle, BundleCache cache)
-			{
-				Debug.Assert(_packetWriter != null);
-				lock (_bundle)
-				{
-					_flushedHandle = new PacketHandle(storageClient, bundleHandle, _packetOffset, _packetLength, cache);
-					ReleaseResources();
-				}
-			}
+			public void CompletePacket(PacketHandle flushedHandle)
+				=> _flushedHandle = flushedHandle;
 
 			public ValueTask<BlobData> ReadAsync(CancellationToken cancellationToken = default)
 			{
-				lock (_bundle)
+				lock (_bundle.LockObject)
 				{
-					if (_packetWriter != null)
+					if (_flushedHandle == null)
 					{
 						throw new NotSupportedException("Reading pending packets is not currently supported.");
 					}
@@ -139,58 +73,50 @@ namespace EpicGames.Horde.Storage.Bundles.V2
 
 			public async ValueTask<BlobData> ReadExportAsync(int exportIdx, CancellationToken cancellationToken = default)
 			{
-				lock (_bundle)
+				lock (_bundle.LockObject)
 				{
-					if (_packetWriter != null)
+					if (_flushedHandle == null)
 					{
-						return _packetWriter.GetExport(exportIdx);
+						return _bundle.GetPendingExport(exportIdx);
 					}
 				}
 				return await _flushedHandle!.ReadExportAsync(exportIdx, cancellationToken);
 			}
 
 			public bool TryAppendIdentifier(Utf8StringBuilder builder)
-			{
-				if (_flushedHandle != null)
-				{
-					return _flushedHandle.TryAppendIdentifier(builder);
-				}
-				else if (_packet != null)
-				{
-					PacketHandle.AppendIdentifier(builder, _packetOffset, _packetLength);
-					return true;
-				}
-				else
-				{
-					return false;
-				}
-			}
+				=> _flushedHandle?.TryAppendIdentifier(builder) ?? false;
 		}
 
 		// Fragment of a bundle that needs to be written to storage.
-		sealed class PendingBundleHandle : IBlobHandle, IDisposable
+		internal sealed class PendingBundleHandle : IBlobHandle, IDisposable
 		{
+			readonly object _lockObject = new object();
+
 			readonly IStorageClient _storageClient;
 			readonly string? _basePath;
 			readonly BundleCache _cache;
 			readonly BundleOptions _options;
 
 			IBlobHandle? _flushedHandle;
-#pragma warning disable CA2213
-			PendingPacketHandle? _currentPacket;
-			List<PendingPacketHandle>? _pendingPackets = new List<PendingPacketHandle>();
+
+			PacketWriter? _packetWriter;
+			PendingPacketHandle? _packetHandle;
+			List<IBlobHandle>? _bundleReferences;
 			List<(PendingExportHandle, AliasInfo)>? _pendingExportAliases;
-#pragma warning restore CA2213
-			int _length;
-			RefCountedMemoryWriter? _writer;
+			RefCountedMemoryWriter? _encodedPacketWriter;
+
+			/// <summary>
+			/// Object used for locking access to this bundle's state
+			/// </summary>
+			public object LockObject => _lockObject;
 
-			/// <inheritdoc/>
 			public IBlobHandle? Outer => null;
+			public IBlobHandle? FlushedHandle => _flushedHandle;
 
 			/// <summary>
-			/// Uncompressed length of this bundle
+			/// Compressed length of this bundle
 			/// </summary>
-			public int Length => _length;
+			public int Length => _encodedPacketWriter?.Length ?? throw new InvalidOperationException("Bundle has been flushed");
 
 			public PendingBundleHandle(IStorageClient storageClient, string? basePath, BundleCache cache, BundleOptions options)
 			{
@@ -199,117 +125,149 @@ namespace EpicGames.Horde.Storage.Bundles.V2
 				_cache = cache;
 				_options = options;
 
-				_currentPacket = new PendingPacketHandle(this, cache.Allocator);
-				_writer = new RefCountedMemoryWriter(_cache.Allocator, 65536);
+				_bundleReferences = new List<IBlobHandle>();
+				_encodedPacketWriter = new RefCountedMemoryWriter(_cache.Allocator, 65536);
+
+				StartPacket();
 			}
 
-			/// <inheritdoc/>
 			public void Dispose() => ReleaseResources();
 
 			void ReleaseResources()
 			{
-				if (_currentPacket != null)
-				{
-					_currentPacket.Dispose();
-					_currentPacket = null;
-				}
-				if (_pendingPackets != null)
+				if (_packetWriter != null)
 				{
-					foreach (PendingPacketHandle pendingPacket in _pendingPackets)
-					{
-						pendingPacket.Dispose();
-					}
-					_pendingPackets = null;
+					_packetWriter.Dispose();
+					_packetWriter = null;
 				}
-				if (_writer != null)
+				if (_encodedPacketWriter != null)
 				{
-					_writer.Dispose();
-					_writer = null;
+					_encodedPacketWriter.Dispose();
+					_encodedPacketWriter = null;
 				}
 
 				// Also clear out any arrays that can be GC'd
-				_pendingPackets = null;
+				_packetHandle = null;
+				_bundleReferences = null;
 				_pendingExportAliases = null;
 			}
 
 			public Memory<byte> GetOutputBuffer(int usedSize, int desiredSize)
-				=> _currentPacket!.GetOutputBuffer(usedSize, desiredSize);
+			{
+				Debug.Assert(_packetWriter != null);
+				return _packetWriter!.GetOutputBuffer(usedSize, desiredSize);
+			}
+
+			public BlobData GetPendingExport(int exportIdx)
+			{
+				Debug.Assert(_packetWriter != null);
+				return _packetWriter.GetExport(exportIdx);
+			}
 
 			public PendingExportHandle CompleteExport(BlobType type, int size, IReadOnlyList<IBlobHandle> references, IReadOnlyList<AliasInfo> aliases)
 			{
-				PendingExportHandle exportHandle = _currentPacket!.CompleteExport(type, size, references);
+				Debug.Assert(_packetWriter != null);
+				Debug.Assert(_packetHandle != null);
+
+				int exportIdx = _packetWriter.CompleteExport(size, type, references);
+				PendingExportHandle exportHandle = new PendingExportHandle(_packetHandle, exportIdx);
+
 				if (aliases.Count > 0)
 				{
 					_pendingExportAliases ??= new List<(PendingExportHandle, AliasInfo)>();
 					_pendingExportAliases.AddRange(aliases.Select(x => (exportHandle, x)));
 				}
-				if (_currentPacket.GetLength() > Math.Min(_options.MinCompressionPacketSize, _options.MaxBlobSize))
+
+				if (_packetWriter.Length > Math.Min(_options.MinCompressionPacketSize, _options.MaxBlobSize))
 				{
-					CompletePacket();
-					_currentPacket = new PendingPacketHandle(this, _cache.Allocator);
+					FinishPacket();
+					StartPacket();
 				}
+
 				return exportHandle;
 			}
 
-			void CompletePacket()
+			void StartPacket()
+			{
+				Debug.Assert(_packetHandle == null);
+				Debug.Assert(_packetWriter == null);
+
+				_packetHandle = new PendingPacketHandle(this);
+				_packetWriter = new PacketWriter(this, _packetHandle, _cache.Allocator, _lockObject);
+			}
+
+			void FinishPacket()
 			{
-				Debug.Assert(_writer != null);
+				Debug.Assert(_packetHandle != null);
+				Debug.Assert(_packetWriter != null);
+				Debug.Assert(_bundleReferences != null);
+				Debug.Assert(_encodedPacketWriter != null);
 
-				_currentPacket!.CompletePacket(_options.CompressionFormat, _writer);
+				if (_packetWriter.GetExportCount() > 0)
+				{
+					int packetOffset = _encodedPacketWriter.Length;
+					Packet packet = _packetWriter.CompletePacket();
+					packet.Encode(_options.CompressionFormat, _encodedPacketWriter);
+					int packetLength = _encodedPacketWriter.Length - packetOffset;
 
-				_length += _currentPacket.Packet!.Length;
-				_pendingPackets!.Add(_currentPacket);
+					// Point the packet handle to the encoded data
+					lock (_lockObject)
+					{
+						PacketHandle flushedPacketHandle = new PacketHandle(_storageClient, this, packetOffset, packetLength, _cache);
+						_packetHandle.CompletePacket(flushedPacketHandle);
+					}
 
-				_currentPacket = null;
+					// Find all the other bundles that are referenced
+					for (int importIdx = 0; importIdx < packet.GetImportCount(); importIdx++)
+					{
+						PacketImport import = packet.GetImport(importIdx);
+						if (import.BaseIdx == -1)
+						{
+							_bundleReferences.Add(_packetWriter.GetImport(importIdx));
+						}
+					}
+				}
+
+				_packetWriter.Dispose();
+				_packetWriter = null;
+
+				_packetHandle = null;
 			}
 
 			// Write this bundle to storage
 			public async ValueTask FlushAsync(CancellationToken cancellationToken = default)
 			{
-				if (_currentPacket == null)
-				{
-					return;
-				}
-				if (!_currentPacket!.IsEmpty())
-				{
-					CompletePacket();
-				}
-				if (_pendingPackets!.Count == 0)
+				// Check we haven't already flushed this bundle
+				if (_encodedPacketWriter == null)
 				{
 					return;
 				}
 
-				Debug.Assert(_writer != null);
+				FinishPacket();
 
-				// Find all the other bundles that are referenced
-				HashSet<BlobLocator> references = new HashSet<BlobLocator>();
-				foreach (PendingPacketHandle pendingPacket in _pendingPackets)
+				if (_encodedPacketWriter.Length == 0)
 				{
-					for (int importIdx = 0; importIdx < pendingPacket.Packet!.GetImportCount(); importIdx++)
-					{
-						PacketImport import = pendingPacket.Packet.GetImport(importIdx);
-						if (import.BaseIdx == -1)
-						{
-							references.Add(new BlobLocator(import.Fragment));
-						}
-					}
+					return;
 				}
 
-				// TODO: put all the encoded packets into the cache
+				Debug.Assert(_bundleReferences != null);
 
 				// Write the bundle data
-				using (ReadOnlySequenceStream stream = new ReadOnlySequenceStream(_writer.AsSequence()))
+				IBlobHandle flushedHandle;
+				using (ReadOnlySequenceStream stream = new ReadOnlySequenceStream(_encodedPacketWriter.AsSequence()))
 				{
-					IBlobHandle[] referencedHandles = references.Select(x => _storageClient.CreateBlobHandle(x)).ToArray(); 
-					_flushedHandle = await _storageClient.WriteBlobAsync(Bundle.BlobType, stream, referencedHandles, _basePath, cancellationToken);
+					flushedHandle = await _storageClient.WriteBlobAsync(Bundle.BlobType, stream, _bundleReferences, _basePath, cancellationToken);
 				}
 
-				// Update all the packets to point to the flushed bundle
-				foreach(PendingPacketHandle pendingPacket in _pendingPackets)
+				// Release all the intermediate data
+				lock (_lockObject)
 				{
-					pendingPacket.CompleteBundle(_storageClient, _flushedHandle, _cache);
+					_flushedHandle = flushedHandle;
+					ReleaseResources();
 				}
 
+				// TODO: put all the encoded packets into the cache using the final handles
+
 				// Add all the aliases
 				if (_pendingExportAliases != null)
 				{
@@ -325,8 +283,25 @@ namespace EpicGames.Horde.Storage.Bundles.V2
 				=> GetFlushedHandle().ReadAsync(cancellationToken);
 
 			/// <inheritdoc/>
-			public ValueTask<IReadOnlyMemoryOwner<byte>> ReadBodyAsync(int offset, int? length, CancellationToken cancellationToken = default)
-				=> GetFlushedHandle().ReadBodyAsync(offset, length, cancellationToken);
+			public async ValueTask<IReadOnlyMemoryOwner<byte>> ReadBodyAsync(int offset, int? length, CancellationToken cancellationToken = default)
+			{
+				if (_flushedHandle == null)
+				{
+					lock (_lockObject)
+					{
+						if (_flushedHandle == null)
+						{
+							Debug.Assert(_encodedPacketWriter != null);
+							int fetchLength = length ?? (_encodedPacketWriter.Length - offset);
+
+							IRefCountedHandle<ReadOnlyMemory<byte>> handle = _encodedPacketWriter.AsRefCountedMemory(offset, fetchLength);
+							return ReadOnlyMemoryOwner.Create(handle.Target, handle);
+						}
+					}
+				}
+
+				return await _flushedHandle.ReadBodyAsync(offset, length, cancellationToken);
+			}
 
 			/// <inheritdoc/>
 			public bool TryAppendIdentifier(Utf8StringBuilder builder)
diff --git a/Engine/Source/Programs/Shared/EpicGames.Horde/Storage/Bundles/V2/PacketWriter.cs b/Engine/Source/Programs/Shared/EpicGames.Horde/Storage/Bundles/V2/PacketWriter.cs
--- a/Engine/Source/Programs/Shared/EpicGames.Horde/Storage/Bundles/V2/PacketWriter.cs
+++ b/Engine/Source/Programs/Shared/EpicGames.Horde/Storage/Bundles/V2/PacketWriter.cs
@@ -77,6 +77,11 @@ namespace EpicGames.Horde.Storage.Bundles.V2
 			}
 		}
 
+		/// <summary>
+		/// Gets the number of exports currently in this writer
+		/// </summary>
+		public int GetExportCount() => _exportOffsets.Count - 1;
+
 		/// <summary>
 		/// Writes a new blob to this packet
 		/// </summary>
@@ -210,6 +215,12 @@ namespace EpicGames.Horde.Storage.Bundles.V2
 			return idx;
 		}
 
+		/// <summary>
+		/// Gets the import assigned to a particular index
+		/// </summary>
+		public IBlobHandle GetImport(int importIdx)
+			=> _importHandles[importIdx + PacketImport.Bias];
+
 		/// <summary>
 		/// Gets data to write new export
 		/// </summary>
