Horde: Fix memory leak when constructing ReadOnlyMemory from RefCountedMemoryWriter.

diff --git a/Engine/Source/Programs/Shared/EpicGames.Core/MemoryWriter.cs b/Engine/Source/Programs/Shared/EpicGames.Core/MemoryWriter.cs
--- a/Engine/Source/Programs/Shared/EpicGames.Core/MemoryWriter.cs
+++ b/Engine/Source/Programs/Shared/EpicGames.Core/MemoryWriter.cs
@@ -638,10 +638,16 @@ namespace EpicGames.Core
 				return RefCountedHandle.Create(sequence.Target.First, sequence);
 			}
 
-			IRefCountedHandle<Memory<byte>> allocation = RefCountedHandle.Create(_allocator.Alloc(length));
-			sequence.Target.CopyTo(allocation.Target.Span);
-
-			return RefCountedHandle.Create<ReadOnlyMemory<byte>>(allocation.Target.Slice(0, length), allocation);
+			try
+			{
+				IRefCountedHandle<Memory<byte>> allocation = RefCountedHandle.Create(_allocator.Alloc(length));
+				sequence.Target.CopyTo(allocation.Target.Span);
+				return RefCountedHandle.Create<ReadOnlyMemory<byte>>(allocation.Target.Slice(0, length), allocation);
+			}
+			finally
+			{
+				sequence.Dispose();
+			}
 		}
 
 		/// <summary>
