Fix warning about missing working directoryu in Horde AgentMessageHandler.

#rnx

diff --git a/Engine/Source/Programs/Shared/EpicGames.Horde/Compute/AgentMessageHandler.cs b/Engine/Source/Programs/Shared/EpicGames.Horde/Compute/AgentMessageHandler.cs
--- a/Engine/Source/Programs/Shared/EpicGames.Horde/Compute/AgentMessageHandler.cs
+++ b/Engine/Source/Programs/Shared/EpicGames.Horde/Compute/AgentMessageHandler.cs
@@ -362,7 +362,7 @@ namespace EpicGames.Horde.Compute
 					_logger.LogWarning("Executable {Path} does not exist", resolvedExecutable);	
 				}
 				
-				if (!File.Exists(resolvedWorkingDir))
+				if (!Directory.Exists(resolvedWorkingDir))
 				{
 					_logger.LogWarning("Working dir {Path} does not exist", resolvedWorkingDir);	
 				}
