Horde: Force v2 bundle version when running new packet tests.

diff --git a/Engine/Source/Programs/Shared/EpicGames.Horde.Tests/StorageClientTests.cs b/Engine/Source/Programs/Shared/EpicGames.Horde.Tests/StorageClientTests.cs
--- a/Engine/Source/Programs/Shared/EpicGames.Horde.Tests/StorageClientTests.cs
+++ b/Engine/Source/Programs/Shared/EpicGames.Horde.Tests/StorageClientTests.cs
@@ -110,7 +110,7 @@ namespace EpicGames.Horde.Tests
 			using MemoryStorageClient memoryStore = new MemoryStorageClient();
 			using BundleStorageClient store = new BundleStorageClient(memoryStore, cache, NullLogger.Instance);
 
-			await using IStorageWriter writer = store.CreateWriter(options: new BundleOptions { MinCompressionPacketSize = 100, MaxBlobSize = 1024 * 1024 });
+			await using IStorageWriter writer = store.CreateWriter(options: new BundleOptions { MinCompressionPacketSize = 100, MaxBlobSize = 1024 * 1024, MaxVersion = BundleVersion.LatestV2 });
 			NodeRef<TestNode> nodeRef1 = await writer.WriteNodeAsync(new TestNode(123) { Padding = new byte[1024] });
 			await writer.FlushAsync();
 			NodeRef<TestNode> nodeRef2 = await writer.WriteNodeAsync(new TestNode(456, nodeRef1) { Padding = new byte[1024] });
