Horde: Partition files to chunk into bundles into 2gb batches, in order to allow intermediates to be garbage collected between each batch.

diff --git a/Engine/Source/Programs/Horde/Horde/Commands/Bundles/BundleCreate.cs b/Engine/Source/Programs/Horde/Horde/Commands/Bundles/BundleCreate.cs
--- a/Engine/Source/Programs/Horde/Horde/Commands/Bundles/BundleCreate.cs
+++ b/Engine/Source/Programs/Horde/Horde/Commands/Bundles/BundleCreate.cs
@@ -79,7 +79,10 @@ namespace Horde.Commands.Bundles
 			{
 				ChunkingOptions options = new ChunkingOptions();
 
-				DirectoryNode node = await DirectoryNode.CreateAsync(baseDir, files.ConvertAll(x => x.ToFileInfo()), options, writer, new CopyStatsLogger(logger), CancellationToken.None);
+				List<FileInfo> fileInfos = files.ConvertAll(x => x.ToFileInfo());
+				CopyStatsLoggerWithTotals copyStatsLogger = new CopyStatsLoggerWithTotals(files.Count, fileInfos.Sum(x => x.Length), logger);
+
+				DirectoryNode node = await DirectoryNode.CreateAsync(baseDir, files.ConvertAll(x => x.ToFileInfo()), options, writer, copyStatsLogger, CancellationToken.None);
 				HashedNodeRef<DirectoryNode> nodeRef = await writer.WriteHashedNodeAsync(node, CancellationToken.None);
 
 				await writer.FlushAsync();
diff --git a/Engine/Source/Programs/Shared/EpicGames.Horde/Storage/Nodes/DirectoryNode.cs b/Engine/Source/Programs/Shared/EpicGames.Horde/Storage/Nodes/DirectoryNode.cs
--- a/Engine/Source/Programs/Shared/EpicGames.Horde/Storage/Nodes/DirectoryNode.cs
+++ b/Engine/Source/Programs/Shared/EpicGames.Horde/Storage/Nodes/DirectoryNode.cs
@@ -25,27 +25,17 @@ namespace EpicGames.Horde.Storage.Nodes
 		/// <summary>
 		/// Number of files that have been copied
 		/// </summary>
-		int CopiedCount { get; }
+		int Count { get; }
 
 		/// <summary>
 		/// Total size of data to be copied
 		/// </summary>
-		long CopiedSize { get; }
+		long Size { get; }
 
 		/// <summary>
-		/// Total number of files to copy
+		/// Processing speed, in bytes per second
 		/// </summary>
-		int TotalCount { get; }
-
-		/// <summary>
-		/// Total size of data to copy
-		/// </summary>
-		long TotalSize { get; }
-
-		/// <summary>
-		/// Download speed, in bytes per second
-		/// </summary>
-		double DownloadSpeed { get; }
+		double Rate { get; }
 	}
 
 	/// <summary>
@@ -58,16 +48,12 @@ namespace EpicGames.Horde.Storage.Nodes
 		readonly IProgress<ICopyStats> _progress;
 		long _lastTotalSize;
 
-		public int CopiedCount { get; set; }
-		public long CopiedSize { get; set; }
-		public int TotalCount { get; }
-		public long TotalSize { get; }
-		public double DownloadSpeed { get; set; }
+		public int Count { get; set; }
+		public long Size { get; set; }
+		public double Rate { get; set; }
 
-		public CopyStats(int totalCount, long totalSize, IProgress<ICopyStats> progress)
+		public CopyStats(IProgress<ICopyStats> progress)
 		{
-			TotalCount = totalCount;
-			TotalSize = totalSize;
 			_progress = progress;
 		}
 
@@ -75,12 +61,12 @@ namespace EpicGames.Horde.Storage.Nodes
 		{
 			lock (_lockObject)
 			{
-				CopiedCount += count;
-				CopiedSize += size;
-				if (_timer.Elapsed > TimeSpan.FromSeconds(10.0) || CopiedCount == count || CopiedSize == TotalSize)
+				Count += count;
+				Size += size;
+				if (_timer.Elapsed > TimeSpan.FromSeconds(10.0) || Count == count)
 				{
-					DownloadSpeed = (CopiedSize - _lastTotalSize) / _timer.Elapsed.TotalSeconds;
-					_lastTotalSize = CopiedSize;
+					Rate = (Size - _lastTotalSize) / _timer.Elapsed.TotalSeconds;
+					_lastTotalSize = Size;
 
 					_progress.Report(this);
 					_timer.Restart();
@@ -108,20 +94,36 @@ namespace EpicGames.Horde.Storage.Nodes
 		/// <summary>
 		/// Constructor
 		/// </summary>
-		public CopyStatsLogger(ILogger logger) => _logger = logger;
+		public CopyStatsLogger(ILogger logger) 
+			=> _logger = logger;
 
 		/// <inheritdoc/>
 		public void Report(ICopyStats stats)
+			=> _logger.LogInformation("Copied {NumFiles:n0} files ({Size:n1}mb, {Rate:n1}mb/s)", stats.Count, stats.Size / (1024.0 * 1024.0), stats.Rate / (1024.0 * 1024.0));
+	}
+
+	/// <summary>
+	/// Progress logger for writing copy stats
+	/// </summary>
+	public class CopyStatsLoggerWithTotals : IProgress<ICopyStats>
+	{
+		readonly int _totalCount;
+		readonly long _totalSize;
+		readonly ILogger _logger;
+
+		/// <summary>
+		/// Constructor
+		/// </summary>
+		public CopyStatsLoggerWithTotals(int totalCount, long totalSize, ILogger logger)
 		{
-			if (stats.TotalCount == 0)
-			{
-				_logger.LogInformation("Copied {NumFiles:n0} files ({Size:n1}/{TotalSize:n1}mb, {Rate:n1}mb/s, {Pct}%)", stats.CopiedCount, stats.CopiedSize / (1024.0 * 1024.0), stats.TotalSize / (1024.0 * 1024.0), stats.DownloadSpeed / (1024.0 * 1024.0), (int)((Math.Max(stats.CopiedSize, 1) * 100) / Math.Max(stats.TotalSize, 1)));
-			}
-			else
-			{
-				_logger.LogInformation("Copied {NumFiles:n0}/{TotalFiles:n0} files ({Size:n1}/{TotalSize:n1}mb, {Rate:n1}mb/s, {Pct}%)", stats.CopiedCount, stats.TotalCount, stats.CopiedSize / (1024.0 * 1024.0), stats.TotalSize / (1024.0 * 1024.0), stats.DownloadSpeed / (1024.0 * 1024.0), (int)((Math.Max(stats.CopiedSize, 1) * 100) / Math.Max(stats.TotalSize, 1)));
-			}
+			_totalCount = totalCount;
+			_totalSize = totalSize;
+			_logger = logger;
 		}
+
+		/// <inheritdoc/>
+		public void Report(ICopyStats stats)
+			=> _logger.LogInformation("Copied {NumFiles:n0}/{TotalFiles:n0} files ({Size:n1}/{TotalSize:n1}mb, {Rate:n1}mb/s, {Pct}%)", stats.Count, _totalCount, stats.Size / (1024.0 * 1024.0), _totalSize / (1024.0 * 1024.0), stats.Rate / (1024.0 * 1024.0), (int)((Math.Max(stats.Size, 1) * 100) / Math.Max(_totalSize, 1)));
 	}
 
 	/// <summary>
@@ -521,13 +523,13 @@ namespace EpicGames.Horde.Storage.Nodes
 			return directoryNode;
 		}
 
-		/// <inheritdoc cref="AddFilesAsync(DirectoryReference, IReadOnlyList{FileInfo}, ChunkingOptions, IStorageWriter, IProgress{ICopyStats}?, CancellationToken)"/>
+		/// <inheritdoc cref="AddFilesAsync(DirectoryReference, IEnumerable{FileInfo}, ChunkingOptions, IStorageWriter, IProgress{ICopyStats}?, CancellationToken)"/>
 		public Task AddFilesAsync(DirectoryReference baseDir, IEnumerable<FileReference> files, ChunkingOptions options, IStorageWriter writer, IProgress<ICopyStats>? progress, CancellationToken cancellationToken)
 		{
 			return AddFilesAsync(baseDir, files.Select(x => x.ToFileInfo()).ToList(), options, writer, progress, cancellationToken);
 		}
 
-		/// <inheritdoc cref="AddFilesAsync(DirectoryReference, IReadOnlyList{FileInfo}, ChunkingOptions, IStorageWriter, IProgress{ICopyStats}?, CancellationToken)"/>
+		/// <inheritdoc cref="AddFilesAsync(DirectoryReference, IEnumerable{FileInfo}, ChunkingOptions, IStorageWriter, IProgress{ICopyStats}?, CancellationToken)"/>
 		public Task AddFilesAsync(DirectoryInfo baseDir, IEnumerable<FileInfo> files, ChunkingOptions options, IStorageWriter writer, IProgress<ICopyStats>? progress, CancellationToken cancellationToken)
 		{
 			return AddFilesAsync(new DirectoryReference(baseDir), files.ToList(), options, writer, progress, cancellationToken);
@@ -542,55 +544,70 @@ namespace EpicGames.Horde.Storage.Nodes
 		/// <param name="writer">Writer for new node data</param>
 		/// <param name="progress">Feedback interface for progress updates</param>
 		/// <param name="cancellationToken">Cancellation token for the operation</param>
-		public async Task AddFilesAsync(DirectoryReference baseDir, IReadOnlyList<FileInfo> files, ChunkingOptions options, IStorageWriter writer, IProgress<ICopyStats>? progress, CancellationToken cancellationToken)
+		public async Task AddFilesAsync(DirectoryReference baseDir, IEnumerable<FileInfo> files, ChunkingOptions options, IStorageWriter writer, IProgress<ICopyStats>? progress, CancellationToken cancellationToken)
 		{
-			// Get the total size of all the files
-			long totalSize = 0;
-			foreach (FileInfo file in files)
-			{
-				totalSize += file.Length;
-			}
-
 			CopyStats? copyStats = null;
 			if (progress != null)
 			{
-				copyStats = new CopyStats(files.Count, totalSize, progress);
+				copyStats = new CopyStats(progress);
 			}
 
-			// Partition them into blocks for parallel writers to process asynchronously
-			List<(int Start, int Count)> partitions = ComputePartitions(files, totalSize);
-			LeafChunkedData[] leafChunkedFiles = new LeafChunkedData[files.Count];
-			await Parallel.ForEachAsync(partitions, cancellationToken, (filePartition, ctx) => CreateLeafChunkNodesAsync(writer, files, leafChunkedFiles, filePartition.Start, filePartition.Count, copyStats, options, cancellationToken));
-
-			// Create interior nodes for all the leaf chunks
-			ChunkedData[] chunkedFiles = await CreateInteriorChunkNodesAsync(leafChunkedFiles, options.InteriorOptions, writer, cancellationToken);
-
-			// Write all the interior nodes and generate the directory update
 			DirectoryUpdate update = new DirectoryUpdate();
-			for(int idx = 0; idx < files.Count; idx++)
-			{
-				FileInfo file = files[idx];
 
-				FileEntryFlags flags = FileEntryFlags.None;
-				if (RuntimeInformation.IsOSPlatform(OSPlatform.Linux))
+			// Process the input sequence in 2gb batches
+			const long MaxBatchSize = 2 * 1024 * 1024 * 1024L;
+			using (IEnumerator<FileInfo> fileEnumerator = files.GetEnumerator())
+			{
+				List<FileInfo> batch = new List<FileInfo>();
+				for (bool moreData = fileEnumerator.MoveNext(); moreData; )
 				{
-					int mode = FileUtils.GetFileMode_Linux(file.FullName);
-					if ((mode & ((1 << 0) | (1 << 3) | (1 << 6))) != 0)
+					batch.Clear();
+
+					// Take the next batch of files
+					long batchSize = 0;
+					while (batchSize < MaxBatchSize && moreData)
 					{
-						flags |= FileEntryFlags.Executable;
+						FileInfo file = fileEnumerator.Current;
+						batch.Add(file);
+						batchSize += file.Length;
+						moreData = fileEnumerator.MoveNext();
 					}
-				}
-				else if (RuntimeInformation.IsOSPlatform(OSPlatform.OSX))
-				{
-					int mode = FileUtils.GetFileMode_Mac(file.FullName);
-					if ((mode & ((1 << 0) | (1 << 3) | (1 << 6))) != 0)
+
+					// Partition them up into parallel writers
+					List<(int Start, int Count)> partitions = ComputePartitions(batch, batchSize);
+					LeafChunkedData[] leafChunkedFiles = new LeafChunkedData[batch.Count];
+					await Parallel.ForEachAsync(partitions, cancellationToken, (filePartition, ctx) => CreateLeafChunkNodesAsync(writer, batch, leafChunkedFiles, filePartition.Start, filePartition.Count, copyStats, options, cancellationToken));
+
+					// Create interior nodes for all the leaf chunks
+					ChunkedData[] chunkedFiles = await CreateInteriorChunkNodesAsync(leafChunkedFiles, options.InteriorOptions, writer, cancellationToken);
+
+					// Write all the interior nodes and generate the directory update
+					for (int idx = 0; idx < batch.Count; idx++)
 					{
-						flags |= FileEntryFlags.Executable;
+						FileInfo file = batch[idx];
+
+						FileEntryFlags flags = FileEntryFlags.None;
+						if (RuntimeInformation.IsOSPlatform(OSPlatform.Linux))
+						{
+							int mode = FileUtils.GetFileMode_Linux(file.FullName);
+							if ((mode & ((1 << 0) | (1 << 3) | (1 << 6))) != 0)
+							{
+								flags |= FileEntryFlags.Executable;
+							}
+						}
+						else if (RuntimeInformation.IsOSPlatform(OSPlatform.OSX))
+						{
+							int mode = FileUtils.GetFileMode_Mac(file.FullName);
+							if ((mode & ((1 << 0) | (1 << 3) | (1 << 6))) != 0)
+							{
+								flags |= FileEntryFlags.Executable;
+							}
+						}
+
+						FileEntry entry = new FileEntry(file.Name, flags, file.Length, chunkedFiles[idx]);
+						update.AddFile(new FileReference(file).MakeRelativeTo(baseDir), entry);
 					}
 				}
-
-				FileEntry entry = new FileEntry(file.Name, flags, file.Length, chunkedFiles[idx]);
-				update.AddFile(new FileReference(file).MakeRelativeTo(baseDir), entry);
 			}
 
 			// Add all the new entries to the tree
@@ -600,7 +617,7 @@ namespace EpicGames.Horde.Storage.Nodes
 		static List<(int Index, int Count)> ComputePartitions(IReadOnlyList<FileInfo> files, long totalSize)
 		{
 			// Maximum number of streams to write in parallel
-			const int MaxPartitions = 8;
+			const int MaxPartitions = 16;
 
 			// Minimum size of output payload for each writer
 			const long MinSizePerPartition = 1024 * 1024;
@@ -786,7 +803,7 @@ namespace EpicGames.Horde.Storage.Nodes
 			CopyStats? copyStats = null;
 			if (progress != null)
 			{
-				copyStats = new CopyStats(0, Length, progress);
+				copyStats = new CopyStats(progress);
 			}
 
 			List<Task> tasks = new List<Task>();
