Horde: Fix missing parentheses from alignment expressions.

diff --git a/Engine/Source/Programs/Shared/EpicGames.Horde/Storage/Bundles/V2/PacketWriter.cs b/Engine/Source/Programs/Shared/EpicGames.Horde/Storage/Bundles/V2/PacketWriter.cs
--- a/Engine/Source/Programs/Shared/EpicGames.Horde/Storage/Bundles/V2/PacketWriter.cs
+++ b/Engine/Source/Programs/Shared/EpicGames.Horde/Storage/Bundles/V2/PacketWriter.cs
@@ -242,7 +242,7 @@ namespace EpicGames.Horde.Storage.Bundles.V2
 		/// <summary>
 		/// Aligns an offset to a power-of-2 boundary
 		/// </summary>
-		static int Align(int value, int alignment) => value + (alignment - 1) & ~(alignment - 1);
+		static int Align(int value, int alignment) => (value + (alignment - 1)) & ~(alignment - 1);
 
 		/// <summary>
 		/// Gets data to write new export
@@ -253,7 +253,7 @@ namespace EpicGames.Horde.Storage.Bundles.V2
 			{
 				lock (_lockObject)
 				{
-					int newSize = _length + desiredSize + 4096 + 16384 & ~16384;
+					int newSize = (_length + desiredSize + 4096 + 16384) & ~16384;
 
 					IRefCountedHandle<Memory<byte>> newBufferHandle = RefCountedHandle.Create(_allocator.Alloc(newSize));
 					_buffer.Slice(0, _length + usedSize).CopyTo(newBufferHandle.Target);
diff --git a/Engine/Source/Programs/Shared/EpicGames.Horde/Storage/RefName.cs b/Engine/Source/Programs/Shared/EpicGames.Horde/Storage/RefName.cs
--- a/Engine/Source/Programs/Shared/EpicGames.Horde/Storage/RefName.cs
+++ b/Engine/Source/Programs/Shared/EpicGames.Horde/Storage/RefName.cs
@@ -12,7 +12,7 @@ using System.Text.Json.Serialization;
 namespace EpicGames.Horde.Storage
 {
 	/// <summary>
-	/// Identifier for a storage namespace
+	/// Identifier for a ref in the storage system. Refs serve as GC roots, and are persistent entry points to expanding data structures within the store.
 	/// </summary>
 	[JsonSchemaString]
 	[JsonConverter(typeof(RefNameJsonConverter))]
