Horde: Fix issue with UnrealBuildBase dependency introduced in CL 29712066
#rnx

diff --git a/Engine/Source/Programs/Horde/Horde.Agent/AgentSettings.cs b/Engine/Source/Programs/Horde/Horde.Agent/AgentSettings.cs
--- a/Engine/Source/Programs/Horde/Horde.Agent/AgentSettings.cs
+++ b/Engine/Source/Programs/Horde/Horde.Agent/AgentSettings.cs
@@ -357,7 +357,7 @@ namespace Horde.Agent
 
 		internal string GetAgentName()
 		{
-			return Name ?? Unreal.MachineName;
+			return Name ?? Environment.MachineName;
 		}
 	}
 
diff --git a/Engine/Source/Programs/Shared/EpicGames.Perforce.Managed.Tests/ManagedWorkspaceTest.cs b/Engine/Source/Programs/Shared/EpicGames.Perforce.Managed.Tests/ManagedWorkspaceTest.cs
--- a/Engine/Source/Programs/Shared/EpicGames.Perforce.Managed.Tests/ManagedWorkspaceTest.cs
+++ b/Engine/Source/Programs/Shared/EpicGames.Perforce.Managed.Tests/ManagedWorkspaceTest.cs
@@ -295,7 +295,7 @@ public class ManagedWorkspaceTest : BasePerforceFixtureTest
 	private async Task<ManagedWorkspace> CreateManagedWorkspaceAsync(bool useHaveTable)
 	{
 		ManagedWorkspaceOptions options = new() { UseHaveTable = useHaveTable };
-		ManagedWorkspace ws = await ManagedWorkspace.CreateAsync(Unreal.MachineName, TempDir, options, _mwLogger, CancellationToken.None);
+		ManagedWorkspace ws = await ManagedWorkspace.CreateAsync(Environment.MachineName, TempDir, options, _mwLogger, CancellationToken.None);
 		await ws.SetupAsync(PerforceConnection, StreamName, CancellationToken.None);
 		return ws;
 	}
