fix for xml doc comments error

diff --git a/Engine/Source/Programs/Shared/EpicGames.Core/UnrealEngineTypes.cs b/Engine/Source/Programs/Shared/EpicGames.Core/UnrealEngineTypes.cs
--- a/Engine/Source/Programs/Shared/EpicGames.Core/UnrealEngineTypes.cs
+++ b/Engine/Source/Programs/Shared/EpicGames.Core/UnrealEngineTypes.cs
@@ -1195,7 +1195,7 @@ namespace EpicGames.Core
 		SkipSerialization = 0x0080000000000000,
 
 		/// <summary>
-		/// Property is a TObjectPtr<T> instead of a USomething*. Need to differentiate between TObjectclassOf and TObjectPtr
+		/// Property is a TObjectPtr instead of a USomething*. Need to differentiate between TObjectclassOf and TObjectPtr
 		/// </summary>
 		TObjectPtr = 0x0100000000000000,
 
