Horde: Fix exception writing temp storage using new storage backend.

diff --git a/Engine/Source/Programs/Shared/EpicGames.Horde/Storage/Node.cs b/Engine/Source/Programs/Shared/EpicGames.Horde/Storage/Node.cs
--- a/Engine/Source/Programs/Shared/EpicGames.Horde/Storage/Node.cs
+++ b/Engine/Source/Programs/Shared/EpicGames.Horde/Storage/Node.cs
@@ -268,6 +268,7 @@ namespace EpicGames.Horde.Storage
 		{
 			await using IStorageWriter writer = store.CreateWriter(name);
 			HashedNodeRef<Node> nodeRef = await writer.WriteHashedNodeAsync(node, cancellationToken);
+			await writer.FlushAsync();
 			await store.WriteRefTargetAsync(name, nodeRef.Handle, refOptions, cancellationToken);
 			return nodeRef.Handle;
 		}
