Horde: Add relay metadata to compute lease

diff --git a/Engine/Source/Programs/Shared/EpicGames.Horde/Compute/Clients/AgentComputeClient.cs b/Engine/Source/Programs/Shared/EpicGames.Horde/Compute/Clients/AgentComputeClient.cs
--- a/Engine/Source/Programs/Shared/EpicGames.Horde/Compute/Clients/AgentComputeClient.cs
+++ b/Engine/Source/Programs/Shared/EpicGames.Horde/Compute/Clients/AgentComputeClient.cs
@@ -31,6 +31,15 @@ namespace EpicGames.Horde.Compute.Clients
 
 			/// <inheritdoc/>
 			public RemoteComputeSocket Socket => _source.Current;
+			
+			/// <inheritdoc/>
+			public string Ip => "127.0.0.1";
+			
+			/// <inheritdoc/>
+			public ConnectionMode ConnectionMode => ConnectionMode.Direct;
+			
+			/// <inheritdoc/>
+			public IReadOnlyDictionary<string, ConnectionMetadataPort> Ports => new Dictionary<string, ConnectionMetadataPort>();
 
 			public LeaseImpl(IAsyncEnumerator<RemoteComputeSocket> source) => _source = source;
 
diff --git a/Engine/Source/Programs/Shared/EpicGames.Horde/Compute/Clients/LocalComputeClient.cs b/Engine/Source/Programs/Shared/EpicGames.Horde/Compute/Clients/LocalComputeClient.cs
--- a/Engine/Source/Programs/Shared/EpicGames.Horde/Compute/Clients/LocalComputeClient.cs
+++ b/Engine/Source/Programs/Shared/EpicGames.Horde/Compute/Clients/LocalComputeClient.cs
@@ -21,6 +21,9 @@ namespace EpicGames.Horde.Compute.Clients
 			public IReadOnlyList<string> Properties { get; } = new List<string>();
 			public IReadOnlyDictionary<string, int> AssignedResources => new Dictionary<string, int>();
 			public RemoteComputeSocket Socket => _socket;
+			public string Ip => "127.0.0.1";
+			public ConnectionMode ConnectionMode => ConnectionMode.Direct;
+			public IReadOnlyDictionary<string, ConnectionMetadataPort> Ports => new Dictionary<string, ConnectionMetadataPort>();
 
 			readonly RemoteComputeSocket _socket;
 
diff --git a/Engine/Source/Programs/Shared/EpicGames.Horde/Compute/Clients/ServerComputeClient.cs b/Engine/Source/Programs/Shared/EpicGames.Horde/Compute/Clients/ServerComputeClient.cs
--- a/Engine/Source/Programs/Shared/EpicGames.Horde/Compute/Clients/ServerComputeClient.cs
+++ b/Engine/Source/Programs/Shared/EpicGames.Horde/Compute/Clients/ServerComputeClient.cs
@@ -98,17 +98,25 @@ namespace EpicGames.Horde.Compute.Clients
 		/// </summary>
 		public const int NonceLength = 64;
 
-		record class LeaseInfo(IReadOnlyList<string> Properties, IReadOnlyDictionary<string, int> AssignedResources, RemoteComputeSocket Socket);
+		record LeaseInfo(
+			IReadOnlyList<string> Properties,
+			IReadOnlyDictionary<string, int> AssignedResources,
+			RemoteComputeSocket Socket,
+			string Ip,
+			ConnectionMode ConnectionMode,
+			IReadOnlyDictionary<string, ConnectionMetadataPort> Ports);
 
 		class LeaseImpl : IComputeLease
 		{
-			readonly IAsyncEnumerator<LeaseInfo> _source;
-			
-			BackgroundTask? _pingTask;
-
 			public IReadOnlyList<string> Properties => _source.Current.Properties;
 			public IReadOnlyDictionary<string, int> AssignedResources => _source.Current.AssignedResources;
 			public RemoteComputeSocket Socket => _source.Current.Socket;
+			public string Ip => _source.Current.Ip;
+			public ConnectionMode ConnectionMode => _source.Current.ConnectionMode;
+			public IReadOnlyDictionary<string, ConnectionMetadataPort> Ports => _source.Current.Ports;
+
+			private readonly IAsyncEnumerator<LeaseInfo> _source;
+			private BackgroundTask? _pingTask;
 
 			public LeaseImpl(IAsyncEnumerator<LeaseInfo> source)
 			{
@@ -224,64 +232,65 @@ namespace EpicGames.Horde.Compute.Clients
 			request.RequestId = requestId;
 			request.Connection = connection;
 
-			AssignComputeResponse? responseMessage;
-			using (HttpResponseMessage response = await HordeHttpClient.PostAsync(client, $"api/v2/compute/{clusterId}", request, _cancellationSource.Token))
+			AssignComputeResponse? response;
+			using (HttpResponseMessage httpResponse = await HordeHttpClient.PostAsync(client, $"api/v2/compute/{clusterId}", request, _cancellationSource.Token))
 			{
-				if (response.StatusCode == HttpStatusCode.NotFound)
+				if (httpResponse.StatusCode == HttpStatusCode.NotFound)
 				{
 					throw new NoComputeAgentsFoundException(clusterId, requirements);
 				}
 
-				if (response.StatusCode == HttpStatusCode.ServiceUnavailable)
+				if (httpResponse.StatusCode == HttpStatusCode.ServiceUnavailable)
 				{
 					_logger.LogDebug("No compute resource is available.");
 					yield break;
 				}
 
-				response.EnsureSuccessStatusCode();
-
-				responseMessage = await response.Content.ReadFromJsonAsync<AssignComputeResponse>(HordeHttpClient.JsonSerializerOptions, cancellationToken);
-				if (responseMessage == null)
+				httpResponse.EnsureSuccessStatusCode();
+				response = await httpResponse.Content.ReadFromJsonAsync<AssignComputeResponse>(HordeHttpClient.JsonSerializerOptions, cancellationToken);
+				if (response == null)
 				{
 					throw new InvalidOperationException();
 				}
 			}
 
-			string agentAddress = $"{responseMessage.Ip}:{responseMessage.Port}";
+			string agentAddress = $"{response.Ip}:{response.Port}";
 
 			// Connect to the remote machine
 			using Socket socket = new Socket(SocketType.Stream, ProtocolType.Tcp);
 			
-			workerLogger.LogDebug("Connecting to {AgentId} at {AgentAddress} ({ConnectionType} via {ConnectionAddress}) with nonce {Nonce}...", responseMessage.AgentId, agentAddress, responseMessage.ConnectionMode, responseMessage.ConnectionAddress ?? "None", responseMessage.Nonce);
-			switch (responseMessage.ConnectionMode)
+			workerLogger.LogDebug("Connecting to {AgentId} at {AgentAddress} ({ConnectionType} via {ConnectionAddress}) with nonce {Nonce}...", response.AgentId, agentAddress, response.ConnectionMode, response.ConnectionAddress ?? "None", response.Nonce);
+			switch (response.ConnectionMode)
 			{
 				case ConnectionMode.Direct:
-					await socket.ConnectAsync(IPAddress.Parse(responseMessage.Ip), responseMessage.Port, cancellationToken);
+					await socket.ConnectAsync(IPAddress.Parse(response.Ip), response.Port, cancellationToken);
 					break;
 
-				case ConnectionMode.Tunnel when !String.IsNullOrEmpty(responseMessage.ConnectionAddress):
-					(string host, int port) = ParseHostPort(responseMessage.ConnectionAddress);
+				case ConnectionMode.Tunnel when !String.IsNullOrEmpty(response.ConnectionAddress):
+					(string host, int port) = ParseHostPort(response.ConnectionAddress);
 					await socket.ConnectAsync(host, port, cancellationToken);
-					await TunnelHandshakeAsync(socket, responseMessage, cancellationToken);
+					await TunnelHandshakeAsync(socket, response, cancellationToken);
 					break;
 				
-				case ConnectionMode.Relay when !String.IsNullOrEmpty(responseMessage.ConnectionAddress):
-					throw new NotImplementedException("Relay connection mode not yet implemented");
+				case ConnectionMode.Relay when !String.IsNullOrEmpty(response.ConnectionAddress):
+					response.Ip = response.ConnectionAddress;
+					await socket.ConnectAsync(IPAddress.Parse(response.ConnectionAddress), response.Ports[ConnectionMetadataPort.ComputeId].Port, cancellationToken);
+					break;
 				
 				default:
-					throw new Exception($"Unable to resolve connection mode ({responseMessage.ConnectionMode}/{responseMessage.ConnectionAddress}");
+					throw new Exception($"Unable to resolve connection mode ({response.ConnectionMode} via {response.ConnectionAddress ?? "none"})");
 			}
 
 			// Send the nonce
-			byte[] nonce = StringUtils.ParseHexString(responseMessage.Nonce);
+			byte[] nonce = StringUtils.ParseHexString(response.Nonce);
 			await socket.SendMessageAsync(nonce, SocketFlags.None, cancellationToken);
-			workerLogger.LogInformation("Connected to {AgentId} ({Ip}) under lease {LeaseId}", responseMessage.AgentId, responseMessage.Ip, responseMessage.LeaseId);
+			workerLogger.LogInformation("Connected to {AgentId} ({Ip}) under lease {LeaseId}", response.AgentId, response.Ip, response.LeaseId);
 
 			// Pass the rest of the call over to the handler
-			byte[] key = StringUtils.ParseHexString(responseMessage.Key);
+			byte[] key = StringUtils.ParseHexString(response.Key);
 
 			await using RemoteComputeSocket computeSocket = new RemoteComputeSocket(new TcpTransport(socket), workerLogger);
-			yield return new LeaseInfo(responseMessage.Properties, responseMessage.AssignedResources, computeSocket);
+			yield return new LeaseInfo(response.Properties, response.AssignedResources, computeSocket, response.Ip, response.ConnectionMode, response.Ports);
 		}
 
 		private static (string host, int port) ParseHostPort(string address)
diff --git a/Engine/Source/Programs/Shared/EpicGames.Horde/Compute/IComputeLease.cs b/Engine/Source/Programs/Shared/EpicGames.Horde/Compute/IComputeLease.cs
--- a/Engine/Source/Programs/Shared/EpicGames.Horde/Compute/IComputeLease.cs
+++ b/Engine/Source/Programs/Shared/EpicGames.Horde/Compute/IComputeLease.cs
@@ -8,7 +8,7 @@ using System.Threading.Tasks;
 namespace EpicGames.Horde.Compute
 {
 	/// <summary>
-	/// Full-duplex channel for sending and reciving messages
+	/// Full-duplex channel for sending and receiving messages
 	/// </summary>
 	public interface IComputeLease : IAsyncDisposable
 	{
@@ -26,6 +26,29 @@ namespace EpicGames.Horde.Compute
 		/// Socket to communicate with the remote
 		/// </summary>
 		RemoteComputeSocket Socket { get; }
+		
+		/// <summary>
+		/// IP address of the remote agent machine running the compute task
+		/// When using relay connection mode, this may be the IP of the relay rather than the remote machine itself.
+		/// </summary>
+		public string Ip { get; }
+		
+		/// <summary>
+		/// How to establish a connection to the remote machine (when not using the default socket)
+		/// </summary>
+		public ConnectionMode ConnectionMode { get; }
+		
+		/// <summary>
+		/// Assigned ports (externally visible port -> local port on agent)
+		///
+		/// Key is an arbitrary name identifying the port (same as was given when requesting the lease>)
+		/// When relay mode is used, ports can mapped to a different externally visible port.
+		/// If compute task uses and listens to port 7000, that port can be externally represented as something else.
+		/// For example, port 32743 can be pointed to port 7000.
+		/// This makes no difference for the compute task process, but the client/initiator making connections must
+		/// pay attention to this mapping.
+		/// </summary>
+		IReadOnlyDictionary<string, ConnectionMetadataPort> Ports { get; }
 
 		/// <summary>
 		/// Relinquish the lease gracefully
