Horde: Missing file.

diff --git a/Engine/Source/Programs/Shared/EpicGames.Horde.Tests/BundleStoreTests.cs b/Engine/Source/Programs/Shared/EpicGames.Horde.Tests/BundleStoreTests.cs
--- a/Engine/Source/Programs/Shared/EpicGames.Horde.Tests/BundleStoreTests.cs
+++ b/Engine/Source/Programs/Shared/EpicGames.Horde.Tests/BundleStoreTests.cs
@@ -158,9 +158,11 @@ namespace EpicGames.Horde.Tests
 				SimpleNode node4 = new SimpleNode(new ReadOnlySequence<byte>(new byte[] { 4 }), Array.Empty<HashedNodeRef<SimpleNode>>());
 
 				SimpleNode root = new SimpleNode(new ReadOnlySequence<byte>(new byte[] { 5 }), new[] { await writer.WriteHashedNodeAsync(node4), await writer.WriteHashedNodeAsync(node3) });
+				NodeRef<SimpleNode> rootRef = await writer.WriteNodeAsync(root);
+				await writer.FlushAsync();
 
-				await store.WriteRefTargetAsync(new RefName("test"), await writer.WriteHashedNodeAsync(root));
-
+				await store.WriteRefTargetAsync(new RefName("test"), rootRef);
+			
 				BundleReader reader = new BundleReader(store, BundleCache.None, NullLogger.Instance);
 				await CheckTreeAsync(root);
 			}
@@ -211,7 +213,13 @@ namespace EpicGames.Horde.Tests
 			using MemoryStorageClient store = new MemoryStorageClient();
 
 			RefName refName = new RefName("test");
-			await store.WriteRefAsync(refName, new SimpleNode(new ReadOnlySequence<byte>(new byte[] { (byte)123 }), Array.Empty<HashedNodeRef<SimpleNode>>()));
+
+			NodeRef<SimpleNode> inputRef;
+			await using (IStorageWriter writer = store.CreateWriter(refName))
+			{
+				inputRef = await writer.WriteNodeAsync(new SimpleNode(new ReadOnlySequence<byte>(new byte[] { (byte)123 }), Array.Empty<HashedNodeRef<SimpleNode>>()));
+			}
+			await store.WriteRefAsync(refName, inputRef.Handle);
 
 			SimpleNode node = await store.ReadRefAsync<SimpleNode>(refName);
 
@@ -240,7 +248,7 @@ namespace EpicGames.Horde.Tests
 
 					await writer.FlushAsync();
 
-					await store.WriteRefTargetAsync(new RefName("test"), rootRef);
+					await store.WriteRefTargetAsync(new RefName("test"), rootRef.Handle);
 				}
 			}
 
