Horde: Missing file.

diff --git a/Engine/Source/Programs/Shared/EpicGames.Horde/Storage/NodeRef.cs b/Engine/Source/Programs/Shared/EpicGames.Horde/Storage/NodeRef.cs
--- a/Engine/Source/Programs/Shared/EpicGames.Horde/Storage/NodeRef.cs
+++ b/Engine/Source/Programs/Shared/EpicGames.Horde/Storage/NodeRef.cs
@@ -361,6 +361,19 @@ namespace EpicGames.Horde.Storage
 			await storageClient.WriteRefTargetAsync(refName, node.Handle, refOptions, cancellationToken);
 		}
 
+		/// <summary>
+		/// Writes a node to the given ref
+		/// </summary>
+		/// <param name="storageClient"></param>
+		/// <param name="refName">Name of the ref to write</param>
+		/// <param name="node"></param>
+		/// <param name="refOptions"></param>
+		/// <param name="cancellationToken"></param>
+		public static async ValueTask WriteRefTargetAsync<TNode>(this IStorageClient storageClient, RefName refName, HashedNodeRef<TNode> node, RefOptions? refOptions = null, CancellationToken cancellationToken = default) where TNode : Node
+		{
+			await storageClient.WriteRefTargetAsync(refName, node.Handle, refOptions, cancellationToken);
+		}
+
 		/// <summary>
 		/// Flushes all the current nodes to storage
 		/// </summary>
