Horde: Add additional checks that data is decompressed correctly from bundles.

diff --git a/Engine/Source/Programs/Shared/EpicGames.Horde/Storage/BlobType.cs b/Engine/Source/Programs/Shared/EpicGames.Horde/Storage/BlobType.cs
--- a/Engine/Source/Programs/Shared/EpicGames.Horde/Storage/BlobType.cs
+++ b/Engine/Source/Programs/Shared/EpicGames.Horde/Storage/BlobType.cs
@@ -43,6 +43,6 @@ namespace EpicGames.Horde.Storage
 		}
 
 		/// <inheritdoc/>
-		public override readonly string ToString() => $"{Guid}#{Version}";
+		public override readonly string ToString() => $"{Guid},{Version}";
 	}
 }
diff --git a/Engine/Source/Programs/Shared/EpicGames.Horde/Storage/Bundles/BundleCompression.cs b/Engine/Source/Programs/Shared/EpicGames.Horde/Storage/Bundles/BundleCompression.cs
--- a/Engine/Source/Programs/Shared/EpicGames.Horde/Storage/Bundles/BundleCompression.cs
+++ b/Engine/Source/Programs/Shared/EpicGames.Horde/Storage/Bundles/BundleCompression.cs
@@ -125,8 +125,14 @@ namespace EpicGames.Horde.Storage.Bundles
 					input.CopyTo(output);
 					break;
 				case BundleCompressionFormat.LZ4:
-					LZ4Codec.Decode(input.Span, output.Span);
-					break;
+					{
+						int length = LZ4Codec.Decode(input.Span, output.Span);
+						if (length != output.Length)
+						{
+							throw new InvalidDataException($"Decoded data is shorter than expected (expected {output.Length} bytes, got {length} bytes)");
+						}
+						break;
+					}
 				case BundleCompressionFormat.Gzip:
 					{
 						using ReadOnlyMemoryStream inputStream = new ReadOnlyMemoryStream(input);
diff --git a/Engine/Source/Programs/Shared/EpicGames.Horde/Storage/Bundles/V2/Packet.cs b/Engine/Source/Programs/Shared/EpicGames.Horde/Storage/Bundles/V2/Packet.cs
--- a/Engine/Source/Programs/Shared/EpicGames.Horde/Storage/Bundles/V2/Packet.cs
+++ b/Engine/Source/Programs/Shared/EpicGames.Horde/Storage/Bundles/V2/Packet.cs
@@ -206,7 +206,7 @@ namespace EpicGames.Horde.Storage.Bundles.V2
 			span = span[1..];
 
 			IMemoryOwner<byte> owner = allocator.Alloc(decodedLength);
-			BundleData.Decompress(format, data.Slice(data.Length - span.Length), owner.Memory);
+			BundleData.Decompress(format, data.Slice(data.Length - span.Length), owner.Memory.Slice(0, decodedLength));
 
 			Packet packet = new Packet(owner.Memory);
 			return RefCountedHandle.Create(packet, owner);
