Horde: Use the exact size of the underlying packet data when constructing packet objects.

diff --git a/Engine/Source/Programs/Shared/EpicGames.Horde/Storage/Bundles/V2/Packet.cs b/Engine/Source/Programs/Shared/EpicGames.Horde/Storage/Bundles/V2/Packet.cs
--- a/Engine/Source/Programs/Shared/EpicGames.Horde/Storage/Bundles/V2/Packet.cs
+++ b/Engine/Source/Programs/Shared/EpicGames.Horde/Storage/Bundles/V2/Packet.cs
@@ -206,9 +206,11 @@ namespace EpicGames.Horde.Storage.Bundles.V2
 			span = span[1..];
 
 			IMemoryOwner<byte> owner = allocator.Alloc(decodedLength);
-			BundleData.Decompress(format, data.Slice(data.Length - span.Length), owner.Memory.Slice(0, decodedLength));
+			Memory<byte> memory = owner.Memory.Slice(0, decodedLength);
 
-			Packet packet = new Packet(owner.Memory);
+			BundleData.Decompress(format, data.Slice(data.Length - span.Length), memory);
+
+			Packet packet = new Packet(memory);
 			return RefCountedHandle.Create(packet, owner);
 		}
 	}
