Horde: Add better diagnostic information for exceptions thrown while reading packets.

diff --git a/Engine/Source/Programs/Shared/EpicGames.Horde/Storage/Bundles/V2/PacketHandle.cs b/Engine/Source/Programs/Shared/EpicGames.Horde/Storage/Bundles/V2/PacketHandle.cs
--- a/Engine/Source/Programs/Shared/EpicGames.Horde/Storage/Bundles/V2/PacketHandle.cs
+++ b/Engine/Source/Programs/Shared/EpicGames.Horde/Storage/Bundles/V2/PacketHandle.cs
@@ -110,8 +110,16 @@ namespace EpicGames.Horde.Storage.Bundles.V2
 		/// <param name="cancellationToken">Cancellation token for the operation</param>
 		public async ValueTask<BlobData> ReadExportAsync(int exportIdx, CancellationToken cancellationToken = default)
 		{
-			using IRefCountedHandle<PacketReader> packetReaderHandle = await GetPacketReaderAsync(cancellationToken);
-			return packetReaderHandle.Target.ReadExport(exportIdx);
+			try
+			{
+				using IRefCountedHandle<PacketReader> packetReaderHandle = await GetPacketReaderAsync(cancellationToken);
+				return packetReaderHandle.Target.ReadExport(exportIdx);
+			}
+			catch (Exception ex)
+			{
+				BlobLocator locator = this.GetLocator();
+				throw new StorageException($"Unable to read {locator}: {ex.Message}", ex);
+			}
 		}
 
 		/// <summary>
@@ -121,8 +129,16 @@ namespace EpicGames.Horde.Storage.Bundles.V2
 		/// <param name="cancellationToken">Cancellation token for the operation</param>
 		public async ValueTask<IReadOnlyMemoryOwner<byte>> ReadExportBodyAsync(int exportIdx, CancellationToken cancellationToken = default)
 		{
-			using IRefCountedHandle<PacketReader> packetReaderHandle = await GetPacketReaderAsync(cancellationToken);
-			return packetReaderHandle.Target.ReadExportBody(exportIdx);
+			try
+			{
+				using IRefCountedHandle<PacketReader> packetReaderHandle = await GetPacketReaderAsync(cancellationToken);
+				return packetReaderHandle.Target.ReadExportBody(exportIdx);
+			}
+			catch (Exception ex)
+			{
+				BlobLocator locator = this.GetLocator();
+				throw new StorageException($"Unable to read {locator}: {ex.Message}", ex);
+			}
 		}
 
 		/// <inheritdoc/>
