Horde: Support for reading entire packets through PacketHandle objects.

diff --git a/Engine/Source/Programs/Shared/EpicGames.Horde/Storage/Bundles/V2/Packet.cs b/Engine/Source/Programs/Shared/EpicGames.Horde/Storage/Bundles/V2/Packet.cs
--- a/Engine/Source/Programs/Shared/EpicGames.Horde/Storage/Bundles/V2/Packet.cs
+++ b/Engine/Source/Programs/Shared/EpicGames.Horde/Storage/Bundles/V2/Packet.cs
@@ -47,6 +47,11 @@ namespace EpicGames.Horde.Storage.Bundles.V2
 	[DebuggerTypeProxy(typeof(Packet.DebugProxy))]
 	public sealed class Packet
 	{
+		/// <summary>
+		/// Type for packet blobs
+		/// </summary>
+		public static BlobType BlobType { get; } = new BlobType(Guid.Parse("{CD9A04EF-CAC1-47D3-A605-2A498130E651}"), 1);
+
 		class DebugProxy
 		{
 			public BlobType[] Types { get; }
@@ -69,6 +74,11 @@ namespace EpicGames.Horde.Storage.Bundles.V2
 		/// <param name="data">Data for the packet</param>
 		public Packet(ReadOnlyMemory<byte> data) => _data = data;
 
+		/// <summary>
+		/// Accessor for the underlying packet data
+		/// </summary>
+		public ReadOnlyMemory<byte> Data => _data;
+
 		/// <summary>
 		/// Length of this packet
 		/// </summary>
diff --git a/Engine/Source/Programs/Shared/EpicGames.Horde/Storage/Bundles/V2/PacketHandle.cs b/Engine/Source/Programs/Shared/EpicGames.Horde/Storage/Bundles/V2/PacketHandle.cs
--- a/Engine/Source/Programs/Shared/EpicGames.Horde/Storage/Bundles/V2/PacketHandle.cs
+++ b/Engine/Source/Programs/Shared/EpicGames.Horde/Storage/Bundles/V2/PacketHandle.cs
@@ -14,11 +14,6 @@ namespace EpicGames.Horde.Storage.Bundles.V2
 	/// </summary>
 	public class PacketHandle : IBlobHandle
 	{
-		/// <summary>
-		/// Type for packet blobs
-		/// </summary>
-		public static BlobType BlobType { get; } = new BlobType(Guid.Parse("{CD9A04EF-CAC1-47D3-A605-2A498130E651}"), 1);
-
 		static readonly Utf8String s_fragmentPrefix = new Utf8String("pkt=");
 
 		readonly IStorageClient _storageClient;
@@ -98,9 +93,18 @@ namespace EpicGames.Horde.Storage.Bundles.V2
 		public ValueTask FlushAsync(CancellationToken cancellationToken = default) => default;
 
 		/// <inheritdoc/>
-		public ValueTask<BlobData> ReadAsync(CancellationToken cancellationToken = default)
+		public async ValueTask<BlobData> ReadAsync(CancellationToken cancellationToken = default)
 		{
-			throw new NotImplementedException();
+			try
+			{
+				using IRefCountedHandle<PacketReader> packetReaderHandle = await GetPacketReaderAsync(cancellationToken);
+				return packetReaderHandle.Target.Read();
+			}
+			catch (Exception ex)
+			{
+				BlobLocator locator = this.GetLocator();
+				throw new StorageException($"Unable to read {locator}: {ex.Message}", ex);
+			}
 		}
 
 		/// <summary>
diff --git a/Engine/Source/Programs/Shared/EpicGames.Horde/Storage/Bundles/V2/PacketReader.cs b/Engine/Source/Programs/Shared/EpicGames.Horde/Storage/Bundles/V2/PacketReader.cs
--- a/Engine/Source/Programs/Shared/EpicGames.Horde/Storage/Bundles/V2/PacketReader.cs
+++ b/Engine/Source/Programs/Shared/EpicGames.Horde/Storage/Bundles/V2/PacketReader.cs
@@ -15,9 +15,14 @@ namespace EpicGames.Horde.Storage.Bundles.V2
 		readonly IBlobHandle _bundleHandle;
 		readonly PacketHandle _packetHandle;
 		readonly Packet _decodedPacket;
-		readonly IRefCountedHandle? _memoryOwner;
+		readonly IRefCountedHandle _memoryOwner;
 		readonly IBlobHandle?[] _cachedImportHandles;
 
+		/// <summary>
+		/// Accessor for the underlying packet data
+		/// </summary>
+		public Packet Packet => _decodedPacket;
+
 		/// <summary>
 		/// Constructor
 		/// </summary>
@@ -27,7 +32,7 @@ namespace EpicGames.Horde.Storage.Bundles.V2
 		/// <param name="packetHandle"></param>
 		/// <param name="decodedPacket">Data for the packet</param>
 		/// <param name="memoryOwner">Owner for the packet data</param>
-		public PacketReader(IStorageClient storageClient, BundleCache cache, IBlobHandle bundleHandle, PacketHandle packetHandle, Packet decodedPacket, IRefCountedHandle? memoryOwner)
+		public PacketReader(IStorageClient storageClient, BundleCache cache, IBlobHandle bundleHandle, PacketHandle packetHandle, Packet decodedPacket, IRefCountedHandle memoryOwner)
 		{
 			_storageClient = storageClient;
 			_cache = cache;
@@ -41,6 +46,19 @@ namespace EpicGames.Horde.Storage.Bundles.V2
 		/// <inheritdoc/>
 		public void Dispose() => _memoryOwner?.Dispose();
 
+		/// <summary>
+		/// Reads this packet in its entirety
+		/// </summary>
+		public BlobData Read()
+		{
+			IBlobHandle[] imports = new IBlobHandle[_decodedPacket.GetImportCount()];
+			for (int idx = 0; idx < imports.Length; idx++)
+			{
+				imports[idx] = GetImportHandle(idx);
+			}
+			return new BlobDataWithOwner(Packet.BlobType, _decodedPacket.Data, imports, _memoryOwner.AddRef());
+		}
+
 		/// <summary>
 		/// Reads an export from this packet
 		/// </summary>
