Horde: Fix LZ4 compression errors due to trailing data in encoded buffer.

diff --git a/Engine/Source/Programs/Shared/EpicGames.Horde/Storage/Bundles/V2/Packet.cs b/Engine/Source/Programs/Shared/EpicGames.Horde/Storage/Bundles/V2/Packet.cs
--- a/Engine/Source/Programs/Shared/EpicGames.Horde/Storage/Bundles/V2/Packet.cs
+++ b/Engine/Source/Programs/Shared/EpicGames.Horde/Storage/Bundles/V2/Packet.cs
@@ -199,16 +199,15 @@ namespace EpicGames.Horde.Storage.Bundles.V2
 		/// <returns></returns>
 		public static IRefCountedHandle<Packet> Decode(ReadOnlyMemory<byte> data, IMemoryAllocator<byte> allocator)
 		{
-			ReadOnlySpan<byte> span = data.Span;
-
-			BundleSignature signature = Bundle.ReadSignature(span);
-			span = span[Bundle.SignatureLength..];
-
+			BundleSignature signature = Bundle.ReadSignature(data.Span);
 			if (signature.Version <= BundleVersion.LatestV1 || signature.Version > BundleVersion.LatestV2)
 			{
 				throw new InvalidOperationException($"Cannot read bundle packet; unsupported version {(int)signature.Version}");
 			}
 
+			data = data.Slice(0, Bundle.SignatureLength + signature.HeaderLength);
+			ReadOnlySpan<byte> span = data.Span.Slice(Bundle.SignatureLength);
+
 			int decodedLength = BinaryPrimitives.ReadInt32LittleEndian(span);
 			span = span[4..];
 
