Horde: Fix comments describing structure of bundle file format.

diff --git a/Engine/Source/Programs/Shared/EpicGames.Horde/Storage/Bundles/V2/Packet.cs b/Engine/Source/Programs/Shared/EpicGames.Horde/Storage/Bundles/V2/Packet.cs
--- a/Engine/Source/Programs/Shared/EpicGames.Horde/Storage/Bundles/V2/Packet.cs
+++ b/Engine/Source/Programs/Shared/EpicGames.Horde/Storage/Bundles/V2/Packet.cs
@@ -14,9 +14,14 @@ namespace EpicGames.Horde.Storage.Bundles.V2
 	/// </summary>
 	/// <remarks>
 	/// <code>
-	/// Each packet starts with a header:
-	///  - 4 bytes: Signature/version field ('U', 'E', 'B', [version number])
-	///  - 4 bytes: Packet length
+	/// Each raw packet contains:
+	///  - 8 bytes: Standard bundle signature. The length field specifies the size of the following data.
+	///  - 4 bytes: Decoded packet length
+	///  - 1 byte: Compression format
+	///  - ?? bytes: Compressed packet data
+	/// 
+	/// After decoding, the packet contains the following:
+	///  - 8 bytes: Standard bundle signature.
 	///  - 4 bytes: offset of type table from the start of the packet
 	///  - 4 bytes: offset of import table from the start of the packet
 	///  - 4 bytes: offset of export table from the start of the packet
@@ -39,9 +44,10 @@ namespace EpicGames.Horde.Storage.Bundles.V2
 	///	
 	/// Each export is written as:
 	///  - 4 bytes: Length of payload
-	///  - ? bytes: Payload data
+	///  - ?? bytes: Payload data
 	///  - VarInt: Type index
-	///  - [VarInt]: Import indexes
+	///  - VarInt: Number of imports
+	///  - VarInt * Number of imports: Import index
 	/// </code>
 	/// </remarks>
 	[DebuggerTypeProxy(typeof(Packet.DebugProxy))]
