Horde: Disable the storage backend cache for now; doesn't seem to be a win to over-fetch entire bundles since ranged S3 reads are fast.

diff --git a/Engine/Source/Programs/Shared/EpicGames.Horde/Storage/Clients/HttpStorageClient.cs b/Engine/Source/Programs/Shared/EpicGames.Horde/Storage/Clients/HttpStorageClient.cs
--- a/Engine/Source/Programs/Shared/EpicGames.Horde/Storage/Clients/HttpStorageClient.cs
+++ b/Engine/Source/Programs/Shared/EpicGames.Horde/Storage/Clients/HttpStorageClient.cs
@@ -318,8 +318,11 @@ namespace EpicGames.Horde.Storage.Clients
 				return httpClient;
 			}
 
+			// Fetching entire blobs doesn't seem to be a win on S3; disabling the cache for now.
+			const bool UseBackendCache = false;
+
 			IStorageBackend backend = new HttpStorageBackend(basePath, CreateClient, _backendLogger);
-			if (_backendCache != null)
+			if (_backendCache != null && UseBackendCache)
 			{
 				backend = _backendCache.CreateWrapper(basePath, backend);
 			}
