Horde: Include the size of bundle signature bytes in the length field of the signature itself, for compatibility with legacy bundle format.

diff --git a/Engine/Source/Programs/Shared/EpicGames.Horde/Storage/Bundles/V2/Packet.cs b/Engine/Source/Programs/Shared/EpicGames.Horde/Storage/Bundles/V2/Packet.cs
--- a/Engine/Source/Programs/Shared/EpicGames.Horde/Storage/Bundles/V2/Packet.cs
+++ b/Engine/Source/Programs/Shared/EpicGames.Horde/Storage/Bundles/V2/Packet.cs
@@ -15,13 +15,13 @@ namespace EpicGames.Horde.Storage.Bundles.V2
 	/// <remarks>
 	/// <code>
 	/// Each raw packet contains:
-	///  - 8 bytes: Standard bundle signature. The length field specifies the size of the following data.
+	///  - 8 bytes: Standard bundle signature. The length field specifies the size of the following data, including the signature itself.
 	///  - 4 bytes: Decoded packet length
 	///  - 1 byte: Compression format
 	///  - ?? bytes: Compressed packet data
 	/// 
 	/// After decoding, the packet contains the following:
-	///  - 8 bytes: Standard bundle signature.
+	///  - 8 bytes: Standard bundle signature. The length field specifies the size of the following data, including the signature itself.
 	///  - 4 bytes: offset of type table from the start of the packet
 	///  - 4 bytes: offset of import table from the start of the packet
 	///  - 4 bytes: offset of export table from the start of the packet
@@ -196,7 +196,7 @@ namespace EpicGames.Horde.Storage.Bundles.V2
 			writer.WriteUInt8((byte)format);
 
 			int encodedLength = BundleData.Compress(format, _data, writer);
-			Bundle.WriteSignature(signatureSpan, new BundleSignature(BundleVersion.LatestV2, encodedLength + sizeof(int) + 1));
+			Bundle.WriteSignature(signatureSpan, new BundleSignature(BundleVersion.LatestV2, Bundle.SignatureLength + encodedLength + sizeof(int) + 1));
 		}
 
 		/// <summary>
@@ -211,7 +211,7 @@ namespace EpicGames.Horde.Storage.Bundles.V2
 				throw new InvalidOperationException($"Cannot read bundle packet; unsupported version {(int)signature.Version}");
 			}
 
-			data = data.Slice(0, Bundle.SignatureLength + signature.HeaderLength);
+			data = data.Slice(0, signature.HeaderLength);
 			ReadOnlySpan<byte> span = data.Span.Slice(Bundle.SignatureLength);
 
 			int decodedLength = BinaryPrimitives.ReadInt32LittleEndian(span);
diff --git a/Engine/Source/Programs/Shared/EpicGames.Horde/Storage/Clients/BundleStorageClient.cs b/Engine/Source/Programs/Shared/EpicGames.Horde/Storage/Clients/BundleStorageClient.cs
--- a/Engine/Source/Programs/Shared/EpicGames.Horde/Storage/Clients/BundleStorageClient.cs
+++ b/Engine/Source/Programs/Shared/EpicGames.Horde/Storage/Clients/BundleStorageClient.cs
@@ -128,7 +128,7 @@ namespace EpicGames.Horde.Storage.Clients
 						}
 					}
 
-					data = data.Slice(Bundle.SignatureLength + signature.HeaderLength);
+					data = data.Slice(signature.HeaderLength);
 				}
 				return locators;
 			}
