Horde: Fix comment.

diff --git a/Engine/Source/Programs/Shared/EpicGames.Horde/Storage/Bundles/BundleCompression.cs b/Engine/Source/Programs/Shared/EpicGames.Horde/Storage/Bundles/BundleCompression.cs
--- a/Engine/Source/Programs/Shared/EpicGames.Horde/Storage/Bundles/BundleCompression.cs
+++ b/Engine/Source/Programs/Shared/EpicGames.Horde/Storage/Bundles/BundleCompression.cs
@@ -30,7 +30,7 @@ namespace EpicGames.Horde.Storage.Bundles
 		Gzip = 2,
 
 		/// <summary>
-		/// Oodle compression (Kraken)
+		/// Oodle compression (Selkie)
 		/// </summary>
 		Oodle = 3,
 
