Horde: Over-read data for V2 bundle packets, and put any additional packets in the memory cache.

diff --git a/Engine/Source/Programs/Shared/EpicGames.Core/MemoryWriter.cs b/Engine/Source/Programs/Shared/EpicGames.Core/MemoryWriter.cs
--- a/Engine/Source/Programs/Shared/EpicGames.Core/MemoryWriter.cs
+++ b/Engine/Source/Programs/Shared/EpicGames.Core/MemoryWriter.cs
@@ -422,7 +422,8 @@ namespace EpicGames.Core
 		public ReadOnlySequence<byte> AsSequence(int offset, int length)
 		{
 			// TODO: could do a binary search for offset and work fowards from there
-			return AsSequence().Slice(offset, length);
+			ReadOnlySequence<byte> sequence = AsSequence();
+			return sequence.Slice(offset, Math.Min(length, sequence.Length - offset));
 		}
 
 		/// <summary>
diff --git a/Engine/Source/Programs/Shared/EpicGames.Horde/Storage/Bundles/V2/BundleWriter.cs b/Engine/Source/Programs/Shared/EpicGames.Horde/Storage/Bundles/V2/BundleWriter.cs
--- a/Engine/Source/Programs/Shared/EpicGames.Horde/Storage/Bundles/V2/BundleWriter.cs
+++ b/Engine/Source/Programs/Shared/EpicGames.Horde/Storage/Bundles/V2/BundleWriter.cs
@@ -3,6 +3,7 @@
 using System;
 using System.Collections.Generic;
 using System.Diagnostics;
+using System.IO;
 using System.Linq;
 using System.Threading;
 using System.Threading.Tasks;
@@ -282,6 +283,18 @@ namespace EpicGames.Horde.Storage.Bundles.V2
 			public ValueTask<BlobData> ReadAsync(CancellationToken cancellationToken)
 				=> GetFlushedHandle().ReadAsync(cancellationToken);
 
+			/// <inheritdoc/>
+			public async Task<Stream> OpenBodyAsync(int offset, int? length, CancellationToken cancellationToken = default)
+			{
+				if (_flushedHandle == null)
+				{
+					IReadOnlyMemoryOwner<byte> owner = await ReadBodyAsync(offset, length, cancellationToken);
+					return owner.AsStream();
+				}
+
+				return await _flushedHandle.OpenBodyAsync(offset, length, cancellationToken);
+			}
+
 			/// <inheritdoc/>
 			public async ValueTask<IReadOnlyMemoryOwner<byte>> ReadBodyAsync(int offset, int? length, CancellationToken cancellationToken = default)
 			{
diff --git a/Engine/Source/Programs/Shared/EpicGames.Horde/Storage/Bundles/V2/PacketHandle.cs b/Engine/Source/Programs/Shared/EpicGames.Horde/Storage/Bundles/V2/PacketHandle.cs
--- a/Engine/Source/Programs/Shared/EpicGames.Horde/Storage/Bundles/V2/PacketHandle.cs
+++ b/Engine/Source/Programs/Shared/EpicGames.Horde/Storage/Bundles/V2/PacketHandle.cs
@@ -1,7 +1,9 @@
 // Copyright Epic Games, Inc. All Rights Reserved.
 
 using System;
+using System.Buffers;
 using System.Buffers.Text;
+using System.IO;
 using System.Text;
 using System.Threading;
 using System.Threading.Tasks;
@@ -201,9 +203,57 @@ namespace EpicGames.Horde.Storage.Bundles.V2
 
 		async Task<IReadOnlyMemoryOwner<byte>> ReadEncodedPacketInternalAsync(EncodedPacketCacheKey key, CancellationToken cancellationToken)
 		{
-			// TODO: If this is more than one packet, slice it up into separate things and add them to the cache.
-			// Use a BundleSignature or BundleFrame
-			return await _outer.ReadBodyAsync(_packetOffset, _packetLength, cancellationToken);
+			IMemoryOwner<byte>? leadingPacket = null;
+			IMemoryOwner<byte>? trailingPacket = null;
+			try
+			{
+				// Read more data than was requested so we can add additional packets to the cache
+				int readLength = Math.Max(_packetLength, 512 * 1024);
+				using Stream stream = await _outer.OpenBodyAsync(_packetOffset, readLength, cancellationToken);
+
+				// Read the first packet
+				leadingPacket = _cache.Allocate(_packetLength);
+				Memory<byte> memory = leadingPacket.Memory.Slice(0, _packetLength);
+				await stream.ReadFixedLengthBytesAsync(memory, cancellationToken);
+
+				// Read any other packets in the same stream
+				byte[] header = new byte[Bundle.SignatureLength];
+				for (int readOffset = _packetLength; readOffset + Bundle.SignatureLength < readLength;)
+				{
+					int readBytes = await stream.ReadGreedyAsync(header, cancellationToken);
+					if (readBytes < header.Length)
+					{
+						break;
+					}
+
+					BundleSignature signature = Bundle.ReadSignature(header);
+					if (readOffset + signature.HeaderLength >= readLength)
+					{
+						break;
+					}
+
+					trailingPacket = _cache.Allocate(signature.HeaderLength);
+					header.CopyTo(trailingPacket.Memory);
+					memory = trailingPacket.Memory.Slice(Bundle.SignatureLength, signature.HeaderLength - Bundle.SignatureLength);
+					await stream.ReadFixedLengthBytesAsync(memory, cancellationToken);
+
+					EncodedPacketCacheKey trailingKey = new EncodedPacketCacheKey(key.Bundle, _packetOffset + readOffset);
+					if (!_cache.TryAdd(trailingKey, () => ReadOnlyMemoryOwner.Create<byte>(memory, trailingPacket)))
+					{
+						trailingPacket.Dispose();
+					}
+					trailingPacket = null;
+
+					readOffset += signature.HeaderLength;
+				}
+			}
+			catch
+			{
+				trailingPacket?.Dispose();
+				leadingPacket?.Dispose();
+				throw;
+			}
+			return ReadOnlyMemoryOwner.Create<byte>(leadingPacket.Memory.Slice(0, _packetLength), leadingPacket);
 		}
 
 		#endregion
