Horde: Fix static analysis warning.

diff --git a/Engine/Source/Programs/Shared/EpicGames.Horde/Storage/Bundles/V2/PacketHandle.cs b/Engine/Source/Programs/Shared/EpicGames.Horde/Storage/Bundles/V2/PacketHandle.cs
--- a/Engine/Source/Programs/Shared/EpicGames.Horde/Storage/Bundles/V2/PacketHandle.cs
+++ b/Engine/Source/Programs/Shared/EpicGames.Horde/Storage/Bundles/V2/PacketHandle.cs
@@ -249,7 +249,9 @@ namespace EpicGames.Horde.Storage.Bundles.V2
 			}
 			catch
 			{
+#pragma warning disable CA1508 // trailingPacket is always null (...not if it's thrown an exception)
 				trailingPacket?.Dispose();
+#pragma warning restore CA1508
 				leadingPacket?.Dispose();
 				throw;
 			}
diff --git a/Engine/Source/Programs/Shared/EpicGames.Horde/Storage/Clients/HttpStorageClient.cs b/Engine/Source/Programs/Shared/EpicGames.Horde/Storage/Clients/HttpStorageClient.cs
--- a/Engine/Source/Programs/Shared/EpicGames.Horde/Storage/Clients/HttpStorageClient.cs
+++ b/Engine/Source/Programs/Shared/EpicGames.Horde/Storage/Clients/HttpStorageClient.cs
@@ -306,7 +306,8 @@ namespace EpicGames.Horde.Storage.Clients
 		/// </summary>
 		/// <param name="basePath">Base path for all requests</param>
 		/// <param name="accessToken">Custom access token to use for requests</param>
-		public IStorageClient CreateClientWithPath(string basePath, string? accessToken = null)
+		/// <param name="withBackendCache"></param>
+		public IStorageClient CreateClientWithPath(string basePath, string? accessToken = null, bool withBackendCache = true)
 		{
 			HttpClient CreateClient()
 			{
@@ -319,7 +320,7 @@ namespace EpicGames.Horde.Storage.Clients
 			}
 
 			IStorageBackend backend = new HttpStorageBackend(basePath, CreateClient, _backendLogger);
-			if (_backendCache != null)
+			if (_backendCache != null && withBackendCache)
 			{
 				backend = _backendCache.CreateWrapper(basePath, backend);
 			}
