Horde: Fix signature corruption when reading additional bundle packets into cache.

diff --git a/Engine/Source/Programs/Shared/EpicGames.Horde/Storage/Bundles/V2/PacketHandle.cs b/Engine/Source/Programs/Shared/EpicGames.Horde/Storage/Bundles/V2/PacketHandle.cs
--- a/Engine/Source/Programs/Shared/EpicGames.Horde/Storage/Bundles/V2/PacketHandle.cs
+++ b/Engine/Source/Programs/Shared/EpicGames.Horde/Storage/Bundles/V2/PacketHandle.cs
@@ -233,9 +233,9 @@ namespace EpicGames.Horde.Storage.Bundles.V2
 					}
 
 					trailingPacket = _cache.Allocate(signature.HeaderLength);
-					header.CopyTo(trailingPacket.Memory);
-					memory = trailingPacket.Memory.Slice(Bundle.SignatureLength, signature.HeaderLength - Bundle.SignatureLength);
-					await stream.ReadFixedLengthBytesAsync(memory, cancellationToken);
+					memory = trailingPacket.Memory.Slice(0, signature.HeaderLength);
+					header.CopyTo(memory);
+					await stream.ReadFixedLengthBytesAsync(memory.Slice(Bundle.SignatureLength), cancellationToken);
 
 					EncodedPacketCacheKey trailingKey = new EncodedPacketCacheKey(key.Bundle, _packetOffset + readOffset);
 					if (!_cache.TryAdd(trailingKey, () => ReadOnlyMemoryOwner.Create<byte>(memory, trailingPacket)))
