Horde: Fix test failure due to JsonPath not having a schema type annotation.

diff --git a/Engine/Source/Programs/Horde/Horde.Server/Telemetry/Metrics/MetricConfig.cs b/Engine/Source/Programs/Horde/Horde.Server/Telemetry/Metrics/MetricConfig.cs
--- a/Engine/Source/Programs/Horde/Horde.Server/Telemetry/Metrics/MetricConfig.cs
+++ b/Engine/Source/Programs/Horde/Horde.Server/Telemetry/Metrics/MetricConfig.cs
@@ -2,6 +2,7 @@
 
 using System;
 using System.Text.Json.Serialization;
+using EpicGames.Core;
 using Horde.Server.Utilities;
 using Json.Path;
 
@@ -56,11 +57,13 @@ namespace Horde.Server.Telemetry.Metrics
 		/// <summary>
 		/// Property to aggregate
 		/// </summary>
+		[JsonSchemaString]
 		public JsonPath? Property { get; set; }
 
 		/// <summary>
 		/// Property to group by
 		/// </summary>
+		[JsonSchemaString]
 		public JsonPath? GroupBy { get; set; }
 
 		/// <summary>
diff --git a/Engine/Source/Programs/Shared/EpicGames.Core/JsonSchema.cs b/Engine/Source/Programs/Shared/EpicGames.Core/JsonSchema.cs
--- a/Engine/Source/Programs/Shared/EpicGames.Core/JsonSchema.cs
+++ b/Engine/Source/Programs/Shared/EpicGames.Core/JsonSchema.cs
@@ -296,6 +296,19 @@ namespace EpicGames.Core
 			return new JsonSchema(schemaAttribute?.Id, CreateSchemaType(type, new Dictionary<Type, JsonSchemaType>(), xmlDocReader));
 		}
 
+		/// <summary>
+		/// Constructs a schema type for the given property
+		/// </summary>
+		static JsonSchemaType CreateSchemaType(PropertyInfo property, Dictionary<Type, JsonSchemaType> typeCache, XmlDocReader? xmlDocReader)
+		{
+			if (property.GetCustomAttribute<JsonSchemaStringAttribute>() != null)
+			{
+				return new JsonSchemaString();
+			}
+
+			return CreateSchemaType(property.PropertyType, typeCache, xmlDocReader);
+		}
+
 		/// <summary>
 		/// Constructs a schema type from the given type object
 		/// </summary>
@@ -440,7 +453,7 @@ namespace EpicGames.Core
 				if (property.GetCustomAttribute<JsonIgnoreAttribute>() == null)
 				{
 					string? description = xmlDocReader?.GetDescription(property);
-					JsonSchemaType propertyType = CreateSchemaType(property.PropertyType, typeCache, xmlDocReader);
+					JsonSchemaType propertyType = CreateSchemaType(property, typeCache, xmlDocReader);
 					obj.Properties.Add(new JsonSchemaProperty(property.Name, description, propertyType));
 				}
 			}
