Horde: Fix aliases not being flushed when using V2 packet format.

diff --git a/Engine/Source/Programs/Shared/EpicGames.Horde/Storage/Bundles/V2/BundleWriter.cs b/Engine/Source/Programs/Shared/EpicGames.Horde/Storage/Bundles/V2/BundleWriter.cs
--- a/Engine/Source/Programs/Shared/EpicGames.Horde/Storage/Bundles/V2/BundleWriter.cs
+++ b/Engine/Source/Programs/Shared/EpicGames.Horde/Storage/Bundles/V2/BundleWriter.cs
@@ -261,6 +261,7 @@ namespace EpicGames.Horde.Storage.Bundles.V2
 				}
 
 				// Release all the intermediate data
+				List<(PendingExportHandle, AliasInfo)>? pendingExportAliases = _pendingExportAliases;
 				lock (_lockObject)
 				{
 					_flushedHandle = flushedHandle;
@@ -270,9 +271,9 @@ namespace EpicGames.Horde.Storage.Bundles.V2
 				// TODO: put all the encoded packets into the cache using the final handles
 
 				// Add all the aliases
-				if (_pendingExportAliases != null)
+				if (pendingExportAliases != null)
 				{
-					foreach ((PendingExportHandle exportHandle, AliasInfo aliasInfo) in _pendingExportAliases)
+					foreach ((PendingExportHandle exportHandle, AliasInfo aliasInfo) in pendingExportAliases)
 					{
 						await _storageClient.AddAliasAsync(aliasInfo.Name, exportHandle, aliasInfo.Rank, aliasInfo.Data, cancellationToken);
 					}
