Horde: Add test suite for BundleCache, and fix a few issues with inconsistent ref counts for managed objects.

diff --git a/Engine/Source/Programs/Shared/EpicGames.Core/RefCountedHandle.cs b/Engine/Source/Programs/Shared/EpicGames.Core/RefCountedHandle.cs
--- a/Engine/Source/Programs/Shared/EpicGames.Core/RefCountedHandle.cs
+++ b/Engine/Source/Programs/Shared/EpicGames.Core/RefCountedHandle.cs
@@ -23,6 +23,11 @@ namespace EpicGames.Core
 	/// <typeparam name="T">Type of the target resource</typeparam>
 	public interface IRefCountedHandle<T> : IRefCountedHandle
 	{
+		/// <summary>
+		/// Accessor for the current ref count, for debugging purposes.
+		/// </summary>
+		int RefCount { get; }
+
 		/// <summary>
 		/// Target of the reference
 		/// </summary>
@@ -73,6 +78,9 @@ namespace EpicGames.Core
 		readonly RefCountedDisposer? _refCountedDisposer;
 		T _target;
 
+		/// <inheritdoc/>
+		public int RefCount => _refCountedDisposer?.RefCount ?? throw new ObjectDisposedException(typeof(RefCountedHandle<T>).Name);
+
 		/// <inheritdoc/>
 		public T Target => _target ?? throw new ObjectDisposedException(typeof(RefCountedHandle<T>).Name);
 
@@ -124,6 +132,11 @@ namespace EpicGames.Core
 		IDisposable? _owner;
 		int _value = 1;
 
+		/// <summary>
+		/// Current reference count
+		/// </summary>
+		public int RefCount => _value;
+
 		/// <summary>
 		/// Constructor
 		/// </summary>
diff --git a/Engine/Source/Programs/Shared/EpicGames.Horde.Tests/BundleCacheTests.cs b/Engine/Source/Programs/Shared/EpicGames.Horde.Tests/BundleCacheTests.cs
new file mode 100644
--- /dev/null
+++ b/Engine/Source/Programs/Shared/EpicGames.Horde.Tests/BundleCacheTests.cs
@@ -0,0 +1,152 @@
+// Copyright Epic Games, Inc. All Rights Reserved.
+
+using System;
+using System.Buffers;
+using System.Collections.Generic;
+using System.Linq;
+using System.Text;
+using System.Threading;
+using System.Threading.Tasks;
+using EpicGames.Core;
+using EpicGames.Horde.Storage;
+using Microsoft.VisualStudio.TestTools.UnitTesting;
+
+namespace EpicGames.Horde.Tests
+{
+	[TestClass]
+	public class BundleCacheTests
+	{
+		[TestMethod]
+		public async Task AllocatorAsync()
+		{
+			await using BundleCache cache = new BundleCache(new BundleCacheOptions { MaxSize = 4096 });
+
+			Assert.AreEqual(0, cache.CurrentSize);
+			using (IMemoryOwner<byte> owner = cache.Allocator.Alloc(1024))
+			{
+				Assert.AreEqual(1024, cache.CurrentSize);
+			}
+			Assert.AreEqual(0, cache.CurrentSize);
+			using (IMemoryOwner<byte> owner = cache.Allocator.Alloc(1024))
+			{
+				Assert.AreEqual(1024, cache.CurrentSize);
+			}
+			Assert.AreEqual(0, cache.CurrentSize);
+		}
+
+		class TestData : IDisposable
+		{
+			readonly IMemoryOwner<byte> _owner;
+
+			public TestData(IMemoryOwner<byte> owner) => _owner = owner;
+			public void Dispose() => _owner.Dispose();
+		}
+
+		[TestMethod]
+		public async Task FindOrAddAsync()
+		{
+			await using BundleCache cache = new BundleCache(new BundleCacheOptions { MaxSize = 1024 }, ManagedHeapAllocator.Instance);
+
+			Assert.AreEqual(0, cache.CurrentSize);
+
+			IRefCountedHandle<TestData> result = await cache.FindOrAddAsync("test", (key, ctx) => Task.FromResult(new TestData(cache.Allocator.Alloc(1000))));
+			Assert.AreEqual(1000, cache.CurrentSize);
+			Assert.AreEqual(2, result.RefCount); // cache and result
+
+			IRefCountedHandle<TestData>? value = cache.Find<string, TestData>("test");
+			Assert.IsNotNull(value);
+			Assert.AreEqual(1000, cache.CurrentSize);
+			Assert.IsTrue(ReferenceEquals(result.Target, value.Target));
+			Assert.AreEqual(3, result.RefCount); // cache, result, value
+
+			value.Dispose();
+
+			Assert.AreEqual(1000, cache.CurrentSize);
+			Assert.AreEqual(2, result.RefCount); // cache and result
+
+			result.Dispose();
+
+			Assert.AreEqual(1000, cache.CurrentSize);
+
+			// Check other things don't return a value
+			Assert.IsNull(cache.Find<string, TestData>("other"));
+
+			// Allocate above the 1024 byte budget and check it's released
+			IMemoryOwner<byte> owner = cache.Allocator.Alloc(25);
+			Assert.AreEqual(25, cache.CurrentSize);
+			Assert.IsNull(cache.Find<string, TestData>("test"));
+
+			// Allow allocating more than the limit if we don't have any cache values to free
+			IMemoryOwner<byte> owner2 = cache.Allocator.Alloc(2000);
+			Assert.AreEqual(2025, cache.CurrentSize);
+
+			owner2.Dispose();
+			owner.Dispose();
+			Assert.AreEqual(0, cache.CurrentSize);
+		}
+
+		[TestMethod]
+		public async Task TrimAsync()
+		{
+			await using BundleCache cache = new BundleCache(new BundleCacheOptions { MaxSize = 1024 }, ManagedHeapAllocator.Instance);
+
+			Assert.AreEqual(0, cache.CurrentSize);
+
+			IRefCountedHandle<TestData> result = await cache.FindOrAddAsync("test", (key, ctx) => Task.FromResult(new TestData(cache.Allocator.Alloc(1000))));
+			Assert.AreEqual(1000, cache.CurrentSize);
+			result.Dispose();
+
+			IMemoryOwner<byte> owner = cache.Allocator.Alloc(10);
+
+			Assert.AreEqual(1010, cache.CurrentSize);
+
+			cache.Trim();
+
+			Assert.AreEqual(10, cache.CurrentSize);
+
+			owner.Dispose();
+		}
+
+		[TestMethod]
+		public async Task TryAddSuccessAsync()
+		{
+			await using BundleCache cache = new BundleCache(new BundleCacheOptions { MaxSize = 1024 }, ManagedHeapAllocator.Instance);
+
+			Assert.AreEqual(0, cache.CurrentSize);
+
+			Assert.IsTrue(cache.TryAdd("test", new TestData(cache.Allocator.Alloc(20))));
+			Assert.AreEqual(20, cache.CurrentSize);
+
+			IRefCountedHandle<TestData> result = await cache.FindOrAddAsync("test", (key, ctx) => Task.FromResult(new TestData(cache.Allocator.Alloc(20))));
+			Assert.AreEqual(20, cache.CurrentSize);
+
+			cache.Trim(); 
+			Assert.AreEqual(20, cache.CurrentSize);
+			result.Dispose();
+			Assert.AreEqual(20, cache.CurrentSize);
+			cache.Trim();
+			Assert.AreEqual(0, cache.CurrentSize);
+		}
+
+		[TestMethod]
+		public async Task TryAddFailureAsync()
+		{
+			await using BundleCache cache = new BundleCache(new BundleCacheOptions { MaxSize = 1024 }, ManagedHeapAllocator.Instance);
+
+			Assert.AreEqual(0, cache.CurrentSize);
+
+			IRefCountedHandle<TestData> result = await cache.FindOrAddAsync("test", (key, ctx) => Task.FromResult(new TestData(cache.Allocator.Alloc(20))));
+			Assert.AreEqual(20, cache.CurrentSize);
+
+			TestData testData = new TestData(cache.Allocator.Alloc(20));
+			Assert.IsFalse(cache.TryAdd("test", testData));
+			Assert.AreEqual(40, cache.CurrentSize);
+			testData.Dispose();
+
+			result.Dispose();
+			Assert.AreEqual(20, cache.CurrentSize);
+			cache.Trim();
+			Assert.AreEqual(0, cache.CurrentSize);
+		}
+	}
+}
diff --git a/Engine/Source/Programs/Shared/EpicGames.Horde/Storage/Bundles/BundleCache.cs b/Engine/Source/Programs/Shared/EpicGames.Horde/Storage/Bundles/BundleCache.cs
--- a/Engine/Source/Programs/Shared/EpicGames.Horde/Storage/Bundles/BundleCache.cs
+++ b/Engine/Source/Programs/Shared/EpicGames.Horde/Storage/Bundles/BundleCache.cs
@@ -40,21 +40,26 @@ namespace EpicGames.Horde.Storage
 	{
 		class CacheValue : IDisposable
 		{
+			public object Key { get; }
 			public LinkedListNode<CacheValue> Node { get; }
 			public Task<IDisposable> InitTask { get; }
-			int _lockCount;
+			int _refCount;
 
-			public CacheValue(IDisposable value)
+			public int RefCount => Interlocked.CompareExchange(ref _refCount, 0, 0);
+
+			public CacheValue(object key, IDisposable value)
 			{
+				_refCount = 1;
+				Key = key;
 				Node = new LinkedListNode<CacheValue>(this);
-				Lock();
 				InitTask = Task.FromResult<IDisposable>(value);
 			}
 
-			public CacheValue(Func<Task<IDisposable>> initTask)
+			public CacheValue(object key, Func<Task<IDisposable>> initTask)
 			{
+				_refCount = 2; // Will be released by RunAndUnlock
+				Key = key;
 				Node = new LinkedListNode<CacheValue>(this);
-				Lock();
 				InitTask = Task.Run(() => RunAndUnlockAsync(initTask));
 			}
 
@@ -71,14 +76,12 @@ namespace EpicGames.Horde.Storage
 				}
 				finally
 				{
-					Unlock();
+					Release();
 				}
 			}
 
-			public bool IsLocked() => Interlocked.CompareExchange(ref _lockCount, 0, 0) > 0;
-
-			public void Lock() => Interlocked.Increment(ref _lockCount);
-			public void Unlock() => Interlocked.Decrement(ref _lockCount);
+			public void AddRef() => Interlocked.Increment(ref _refCount);
+			public void Release() => Interlocked.Decrement(ref _refCount);
 		}
 
 		class CacheValueHandle<T> : IRefCountedHandle<T> where T : class, IDisposable
@@ -86,6 +89,8 @@ namespace EpicGames.Horde.Storage
 			CacheValue? _item;
 			T? _target;
 
+			public int RefCount => _item?.RefCount ?? throw new ObjectDisposedException(nameof(CacheValueHandle<T>));
+
 			public T Target => _target ?? throw new ObjectDisposedException(nameof(CacheValueHandle<T>));
 
 			public CacheValueHandle(CacheValue? item, T? value)
@@ -98,7 +103,7 @@ namespace EpicGames.Horde.Storage
 			{
 				if (_item != null)
 				{
-					_item.Unlock();
+					_item.Release();
 					_item = null;
 					_target = null;
 				}
@@ -111,49 +116,40 @@ namespace EpicGames.Horde.Storage
 					throw new ObjectDisposedException(nameof(CacheValueHandle<T>));
 				}
 
-				_item.Lock();
+				_item.AddRef();
 				return new CacheValueHandle<T>(_item, _target);
 			}
 		}
 
-		class MemoryReservation : IDisposable
+		// Tracks an owned blocks of memory against the cache budged.
+		sealed class MemoryAllocation : IMemoryOwner<byte>
 		{
 			readonly BundleCache _outer;
-			long _size;
-
-			public MemoryReservation(BundleCache outer, long size)
-			{
-				_outer = outer;
-				_size = size;
-			}
-
-			public virtual void Dispose()
-			{
-				Interlocked.Add(ref _outer._currentSize, -_size);
-				_size = 0;
-			}
-		}
-
-		sealed class MemoryAllocation : MemoryReservation, IMemoryOwner<byte>
-		{
-			readonly IMemoryOwner<byte> _owner;
+			IMemoryOwner<byte> _owner;
 
 			public Memory<byte> Memory => _owner.Memory;
 
 			public MemoryAllocation(BundleCache outer, IMemoryOwner<byte> owner)
-				: base(outer, owner.Memory.Length)
 			{
+				_outer = outer;
 				_owner = owner;
 			}
 
-			public override void Dispose()
+			/// <inheritdoc/>
+			public void Dispose()
 			{
-				base.Dispose();
-
-				_owner.Dispose();
+				if (_owner != null)
+				{
+					long size = _owner.Memory.Length;
+					_owner.Dispose();
+					_outer.ReleaseSpace(size);
+					_owner = null!;
+				}
 			}
 		}
 
+		// Custom memory allocator which tracks allocated blocks against the cache's budget. Older cache entries will be 
+		// disposed to create space for new allocations.
 		class MemoryAllocator : IMemoryAllocator<byte>
 		{
 			readonly BundleCache _outer;
@@ -165,7 +161,13 @@ namespace EpicGames.Horde.Storage
 				_inner = inner;
 			}
 
-			public IMemoryOwner<byte> Alloc(int minSize) => new MemoryAllocation(_outer, _inner.Alloc(minSize));
+			public IMemoryOwner<byte> Alloc(int minSize)
+			{
+				_outer.CreateSpace(minSize);
+				IMemoryOwner<byte> owner = _inner.Alloc(minSize);
+				_outer.CreateSpace(owner.Memory.Length - minSize);
+				return new MemoryAllocation(_outer, owner);
+			}
 		}
 
 		readonly object _lockObject = new object();
@@ -173,7 +175,7 @@ namespace EpicGames.Horde.Storage
 		readonly LinkedList<CacheValue> _items = new LinkedList<CacheValue>();
 		readonly Dictionary<object, CacheValue> _itemLookup = new Dictionary<object, CacheValue>();
 		readonly CancellationTokenSource _cancellationSource = new CancellationTokenSource();
-		readonly IMemoryAllocator<byte> _allocator;
+		readonly MemoryAllocator _allocator;
 
 		long _currentSize;
 
@@ -190,6 +192,11 @@ namespace EpicGames.Horde.Storage
 		/// </summary>
 		public IMemoryAllocator<byte> Allocator => _allocator;
 
+		/// <summary>
+		/// Current size of data allocated or in the cache
+		/// </summary>
+		public long CurrentSize => _currentSize;
+
 		/// <summary>
 		/// Size of the configured header cache
 		/// </summary>
@@ -208,17 +215,29 @@ namespace EpicGames.Horde.Storage
 		/// <summary>
 		/// Constructor
 		/// </summary>
-		public BundleCache() : this(new BundleCacheOptions()) 
+		public BundleCache() 
+			: this(new BundleCacheOptions()) 
 		{ 
 		}
 
 		/// <summary>
 		/// Constructor
 		/// </summary>
+		/// <param name="options">Options for the cache</param>
 		public BundleCache(BundleCacheOptions options)
+			: this(options, PoolAllocator.Shared)
+		{
+		}
+
+		/// <summary>
+		/// Constructor
+		/// </summary>
+		/// <param name="options">Options for the cache</param>
+		/// <param name="innerAllocator">Inner allocator to use. Will be wrapped in an allocator that tracks allocations against the cache's budget.</param>
+		public BundleCache(BundleCacheOptions options, IMemoryAllocator<byte> innerAllocator)
 		{
 			_options = options;
-			_allocator = new MemoryAllocator(this, PoolAllocator.Shared);
+			_allocator = new MemoryAllocator(this, innerAllocator);
 
 			if (options.HeaderCacheSize > 0)
 			{
@@ -257,6 +276,15 @@ namespace EpicGames.Horde.Storage
 			_cancellationSource.Dispose();
 		}
 
+		/// <summary>
+		/// Empty the cache
+		/// </summary>
+		public void Trim()
+		{
+			CreateSpace(_options.MaxSize);
+			ReleaseSpace(_options.MaxSize);
+		}
+
 		/// <summary>
 		/// Find or add a new cached value to the cache
 		/// </summary>
@@ -271,7 +299,7 @@ namespace EpicGames.Horde.Storage
 				CacheValue? item;
 				if (_itemLookup.TryGetValue(key, out item) && item.InitTask.TryGetResult(out IDisposable result))
 				{
-					item.Lock();
+					item.AddRef();
 					return new CacheValueHandle<TValue>(item, (TValue)result);
 				}
 			}
@@ -279,12 +307,12 @@ namespace EpicGames.Horde.Storage
 		}
 
 		/// <summary>
-		/// Attempts to add a value to the cache
+		/// Attempts to add a value to the cache. 
 		/// </summary>
-		/// <param name="key"></param>
-		/// <param name="createFunc"></param>
+		/// <param name="key">Key to add the item to the cache with</param>
+		/// <param name="value">Value to be added. If the item is added, ownership is implicitly transferred to the cache.</param>
 		/// <returns>True if the value was added</returns>
-		public bool TryAdd<TKey, TValue>(TKey key, Func<TValue> createFunc)
+		public bool TryAdd<TKey, TValue>(TKey key, TValue value)
 			where TKey : notnull
 			where TValue : class, IDisposable
 		{
@@ -292,7 +320,8 @@ namespace EpicGames.Horde.Storage
 			{
 				if (!_itemLookup.ContainsKey(key))
 				{
-					CacheValue item = new CacheValue(createFunc());
+					CacheValue item = new CacheValue(key, value);
+					_items.AddFirst(item);
 					_itemLookup.Add(key, item);
 					return true;
 				}
@@ -307,70 +336,56 @@ namespace EpicGames.Horde.Storage
 		/// <param name="createAsync"></param>
 		/// <param name="cancellationToken"></param>
 		/// <returns>Handle to the item that was read. Must be disposed by the caller.</returns>
-		public async Task<IRefCountedHandle<TValue>> FindOrAddAsync<TKey, TValue>(TKey key, Func<TKey, CancellationToken, Task<TValue>> createAsync, CancellationToken cancellationToken)
+		public async Task<IRefCountedHandle<TValue>> FindOrAddAsync<TKey, TValue>(TKey key, Func<TKey, CancellationToken, Task<TValue>> createAsync, CancellationToken cancellationToken = default)
 			where TKey : notnull
 			where TValue : class, IDisposable
 		{
-			CacheValue? item = null;
-			try
+			CacheValue? item;
+			lock (_lockObject)
 			{
-				lock (_lockObject)
+				CacheValue? untypedItem;
+				if (_itemLookup.TryGetValue(key, out untypedItem))
 				{
-					CacheValue? untypedItem;
-					if (_itemLookup.TryGetValue(key, out untypedItem))
-					{
-						item = untypedItem;
-						item.Lock();
-					}
-					else
-					{
-						item = new CacheValue(async () => await createAsync(key, _cancellationSource.Token));
-						_items.AddFirst(item);
-						_itemLookup.Add(key, item);
-					}
+					item = untypedItem;
+				}
+				else
+				{
+					item = new CacheValue(key, async () => await createAsync(key, _cancellationSource.Token));
+					_items.AddFirst(item);
+					_itemLookup.Add(key, item);
 				}
-
-				TValue value = (TValue)await item.InitTask.WaitAsync(cancellationToken);
-				return new CacheValueHandle<TValue>(item, value);
-			}
-			finally
-			{
-				item?.Unlock();
 			}
-		}
 
-		/// <summary>
-		/// Allocate a block of memory from the shared heap
-		/// </summary>
-		/// <param name="size">Size of the allocation. The returned memory may be larger than this.</param>
-		public IMemoryOwner<byte> Allocate(int size)
-		{
-			return _allocator.Alloc(size);
+			TValue value = (TValue)await item.InitTask.WaitAsync(cancellationToken);
+			item.AddRef();
+			return new CacheValueHandle<TValue>(item, value);
 		}
 
-		/// <summary>
-		/// Add a reservation against the maximum allowed size of cached data. 
-		/// </summary>
-		public IDisposable Reserve(long size)
+		void CreateSpace(long size)
 		{
-			lock (_lockObject)
+			if (Interlocked.Add(ref _currentSize, size) > _options.MaxSize)
 			{
-				Interlocked.Add(ref _currentSize, size);
-
-				for (LinkedListNode<CacheValue>? lastNode = _items.Last; lastNode != null && Interlocked.CompareExchange(ref _currentSize, 0, 0) > _options.MaxSize; lastNode = lastNode.Previous)
+				lock (_lockObject)
 				{
-					CacheValue lastItem = lastNode.Value;
-					if (!lastItem.IsLocked())
+					for (LinkedListNode<CacheValue>? lastNode = _items.Last; lastNode != null && Interlocked.CompareExchange(ref _currentSize, 0, 0) > _options.MaxSize; lastNode = lastNode.Previous)
 					{
-						_items.Remove(lastItem);
-						lastItem.Dispose();
+						CacheValue lastItem = lastNode.Value;
+						if (lastItem.RefCount == 1)
+						{
+							_items.Remove(lastItem);
+							_itemLookup.Remove(lastItem.Key);
+							lastItem.Dispose();
+						}
 					}
 				}
-
-				return new MemoryReservation(this, size);
 			}
 		}
 
+		void ReleaseSpace(long size)
+		{
+			Interlocked.Add(ref _currentSize, -size);
+		}
+
 		#region V1
 
 		static void AddCachedValue(IMemoryCache? cache, string cacheKey, object value, int size)
diff --git a/Engine/Source/Programs/Shared/EpicGames.Horde/Storage/Bundles/V2/PacketHandle.cs b/Engine/Source/Programs/Shared/EpicGames.Horde/Storage/Bundles/V2/PacketHandle.cs
--- a/Engine/Source/Programs/Shared/EpicGames.Horde/Storage/Bundles/V2/PacketHandle.cs
+++ b/Engine/Source/Programs/Shared/EpicGames.Horde/Storage/Bundles/V2/PacketHandle.cs
@@ -212,7 +212,7 @@ namespace EpicGames.Horde.Storage.Bundles.V2
 				using Stream stream = await _outer.OpenBodyAsync(_packetOffset, readLength, cancellationToken);
 
 				// Read the first packet
-				leadingPacket = _cache.Allocate(_packetLength);
+				leadingPacket = _cache.Allocator.Alloc(_packetLength);
 				Memory<byte> memory = leadingPacket.Memory.Slice(0, _packetLength);
 				await stream.ReadFixedLengthBytesAsync(memory, cancellationToken);
 
@@ -232,15 +232,16 @@ namespace EpicGames.Horde.Storage.Bundles.V2
 						break;
 					}
 
-					trailingPacket = _cache.Allocate(signature.HeaderLength);
+					trailingPacket = _cache.Allocator.Alloc(signature.HeaderLength);
 					memory = trailingPacket.Memory.Slice(0, signature.HeaderLength);
 					header.CopyTo(memory);
 					await stream.ReadFixedLengthBytesAsync(memory.Slice(Bundle.SignatureLength), cancellationToken);
 
 					EncodedPacketCacheKey trailingKey = new EncodedPacketCacheKey(key.Bundle, _packetOffset + readOffset);
-					if (!_cache.TryAdd(trailingKey, () => ReadOnlyMemoryOwner.Create<byte>(memory, trailingPacket)))
+					IReadOnlyMemoryOwner<byte> trailingValue = ReadOnlyMemoryOwner.Create<byte>(memory, trailingPacket);
+					if (!_cache.TryAdd(trailingKey, trailingValue))
 					{
-						trailingPacket.Dispose();
+						trailingValue.Dispose();
 					}
 					trailingPacket = null;
 
