Horde: Fix static analysis warnings.

diff --git a/Engine/Source/Programs/Shared/EpicGames.Horde.Tests/BundleCacheTests.cs b/Engine/Source/Programs/Shared/EpicGames.Horde.Tests/BundleCacheTests.cs
--- a/Engine/Source/Programs/Shared/EpicGames.Horde.Tests/BundleCacheTests.cs
+++ b/Engine/Source/Programs/Shared/EpicGames.Horde.Tests/BundleCacheTests.cs
@@ -2,10 +2,6 @@
 
 using System;
 using System.Buffers;
-using System.Collections.Generic;
-using System.Linq;
-using System.Text;
-using System.Threading;
 using System.Threading.Tasks;
 using EpicGames.Core;
 using EpicGames.Horde.Storage;
diff --git a/Engine/Source/Programs/Shared/EpicGames.Horde/Storage/Backends/FileStorageBackend.cs b/Engine/Source/Programs/Shared/EpicGames.Horde/Storage/Backends/FileStorageBackend.cs
--- a/Engine/Source/Programs/Shared/EpicGames.Horde/Storage/Backends/FileStorageBackend.cs
+++ b/Engine/Source/Programs/Shared/EpicGames.Horde/Storage/Backends/FileStorageBackend.cs
@@ -241,12 +241,12 @@ namespace EpicGames.Horde.Storage.Backends
 		}
 
 		/// <inheritdoc/>
-		public Task<Stream> OpenAsync(string path, int offset, int? length, CancellationToken cancellationToken)
+		public async Task<Stream> OpenAsync(string path, int offset, int? length, CancellationToken cancellationToken)
 		{
 #pragma warning disable CA2000 // Dispose objects before losing scope
-			IReadOnlyMemoryOwner<byte> storageObject = Read(path, offset, length);
+			IReadOnlyMemoryOwner<byte> storageObject = await ReadAsync(path, offset, length, cancellationToken);
+			return storageObject.AsStream();
 #pragma warning restore CA2000 // Dispose objects before losing scope
-			return Task.FromResult<Stream>(storageObject.AsStream());
 		}
 
 		/// <summary>
diff --git a/Engine/Source/Programs/Shared/EpicGames.Horde/Storage/Bundles/V2/PacketHandle.cs b/Engine/Source/Programs/Shared/EpicGames.Horde/Storage/Bundles/V2/PacketHandle.cs
--- a/Engine/Source/Programs/Shared/EpicGames.Horde/Storage/Bundles/V2/PacketHandle.cs
+++ b/Engine/Source/Programs/Shared/EpicGames.Horde/Storage/Bundles/V2/PacketHandle.cs
@@ -238,12 +238,15 @@ namespace EpicGames.Horde.Storage.Bundles.V2
 					await stream.ReadFixedLengthBytesAsync(memory.Slice(Bundle.SignatureLength), cancellationToken);
 
 					EncodedPacketCacheKey trailingKey = new EncodedPacketCacheKey(key.Bundle, _packetOffset + readOffset);
+
+#pragma warning disable CA2000
 					IReadOnlyMemoryOwner<byte> trailingValue = ReadOnlyMemoryOwner.Create<byte>(memory, trailingPacket);
 					if (!_cache.TryAdd(trailingKey, trailingValue))
 					{
 						trailingValue.Dispose();
 					}
 					trailingPacket = null;
+#pragma warning restore CA2000
 
 					readOffset += signature.HeaderLength;
 				}
diff --git a/Engine/Source/Programs/Shared/EpicGames.Horde/Storage/Clients/BundleStorageClient.cs b/Engine/Source/Programs/Shared/EpicGames.Horde/Storage/Clients/BundleStorageClient.cs
--- a/Engine/Source/Programs/Shared/EpicGames.Horde/Storage/Clients/BundleStorageClient.cs
+++ b/Engine/Source/Programs/Shared/EpicGames.Horde/Storage/Clients/BundleStorageClient.cs
@@ -57,10 +57,10 @@ namespace EpicGames.Horde.Storage.Clients
 					Bundles.V1.BundleHeader header = await _storageClient.ReadHeaderAsync(_inner.GetLocator(), cancellationToken);
 					if (header.Exports.Count > 0)
 					{
-					foreach (BlobLocator import in header.Imports)
-					{
-						refs.Add(_storageClient.CreateBlobHandle(new BlobLocator(import.Path)));
-					}
+						foreach (BlobLocator import in header.Imports)
+						{
+							refs.Add(_storageClient.CreateBlobHandle(new BlobLocator(import.Path)));
+						}
 					}
 					else
 					{
@@ -82,22 +82,22 @@ namespace EpicGames.Horde.Storage.Clients
 			{
 				IReadOnlyMemoryOwner<byte> data = await _inner.ReadBodyAsync(cancellationToken);
 
-				List<BlobLocator> importLocators = ReadImports(data.Memory).ToList();
+				List<BlobLocator> importLocators = ReadImportsFromData(data.Memory).ToList();
 				List<IBlobHandle> importHandles = importLocators.ConvertAll(x => _storageClient.CreateBlobHandle(x));
 
 				return new BlobDataWithOwner(Bundle.BlobType, data.Memory, importHandles, data);
 			}
 
-			IEnumerable<BlobLocator> ReadImports(ReadOnlyMemory<byte> data)
+			IEnumerable<BlobLocator> ReadImportsFromData(ReadOnlyMemory<byte> data)
 			{
 				BundleSignature signature = Bundle.ReadSignature(data.Span);
 				if (signature.Version <= BundleVersion.LatestV1)
 				{
-					return ReadImportsV1(data);
+					return ReadImportsFromDataV1(data);
 				}
 				else if (signature.Version <= BundleVersion.LatestV2)
 				{
-					return ReadImportsV2(data);
+					return ReadImportsFromDataV2(data);
 				}
 				else
 				{
@@ -105,13 +105,13 @@ namespace EpicGames.Horde.Storage.Clients
 				}
 			}
 
-			static IEnumerable<BlobLocator> ReadImportsV1(ReadOnlyMemory<byte> data)
+			static IEnumerable<BlobLocator> ReadImportsFromDataV1(ReadOnlyMemory<byte> data)
 			{
 				BundleHeader header = BundleHeader.Read(data);
 				return header.Imports.Select(x => x.BaseLocator);
 			}
 
-			IEnumerable<BlobLocator> ReadImportsV2(ReadOnlyMemory<byte> data)
+			IEnumerable<BlobLocator> ReadImportsFromDataV2(ReadOnlyMemory<byte> data)
 			{
 				HashSet<BlobLocator> locators = new HashSet<BlobLocator>();
 				while (data.Length > 0)
