Horde: Fix a use-after-dispose race in the bundle cache due to order in which reference count is updated. Also null out fields in PacketReader to allow catching use of disposed objects sooner.

diff --git a/Engine/Source/Programs/Shared/EpicGames.Horde/Storage/Bundles/BundleCache.cs b/Engine/Source/Programs/Shared/EpicGames.Horde/Storage/Bundles/BundleCache.cs
--- a/Engine/Source/Programs/Shared/EpicGames.Horde/Storage/Bundles/BundleCache.cs
+++ b/Engine/Source/Programs/Shared/EpicGames.Horde/Storage/Bundles/BundleCache.cs
@@ -340,25 +340,33 @@ namespace EpicGames.Horde.Storage
 			where TKey : notnull
 			where TValue : class, IDisposable
 		{
-			CacheValue? item;
-			lock (_lockObject)
+			CacheValue? item = null;
+			try
 			{
-				CacheValue? untypedItem;
-				if (_itemLookup.TryGetValue(key, out untypedItem))
-				{
-					item = untypedItem;
-				}
-				else
+				lock (_lockObject)
 				{
-					item = new CacheValue(key, async () => await createAsync(key, _cancellationSource.Token));
-					_items.AddFirst(item);
-					_itemLookup.Add(key, item);
+					CacheValue? untypedItem;
+					if (_itemLookup.TryGetValue(key, out untypedItem))
+					{
+						item = untypedItem;
+					}
+					else
+					{
+						item = new CacheValue(key, async () => await createAsync(key, _cancellationSource.Token));
+						_items.AddFirst(item);
+						_itemLookup.Add(key, item);
+					}
+					item.AddRef(); // Don't allow the item to be freed while we wait for it
 				}
-			}
 
-			TValue value = (TValue)await item.InitTask.WaitAsync(cancellationToken);
-			item.AddRef();
-			return new CacheValueHandle<TValue>(item, value);
+				TValue value = (TValue)await item.InitTask.WaitAsync(cancellationToken);
+				return new CacheValueHandle<TValue>(item, value);
+			}
+			catch
+			{
+				item?.Release();
+				throw;
+			}
 		}
 
 		void CreateSpace(long size)
diff --git a/Engine/Source/Programs/Shared/EpicGames.Horde/Storage/Bundles/V2/PacketReader.cs b/Engine/Source/Programs/Shared/EpicGames.Horde/Storage/Bundles/V2/PacketReader.cs
--- a/Engine/Source/Programs/Shared/EpicGames.Horde/Storage/Bundles/V2/PacketReader.cs
+++ b/Engine/Source/Programs/Shared/EpicGames.Horde/Storage/Bundles/V2/PacketReader.cs
@@ -14,9 +14,10 @@ namespace EpicGames.Horde.Storage.Bundles.V2
 		readonly BundleCache _cache;
 		readonly IBlobHandle _bundleHandle;
 		readonly PacketHandle _packetHandle;
-		readonly Packet _decodedPacket;
-		readonly IRefCountedHandle _memoryOwner;
-		readonly IBlobHandle?[] _cachedImportHandles;
+
+		Packet _decodedPacket;
+		IBlobHandle?[] _cachedImportHandles;
+		IRefCountedHandle _memoryOwner;
 
 		/// <summary>
 		/// Accessor for the underlying packet data
@@ -44,7 +45,17 @@ namespace EpicGames.Horde.Storage.Bundles.V2
 		}
 
 		/// <inheritdoc/>
-		public void Dispose() => _memoryOwner?.Dispose();
+		public void Dispose()
+		{
+			if(_memoryOwner != null)
+			{
+				_memoryOwner.Dispose();
+				_memoryOwner = null!;
+			}
+
+			_decodedPacket = null!;
+			_cachedImportHandles = null!;
+		}
 
 		/// <summary>
 		/// Reads this packet in its entirety
