Horde: Fix another use-after-free error when extracting data from new bundle format.

diff --git a/Engine/Source/Programs/Horde/Horde/Commands/Vcs/VcsCheckout.cs b/Engine/Source/Programs/Horde/Horde/Commands/Vcs/VcsCheckout.cs
--- a/Engine/Source/Programs/Horde/Horde/Commands/Vcs/VcsCheckout.cs
+++ b/Engine/Source/Programs/Horde/Horde/Commands/Vcs/VcsCheckout.cs
@@ -117,8 +117,7 @@ namespace Horde.Commands.Vcs
 		static async Task<FileState> CheckoutFileAsync(FileEntry fileRef, FileInfo fileInfo, ILogger logger)
 		{
 			logger.LogInformation("Updating {File} to {Hash}", fileInfo, fileRef.Hash);
-			ChunkedDataNode fileNode = await fileRef.ExpandAsync();
-			await fileNode.CopyToFileAsync(fileInfo, CancellationToken.None);
+			await ChunkedDataNode.CopyToFileAsync(fileRef.Handle, fileInfo, CancellationToken.None);
 			fileInfo.Refresh();
 			return new FileState(fileInfo, fileRef.Hash);
 		}
diff --git a/Engine/Source/Programs/Shared/EpicGames.Core/IMemoryAllocator.cs b/Engine/Source/Programs/Shared/EpicGames.Core/IMemoryAllocator.cs
--- a/Engine/Source/Programs/Shared/EpicGames.Core/IMemoryAllocator.cs
+++ b/Engine/Source/Programs/Shared/EpicGames.Core/IMemoryAllocator.cs
@@ -4,6 +4,7 @@ using System;
 using System.Buffers;
 using System.IO.MemoryMappedFiles;
 using System.Runtime.InteropServices;
+using System.Threading;
 
 namespace EpicGames.Core
 {
@@ -29,19 +30,42 @@ namespace EpicGames.Core
 	{
 		class MemoryOwner : IMemoryOwner<byte>
 		{
-			public Memory<byte> Memory { get; }
+			readonly ManagedHeapAllocator _outer;
 
-			public MemoryOwner(int size) => Memory = new byte[size];
-			public void Dispose() { }
+			public Memory<byte> Memory { get; private set; }
+
+			public MemoryOwner(ManagedHeapAllocator outer, int size)
+			{
+				_outer = outer;
+				Interlocked.Add(ref _outer._allocatedSize, size);
+				Memory = new byte[size];
+			}
+
+			public void Dispose()
+			{
+#if DEBUG
+				Memory.Span.Fill(0xfe);
+#endif
+
+				Interlocked.Add(ref _outer._allocatedSize, -Memory.Length);
+				Memory = Memory<byte>.Empty;
+			}
 		}
 
+		long _allocatedSize;
+
 		/// <summary>
 		/// Default shared instance
 		/// </summary>
-		public static ManagedHeapAllocator Instance { get; } = new ManagedHeapAllocator();
+		public static ManagedHeapAllocator Shared { get; } = new ManagedHeapAllocator();
+
+		/// <summary>
+		/// Currently allocated size using this heap
+		/// </summary>
+		public long AllocatedSize => _allocatedSize;
 
 		/// <inheritdoc/>
-		public IMemoryOwner<byte> Alloc(int minSize) => new MemoryOwner(minSize);
+		public IMemoryOwner<byte> Alloc(int minSize) => new MemoryOwner(this, minSize);
 	}
 
 	/// <summary>
@@ -72,13 +96,17 @@ namespace EpicGames.Core
 	{
 		unsafe class Allocation : MemoryManager<byte>
 		{
+			readonly GlobalHeapAllocator _outer;
 			IntPtr _handle;
 			int _length;
 
-			public Allocation(int size)
+			public Allocation(GlobalHeapAllocator outer, int size)
 			{
+				_outer = outer;
 				_handle = Marshal.AllocHGlobal(size);
 				_length = size;
+
+				Interlocked.Add(ref _outer._allocatedSize, _length);
 			}
 
 			byte* GetPointer() => (byte*)_handle.ToPointer();
@@ -99,14 +127,28 @@ namespace EpicGames.Core
 				{
 					Marshal.FreeHGlobal(_handle);
 					_handle = IntPtr.Zero;
+
+					Interlocked.Add(ref _outer._allocatedSize, -_length);
 				}
 
 				_length = -1;
 			}
 		}
 
+		long _allocatedSize;
+
+		/// <summary>
+		/// Shared heap allocator
+		/// </summary>
+		public static GlobalHeapAllocator Shared { get; } = new GlobalHeapAllocator();
+
+		/// <summary>
+		/// Currently allocated size using this heap
+		/// </summary>
+		public long AllocatedSize => _allocatedSize;
+
 		/// <inheritdoc/>
-		public IMemoryOwner<byte> Alloc(int minSize) => new Allocation(minSize);
+		public IMemoryOwner<byte> Alloc(int minSize) => new Allocation(this, minSize);
 	}
 	
 	/// <summary>
diff --git a/Engine/Source/Programs/Shared/EpicGames.Horde/Storage/Bundles/V2/PacketReader.cs b/Engine/Source/Programs/Shared/EpicGames.Horde/Storage/Bundles/V2/PacketReader.cs
--- a/Engine/Source/Programs/Shared/EpicGames.Horde/Storage/Bundles/V2/PacketReader.cs
+++ b/Engine/Source/Programs/Shared/EpicGames.Horde/Storage/Bundles/V2/PacketReader.cs
@@ -117,19 +117,20 @@ namespace EpicGames.Horde.Storage.Bundles.V2
 			{
 				PacketImport import = _decodedPacket.GetImport(index);
 
+				Utf8String fragment = import.Fragment.Clone(); // Handles may outlive the current packet reader, so duplicate the fragment string.
 				switch (import.BaseIdx)
 				{
 					case PacketImport.InvalidBaseIdx:
-						importHandle = _storageClient.CreateBlobHandle(new BlobLocator(import.Fragment));
+						importHandle = _storageClient.CreateBlobHandle(new BlobLocator(fragment));
 						break;
 					case PacketImport.CurrentBundleBaseIdx:
-						importHandle = new PacketHandle(_storageClient, _bundleHandle, import.Fragment, _cache);
+						importHandle = new PacketHandle(_storageClient, _bundleHandle, fragment, _cache);
 						break;
 					case PacketImport.CurrentPacketBaseIdx:
-						importHandle = new ExportHandle(_packetHandle, import.Fragment);
+						importHandle = new ExportHandle(_packetHandle, fragment);
 						break;
 					default:
-						importHandle = GetImportHandle(import.BaseIdx).GetFragmentHandle(import.Fragment);
+						importHandle = GetImportHandle(import.BaseIdx).GetFragmentHandle(fragment);
 						break;
 				}
 
diff --git a/Engine/Source/Programs/Shared/EpicGames.Horde/Storage/Nodes/ChunkedDataNode.cs b/Engine/Source/Programs/Shared/EpicGames.Horde/Storage/Nodes/ChunkedDataNode.cs
--- a/Engine/Source/Programs/Shared/EpicGames.Horde/Storage/Nodes/ChunkedDataNode.cs
+++ b/Engine/Source/Programs/Shared/EpicGames.Horde/Storage/Nodes/ChunkedDataNode.cs
@@ -55,10 +55,11 @@ namespace EpicGames.Horde.Storage.Nodes
 		/// <summary>
 		/// Extracts the contents of this node to a file
 		/// </summary>
+		/// <param name="handle">Handle to the data to read</param>
 		/// <param name="file">File to write with the contents of this node</param>
 		/// <param name="cancellationToken">Cancellation token for the operation</param>
 		/// <returns></returns>
-		public async Task CopyToFileAsync(FileInfo file, CancellationToken cancellationToken)
+		public static async Task CopyToFileAsync(IBlobHandle handle, FileInfo file, CancellationToken cancellationToken)
 		{
 			if(file.Exists && (file.Attributes & FileAttributes.ReadOnly) != 0)
 			{
@@ -66,7 +67,7 @@ namespace EpicGames.Horde.Storage.Nodes
 			}
 			using (FileStream stream = file.Open(FileMode.Create, FileAccess.Write, FileShare.Read))
 			{
-				await CopyToStreamAsync(stream, cancellationToken);
+				await CopyToStreamAsync(handle, stream, cancellationToken);
 			}
 		}
 
@@ -188,7 +189,7 @@ namespace EpicGames.Horde.Storage.Nodes
 		public LeafChunkedDataNode(IBlobReader reader)
 		{
 			// Keep this code in sync with CopyToStreamAsync
-			Data = reader.GetMemory();
+			Data = reader.GetMemory().ToArray();
 		}
 
 		/// <inheritdoc/>
diff --git a/Engine/Source/Programs/Shared/EpicGames.Horde/Storage/Nodes/FileEntry.cs b/Engine/Source/Programs/Shared/EpicGames.Horde/Storage/Nodes/FileEntry.cs
--- a/Engine/Source/Programs/Shared/EpicGames.Horde/Storage/Nodes/FileEntry.cs
+++ b/Engine/Source/Programs/Shared/EpicGames.Horde/Storage/Nodes/FileEntry.cs
@@ -156,8 +156,7 @@ namespace EpicGames.Horde.Storage.Nodes
 		{
 			try
 			{
-				ChunkedDataNode node = await ExpandAsync(cancellationToken);
-				await node.CopyToFileAsync(file, cancellationToken);
+				await ChunkedDataNode.CopyToFileAsync(Handle, file, cancellationToken);
 
 				if ((Flags & FileEntryFlags.Executable) != 0)
 				{
