Horde: Fix tests.

diff --git a/Engine/Source/Programs/Shared/EpicGames.Horde.Tests/BundleCacheTests.cs b/Engine/Source/Programs/Shared/EpicGames.Horde.Tests/BundleCacheTests.cs
--- a/Engine/Source/Programs/Shared/EpicGames.Horde.Tests/BundleCacheTests.cs
+++ b/Engine/Source/Programs/Shared/EpicGames.Horde.Tests/BundleCacheTests.cs
@@ -41,7 +41,7 @@ namespace EpicGames.Horde.Tests
 		[TestMethod]
 		public async Task FindOrAddAsync()
 		{
-			await using BundleCache cache = new BundleCache(new BundleCacheOptions { MaxSize = 1024 }, ManagedHeapAllocator.Instance);
+			await using BundleCache cache = new BundleCache(new BundleCacheOptions { MaxSize = 1024 }, ManagedHeapAllocator.Shared);
 
 			Assert.AreEqual(0, cache.CurrentSize);
 
@@ -84,7 +84,7 @@ namespace EpicGames.Horde.Tests
 		[TestMethod]
 		public async Task TrimAsync()
 		{
-			await using BundleCache cache = new BundleCache(new BundleCacheOptions { MaxSize = 1024 }, ManagedHeapAllocator.Instance);
+			await using BundleCache cache = new BundleCache(new BundleCacheOptions { MaxSize = 1024 }, ManagedHeapAllocator.Shared);
 
 			Assert.AreEqual(0, cache.CurrentSize);
 
@@ -106,7 +106,7 @@ namespace EpicGames.Horde.Tests
 		[TestMethod]
 		public async Task TryAddSuccessAsync()
 		{
-			await using BundleCache cache = new BundleCache(new BundleCacheOptions { MaxSize = 1024 }, ManagedHeapAllocator.Instance);
+			await using BundleCache cache = new BundleCache(new BundleCacheOptions { MaxSize = 1024 }, ManagedHeapAllocator.Shared);
 
 			Assert.AreEqual(0, cache.CurrentSize);
 
@@ -129,7 +129,7 @@ namespace EpicGames.Horde.Tests
 		[TestMethod]
 		public async Task TryAddFailureAsync()
 		{
-			await using BundleCache cache = new BundleCache(new BundleCacheOptions { MaxSize = 1024 }, ManagedHeapAllocator.Instance);
+			await using BundleCache cache = new BundleCache(new BundleCacheOptions { MaxSize = 1024 }, ManagedHeapAllocator.Shared);
 
 			Assert.AreEqual(0, cache.CurrentSize);
 
diff --git a/Engine/Source/Programs/Shared/EpicGames.Horde.Tests/BundlePacketTests.cs b/Engine/Source/Programs/Shared/EpicGames.Horde.Tests/BundlePacketTests.cs
--- a/Engine/Source/Programs/Shared/EpicGames.Horde.Tests/BundlePacketTests.cs
+++ b/Engine/Source/Programs/Shared/EpicGames.Horde.Tests/BundlePacketTests.cs
@@ -67,7 +67,7 @@ namespace EpicGames.Horde.Tests
 			FakeHandle bundleHandle = new FakeHandle(null, null);
 			FakeHandle packetHandle = new FakeHandle(bundleHandle, null);
 
-			using PacketWriter writer = new PacketWriter(bundleHandle, packetHandle, ManagedHeapAllocator.Instance, new object());
+			using PacketWriter writer = new PacketWriter(bundleHandle, packetHandle, ManagedHeapAllocator.Shared, new object());
 			int type1 = writer.FindOrAddType(blobType1);
 			int type2 = writer.FindOrAddType(blobType2);
 			int type1b = writer.FindOrAddType(blobType1); // Check it's deduped
@@ -101,7 +101,7 @@ namespace EpicGames.Horde.Tests
 			packet.Encode(BundleCompressionFormat.None, memoryWriter);
 
 			// Check we can read it back in
-			IRefCountedHandle<Packet> handle = Packet.Decode(memoryWriter.WrittenMemory, ManagedHeapAllocator.Instance);
+			IRefCountedHandle<Packet> handle = Packet.Decode(memoryWriter.WrittenMemory, ManagedHeapAllocator.Shared);
 			packet = handle.Target;
 
 			Assert.AreEqual(2, packet.GetTypeCount());
