Horde: More diagnostic information for tool download.

diff --git a/Engine/Source/Programs/Horde/Horde.Server/Tools/ToolCollection.cs b/Engine/Source/Programs/Horde/Horde.Server/Tools/ToolCollection.cs
--- a/Engine/Source/Programs/Horde/Horde.Server/Tools/ToolCollection.cs
+++ b/Engine/Source/Programs/Horde/Horde.Server/Tools/ToolCollection.cs
@@ -448,15 +448,7 @@ namespace Horde.Server.Tools
 			{
 				DirectoryNode node = await client.ReadRefAsync<DirectoryNode>(deployment.RefName, DateTime.UtcNow - TimeSpan.FromDays(2.0), cancellationToken);
 				_logger.LogInformation("Tool {ToolId} has {NumDirectories}, {NumFiles}, total size {NumBytes}", tool.Id, node.NameToFile.Count, node.NameToDirectory.Count, node.Length);
-
-				using Stream zipStream = node.AsZipStream();
-
-				MemoryStream memoryStream = new MemoryStream();
-				await zipStream.CopyToAsync(memoryStream, cancellationToken);
-				_logger.LogInformation("Tool {ToolId} zip stream is {NumBytes}", tool.Id, memoryStream.Length);
-
-				memoryStream.Position = 0;
-				return memoryStream;
+				return node.AsZipStream(logger: _logger);
 			}
 			else
 			{
diff --git a/Engine/Source/Programs/Shared/EpicGames.Horde/Storage/Nodes/DirectoryNode.cs b/Engine/Source/Programs/Shared/EpicGames.Horde/Storage/Nodes/DirectoryNode.cs
--- a/Engine/Source/Programs/Shared/EpicGames.Horde/Storage/Nodes/DirectoryNode.cs
+++ b/Engine/Source/Programs/Shared/EpicGames.Horde/Storage/Nodes/DirectoryNode.cs
@@ -857,8 +857,9 @@ namespace EpicGames.Horde.Storage.Nodes
 		/// Returns a stream containing the zipped contents of this directory
 		/// </summary>
 		/// <param name="filter">Filter for files to include in the zip</param>
+		/// <param name="logger">Logger for diagnostic output</param>
 		/// <returns>Stream containing zipped archive data</returns>
-		public Stream AsZipStream(FileFilter? filter = null) => new DirectoryNodeZipStream(this, filter);
+		public Stream AsZipStream(FileFilter? filter = null, ILogger? logger = null) => new DirectoryNodeZipStream(this, filter, logger);
 	}
 
 	/// <summary>
@@ -1023,6 +1024,7 @@ namespace EpicGames.Horde.Storage.Nodes
 		public override long Position { get => _position; set => throw new NotImplementedException(); }
 
 		readonly Pipe _pipe;
+		readonly ILogger? _logger;
 		readonly BackgroundTask _backgroundTask;
 
 		long _position;
@@ -1033,10 +1035,12 @@ namespace EpicGames.Horde.Storage.Nodes
 		/// </summary>
 		/// <param name="node">Root node to copy from</param>
 		/// <param name="filter">Filter for files to include in the zip</param>
-		public DirectoryNodeZipStream(DirectoryNode node, FileFilter? filter)
+		/// <param name="logger">Optional logger for debug tracing</param>
+		public DirectoryNodeZipStream(DirectoryNode node, FileFilter? filter, ILogger? logger)
 		{
 			_pipe = new Pipe();
-			_backgroundTask = BackgroundTask.StartNew(ctx => CopyToPipeAsync(node, filter, _pipe.Writer, ctx));
+			_backgroundTask = BackgroundTask.StartNew(ctx => CopyToPipeAsync(node, filter, _pipe.Writer, logger, ctx));
+			_logger = logger;
 		}
 
 		/// <inheritdoc/>
@@ -1057,6 +1061,7 @@ namespace EpicGames.Horde.Storage.Nodes
 
 				if (result.IsCompleted && _current.Length == 0)
 				{
+					_logger?.LogInformation("Zip file was read to end");
 					return 0;
 				}
 			}
@@ -1080,25 +1085,31 @@ namespace EpicGames.Horde.Storage.Nodes
 			return length;
 		}
 
-		static async Task CopyToPipeAsync(DirectoryNode node, FileFilter? filter, PipeWriter writer, CancellationToken cancellationToken)
+		static async Task CopyToPipeAsync(DirectoryNode node, FileFilter? filter, PipeWriter writer, ILogger? logger, CancellationToken cancellationToken)
 		{
 			using Stream outputStream = writer.AsStream();
 			using ZipArchive archive = new ZipArchive(outputStream, ZipArchiveMode.Create);
-			await CopyFilesAsync(node, "", filter, archive, cancellationToken);
+			await CopyFilesAsync(node, "", filter, archive, logger, cancellationToken);
 		}
 
-		static async Task CopyFilesAsync(DirectoryNode directory, string prefix, FileFilter? filter, ZipArchive archive, CancellationToken cancellationToken)
+		static async Task CopyFilesAsync(DirectoryNode directory, string prefix, FileFilter? filter, ZipArchive archive, ILogger? logger, CancellationToken cancellationToken)
 		{
+			int numDirs = directory.Directories.Count;
+			int numFiles = directory.Files.Count;
+
+			int numCopiedDirs = 0;
 			foreach (DirectoryEntry directoryEntry in directory.Directories)
 			{
 				string directoryPath = $"{prefix}{directoryEntry.Name}/";
 				if (filter == null || filter.PossiblyMatches(directoryPath))
 				{
 					DirectoryNode node = await directoryEntry.ExpandAsync(cancellationToken);
-					await CopyFilesAsync(node, directoryPath, filter, archive, cancellationToken);
+					await CopyFilesAsync(node, directoryPath, filter, archive, logger, cancellationToken);
 				}
+				numCopiedDirs++;
 			}
 
+			int numCopiedFiles = 0;
 			foreach (FileEntry fileEntry in directory.Files)
 			{
 				string filePath = $"{prefix}{fileEntry}";
@@ -1118,7 +1129,10 @@ namespace EpicGames.Horde.Storage.Nodes
 					using Stream entryStream = entry.Open();
 					await fileEntry.CopyToStreamAsync(entryStream, cancellationToken);
 				}
+				numCopiedFiles++;
 			}
+
+			logger?.LogInformation("Zip file path {Prefix} has {NumCopiedDirectories}/{NumDirectories} directories, {NumCopiedFiles}/{NumFiles} files", prefix, numCopiedDirs, numDirs, numCopiedFiles, numFiles);
 		}
 
 		/// <inheritdoc/>
