Horde: Fix swallowing of exceptions in zip file background tasks.

diff --git a/Engine/Source/Programs/Horde/Horde.Server/Server/RpcService.cs b/Engine/Source/Programs/Horde/Horde.Server/Server/RpcService.cs
--- a/Engine/Source/Programs/Horde/Horde.Server/Server/RpcService.cs
+++ b/Engine/Source/Programs/Horde/Horde.Server/Server/RpcService.cs
@@ -498,7 +498,7 @@ namespace Horde.Server.Server
 				throw new StructuredRpcException(StatusCode.NotFound, $"Missing tool version {version}");
 			}
 
-			using Stream stream = await _toolCollection.GetDeploymentZipAsync(tool, deployment, context.CancellationToken);
+			await using Stream stream = await _toolCollection.GetDeploymentZipAsync(tool, deployment, context.CancellationToken);
 			using (IMemoryOwner<byte> buffer = MemoryPool<byte>.Shared.Rent(128 * 1024))
 			{
 				long totalWritten = 0;
diff --git a/Engine/Source/Programs/Shared/EpicGames.Horde/Storage/IBlobHandle.cs b/Engine/Source/Programs/Shared/EpicGames.Horde/Storage/IBlobHandle.cs
--- a/Engine/Source/Programs/Shared/EpicGames.Horde/Storage/IBlobHandle.cs
+++ b/Engine/Source/Programs/Shared/EpicGames.Horde/Storage/IBlobHandle.cs
@@ -221,7 +221,7 @@ namespace EpicGames.Horde.Storage
 		public static bool TryGetLocator(this IBlobHandle handle, [NotNullWhen(true)] out BlobLocator locator)
 		{
 			Utf8StringBuilder builder = new Utf8StringBuilder();
-			if (AppendLocator(handle, builder))
+			if (TryAppendLocator(handle, builder))
 			{
 				locator = new BlobLocator(builder.ToUtf8String());
 				return true;
@@ -236,15 +236,16 @@ namespace EpicGames.Horde.Storage
 		/// <summary>
 		/// Builds a full locator for a blob by traversing the outer chain
 		/// </summary>
-		static bool AppendLocator(IBlobHandle handle, Utf8StringBuilder builder)
+		static bool TryAppendLocator(IBlobHandle handle, Utf8StringBuilder builder)
 		{
 			IBlobHandle? outer = handle.Outer;
 			if (outer != null)
 			{
-				if (!AppendLocator(outer, builder))
+				if (!TryAppendLocator(outer, builder))
 				{
 					return false;
 				}
+
 				if (outer.Outer == null)
 				{
 					builder.Append('#');
diff --git a/Engine/Source/Programs/Shared/EpicGames.Horde/Storage/Nodes/DirectoryNode.cs b/Engine/Source/Programs/Shared/EpicGames.Horde/Storage/Nodes/DirectoryNode.cs
--- a/Engine/Source/Programs/Shared/EpicGames.Horde/Storage/Nodes/DirectoryNode.cs
+++ b/Engine/Source/Programs/Shared/EpicGames.Horde/Storage/Nodes/DirectoryNode.cs
@@ -1043,6 +1043,14 @@ namespace EpicGames.Horde.Storage.Nodes
 			_logger = logger;
 		}
 
+		/// <inheritdoc/>
+		protected override void Dispose(bool disposing)
+		{
+			base.Dispose(disposing);
+
+			_backgroundTask.DisposeAsync().AsTask().Wait();
+		}
+
 		/// <inheritdoc/>
 		public override async ValueTask DisposeAsync()
 		{
@@ -1061,7 +1069,9 @@ namespace EpicGames.Horde.Storage.Nodes
 
 				if (result.IsCompleted && _current.Length == 0)
 				{
-					_logger?.LogInformation("Zip file was read to end");
+					// Wait for the background thread to finish; it may error/have errored, and we want to re-throw on this thread before completing the read.
+					await _backgroundTask.StopAsync(cancellationToken);
+					_logger?.LogInformation("Zip file pipe was read to end");
 					return 0;
 				}
 			}
