OidcToken: Add support for WSL2

Launches browser for authentication under WSL2. The localhost IP bound by server under Linux is accessible on Windows host.

diff --git a/Engine/Source/Programs/Shared/EpicGames.OIDC/OidcTokenManager.cs b/Engine/Source/Programs/Shared/EpicGames.OIDC/OidcTokenManager.cs
--- a/Engine/Source/Programs/Shared/EpicGames.OIDC/OidcTokenManager.cs
+++ b/Engine/Source/Programs/Shared/EpicGames.OIDC/OidcTokenManager.cs
@@ -426,7 +426,12 @@ namespace EpicGames.OIDC
 			}
 			else if (RuntimeInformation.IsOSPlatform(OSPlatform.Linux))
 			{
-				return Process.Start("xdg-open", url);
+				if (IsRunningWsl())
+				{
+					return Process.Start("wslview", url);
+				}
+
+				return Process.Start("xdg-open", url);	
 			}
 			else if (RuntimeInformation.IsOSPlatform(OSPlatform.OSX))
 			{
@@ -438,6 +443,18 @@ namespace EpicGames.OIDC
 			}
 		}
 
+		private static bool IsRunningWsl()
+		{
+			string versionFile = "/proc/version";
+			if (!RuntimeInformation.IsOSPlatform(OSPlatform.Linux) || !File.Exists(versionFile))
+			{
+				return false;
+			}
+
+			string version = File.ReadAllText(versionFile);
+			return version.Contains("microsoft", StringComparison.InvariantCultureIgnoreCase) || version.Contains("wsl2", StringComparison.InvariantCultureIgnoreCase);
+		}
+
 		[SupportedOSPlatform("windows")]
 		static bool TryReadRegistryString(RegistryKey rootKey, string path, string? name, [NotNullWhen(true)] out string? value)
 		{
