Horde: Change blobs to serialize GUIDs in UE-style (4 x 32-bit LE values) rather than Windows/NET-style (1 x 32-bit LE, 2 x 16-bit LE, 8 x 8-bit values), to ease interop with the engine.

In order to prevent invalidating existing data, strings for GUIDs for all existing blob types have been updated to deserialize to the same value in UE format.

diff --git a/Engine/Source/Programs/Horde/Horde.Server.Tests/BlobStoreTests.cs b/Engine/Source/Programs/Horde/Horde.Server.Tests/BlobStoreTests.cs
--- a/Engine/Source/Programs/Horde/Horde.Server.Tests/BlobStoreTests.cs
+++ b/Engine/Source/Programs/Horde/Horde.Server.Tests/BlobStoreTests.cs
@@ -15,7 +15,7 @@ namespace Horde.Server.Tests
 	[TestClass]
 	public class BlobStoreTests : TestSetup
 	{
-		static readonly BlobType s_blobType = new BlobType(Guid.Parse("{AFDF76A7-5333-4DEE-B837-B5F5CA511245}"), 1);
+		static readonly BlobType s_blobType = new BlobType("{AFDF76A7-4DEE-5333-F5B5-37B8451251CA}", 1);
 
 		IStorageClient CreateStorageClient()
 		{
diff --git a/Engine/Source/Programs/Horde/Horde.Server.Tests/GcServiceTests.cs b/Engine/Source/Programs/Horde/Horde.Server.Tests/GcServiceTests.cs
--- a/Engine/Source/Programs/Horde/Horde.Server.Tests/GcServiceTests.cs
+++ b/Engine/Source/Programs/Horde/Horde.Server.Tests/GcServiceTests.cs
@@ -100,7 +100,7 @@ namespace Horde.Server.Tests
 				}
 			}
 
-			BlobType blobType = new BlobType(Guid.Parse("{AFDF76A7-5333-4DEE-B837-B5F5CA511245}"), 0);
+			BlobType blobType = new BlobType(Guid.Parse("{AFDF76A7-4DEE-5333-F5B5-37B8451251CA}"), 0);
 
 			BlobLocator[] locators = new BlobLocator[children.Length];
 			for (int idx = numNodes - 1; idx >= 0; idx--)
diff --git a/Engine/Source/Programs/Horde/Horde.Server.Tests/StorageServiceTests.cs b/Engine/Source/Programs/Horde/Horde.Server.Tests/StorageServiceTests.cs
--- a/Engine/Source/Programs/Horde/Horde.Server.Tests/StorageServiceTests.cs
+++ b/Engine/Source/Programs/Horde/Horde.Server.Tests/StorageServiceTests.cs
@@ -18,10 +18,10 @@ namespace Horde.Server.Tests
 			StorageService storageService = ServiceProvider.GetRequiredService<StorageService>();
 			using IStorageClient client = storageService.CreateClient(new NamespaceId("memory"));
 
-			BlobType type1 = new BlobType(Guid.Parse("{11C2D886-3349-4164-946F-E9D10BD12E3D}"), 0);
+			BlobType type1 = new BlobType(Guid.Parse("{11C2D886-4164-3349-D1E9-6F943D2ED10B}"), 0);
 			byte[] data1 = new byte[] { 1, 2, 3 };
 
-			BlobType type2 = new BlobType(Guid.Parse("{6CB3A005-26BA-4787-86D2-793ED13771CB}"), 0);
+			BlobType type2 = new BlobType(Guid.Parse("{6CB3A005-4787-26BA-3E79-D286CB7137D1}"), 0);
 			byte[] data2 = new byte[] { 4, 5, 6 };
 
 			IBlobHandle handle1a;
diff --git a/Engine/Source/Programs/Horde/Horde.Server/Perforce/PerforceReplicator.cs b/Engine/Source/Programs/Horde/Horde.Server/Perforce/PerforceReplicator.cs
--- a/Engine/Source/Programs/Horde/Horde.Server/Perforce/PerforceReplicator.cs
+++ b/Engine/Source/Programs/Horde/Horde.Server/Perforce/PerforceReplicator.cs
@@ -37,7 +37,7 @@ namespace Horde.Server.Perforce
 	/// </summary>
 	class PerforceReplicator
 	{
-		[NodeType("{8C874966-2E89-4273-A4AB-AC9F5491C86D}")]
+		[BlobType("{8C874966-4273-2E89-9FAC-ABA46DC89154}")]
 		class SyncNode : Node
 		{
 			public int Change { get; }
diff --git a/Engine/Source/Programs/Shared/EpicGames.Core/GuidUtils.cs b/Engine/Source/Programs/Shared/EpicGames.Core/GuidUtils.cs
new file mode 100644
--- /dev/null
+++ b/Engine/Source/Programs/Shared/EpicGames.Core/GuidUtils.cs
@@ -0,0 +1,84 @@
+// Copyright Epic Games, Inc. All Rights Reserved.
+
+using System;
+using System.Buffers.Binary;
+using System.Collections.Generic;
+using System.Linq;
+using System.Text;
+using System.Threading.Tasks;
+
+namespace EpicGames.Core
+{
+	/// <summary>
+	/// Utility methods for serializing GUIDs in UE format. UE stores GUIDs as 4 x 32-bit LE values, whereas NET stores GUIDs as 1 x 32-bit LE, 2 x 16-bit LE, 8 x 8-bit values.
+	/// </summary>
+	public static class GuidUtils
+	{
+		// String order:              0   1   2   3   4   5   6   7   8   9   a   b   c   d   e   f
+		//
+		// NET serialiazation order:  3   2   1   0 | 5   4 | 7   6 | 8   9   a   b   c   d   e   f
+		// UE serialziation order:    3   2   1   0 | 7   6   5   4 | b   a   9   8 | f   e   d   c
+
+		/// <summary>
+		/// Reads a GUID serialized in UE order
+		/// </summary>
+		public static Guid ReadGuidUnrealOrder(ReadOnlySpan<byte> span)
+		{
+			Span<byte> buffer = stackalloc byte[16];
+
+			buffer[0x0] = span[0x0];
+			buffer[0x1] = span[0x1];
+			buffer[0x2] = span[0x2];
+			buffer[0x3] = span[0x3];
+
+			buffer[0x4] = span[0x6];
+			buffer[0x5] = span[0x7];
+			buffer[0x6] = span[0x4];
+			buffer[0x7] = span[0x5];
+
+			buffer[0x8] = span[0xb];
+			buffer[0x9] = span[0xa];
+			buffer[0xa] = span[0x9];
+			buffer[0xb] = span[0x8];
+
+			buffer[0xc] = span[0xf];
+			buffer[0xd] = span[0xe];
+			buffer[0xe] = span[0xd];
+			buffer[0xf] = span[0xc];
+
+			return new Guid(buffer);
+		}
+
+		/// <summary>
+		/// Reads a GUID serialized in UE order
+		/// </summary>
+		public static void WriteGuidUnrealOrder(Span<byte> span, Guid guid)
+		{
+			Span<byte> buffer = stackalloc byte[16];
+			if (!guid.TryWriteBytes(buffer))
+			{
+				throw new InvalidOperationException("Unable to write guid to buffer");
+			}
+
+			span[0x0] = buffer[0x0];
+			span[0x1] = buffer[0x1];
+			span[0x2] = buffer[0x2];
+			span[0x3] = buffer[0x3];
+
+			span[0x4] = buffer[0x6];
+			span[0x5] = buffer[0x7];
+			span[0x6] = buffer[0x4];
+			span[0x7] = buffer[0x5];
+
+			span[0x8] = buffer[0xb];
+			span[0x9] = buffer[0xa];
+			span[0xa] = buffer[0x9];
+			span[0xb] = buffer[0x8];
+
+			span[0xc] = buffer[0xf];
+			span[0xd] = buffer[0xe];
+			span[0xe] = buffer[0xd];
+			span[0xf] = buffer[0xc];
+		}
+	}
+}
diff --git a/Engine/Source/Programs/Shared/EpicGames.Core/MemoryReader.cs b/Engine/Source/Programs/Shared/EpicGames.Core/MemoryReader.cs
--- a/Engine/Source/Programs/Shared/EpicGames.Core/MemoryReader.cs
+++ b/Engine/Source/Programs/Shared/EpicGames.Core/MemoryReader.cs
@@ -201,16 +201,27 @@ namespace EpicGames.Core
 		}
 
 		/// <summary>
-		/// Reads a GUID from the memory buffer
+		/// Reads a GUID from the memory buffer using NET-ordered serialization.
 		/// </summary>
 		/// <param name="reader">Reader to deserialize from</param>
-		public static Guid ReadGuid(this IMemoryReader reader)
+		public static Guid ReadGuidNetOrder(this IMemoryReader reader)
 		{
 			Guid guid = new Guid(reader.GetMemory(16).Slice(0, 16).Span);
 			reader.Advance(16);
 			return guid;
 		}
 
+		/// <summary>
+		/// Reads a GUID from the memory buffer using UE-ordered serialization.
+		/// </summary>
+		/// <param name="reader">Reader to deserialize from</param>
+		public static Guid ReadGuidUnrealOrder(this IMemoryReader reader)
+		{
+			Guid guid = GuidUtils.ReadGuidUnrealOrder(reader.GetSpan(16));
+			reader.Advance(16);
+			return guid;
+		}
+
 		/// <summary>
 		/// Reads a sequence of bytes from the buffer
 		/// </summary>
diff --git a/Engine/Source/Programs/Shared/EpicGames.Core/MemoryWriter.cs b/Engine/Source/Programs/Shared/EpicGames.Core/MemoryWriter.cs
--- a/Engine/Source/Programs/Shared/EpicGames.Core/MemoryWriter.cs
+++ b/Engine/Source/Programs/Shared/EpicGames.Core/MemoryWriter.cs
@@ -834,11 +834,11 @@ namespace EpicGames.Core
 		}
 
 		/// <summary>
-		/// Writes a Guid to the memory writer
+		/// Writes a Guid to the memory writer using NET order serialization.
 		/// </summary>
 		/// <param name="writer">Writer to serialize to</param>
 		/// <param name="guid">Value to write</param>
-		public static void WriteGuid(this IMemoryWriter writer, Guid guid)
+		public static void WriteGuidNetOrder(this IMemoryWriter writer, Guid guid)
 		{
 			Memory<byte> buffer = writer.GetMemory(16);
 			if (!guid.TryWriteBytes(buffer.Slice(0, 16).Span))
@@ -848,6 +848,17 @@ namespace EpicGames.Core
 			writer.Advance(16);
 		}
 
+		/// <summary>
+		/// Writes a Guid to the memory writer (uses UE rather than NET serialization order).
+		/// </summary>
+		/// <param name="writer">Writer to serialize to</param>
+		/// <param name="guid">Value to write</param>
+		public static void WriteGuidUnrealOrder(this IMemoryWriter writer, Guid guid)
+		{
+			GuidUtils.WriteGuidUnrealOrder(writer.GetSpan(16), guid);
+			writer.Advance(16);
+		}
+
 		/// <summary>
 		/// Appends a span of bytes to the buffer, without a length field.
 		/// </summary>
diff --git a/Engine/Source/Programs/Shared/EpicGames.Horde.Tests/BundlePacketTests.cs b/Engine/Source/Programs/Shared/EpicGames.Horde.Tests/BundlePacketTests.cs
--- a/Engine/Source/Programs/Shared/EpicGames.Horde.Tests/BundlePacketTests.cs
+++ b/Engine/Source/Programs/Shared/EpicGames.Horde.Tests/BundlePacketTests.cs
@@ -58,7 +58,7 @@ namespace EpicGames.Horde.Tests
 		[TestMethod]
 		public void TestTypes()
 		{
-			Guid guid = Guid.Parse("{30F2CA83-B5D4-494C-8802-0661454CCD58}");
+			Guid guid = Guid.Parse("{30F2CA83-494C-B5D4-6106-028858CD4C45}");
 
 			// Write out some data
 			BlobType blobType1 = new BlobType(guid, 1);
diff --git a/Engine/Source/Programs/Shared/EpicGames.Horde.Tests/BundleStoreTests.cs b/Engine/Source/Programs/Shared/EpicGames.Horde.Tests/BundleStoreTests.cs
--- a/Engine/Source/Programs/Shared/EpicGames.Horde.Tests/BundleStoreTests.cs
+++ b/Engine/Source/Programs/Shared/EpicGames.Horde.Tests/BundleStoreTests.cs
@@ -32,7 +32,7 @@ namespace EpicGames.Horde.Tests
 			Assert.IsTrue(a.SequenceEqual(b));
 		}
 
-		[NodeType("{F63606D4-5DBB-4061-A655-6F444F65229E}")]
+		[BlobType("{F63606D4-4061-5DBB-446F-55A69E22654F}")]
 		class TextNode : Node
 		{
 			public string Text { get; }
@@ -77,9 +77,9 @@ namespace EpicGames.Horde.Tests
 			byte[] payload = payloadWriter.WrittenMemory.ToArray();
 
 			List<BlobType> types = new List<BlobType>();
-			types.Add(new BlobType(Guid.Parse("F63606D4-5DBB-4061-A655-6F444F65229E"), 1));
+			types.Add(new BlobType("{F63606D4-4061-5DBB-446F-55A69E22654F}", 1));
 
-			List<BundleExport> exports = new List<BundleExport>();
+			List <BundleExport> exports = new List<BundleExport>();
 			exports.Add(new BundleExport(0, 0, 0, payload.Length, Array.Empty<BundleExportRef>()));
 
 			List<BundlePacket> packets = new List<BundlePacket>();
@@ -121,7 +121,7 @@ namespace EpicGames.Horde.Tests
 			Assert.AreEqual(1, blobStore.Refs.Count);
 		}
 
-		[NodeType("{F63606D4-5DBB-4061-A655-6F444F65229F}")]
+		[BlobType("{F63606D4-4061-5DBB-446F-55A69F22654F}")]
 		class SimpleNode : Node
 		{
 			public ReadOnlySequence<byte> Data { get; }
diff --git a/Engine/Source/Programs/Shared/EpicGames.Horde.Tests/StorageClientTests.cs b/Engine/Source/Programs/Shared/EpicGames.Horde.Tests/StorageClientTests.cs
--- a/Engine/Source/Programs/Shared/EpicGames.Horde.Tests/StorageClientTests.cs
+++ b/Engine/Source/Programs/Shared/EpicGames.Horde.Tests/StorageClientTests.cs
@@ -15,7 +15,7 @@ namespace EpicGames.Horde.Tests
 	[TestClass]
 	public class StorageClientTests
 	{
-		[NodeType("{99601905-A089-4F6E-87F1-D603EEAF1B71}", 1)]
+		[BlobType("{99601905-4F6E-A089-03D6-F187711BAFEE}", 1)]
 		class TestNode : Node
 		{
 			public int Value { get; }
diff --git a/Engine/Source/Programs/Shared/EpicGames.Horde/Logs/LogChunkNode.cs b/Engine/Source/Programs/Shared/EpicGames.Horde/Logs/LogChunkNode.cs
--- a/Engine/Source/Programs/Shared/EpicGames.Horde/Logs/LogChunkNode.cs
+++ b/Engine/Source/Programs/Shared/EpicGames.Horde/Logs/LogChunkNode.cs
@@ -14,7 +14,7 @@ namespace EpicGames.Horde.Logs
 	/// <summary>
 	/// Read-only buffer for log text, with indexed line offsets.
 	/// </summary>
-	[NodeType("{7020B6CA-0F72-4174-B6AA-06AA60A3EF30}", 1)]
+	[BlobType("{7020B6CA-4174-0F72-AA06-AAB630EFA360}", 1)]
 	public class LogChunkNode : Node
 	{
 		/// <summary>
diff --git a/Engine/Source/Programs/Shared/EpicGames.Horde/Logs/LogIndexNode.cs b/Engine/Source/Programs/Shared/EpicGames.Horde/Logs/LogIndexNode.cs
--- a/Engine/Source/Programs/Shared/EpicGames.Horde/Logs/LogIndexNode.cs
+++ b/Engine/Source/Programs/Shared/EpicGames.Horde/Logs/LogIndexNode.cs
@@ -28,7 +28,7 @@ namespace EpicGames.Horde.Logs
 	/// Since alignment of ngrams may not match alignment of ngrams in the search term, we offset the search term by
 	/// 1-4 bytes and include the union of blocks matching at any offset.
 	/// </summary>
-	[NodeType("{BAE1A00E-FD63-474E-A804-8081E20134F9}", 1)]
+	[BlobType("{BAE1A00E-474E-FD63-8180-04A8F93401E2}", 1)]
 	public class LogIndexNode : Node
 	{
 		/// <summary>
diff --git a/Engine/Source/Programs/Shared/EpicGames.Horde/Logs/LogNode.cs b/Engine/Source/Programs/Shared/EpicGames.Horde/Logs/LogNode.cs
--- a/Engine/Source/Programs/Shared/EpicGames.Horde/Logs/LogNode.cs
+++ b/Engine/Source/Programs/Shared/EpicGames.Horde/Logs/LogNode.cs
@@ -32,7 +32,7 @@ namespace EpicGames.Horde.Logs
 	/// <summary>
 	/// Represents an entire log
 	/// </summary>
-	[NodeType("{274DF8F7-9E87-4B4F-8AD5-318CDB25AD33}", 1)]
+	[BlobType("{274DF8F7-4B4F-9E87-8C31-D58A33AD25DB}", 1)]
 	public class LogNode : Node
 	{
 		/// <summary>
diff --git a/Engine/Source/Programs/Shared/EpicGames.Horde/Storage/BlobType.cs b/Engine/Source/Programs/Shared/EpicGames.Horde/Storage/BlobType.cs
--- a/Engine/Source/Programs/Shared/EpicGames.Horde/Storage/BlobType.cs
+++ b/Engine/Source/Programs/Shared/EpicGames.Horde/Storage/BlobType.cs
@@ -2,6 +2,7 @@
 
 using System;
 using System.Buffers.Binary;
+using EpicGames.Core;
 
 namespace EpicGames.Horde.Storage
 {
@@ -20,14 +21,22 @@ namespace EpicGames.Horde.Storage
 		/// <summary>
 		/// Blob type used to indicate a leaf node. Data is an opaque blob and has no references.
 		/// </summary>
-		public static BlobType Leaf { get; } = new BlobType(Guid.Parse("{1080B643-8015-4A4D-ADBB-C2FB586C8D81}"), 1);
+		public static BlobType Leaf { get; } = new BlobType("{1080B643-4A4D-8015-FBC2-BBAD818D6C58}", 1);
+
+		/// <summary>
+		/// Constructor
+		/// </summary>
+		public BlobType(string guid, int version)
+			: this(Guid.Parse(guid), version)
+		{
+		}
 
 		/// <summary>
 		/// Deserialize a type from a byte span
 		/// </summary>
 		public static BlobType Read(ReadOnlySpan<byte> span)
 		{
-			Guid guid = new Guid(span.Slice(0, 16));
+			Guid guid = GuidUtils.ReadGuidUnrealOrder(span);
 			int version = BinaryPrimitives.ReadInt32LittleEndian(span.Slice(16));
 
 			return new BlobType(guid, version);
@@ -38,7 +47,7 @@ namespace EpicGames.Horde.Storage
 		/// </summary>
 		public readonly void Write(Span<byte> data)
 		{
-			Guid.TryWriteBytes(data.Slice(0, 16));
+			GuidUtils.WriteGuidUnrealOrder(data, Guid);
 			BinaryPrimitives.WriteInt32LittleEndian(data.Slice(16), Version);
 		}
 
diff --git a/Engine/Source/Programs/Shared/EpicGames.Horde/Storage/Bundles/Bundle.cs b/Engine/Source/Programs/Shared/EpicGames.Horde/Storage/Bundles/Bundle.cs
--- a/Engine/Source/Programs/Shared/EpicGames.Horde/Storage/Bundles/Bundle.cs
+++ b/Engine/Source/Programs/Shared/EpicGames.Horde/Storage/Bundles/Bundle.cs
@@ -80,7 +80,7 @@ namespace EpicGames.Horde.Storage.Bundles
 		/// <summary>
 		/// Blob type for bundles
 		/// </summary>
-		public static BlobType BlobType { get; } = new BlobType(Guid.Parse("{7C5BA294-2D21-4F92-85BE-852F48CC4C1E}"), 1);
+		public static BlobType BlobType { get; } = new BlobType("{7C5BA294-4F92-2D21-2F85-BE851E4CCC48}", 1);
 
 		/// <summary>
 		/// Number of bytes in a signature when serialized
diff --git a/Engine/Source/Programs/Shared/EpicGames.Horde/Storage/Bundles/V1/BundleHeader.cs b/Engine/Source/Programs/Shared/EpicGames.Horde/Storage/Bundles/V1/BundleHeader.cs
--- a/Engine/Source/Programs/Shared/EpicGames.Horde/Storage/Bundles/V1/BundleHeader.cs
+++ b/Engine/Source/Programs/Shared/EpicGames.Horde/Storage/Bundles/V1/BundleHeader.cs
@@ -305,7 +305,7 @@ namespace EpicGames.Horde.Storage.Bundles.V1
 
 			for (int typeIdx = 0; typeIdx < numTypes; typeIdx++)
 			{
-				Guid guid = reader.ReadGuid();
+				Guid guid = reader.ReadGuidUnrealOrder();
 				int serializerVersion = (int)reader.ReadUnsignedVarInt();
 
 				types.Add(new BlobType(guid, serializerVersion));
diff --git a/Engine/Source/Programs/Shared/EpicGames.Horde/Storage/Bundles/V2/Packet.cs b/Engine/Source/Programs/Shared/EpicGames.Horde/Storage/Bundles/V2/Packet.cs
--- a/Engine/Source/Programs/Shared/EpicGames.Horde/Storage/Bundles/V2/Packet.cs
+++ b/Engine/Source/Programs/Shared/EpicGames.Horde/Storage/Bundles/V2/Packet.cs
@@ -56,7 +56,7 @@ namespace EpicGames.Horde.Storage.Bundles.V2
 		/// <summary>
 		/// Type for packet blobs
 		/// </summary>
-		public static BlobType BlobType { get; } = new BlobType(Guid.Parse("{CD9A04EF-CAC1-47D3-A605-2A498130E651}"), 1);
+		public static BlobType BlobType { get; } = new BlobType("{CD9A04EF-47D3-CAC1-492A-05A651E63081}", 1);
 
 		class DebugProxy
 		{
diff --git a/Engine/Source/Programs/Shared/EpicGames.Horde/Storage/IBlobReader.cs b/Engine/Source/Programs/Shared/EpicGames.Horde/Storage/IBlobReader.cs
--- a/Engine/Source/Programs/Shared/EpicGames.Horde/Storage/IBlobReader.cs
+++ b/Engine/Source/Programs/Shared/EpicGames.Horde/Storage/IBlobReader.cs
@@ -17,7 +17,7 @@ namespace EpicGames.Horde.Storage
 		BlobType Type { get; }
 
 		/// <summary>
-		/// Version of the current node, as specified via <see cref="NodeTypeAttribute"/>
+		/// Version of the current node, as specified via <see cref="BlobTypeAttribute"/>
 		/// </summary>
 		int Version { get; }
 
@@ -43,7 +43,7 @@ namespace EpicGames.Horde.Storage
 		public BlobType Type => _blobData.Type;
 
 		/// <summary>
-		/// Version of the current node, as specified via <see cref="NodeTypeAttribute"/>
+		/// Version of the current node, as specified via <see cref="BlobTypeAttribute"/>
 		/// </summary>
 		public int Version => Type.Version;
 
diff --git a/Engine/Source/Programs/Shared/EpicGames.Horde/Storage/Node.cs b/Engine/Source/Programs/Shared/EpicGames.Horde/Storage/Node.cs
--- a/Engine/Source/Programs/Shared/EpicGames.Horde/Storage/Node.cs
+++ b/Engine/Source/Programs/Shared/EpicGames.Horde/Storage/Node.cs
@@ -7,6 +7,7 @@ using System.Reflection;
 using System.Reflection.Emit;
 using System.Threading;
 using System.Threading.Tasks;
+using EpicGames.Core;
 
 namespace EpicGames.Horde.Storage
 {
@@ -14,7 +15,7 @@ namespace EpicGames.Horde.Storage
 	/// Attribute used to define a factory for a particular node type
 	/// </summary>
 	[AttributeUsage(AttributeTargets.Class)]
-	public sealed class NodeTypeAttribute : Attribute // TODO: THIS SHOULD BE NODEATTRIBUTE, not NODETYPEATTRIBUTE
+	public sealed class BlobTypeAttribute : Attribute
 	{
 		/// <summary>
 		/// Name of the type to store in the bundle header
@@ -29,7 +30,7 @@ namespace EpicGames.Horde.Storage
 		/// <summary>
 		/// Constructor
 		/// </summary>
-		public NodeTypeAttribute(string guid, int version = 1)
+		public BlobTypeAttribute(string guid, int version = 1)
 		{
 			Guid = guid;
 			Version = version;
@@ -73,15 +74,15 @@ namespace EpicGames.Horde.Storage
 		static readonly ConcurrentDictionary<Guid, Type> s_guidToType = new ConcurrentDictionary<Guid, Type>();
 		static readonly ConcurrentDictionary<Guid, Func<BlobData, Node>> s_guidToDeserializer = new ConcurrentDictionary<Guid, Func<BlobData, Node>>();
 
-		static BlobType CreateNodeType(NodeTypeAttribute attribute)
+		static BlobType CreateNodeType(BlobTypeAttribute attribute)
 		{
-			return new BlobType(Guid.Parse(attribute.Guid), attribute.Version);
+			return new BlobType(attribute.Guid, attribute.Version);
 		}
 
 		/// <summary>
 		/// Attempts to get the concrete type with the given node. The type must have been registered via a previous call to <see cref="RegisterType(Type)"/>.
 		/// </summary>
-		/// <param name="guid">Guid specified in the <see cref="NodeTypeAttribute"/></param>
+		/// <param name="guid">Guid specified in the <see cref="BlobTypeAttribute"/></param>
 		/// <param name="type">On success, receives the C# type associated with this GUID</param>
 		/// <returns>True if the type was found</returns>
 		public static bool TryGetConcreteType(Guid guid, [NotNullWhen(true)] out Type? type) => s_guidToType.TryGetValue(guid, out type);
@@ -100,10 +101,10 @@ namespace EpicGames.Horde.Storage
 				{
 					if (!s_typeToNodeType.TryGetValue(type, out nodeType))
 					{
-						NodeTypeAttribute? attribute = type.GetCustomAttribute<NodeTypeAttribute>();
+						BlobTypeAttribute? attribute = type.GetCustomAttribute<BlobTypeAttribute>();
 						if (attribute == null)
 						{
-							throw new InvalidOperationException($"Missing {nameof(NodeTypeAttribute)} from type {type.Name}");
+							throw new InvalidOperationException($"Missing {nameof(BlobTypeAttribute)} from type {type.Name}");
 						}
 						nodeType = s_typeToNodeType.GetOrAdd(type, CreateNodeType(attribute));
 					}
@@ -151,10 +152,10 @@ namespace EpicGames.Horde.Storage
 		/// <exception cref="NotImplementedException"></exception>
 		public static void RegisterType(Type type)
 		{
-			NodeTypeAttribute? attribute = type.GetCustomAttribute<NodeTypeAttribute>();
+			BlobTypeAttribute? attribute = type.GetCustomAttribute<BlobTypeAttribute>();
 			if (attribute == null)
 			{
-				throw new InvalidOperationException($"Missing {nameof(NodeTypeAttribute)} from type {type.Name}");
+				throw new InvalidOperationException($"Missing {nameof(BlobTypeAttribute)} from type {type.Name}");
 			}
 			RegisterType(type, attribute);
 		}
@@ -166,7 +167,7 @@ namespace EpicGames.Horde.Storage
 		public static void RegisterType<T>() where T : Node => RegisterType(typeof(T));
 
 		/// <summary>
-		/// Register all node types with the <see cref="NodeTypeAttribute"/> from the given assembly
+		/// Register all node types with the <see cref="BlobTypeAttribute"/> from the given assembly
 		/// </summary>
 		/// <param name="assembly">Assembly to register types from</param>
 		public static void RegisterTypesFromAssembly(Assembly assembly)
@@ -178,7 +179,7 @@ namespace EpicGames.Horde.Storage
 				{
 					if (type.IsClass)
 					{
-						NodeTypeAttribute? attribute = type.GetCustomAttribute<NodeTypeAttribute>();
+						BlobTypeAttribute? attribute = type.GetCustomAttribute<BlobTypeAttribute>();
 						if (attribute != null)
 						{
 							RegisterType(type, attribute);
@@ -188,7 +189,7 @@ namespace EpicGames.Horde.Storage
 			}
 		}
 
-		static void RegisterType(Type type, NodeTypeAttribute attribute)
+		static void RegisterType(Type type, BlobTypeAttribute attribute)
 		{
 			BlobType nodeType = CreateNodeType(attribute);
 			s_typeToNodeType.TryAdd(type, nodeType);
diff --git a/Engine/Source/Programs/Shared/EpicGames.Horde/Storage/Nodes/CbNode.cs b/Engine/Source/Programs/Shared/EpicGames.Horde/Storage/Nodes/CbNode.cs
--- a/Engine/Source/Programs/Shared/EpicGames.Horde/Storage/Nodes/CbNode.cs
+++ b/Engine/Source/Programs/Shared/EpicGames.Horde/Storage/Nodes/CbNode.cs
@@ -9,7 +9,7 @@ namespace EpicGames.Horde.Storage.Nodes
 	/// <summary>
 	/// A node containing arbitrary compact binary data
 	/// </summary>
-	[NodeType("{34A0793F-8364-42F4-8632-98A71C843229}", 1)]
+	[BlobType("{34A0793F-42F4-8364-A798-32862932841C}", 1)]
 	public class CbNode : Node
 	{
 		class HandleMapper
diff --git a/Engine/Source/Programs/Shared/EpicGames.Horde/Storage/Nodes/ChunkedDataNode.cs b/Engine/Source/Programs/Shared/EpicGames.Horde/Storage/Nodes/ChunkedDataNode.cs
--- a/Engine/Source/Programs/Shared/EpicGames.Horde/Storage/Nodes/ChunkedDataNode.cs
+++ b/Engine/Source/Programs/Shared/EpicGames.Horde/Storage/Nodes/ChunkedDataNode.cs
@@ -163,7 +163,7 @@ namespace EpicGames.Horde.Storage.Nodes
 	/// <summary>
 	/// File node that contains a chunk of data
 	/// </summary>
-	[NodeType("{B27AFB68-9E20-4A4B-A4D8-788A4098D439}", 1)]
+	[BlobType("{B27AFB68-4A4B-9E20-8A78-D8A439D49840}", 1)]
 	public sealed class LeafChunkedDataNode : ChunkedDataNode
 	{
 		/// <summary>
@@ -386,7 +386,7 @@ namespace EpicGames.Horde.Storage.Nodes
 	/// <summary>
 	/// An interior file node
 	/// </summary>
-	[NodeType("{F4DEDDBC-70CB-4C7A-8347-F011AFCCCDB9}", 2)]
+	[BlobType("{F4DEDDBC-4C7A-70CB-11F0-4783B9CDCCAF}", 2)]
 	public class InteriorChunkedDataNode : ChunkedDataNode
 	{
 		/// <summary>
diff --git a/Engine/Source/Programs/Shared/EpicGames.Horde/Storage/Nodes/CommitNode.cs b/Engine/Source/Programs/Shared/EpicGames.Horde/Storage/Nodes/CommitNode.cs
--- a/Engine/Source/Programs/Shared/EpicGames.Horde/Storage/Nodes/CommitNode.cs
+++ b/Engine/Source/Programs/Shared/EpicGames.Horde/Storage/Nodes/CommitNode.cs
@@ -9,7 +9,7 @@ namespace EpicGames.Horde.Storage.Nodes
 	/// <summary>
 	/// A node representing commit metadata
 	/// </summary>
-	[NodeType("{64D50724-6B22-41C0-A890-B51CD6241817}", 1)]
+	[BlobType("{64D50724-41C0-6B22-1CB5-90A8171824D6}", 1)]
 	public class CommitNode : Node
 	{
 		/// <summary>
@@ -96,7 +96,7 @@ namespace EpicGames.Horde.Storage.Nodes
 			Message = reader.ReadString();
 			Time = reader.ReadDateTime();
 			Contents = new DirectoryNodeRef(reader);
-			Metadata = reader.ReadDictionary(() => reader.ReadGuid(), () => reader.ReadHashedNodeRef());
+			Metadata = reader.ReadDictionary(() => reader.ReadGuidUnrealOrder(), () => reader.ReadHashedNodeRef());
 		}
 
 		/// <inheritdoc/>
@@ -111,7 +111,7 @@ namespace EpicGames.Horde.Storage.Nodes
 			writer.WriteString(Message);
 			writer.WriteDateTime(Time);
 			writer.WriteHashedNodeRef(Contents);
-			writer.WriteDictionary(Metadata, key => writer.WriteGuid(key), value => writer.WriteHashedNodeRef(value));
+			writer.WriteDictionary(Metadata, key => writer.WriteGuidUnrealOrder(key), value => writer.WriteHashedNodeRef(value));
 		}
 	}
 }
diff --git a/Engine/Source/Programs/Shared/EpicGames.Horde/Storage/Nodes/DdcRefNode.cs b/Engine/Source/Programs/Shared/EpicGames.Horde/Storage/Nodes/DdcRefNode.cs
--- a/Engine/Source/Programs/Shared/EpicGames.Horde/Storage/Nodes/DdcRefNode.cs
+++ b/Engine/Source/Programs/Shared/EpicGames.Horde/Storage/Nodes/DdcRefNode.cs
@@ -9,7 +9,7 @@ namespace Horde.Server.Ddc
 	/// <summary>
 	/// A node containing ref data
 	/// </summary>
-	[NodeType("{0C7E5F25-454B-4B55-9B4A-F4635106D074}", 1)]
+	[BlobType("{0C7E5F25-4B55-454B-63F4-4A9B74D00651}", 1)]
 	public class DdcRefNode : Node
 	{
 		/// <summary>
diff --git a/Engine/Source/Programs/Shared/EpicGames.Horde/Storage/Nodes/DirectoryNode.cs b/Engine/Source/Programs/Shared/EpicGames.Horde/Storage/Nodes/DirectoryNode.cs
--- a/Engine/Source/Programs/Shared/EpicGames.Horde/Storage/Nodes/DirectoryNode.cs
+++ b/Engine/Source/Programs/Shared/EpicGames.Horde/Storage/Nodes/DirectoryNode.cs
@@ -140,7 +140,7 @@ namespace EpicGames.Horde.Storage.Nodes
 	/// <summary>
 	/// A directory node
 	/// </summary>
-	[NodeType("{0714EC11-291A-4D07-867F-E78AD6809979}", 1)]
+	[BlobType("{0714EC11-4D07-291A-8AE7-7F86799980D6}", 1)]
 	public class DirectoryNode : Node
 	{
 		readonly SortedDictionary<string, FileEntry> _nameToFileEntry = new SortedDictionary<string, FileEntry>(StringComparer.Ordinal);
diff --git a/Engine/Source/Programs/Shared/EpicGames.Horde/Storage/Nodes/RedirectNode.cs b/Engine/Source/Programs/Shared/EpicGames.Horde/Storage/Nodes/RedirectNode.cs
--- a/Engine/Source/Programs/Shared/EpicGames.Horde/Storage/Nodes/RedirectNode.cs
+++ b/Engine/Source/Programs/Shared/EpicGames.Horde/Storage/Nodes/RedirectNode.cs
@@ -5,7 +5,7 @@ namespace EpicGames.Horde.Storage.Nodes
 	/// <summary>
 	/// A node containing arbitrary compact binary data
 	/// </summary>
-	[NodeType("{BE09E54F-7A6B-47CA-BCAF-972A38153B18}", 1)]
+	[BlobType("{BE09E54F-47CA-7A6B-2A97-AFBC183B1538}", 1)]
 	public class RedirectNode : Node 
 	{
 		/// <summary>
diff --git a/Engine/Source/Programs/UnrealBuildTool/Artifacts/HordeStorageArtifactCache.cs b/Engine/Source/Programs/UnrealBuildTool/Artifacts/HordeStorageArtifactCache.cs
--- a/Engine/Source/Programs/UnrealBuildTool/Artifacts/HordeStorageArtifactCache.cs
+++ b/Engine/Source/Programs/UnrealBuildTool/Artifacts/HordeStorageArtifactCache.cs
@@ -123,7 +123,7 @@ namespace UnrealBuildTool.Artifacts
 	/// <summary>
 	/// Horde node that represents a collection of action nodes 
 	/// </summary>
-	[NodeType("{E8DBCD77-861D-4CAE-B77F-5807D26E2533}")]
+	[BlobType("{E8DBCD77-4CAE-861D-0758-7FB733256ED2}")]
 	class ArtifactActionCollectionNode : Node
 	{
 
