Fix static analyzer warnings.

diff --git a/Engine/Source/Programs/Shared/EpicGames.Core/GuidUtils.cs b/Engine/Source/Programs/Shared/EpicGames.Core/GuidUtils.cs
--- a/Engine/Source/Programs/Shared/EpicGames.Core/GuidUtils.cs
+++ b/Engine/Source/Programs/Shared/EpicGames.Core/GuidUtils.cs
@@ -1,11 +1,6 @@
 // Copyright Epic Games, Inc. All Rights Reserved.
 
 using System;
-using System.Buffers.Binary;
-using System.Collections.Generic;
-using System.Linq;
-using System.Text;
-using System.Threading.Tasks;
 
 namespace EpicGames.Core
 {
