Horde: Add external IP resolver for use by compute resource allocation

diff --git a/Engine/Source/Programs/Shared/EpicGames.Horde.Tests/ExternalIpResolverTests.cs b/Engine/Source/Programs/Shared/EpicGames.Horde.Tests/ExternalIpResolverTests.cs
new file mode 100644
--- /dev/null
+++ b/Engine/Source/Programs/Shared/EpicGames.Horde.Tests/ExternalIpResolverTests.cs
@@ -0,0 +1,58 @@
+﻿// Copyright Epic Games, Inc. All Rights Reserved.
+
+using System;
+using System.Net;
+using System.Net.Http;
+using System.Threading;
+using System.Threading.Tasks;
+using EpicGames.Horde.Compute;
+using Microsoft.VisualStudio.TestTools.UnitTesting;
+
+namespace EpicGames.Horde.Tests;
+
+[TestClass]
+public class ExternalIpResolverTests
+{
+	private class StubMessageHandler : HttpMessageHandler
+	{
+		private readonly HttpStatusCode _statusCode;
+		private readonly string _content;
+
+		public StubMessageHandler(HttpStatusCode statusCode, string content)
+		{
+			_statusCode = statusCode;
+			_content = content;
+		}
+
+		protected override Task<HttpResponseMessage> SendAsync(HttpRequestMessage request, CancellationToken cancellationToken)
+		{
+			return Task.FromResult(new HttpResponseMessage(_statusCode) { Content = new StringContent(_content) });
+		}
+	}
+	
+	[TestMethod]
+	public async Task BasicAsync()
+	{
+		Assert.AreEqual(IPAddress.Parse("192.168.2.3"), await GetIpAsync("192.168.2.3"));
+		await Assert.ThrowsExceptionAsync<ExternalIpResolverException>(() => GetIpAsync("bad-ip"));
+		await Assert.ThrowsExceptionAsync<ExternalIpResolverException>(() => GetIpAsync("192.168.2.3", HttpStatusCode.NotFound));
+	}
+	
+	[TestMethod]
+	[Ignore]
+	public async Task IntegrationAsync()
+	{
+		using HttpClient httpClient = new ();
+		ExternalIpResolver resolver = new (httpClient);
+		IPAddress externalIp = await resolver.GetExternalIpAddressAsync();
+		Console.WriteLine("External IP: " + externalIp);
+	}
+
+	private static async Task<IPAddress> GetIpAsync(string content, HttpStatusCode statusCode = HttpStatusCode.OK, CancellationToken cancellationToken = default)
+	{
+		using StubMessageHandler handler = new (statusCode, content);
+		using HttpClient httpClient = new (handler);
+		ExternalIpResolver resolver = new (httpClient);
+		return await resolver.GetExternalIpAddressAsync(cancellationToken);
+	}
+}
\ No newline at end of file
diff --git a/Engine/Source/Programs/Shared/EpicGames.Horde/Compute/ExternalIpResolver.cs b/Engine/Source/Programs/Shared/EpicGames.Horde/Compute/ExternalIpResolver.cs
new file mode 100644
--- /dev/null
+++ b/Engine/Source/Programs/Shared/EpicGames.Horde/Compute/ExternalIpResolver.cs
@@ -0,0 +1,89 @@
+﻿// Copyright Epic Games, Inc. All Rights Reserved.
+
+using System;
+using System.Collections.Generic;
+using System.Net;
+using System.Net.Http;
+using System.Threading;
+using System.Threading.Tasks;
+
+namespace EpicGames.Horde.Compute;
+
+/// <summary>
+/// Exception for external IP resolver
+/// </summary>
+public class ExternalIpResolverException : Exception
+{
+	/// <inheritdoc />
+	public ExternalIpResolverException(string? message) : base(message)
+	{
+	}
+
+	/// <inheritdoc />
+	public ExternalIpResolverException(string? message, Exception? innerException) : base(message, innerException)
+	{
+	}
+}
+
+/// <summary>
+/// Find public IP address of local machine by querying a third-party IP lookup service
+/// </summary>
+public class ExternalIpResolver
+{
+	private readonly HttpClient _httpClient;
+	private readonly List<string> _ipLookupUrls = new()
+	{
+		"http://whatismyip.akamai.com", "http://checkip.amazonaws.com", "http://ifconfig.me/ip"
+	};
+
+	/// <summary>
+	/// Constructor
+	/// </summary>
+	/// <param name="httpClient"></param>
+	public ExternalIpResolver(HttpClient httpClient)
+	{
+		_httpClient = httpClient;
+	}
+
+	/// <summary>
+	/// Get the external, public-facing IP of local machine
+	/// </summary>
+	/// <param name="cancellationToken">Cancellation token</param>
+	/// <returns>External IP address</returns>
+	/// <exception cref="ExternalIpResolverException">If unable to resolve</exception>
+	public Task<IPAddress> GetExternalIpAddressAsync(CancellationToken cancellationToken = default)
+	{
+		ExternalIpResolverException? lastException = null;
+		foreach (string lookupUrl in _ipLookupUrls)
+		{
+			Uri url = new(lookupUrl);
+			try
+			{
+				return GetExternalIpAddressAsync(url, cancellationToken);
+			}
+			catch (ExternalIpResolverException e)
+			{
+				lastException = e;
+			}
+		}
+
+		throw new ExternalIpResolverException("Exhausted list of all IP resolvers", lastException);
+	}
+
+	private async Task<IPAddress> GetExternalIpAddressAsync(Uri url, CancellationToken cancellationToken = default)
+	{
+		HttpResponseMessage res = await _httpClient.GetAsync(url, cancellationToken);
+		if (!res.IsSuccessStatusCode)
+		{
+			throw new ExternalIpResolverException($"Non-successful response code: {res.StatusCode}");
+		}
+		
+		string content = await res.Content.ReadAsStringAsync(cancellationToken);
+		if (!IPAddress.TryParse(content, out IPAddress? ipAddress))
+		{
+			throw new ExternalIpResolverException($"Failed parsing HTTP body as an IP address: {content}");
+		}
+		
+		return ipAddress;
+	}
+}
\ No newline at end of file
