Horde: Fix static analysis warning.

diff --git a/Engine/Source/Programs/Shared/EpicGames.Horde/Storage/Node.cs b/Engine/Source/Programs/Shared/EpicGames.Horde/Storage/Node.cs
--- a/Engine/Source/Programs/Shared/EpicGames.Horde/Storage/Node.cs
+++ b/Engine/Source/Programs/Shared/EpicGames.Horde/Storage/Node.cs
@@ -7,7 +7,6 @@ using System.Reflection;
 using System.Reflection.Emit;
 using System.Threading;
 using System.Threading.Tasks;
-using EpicGames.Core;
 
 namespace EpicGames.Horde.Storage
 {
