Horde: Remove unused functionality for tracking modified nodes.

diff --git a/Engine/Source/Programs/Shared/EpicGames.Horde/Storage/Node.cs b/Engine/Source/Programs/Shared/EpicGames.Horde/Storage/Node.cs
--- a/Engine/Source/Programs/Shared/EpicGames.Horde/Storage/Node.cs
+++ b/Engine/Source/Programs/Shared/EpicGames.Horde/Storage/Node.cs
@@ -41,25 +41,11 @@ namespace EpicGames.Horde.Storage
 	/// </summary>
 	public abstract class Node
 	{
-		/// <summary>
-		/// Revision number of the node. Incremented whenever the node is modified, and used to track whether nodes are modified between 
-		/// writes starting and completing.
-		/// </summary>
-		public uint Revision { get; private set; }
-
 		/// <summary>
 		/// Accessor for the bundle type definition associated with this node
 		/// </summary>
 		public BlobType NodeType => GetNodeType(GetType());
 
-		/// <summary>
-		/// Mark this node as dirty
-		/// </summary>
-		protected void MarkAsDirty()
-		{
-			Revision++;
-		}
-
 		/// <summary>
 		/// Serialize the contents of this node
 		/// </summary>
diff --git a/Engine/Source/Programs/Shared/EpicGames.Horde/Storage/Nodes/DirectoryNode.cs b/Engine/Source/Programs/Shared/EpicGames.Horde/Storage/Nodes/DirectoryNode.cs
--- a/Engine/Source/Programs/Shared/EpicGames.Horde/Storage/Nodes/DirectoryNode.cs
+++ b/Engine/Source/Programs/Shared/EpicGames.Horde/Storage/Nodes/DirectoryNode.cs
@@ -233,8 +233,6 @@ namespace EpicGames.Horde.Storage.Nodes
 		{
 			_nameToFileEntry.Clear();
 			_nameToDirectoryEntry.Clear();
-
-			MarkAsDirty();
 		}
 
 		/// <summary>
@@ -253,7 +251,6 @@ namespace EpicGames.Horde.Storage.Nodes
 		public void AddFile(FileEntry entry)
 		{
 			_nameToFileEntry[entry.Name] = entry;
-			MarkAsDirty();
 		}
 
 		/// <summary>
@@ -316,15 +313,7 @@ namespace EpicGames.Horde.Storage.Nodes
 		/// </summary>
 		/// <param name="name">Name of the entry to delete</param>
 		/// <returns>True if the entry was found, false otherwise</returns>
-		public bool DeleteFile(string name)
-		{
-			if (_nameToFileEntry.Remove(name))
-			{
-				MarkAsDirty();
-				return true;
-			}
-			return false;
-		}
+		public bool DeleteFile(string name) => _nameToFileEntry.Remove(name);
 
 		/// <summary>
 		/// Attempts to get a file entry from a path
@@ -438,7 +427,6 @@ namespace EpicGames.Horde.Storage.Nodes
 			}
 
 			_nameToDirectoryEntry.Add(entry.Name, entry);
-			MarkAsDirty();
 		}
 
 		/// <summary>
@@ -495,15 +483,7 @@ namespace EpicGames.Horde.Storage.Nodes
 		/// </summary>
 		/// <param name="name">Name of the entry to delete</param>
 		/// <returns>True if the entry was found, false otherwise</returns>
-		public bool DeleteDirectory(string name)
-		{
-			if (_nameToDirectoryEntry.Remove(name))
-			{
-				MarkAsDirty();
-				return true;
-			}
-			return false;
-		}
+		public bool DeleteDirectory(string name) => _nameToDirectoryEntry.Remove(name);
 
 		#endregion
 
