Horde: Downgrade warning to information.

diff --git a/Engine/Source/Programs/Shared/EpicGames.Horde/Compute/ComputeSocket.cs b/Engine/Source/Programs/Shared/EpicGames.Horde/Compute/ComputeSocket.cs
--- a/Engine/Source/Programs/Shared/EpicGames.Horde/Compute/ComputeSocket.cs
+++ b/Engine/Source/Programs/Shared/EpicGames.Horde/Compute/ComputeSocket.cs
@@ -505,7 +505,7 @@ namespace EpicGames.Horde.Compute
 		{
 			if (!await TryReadPacketAsync(transport, id, size, cancellationToken))
 			{
-				_logger.LogWarning("Discarding {Size} bytes received on compute channel {Id}; no buffer attached?", size, id);
+				_logger.LogInformation("Discarding {Size} bytes received on compute channel {Id}; no buffer attached?", size, id);
 
 				int bufferSize = Math.Min(size, 65536);
 				using (IMemoryOwner<byte> buffer = MemoryPool<byte>.Shared.Rent(bufferSize))
