Horde: Cache externally resolved IP

diff --git a/Engine/Source/Programs/Shared/EpicGames.Horde.Tests/ExternalIpResolverTests.cs b/Engine/Source/Programs/Shared/EpicGames.Horde.Tests/ExternalIpResolverTests.cs
--- a/Engine/Source/Programs/Shared/EpicGames.Horde.Tests/ExternalIpResolverTests.cs
+++ b/Engine/Source/Programs/Shared/EpicGames.Horde.Tests/ExternalIpResolverTests.cs
@@ -16,17 +16,17 @@ public class ExternalIpResolverTests
 	private class StubMessageHandler : HttpMessageHandler
 	{
 		private readonly HttpStatusCode _statusCode;
-		private readonly string _content;
+		public string Content { get; set; }
 
 		public StubMessageHandler(HttpStatusCode statusCode, string content)
 		{
 			_statusCode = statusCode;
-			_content = content;
+			Content = content;
 		}
 
 		protected override Task<HttpResponseMessage> SendAsync(HttpRequestMessage request, CancellationToken cancellationToken)
 		{
-			return Task.FromResult(new HttpResponseMessage(_statusCode) { Content = new StringContent(_content) });
+			return Task.FromResult(new HttpResponseMessage(_statusCode) { Content = new StringContent(Content) });
 		}
 	}
 	
@@ -38,6 +38,21 @@ public class ExternalIpResolverTests
 		await Assert.ThrowsExceptionAsync<ExternalIpResolverException>(() => GetIpAsync("192.168.2.3", HttpStatusCode.NotFound));
 	}
 	
+	[TestMethod]
+	public async Task CacheIpAddressAsync()
+	{
+		using StubMessageHandler stub = new (HttpStatusCode.OK, "100.200.1.1");
+		using HttpClient httpClient = new (stub);
+		ExternalIpResolver resolver = new (httpClient);
+		IPAddress ip = await resolver.GetExternalIpAddressAsync(CancellationToken.None);
+		Assert.AreEqual(IPAddress.Parse("100.200.1.1"), ip);
+
+		// Querying for IP again should never reach the IP resolver's HTTP server, so an invalid IP is returned to ensure this
+		stub.Content = "invalid-ip";
+		ip = await resolver.GetExternalIpAddressAsync(CancellationToken.None);
+		Assert.AreEqual(IPAddress.Parse("100.200.1.1"), ip);
+	}
+	
 	[TestMethod]
 	[Ignore]
 	public async Task IntegrationAsync()
diff --git a/Engine/Source/Programs/Shared/EpicGames.Horde/Compute/ExternalIpResolver.cs b/Engine/Source/Programs/Shared/EpicGames.Horde/Compute/ExternalIpResolver.cs
--- a/Engine/Source/Programs/Shared/EpicGames.Horde/Compute/ExternalIpResolver.cs
+++ b/Engine/Source/Programs/Shared/EpicGames.Horde/Compute/ExternalIpResolver.cs
@@ -30,6 +30,11 @@ public class ExternalIpResolverException : Exception
 /// </summary>
 public class ExternalIpResolver
 {
+	/// <summary>
+	/// Cache of last resolved IP address. Once resolved, the address is cached for the lifetime of this class.
+	/// </summary>
+	private IPAddress? _resolvedIp;
+	
 	private readonly HttpClient _httpClient;
 	private readonly List<string> _ipLookupUrls = new()
 	{
@@ -51,15 +56,21 @@ public class ExternalIpResolver
 	/// <param name="cancellationToken">Cancellation token</param>
 	/// <returns>External IP address</returns>
 	/// <exception cref="ExternalIpResolverException">If unable to resolve</exception>
-	public Task<IPAddress> GetExternalIpAddressAsync(CancellationToken cancellationToken = default)
+	public async Task<IPAddress> GetExternalIpAddressAsync(CancellationToken cancellationToken = default)
 	{
+		if (_resolvedIp != null)
+		{
+			return _resolvedIp;
+		}
+		
 		ExternalIpResolverException? lastException = null;
 		foreach (string lookupUrl in _ipLookupUrls)
 		{
 			Uri url = new(lookupUrl);
 			try
 			{
-				return GetExternalIpAddressAsync(url, cancellationToken);
+				_resolvedIp = await GetExternalIpAddressAsync(url, cancellationToken);
+				return _resolvedIp;
 			}
 			catch (ExternalIpResolverException e)
 			{
