OidcToken - Allow for allocating of an access token without a identity token, we never actually used the IdentityToken and allowing to get a access token without it allows us to sidestep issues with some IdPs only providing thin tokens when you request both identity and access tokens at the same time.

diff --git a/Engine/Source/Programs/Shared/EpicGames.OIDC/OidcTokenManager.cs b/Engine/Source/Programs/Shared/EpicGames.OIDC/OidcTokenManager.cs
--- a/Engine/Source/Programs/Shared/EpicGames.OIDC/OidcTokenManager.cs
+++ b/Engine/Source/Programs/Shared/EpicGames.OIDC/OidcTokenManager.cs
@@ -168,12 +168,11 @@ namespace EpicGames.OIDC
 
 	public class OidcTokenInfo
 	{
-		public string? IdentityToken { get; set; }
 		public string? RefreshToken { get; set; }
 		public string? AccessToken { get; set; }
 		public DateTimeOffset TokenExpiry { get; set; }
 
-		public bool IsValid => IdentityToken != null && RefreshToken != null && AccessToken != null;
+		public bool IsValid => RefreshToken != null && AccessToken != null;
 	};
 
 	public class OidcTokenClient
@@ -319,7 +318,6 @@ namespace EpicGames.OIDC
 
 			return new OidcTokenInfo
 			{
-				IdentityToken = loginResult.IdentityToken,
 				RefreshToken = loginResult.RefreshToken,
 				AccessToken = loginResult.AccessToken,
 				TokenExpiry = loginResult.AccessTokenExpiration
@@ -504,7 +502,6 @@ namespace EpicGames.OIDC
 
 			return new OidcTokenInfo
 			{
-				IdentityToken = refreshTokenResult.IdentityToken,
 				RefreshToken = refreshTokenResult.RefreshToken,
 				AccessToken = refreshTokenResult.AccessToken,
 				TokenExpiry = refreshTokenResult.AccessTokenExpiration
