Horde: Resolve external IP when using relay connection mode

Provides necessary info to server to properly scope access to client's IP only.

diff --git a/Engine/Source/Programs/Shared/EpicGames.Horde/Compute/Clients/ServerComputeClient.cs b/Engine/Source/Programs/Shared/EpicGames.Horde/Compute/Clients/ServerComputeClient.cs
--- a/Engine/Source/Programs/Shared/EpicGames.Horde/Compute/Clients/ServerComputeClient.cs
+++ b/Engine/Source/Programs/Shared/EpicGames.Horde/Compute/Clients/ServerComputeClient.cs
@@ -163,6 +163,7 @@ namespace EpicGames.Horde.Compute.Clients
 		readonly CancellationTokenSource _cancellationSource = new CancellationTokenSource();
 		readonly string _sessionId;
 		readonly ILogger _logger;
+		readonly ExternalIpResolver _externalIpResolver;
 
 		/// <summary>
 		/// Constructor
@@ -184,6 +185,7 @@ namespace EpicGames.Horde.Compute.Clients
 			_httpClientFactory = httpClientFactory;
 			_sessionId = sessionId ?? Guid.NewGuid().ToString();
 			_logger = logger;
+			_externalIpResolver = new ExternalIpResolver(_httpClientFactory.CreateClient(HordeHttpClient.HttpClientName));
 		}
 
 		/// <inheritdoc/>
@@ -232,6 +234,11 @@ namespace EpicGames.Horde.Compute.Clients
 			request.RequestId = requestId;
 			request.Connection = connection;
 
+			if (connection is { ModePreference: ConnectionMode.Relay })
+			{
+				connection.ClientPublicIp = (await _externalIpResolver.GetExternalIpAddressAsync(cancellationToken)).ToString();
+			}
+
 			AssignComputeResponse? response;
 			using (HttpResponseMessage httpResponse = await HordeHttpClient.PostAsync(client, $"api/v2/compute/{clusterId}", request, _cancellationSource.Token))
 			{
