Horde: Fix static analysis warning.

diff --git a/Engine/Source/Programs/Shared/EpicGames.Horde.Tests/BundlePacketTests.cs b/Engine/Source/Programs/Shared/EpicGames.Horde.Tests/BundlePacketTests.cs
--- a/Engine/Source/Programs/Shared/EpicGames.Horde.Tests/BundlePacketTests.cs
+++ b/Engine/Source/Programs/Shared/EpicGames.Horde.Tests/BundlePacketTests.cs
@@ -29,14 +29,9 @@ namespace EpicGames.Horde.Tests
 				_fragment = fragment;
 			}
 
-			public ValueTask<BlobData> ReadBlobDataAsync(CancellationToken cancellationToken = default) => throw new NotImplementedException();
+			public ValueTask<BlobData> ReadAsync(CancellationToken cancellationToken = default) => throw new NotImplementedException();
 			public ValueTask FlushAsync(CancellationToken cancellationToken = default) => throw new NotImplementedException();
 
-			public ValueTask<BlobData> ReadAsync(CancellationToken cancellationToken = default)
-			{
-				throw new NotImplementedException();
-			}
-
 			public bool TryAppendIdentifier(Utf8StringBuilder builder)
 			{
 				if (_fragment == null)
