Horde: Do not dispose of storage clients until all data has been read from streams that use them.

- Adds a new StreamWrapper class to EpicGames.Core which wraps another stream and forwards all overridable methods to it.
- Adds an extension method to create a StreamWrapper that overrides Dispose/DisposeAsync methods to dispose of other owned resources.

diff --git a/Engine/Source/Programs/Horde/Horde.Server/Artifacts/ArtifactsController.cs b/Engine/Source/Programs/Horde/Horde.Server/Artifacts/ArtifactsController.cs
--- a/Engine/Source/Programs/Horde/Horde.Server/Artifacts/ArtifactsController.cs
+++ b/Engine/Source/Programs/Horde/Horde.Server/Artifacts/ArtifactsController.cs
@@ -321,11 +321,21 @@ namespace Horde.Server.Artifacts
 				filter = new FileFilter(fileFilter);
 			}
 
-			using IStorageClient storageClient = _storageClientFactory.CreateClient(artifact.NamespaceId);
-			DirectoryNode directory = await storageClient.ReadRefAsync<DirectoryNode>(artifact.RefName, DateTime.UtcNow.AddHours(1.0), cancellationToken);
+#pragma warning disable CA2000
+			IStorageClient storageClient = _storageClientFactory.CreateClient(artifact.NamespaceId);
+			try
+			{
+				DirectoryNode directory = await storageClient.ReadRefAsync<DirectoryNode>(artifact.RefName, DateTime.UtcNow.AddHours(1.0), cancellationToken);
 
-			Stream stream = directory.AsZipStream(filter);
-			return new FileStreamResult(stream, "application/zip") { FileDownloadName = $"{artifact.RefName}.zip" };
+				Stream stream = directory.AsZipStream(filter).WrapOwnership(storageClient);
+				return new FileStreamResult(stream, "application/zip") { FileDownloadName = $"{artifact.RefName}.zip" };
+			}
+			catch
+			{
+				storageClient.Dispose();
+				throw;
+			}
+#pragma warning restore CA2000
 		}
 
 		/// <summary>
diff --git a/Engine/Source/Programs/Horde/Horde.Server/Tools/ToolCollection.cs b/Engine/Source/Programs/Horde/Horde.Server/Tools/ToolCollection.cs
--- a/Engine/Source/Programs/Horde/Horde.Server/Tools/ToolCollection.cs
+++ b/Engine/Source/Programs/Horde/Horde.Server/Tools/ToolCollection.cs
@@ -442,20 +442,19 @@ namespace Horde.Server.Tools
 		/// <returns>Stream for the data</returns>
 		public async Task<Stream> GetDeploymentZipAsync(ITool tool, IToolDeployment deployment, CancellationToken cancellationToken)
 		{
-			using IStorageClient client = CreateStorageClient(tool);
-
-			if (tool.Id == new ToolId("horde-agent-canary"))
+#pragma warning disable CA2000
+			IStorageClient client = CreateStorageClient(tool);
+			try
 			{
 				DirectoryNode node = await client.ReadRefAsync<DirectoryNode>(deployment.RefName, DateTime.UtcNow - TimeSpan.FromDays(2.0), cancellationToken);
-				_logger.LogInformation("Tool {ToolId} has {NumDirectories}, {NumFiles}, total size {NumBytes}", tool.Id, node.NameToFile.Count, node.NameToDirectory.Count, node.Length);
-				return node.AsZipStream(logger: _logger);
+				return node.AsZipStream().WrapOwnership(client);
 			}
-			else
+			catch
 			{
-				DirectoryNode node = await client.ReadRefAsync<DirectoryNode>(deployment.RefName, DateTime.UtcNow - TimeSpan.FromDays(2.0), cancellationToken);
-
-				return node.AsZipStream();
+				client.Dispose();
+				throw;
 			}
+#pragma warning restore CA2000
 		}
 	}
 }
diff --git a/Engine/Source/Programs/Horde/Horde.Server/Tools/ToolsController.cs b/Engine/Source/Programs/Horde/Horde.Server/Tools/ToolsController.cs
--- a/Engine/Source/Programs/Horde/Horde.Server/Tools/ToolsController.cs
+++ b/Engine/Source/Programs/Horde/Horde.Server/Tools/ToolsController.cs
@@ -495,25 +495,32 @@ namespace Horde.Server.Tools
 				return Ok(response);
 			}
 
-			using IStorageClient client = _toolCollection.CreateStorageClient(tool);
-
-			DirectoryNode node = await client.ReadRefAsync<DirectoryNode>(deployment.RefName, DateTime.UtcNow - TimeSpan.FromDays(2.0), cancellationToken);
-
-			if (node.Directories.Count == 0 && node.Files.Count == 1 && action != GetToolAction.Zip)
+			IStorageClient client = _toolCollection.CreateStorageClient(tool);
+			try
 			{
-				FileEntry entry = node.Files.First();
+				DirectoryNode node = await client.ReadRefAsync<DirectoryNode>(deployment.RefName, DateTime.UtcNow - TimeSpan.FromDays(2.0), cancellationToken);
 
-				string? contentType;
-				if (!new FileExtensionContentTypeProvider().TryGetContentType(entry.Name.ToString(), out contentType))
+				if (node.Directories.Count == 0 && node.Files.Count == 1 && action != GetToolAction.Zip)
 				{
-					contentType = "application/octet-stream";
+					FileEntry entry = node.Files.First();
+
+					string? contentType;
+					if (!new FileExtensionContentTypeProvider().TryGetContentType(entry.Name.ToString(), out contentType))
+					{
+						contentType = "application/octet-stream";
+					}
+
+					return new FileStreamResult(entry.OpenAsStream(), contentType) { FileDownloadName = entry.Name.ToString() };
 				}
 
-				return new FileStreamResult(entry.OpenAsStream(), contentType) { FileDownloadName = entry.Name.ToString() };
+				Stream stream = node.AsZipStream().WrapOwnership(client);
+				return new FileStreamResult(stream, "application/zip") { FileDownloadName = $"{tool.Id}-{deployment.Version}.zip" };
+			}
+			catch
+			{
+				client.Dispose();
+				throw;
 			}
-
-			Stream stream = node.AsZipStream();
-			return new FileStreamResult(stream, "application/zip") { FileDownloadName = $"{tool.Id}-{deployment.Version}.zip" };
 		}
 
 		private async Task<GetToolDeploymentResponse> GetDeploymentInfoResponseAsync(ITool tool, IToolDeployment deployment, CancellationToken cancellationToken)
diff --git a/Engine/Source/Programs/Shared/EpicGames.Core/StreamWrapper.cs b/Engine/Source/Programs/Shared/EpicGames.Core/StreamWrapper.cs
new file mode 100644
--- /dev/null
+++ b/Engine/Source/Programs/Shared/EpicGames.Core/StreamWrapper.cs
@@ -0,0 +1,191 @@
+// Copyright Epic Games, Inc. All Rights Reserved.
+
+using System;
+using System.IO;
+using System.Threading;
+using System.Threading.Tasks;
+
+#pragma warning disable CA1710
+
+namespace EpicGames.Core
+{
+	/// <summary>
+	/// <see cref="Stream"/> implementation that wraps another stream, allowing derived classes to override some methods of their chosing.
+	/// </summary>
+	public class StreamWrapper : Stream
+	{
+		readonly Stream _inner;
+
+		/// <summary>
+		/// Accessor for the wrapped stream
+		/// </summary>
+		protected Stream Inner => _inner;
+
+		/// <inheritdoc/>
+		public override bool CanRead => _inner.CanRead;
+
+		/// <inheritdoc/>
+		public override bool CanSeek => _inner.CanSeek;
+
+		/// <inheritdoc/>
+		public override bool CanTimeout => _inner.CanTimeout;
+
+		/// <inheritdoc/>
+		public override bool CanWrite => _inner.CanWrite;
+
+		/// <inheritdoc/>
+		public override long Length => _inner.Length;
+
+		/// <inheritdoc/>
+		public override long Position { get => _inner.Position; set => _inner.Position = value; }
+
+		/// <inheritdoc/>
+		public override int ReadTimeout { get => _inner.ReadTimeout; set => _inner.ReadTimeout = value; }
+
+		/// <inheritdoc/>
+		public override int WriteTimeout { get => _inner.WriteTimeout; set => _inner.WriteTimeout = value; }
+
+		/// <summary>
+		/// Constructor
+		/// </summary>
+		/// <param name="inner">The stream to wrap</param>
+		public StreamWrapper(Stream inner) => _inner = inner;
+
+		/// <inheritdoc/>
+		protected override void Dispose(bool disposing)
+		{
+			base.Dispose(disposing);
+
+			if (disposing)
+			{
+				_inner.Dispose();
+			}
+		}
+
+		/// <inheritdoc/>
+		public override async ValueTask DisposeAsync()
+		{
+			await base.DisposeAsync();
+
+			await _inner.DisposeAsync();
+
+			GC.SuppressFinalize(this);
+		}
+
+		/// <inheritdoc/>
+		public override IAsyncResult BeginRead(byte[] buffer, int offset, int count, AsyncCallback? callback, object? state) => _inner.BeginRead(buffer, offset, count, callback, state);
+
+		/// <inheritdoc/>
+		public override IAsyncResult BeginWrite(byte[] buffer, int offset, int count, AsyncCallback? callback, object? state) => _inner.BeginWrite(buffer, offset, count, callback, state);
+
+		/// <inheritdoc/>
+		public override void Close() => _inner.Close();
+
+		/// <inheritdoc/>
+		public override void CopyTo(Stream destination, int bufferSize) => _inner.CopyTo(destination, bufferSize);
+
+		/// <inheritdoc/>
+		public override Task CopyToAsync(Stream destination, int bufferSize, CancellationToken cancellationToken) => _inner.CopyToAsync(destination, bufferSize, cancellationToken);
+
+		/// <inheritdoc/>
+		public override int EndRead(IAsyncResult asyncResult) => _inner.EndRead(asyncResult);
+
+		/// <inheritdoc/>
+		public override void EndWrite(IAsyncResult asyncResult) => _inner.EndWrite(asyncResult);
+
+		/// <inheritdoc/>
+		public override void Flush() => _inner.Flush();
+
+		/// <inheritdoc/>
+		public override Task FlushAsync(CancellationToken cancellationToken) => _inner.FlushAsync(cancellationToken);
+
+		/// <inheritdoc/>
+		public override int Read(byte[] buffer, int offset, int count) => _inner.Read(buffer, offset, count);
+
+		/// <inheritdoc/>
+		public override int Read(Span<byte> buffer) => _inner.Read(buffer);
+
+		/// <inheritdoc/>
+		public override Task<int> ReadAsync(byte[] buffer, int offset, int count, CancellationToken cancellationToken) => _inner.ReadAsync(buffer, offset, count, cancellationToken);
+
+		/// <inheritdoc/>
+		public override ValueTask<int> ReadAsync(Memory<byte> buffer, CancellationToken cancellationToken = default) => _inner.ReadAsync(buffer, cancellationToken);
+
+		/// <inheritdoc/>
+		public override int ReadByte() => _inner.ReadByte();
+
+		/// <inheritdoc/>
+		public override long Seek(long offset, SeekOrigin origin) => _inner.Seek(offset, origin);
+
+		/// <inheritdoc/>
+		public override void SetLength(long value) => _inner.SetLength(value);
+
+		/// <inheritdoc/>
+		public override void Write(byte[] buffer, int offset, int count) => _inner.Write(buffer, offset, count);
+
+		/// <inheritdoc/>
+		public override void Write(ReadOnlySpan<byte> buffer) => _inner.Write(buffer);
+
+		/// <inheritdoc/>
+		public override Task WriteAsync(byte[] buffer, int offset, int count, CancellationToken cancellationToken) => _inner.WriteAsync(buffer, offset, count, cancellationToken);
+
+		/// <inheritdoc/>
+		public override ValueTask WriteAsync(ReadOnlyMemory<byte> buffer, CancellationToken cancellationToken = default) => _inner.WriteAsync(buffer, cancellationToken);
+
+		/// <inheritdoc/>
+		public override void WriteByte(byte value) => _inner.WriteByte(value);
+	}
+
+	/// <summary>
+	/// Extension methods for <see cref="StreamWrapper"/>
+	/// </summary>
+	public static class StreamWrapperExtensions
+	{
+		class StreamResourceOwner : StreamWrapper
+		{
+			IDisposable? _resource;
+
+			public StreamResourceOwner(Stream inner, IDisposable resource)
+				: base(inner)
+			{
+				_resource = resource;
+			}
+
+			/// <inheritdoc/>
+			protected override void Dispose(bool disposing)
+			{
+				base.Dispose(disposing);
+
+				if (disposing)
+				{
+					DisposeResource();
+				}
+			}
+
+			/// <inheritdoc/>
+			public override async ValueTask DisposeAsync()
+			{
+				DisposeResource();
+
+				await base.DisposeAsync();
+			}
+
+			void DisposeResource()
+			{
+				if (_resource != null)
+				{
+					_resource.Dispose();
+					_resource = null;
+				}
+			}
+		}
+
+		/// <summary>
+		/// Wraps ownership of another resource with a stream
+		/// </summary>
+		/// <param name="stream">Stream to wrap</param>
+		/// <param name="resource">Additional resource to control ownership of</param>
+		/// <returns>New stream which will dispose of the given resource when it is disposed</returns>
+		public static Stream WrapOwnership(this Stream stream, IDisposable resource) => new StreamResourceOwner(stream, resource);
+	}
+}
