Horde: Fix null reference exception when trying to use server without auth.

diff --git a/Engine/Source/Programs/Shared/EpicGames.Horde/HordeHttpAuthHandler.cs b/Engine/Source/Programs/Shared/EpicGames.Horde/HordeHttpAuthHandler.cs
--- a/Engine/Source/Programs/Shared/EpicGames.Horde/HordeHttpAuthHandler.cs
+++ b/Engine/Source/Programs/Shared/EpicGames.Horde/HordeHttpAuthHandler.cs
@@ -58,8 +58,11 @@ namespace EpicGames.Horde
 			}
 
 			// Otherwise update the auth header and try again
-			_authHeader = await _authState.GetAuthHeaderAsync(cancellationToken);
-			request.Headers.Authorization = _authHeader;
+			_authHeader = await _authState.TryGetAuthHeaderAsync(cancellationToken);
+			if (_authHeader != null)
+			{
+				request.Headers.Authorization = _authHeader;
+			}
 			return await base.SendAsync(request, cancellationToken);
 		}
 	}
@@ -76,7 +79,7 @@ namespace EpicGames.Horde
 
 		readonly object _lockObject = new object();
 		readonly CancellationTokenSource _cancellationTokenSource = new CancellationTokenSource();
-		Task<AuthenticationHeaderValue>? _authHeaderTask;
+		Task<AuthenticationHeaderValue?>? _authHeaderTask;
 		readonly IHttpClientFactory _httpClientFactory;
 		readonly ILogger _logger;
 
@@ -109,7 +112,7 @@ namespace EpicGames.Horde
 		{
 			lock (_lockObject)
 			{
-				if (_authHeaderTask != null && _authHeaderTask.IsCompleted && _authHeaderTask.Result.Equals(authHeader))
+				if (_authHeaderTask != null && _authHeaderTask.IsCompleted && Object.Equals(_authHeaderTask.Result, authHeader))
 				{
 					_authHeaderTask = null;
 				}
@@ -120,9 +123,9 @@ namespace EpicGames.Horde
 		/// Gets a new auth header
 		/// </summary>
 		/// <param name="cancellationToken">Cancellation token for the operation</param>
-		public async ValueTask<AuthenticationHeaderValue> GetAuthHeaderAsync(CancellationToken cancellationToken)
+		public async ValueTask<AuthenticationHeaderValue?> TryGetAuthHeaderAsync(CancellationToken cancellationToken)
 		{
-			Task<AuthenticationHeaderValue>? authHeaderTask = _authHeaderTask;
+			Task<AuthenticationHeaderValue?>? authHeaderTask = _authHeaderTask;
 			if (authHeaderTask == null)
 			{
 				lock (_lockObject)
@@ -137,7 +140,7 @@ namespace EpicGames.Horde
 		/// <summary>
 		/// Get an access token for the server specified in a config instance
 		/// </summary>
-		async Task<AuthenticationHeaderValue> GetNewAuthHeaderAsync(CancellationToken cancellationToken)
+		async Task<AuthenticationHeaderValue?> GetNewAuthHeaderAsync(CancellationToken cancellationToken)
 		{
 			Uri serverUrl;
 
@@ -162,13 +165,19 @@ namespace EpicGames.Horde
 				}
 			}
 
+			string? localRedirectUrl = authConfig.LocalRedirectUrls?.FirstOrDefault();
+			if (String.IsNullOrEmpty(authConfig.ServerUrl) || String.IsNullOrEmpty(localRedirectUrl))
+			{
+				return null;
+			}
+
 			const string OidcProvider = "Horde";
 
 			Dictionary<string, string?> values = new Dictionary<string, string?>();
 			values[$"Providers:{OidcProvider}:DisplayName"] = "Horde";
 			values[$"Providers:{OidcProvider}:ServerUri"] = authConfig.ServerUrl;
 			values[$"Providers:{OidcProvider}:ClientId"] = authConfig.ClientId;
-			values[$"Providers:{OidcProvider}:RedirectUri"] = authConfig.LocalRedirectUrls?.FirstOrDefault();
+			values[$"Providers:{OidcProvider}:RedirectUri"] = localRedirectUrl;
 
 			ConfigurationBuilder builder = new ConfigurationBuilder();
 			builder.AddInMemoryCollection(values);
