Horde: Fix path mangling when using file storage backend with a path prefix.

diff --git a/Engine/Source/Programs/Shared/EpicGames.Horde/Storage/Backends/FileStorageBackend.cs b/Engine/Source/Programs/Shared/EpicGames.Horde/Storage/Backends/FileStorageBackend.cs
--- a/Engine/Source/Programs/Shared/EpicGames.Horde/Storage/Backends/FileStorageBackend.cs
+++ b/Engine/Source/Programs/Shared/EpicGames.Horde/Storage/Backends/FileStorageBackend.cs
@@ -96,7 +96,7 @@ namespace EpicGames.Horde.Storage.Backends
 				}
 				else
 				{
-					return _data.Slice(offset, length.Value);
+					return _data.Slice(offset, Math.Min(length.Value, _data.Length - offset));
 				}
 			}
 
diff --git a/Engine/Source/Programs/Shared/EpicGames.Horde/Storage/Backends/PrefixStorageBackend.cs b/Engine/Source/Programs/Shared/EpicGames.Horde/Storage/Backends/PrefixStorageBackend.cs
--- a/Engine/Source/Programs/Shared/EpicGames.Horde/Storage/Backends/PrefixStorageBackend.cs
+++ b/Engine/Source/Programs/Shared/EpicGames.Horde/Storage/Backends/PrefixStorageBackend.cs
@@ -69,7 +69,15 @@ namespace EpicGames.Horde.Storage.Backends
 		public ValueTask<(string, Uri)?> TryGetWriteRedirectAsync(string? basePath = null, CancellationToken cancellationToken = default) => _inner.TryGetWriteRedirectAsync($"{_prefix}{basePath}", cancellationToken);
 
 		/// <inheritdoc/>
-		public Task<string> WriteAsync(Stream stream, string? basePath = null, CancellationToken cancellationToken = default) => _inner.WriteAsync(stream, $"{_prefix}{basePath}", cancellationToken);
+		public async Task<string> WriteAsync(Stream stream, string? basePath = null, CancellationToken cancellationToken = default)
+		{
+			string path = await _inner.WriteAsync(stream, $"{_prefix}{basePath}", cancellationToken);
+			if (!path.StartsWith(_prefix))
+			{
+				throw new InvalidOperationException($"Expected written blob to start with requested prefix '{_prefix}{basePath}'. Got '{path}'.");
+			}
+			return path.Substring(_prefix.Length);
+		}
 
 		/// <inheritdoc/>
 		[Obsolete("Use WriteAsync() instead. Ability to specify an explicit path is deprecated and will be removed in a future release.")]
