Horde: Fix static analysis warning.

diff --git a/Engine/Source/Programs/Shared/EpicGames.Horde/Storage/Backends/PrefixStorageBackend.cs b/Engine/Source/Programs/Shared/EpicGames.Horde/Storage/Backends/PrefixStorageBackend.cs
--- a/Engine/Source/Programs/Shared/EpicGames.Horde/Storage/Backends/PrefixStorageBackend.cs
+++ b/Engine/Source/Programs/Shared/EpicGames.Horde/Storage/Backends/PrefixStorageBackend.cs
@@ -72,7 +72,7 @@ namespace EpicGames.Horde.Storage.Backends
 		public async Task<string> WriteAsync(Stream stream, string? basePath = null, CancellationToken cancellationToken = default)
 		{
 			string path = await _inner.WriteAsync(stream, $"{_prefix}{basePath}", cancellationToken);
-			if (!path.StartsWith(_prefix))
+			if (!path.StartsWith(_prefix, StringComparison.Ordinal))
 			{
 				throw new InvalidOperationException($"Expected written blob to start with requested prefix '{_prefix}{basePath}'. Got '{path}'.");
 			}
