Horde: Fix an issue where empty files or directories may not be extracted from bundles due to the way that the input data is partitioned into threads.

diff --git a/Engine/Source/Programs/Shared/EpicGames.Horde/Storage/Nodes/DirectoryNode.cs b/Engine/Source/Programs/Shared/EpicGames.Horde/Storage/Nodes/DirectoryNode.cs
--- a/Engine/Source/Programs/Shared/EpicGames.Horde/Storage/Nodes/DirectoryNode.cs
+++ b/Engine/Source/Programs/Shared/EpicGames.Horde/Storage/Nodes/DirectoryNode.cs
@@ -794,6 +794,13 @@ namespace EpicGames.Horde.Storage.Nodes
 				{
 					long minOffset = offset;
 					long maxOffset = (Length * (taskIdx + 1)) / numTasks;
+
+					// Entries may be zero-length, so need to make sure the last window will include everything
+					if (taskIdx == numTasks - 1)
+					{
+						maxOffset++;
+					}
+
 					tasks.Add(Task.Run(() => CopyToDirectoryInternalAsync(directoryInfo, minOffset, maxOffset - minOffset, copyStats, logger, cancellationToken), cancellationToken));
 					offset = maxOffset;
 				}
