Add support for JSON log events from stdout to Gauntlet, and re-enable automatically turning on JSON output when UE_LOG_JSON_TO_STDOUT environment variable is set.

diff --git a/Engine/Source/Programs/AutomationTool/Gauntlet/Unreal/Utils/Gauntlet.UnrealLogParser.cs b/Engine/Source/Programs/AutomationTool/Gauntlet/Unreal/Utils/Gauntlet.UnrealLogParser.cs
--- a/Engine/Source/Programs/AutomationTool/Gauntlet/Unreal/Utils/Gauntlet.UnrealLogParser.cs
+++ b/Engine/Source/Programs/AutomationTool/Gauntlet/Unreal/Utils/Gauntlet.UnrealLogParser.cs
@@ -9,6 +9,11 @@ using System.Threading;
 using System.Text.RegularExpressions;
 using System.Drawing;
 using System.Linq;
+using System.Text;
+using EpicGames.Core;
+using System.Text.Json;
+using System.Runtime.CompilerServices;
+using Microsoft.Extensions.Logging;
 
 namespace Gauntlet
 {
@@ -234,8 +239,7 @@ namespace Gauntlet
 		/// <returns></returns>
 		public UnrealLogParser(string InContent)
 		{
-			// convert linefeed to remove \r which is captured in regex's :(
-			Content = InContent.Replace(Environment.NewLine, "\n");
+			Content = SanitizeLogText(InContent);
 
 			// Search for LogFoo: <Display|Error|etc>: Message
 			// Also need to handle 'Log' not always being present, and the category being empty for a level of 'Log'
@@ -271,6 +275,61 @@ namespace Gauntlet
 			LogEntries = ParsedEntries;
 		}
 
+		static string SanitizeLogText(string InContent)
+		{
+			StringBuilder ContentBuilder = new StringBuilder();
+
+			for (int BaseIdx = 0; BaseIdx < InContent.Length;)
+			{
+				// Extract the next line
+				int EndIdx = InContent.IndexOf('\n', BaseIdx);
+				if (EndIdx == -1)
+				{
+					break;
+				}
+
+				// Skip over any windows CR-LF line endings
+				int LineEndIdx = EndIdx;
+				if (LineEndIdx > BaseIdx && InContent[LineEndIdx - 1] == '\r')
+				{
+					LineEndIdx--;
+				}
+
+				// Render any JSON log events
+				string Line = InContent.Substring(BaseIdx, LineEndIdx - BaseIdx);
+				if (Line.Length > 0 && Line[0] == '{')
+				{
+					try
+					{
+						byte[] Buffer = Encoding.UTF8.GetBytes(Line);
+						JsonLogEvent JsonEvent = JsonLogEvent.Parse(Buffer);
+						Line = JsonEvent.GetRenderedMessage().ToString();
+					}
+					catch (Exception ex)
+					{
+						EpicGames.Core.Log.Logger.LogInformation(ex, "Unable to parse log line: {Line}, Exception: {Ex}", Line, ex.ToString());
+
+						int MinIdx = Math.Max(BaseIdx - 2048, 0);
+						int MaxIdx = Math.Min(BaseIdx + 2048, InContent.Length);
+
+						string[] Context = InContent.Substring(MinIdx, MaxIdx - MinIdx).Split('\n');
+						for (int idx = 1; idx < Context.Length - 1; idx++)
+						{
+							EpicGames.Core.Log.Logger.LogInformation("Context {Idx}: {Line}", idx, Context[idx].TrimEnd());
+						}
+					}
+				}
+
+				ContentBuilder.Append(Line);
+				ContentBuilder.Append('\n');
+
+				// Move to the next line
+				BaseIdx = EndIdx + 1;
+			}
+
+			return ContentBuilder.ToString();
+		}
+
 		public UnrealLog GetSummary()
 		{
 			if (Summary == null)
diff --git a/Engine/Source/Programs/Shared/EpicGames.Core/JsonLogEvent.cs b/Engine/Source/Programs/Shared/EpicGames.Core/JsonLogEvent.cs
--- a/Engine/Source/Programs/Shared/EpicGames.Core/JsonLogEvent.cs
+++ b/Engine/Source/Programs/Shared/EpicGames.Core/JsonLogEvent.cs
@@ -87,7 +87,7 @@ namespace EpicGames.Core
 		public static JsonLogEvent Parse(ReadOnlyMemory<byte> data)
 		{
 			JsonLogEvent logEvent;
-			if (!TryParse(data, out logEvent))
+			if (!TryParseInternal(data, out logEvent))
 			{
 				throw new InvalidOperationException("Cannot parse string");
 			}
@@ -104,53 +104,62 @@ namespace EpicGames.Core
 		{
 			try
 			{
-				LogLevel level = LogLevel.None;
-				int eventId = 0;
-				int lineIndex = 0;
-				int lineCount = 1;
+				return TryParseInternal(data, out logEvent);
+			}
+			catch
+			{
+				logEvent = default;
+				return false;
+			}
+		}
 
-				Utf8JsonReader reader = new Utf8JsonReader(data.Span);
-				if (reader.Read() && reader.TokenType == JsonTokenType.StartObject)
+		static bool TryParseInternal(ReadOnlyMemory<byte> data, out JsonLogEvent logEvent)
+		{
+			LogLevel level = LogLevel.None;
+			int eventId = 0;
+			int lineIndex = 0;
+			int lineCount = 1;
+
+			Utf8JsonReader reader = new Utf8JsonReader(data.Span);
+			if (reader.Read() && reader.TokenType == JsonTokenType.StartObject)
+			{
+				while (reader.Read() && reader.TokenType == JsonTokenType.PropertyName)
 				{
-					while (reader.Read() && reader.TokenType == JsonTokenType.PropertyName)
+					ReadOnlySpan<byte> propertyName = reader.ValueSpan;
+					if (!reader.Read())
 					{
-						ReadOnlySpan<byte> propertyName = reader.ValueSpan;
-						if (!reader.Read())
-						{
-							break;
-						}
-						else if (propertyName.SequenceEqual(LogEventPropertyName.Level) && reader.TokenType == JsonTokenType.String)
-						{
-							level = ParseLevel(reader.ValueSpan);
-						}
-						else if (propertyName.SequenceEqual(LogEventPropertyName.Id) && reader.TokenType == JsonTokenType.Number)
-						{
-							eventId = reader.GetInt32();
-						}
-						else if (propertyName.SequenceEqual(LogEventPropertyName.Line) && reader.TokenType == JsonTokenType.Number)
-						{
-							reader.TryGetInt32(out lineIndex);
-						}
-						else if (propertyName.SequenceEqual(LogEventPropertyName.LineCount) && reader.TokenType == JsonTokenType.Number)
-						{
-							reader.TryGetInt32(out lineCount);
-						}
-						reader.Skip();
+						break;
+					}
+					else if (propertyName.SequenceEqual(LogEventPropertyName.Level) && reader.TokenType == JsonTokenType.String)
+					{
+						level = ParseLevel(reader.ValueSpan);
+					}
+					else if (propertyName.SequenceEqual(LogEventPropertyName.Id) && reader.TokenType == JsonTokenType.Number)
+					{
+						eventId = reader.GetInt32();
 					}
+					else if (propertyName.SequenceEqual(LogEventPropertyName.Line) && reader.TokenType == JsonTokenType.Number)
+					{
+						reader.TryGetInt32(out lineIndex);
+					}
+					else if (propertyName.SequenceEqual(LogEventPropertyName.LineCount) && reader.TokenType == JsonTokenType.Number)
+					{
+						reader.TryGetInt32(out lineCount);
+					}
+					reader.Skip();
 				}
+			}
 
-				if (reader.TokenType == JsonTokenType.EndObject && level != LogLevel.None && reader.BytesConsumed == data.Length)
-				{
-					logEvent = new JsonLogEvent(level, new EventId(eventId), lineIndex, lineCount, data.ToArray());
-					return true;
-				}
+			if (reader.TokenType == JsonTokenType.EndObject && level != LogLevel.None && reader.BytesConsumed == data.Length)
+			{
+				logEvent = new JsonLogEvent(level, new EventId(eventId), lineIndex, lineCount, data.ToArray());
+				return true;
 			}
-			catch
+			else
 			{
+				logEvent = default;
+				return false;
 			}
-
-			logEvent = default;
-			return false;
 		}
 
 		static readonly sbyte[] s_firstCharToLogLevel;
diff --git a/Engine/Source/Runtime/Launch/Private/LaunchEngineLoop.cpp b/Engine/Source/Runtime/Launch/Private/LaunchEngineLoop.cpp
--- a/Engine/Source/Runtime/Launch/Private/LaunchEngineLoop.cpp
+++ b/Engine/Source/Runtime/Launch/Private/LaunchEngineLoop.cpp
@@ -385,6 +385,13 @@ public:
 #endif
 
 		bIsJsonOutput = FParse::Param(FCommandLine::Get(), TEXT("JsonStdOut"));
+		if (!bIsJsonOutput)
+		{
+			if (FString Env = FPlatformMisc::GetEnvironmentVariable(TEXT("UE_LOG_JSON_TO_STDOUT")); !Env.IsEmpty())
+			{
+				bIsJsonOutput = FCString::Atoi(*Env) != 0;
+			}
+		}
 
 		if (FParse::Param(FCommandLine::Get(), TEXT("AllowStdOutLogVerbosity")))
 		{
