Add Description property back to command line attribute. Non-Horde programs use this!

diff --git a/Engine/Source/Programs/Shared/EpicGames.Core/CommandLineAttribute.cs b/Engine/Source/Programs/Shared/EpicGames.Core/CommandLineAttribute.cs
--- a/Engine/Source/Programs/Shared/EpicGames.Core/CommandLineAttribute.cs
+++ b/Engine/Source/Programs/Shared/EpicGames.Core/CommandLineAttribute.cs
@@ -37,6 +37,11 @@ namespace EpicGames.Core
 		/// </summary>
 		public char ListSeparator { get; set; } = '\0';
 
+		/// <summary>
+		/// Description of the operation.
+		/// </summary>
+		public string? Description { get; set; }
+
 		/// <summary>
 		/// Constructor
 		/// </summary>
