Change SanitizeLogOutput to restore legacy log format including timestamps

diff --git a/Engine/Source/Programs/AutomationTool/Gauntlet/Unreal/Utils/Gauntlet.UnrealLogParser.cs b/Engine/Source/Programs/AutomationTool/Gauntlet/Unreal/Utils/Gauntlet.UnrealLogParser.cs
--- a/Engine/Source/Programs/AutomationTool/Gauntlet/Unreal/Utils/Gauntlet.UnrealLogParser.cs
+++ b/Engine/Source/Programs/AutomationTool/Gauntlet/Unreal/Utils/Gauntlet.UnrealLogParser.cs
@@ -303,7 +303,7 @@ namespace Gauntlet
 					{
 						byte[] Buffer = Encoding.UTF8.GetBytes(Line);
 						JsonLogEvent JsonEvent = JsonLogEvent.Parse(Buffer);
-						Line = JsonEvent.GetRenderedMessage().ToString();
+						Line = JsonEvent.GetLegacyLogLine();
 					}
 					catch (Exception ex)
 					{
diff --git a/Engine/Source/Programs/Shared/EpicGames.Core/JsonLogEvent.cs b/Engine/Source/Programs/Shared/EpicGames.Core/JsonLogEvent.cs
--- a/Engine/Source/Programs/Shared/EpicGames.Core/JsonLogEvent.cs
+++ b/Engine/Source/Programs/Shared/EpicGames.Core/JsonLogEvent.cs
@@ -219,6 +219,49 @@ namespace EpicGames.Core
 			return Utf8String.Empty;
 		}
 
+		/// <summary>
+		/// Gets the event data rendered as a legacy unreal log line of the format:
+		/// [timestamp][frame number]LogChannel: LogVerbosity: Message
+		/// </summary>
+		public string GetLegacyLogLine()
+		{
+			Utf8JsonReader reader = new Utf8JsonReader(Data.Span);
+			if (reader.Read() && reader.TokenType == JsonTokenType.StartObject)
+			{
+				Utf8String? message = null;
+				DateTime? time = null;
+				while (reader.Read() && reader.TokenType == JsonTokenType.PropertyName)
+				{
+					ReadOnlySpan<byte> propertyName = reader.ValueSpan;
+					if (!reader.Read())
+					{
+						break;
+					}
+					else if (propertyName.SequenceEqual(LogEventPropertyName.Time) && reader.TokenType == JsonTokenType.String)
+					{
+						time = reader.GetDateTime();
+					}
+					else if (propertyName.SequenceEqual(LogEventPropertyName.Message) && reader.TokenType == JsonTokenType.String)
+					{
+						message = new Utf8String(reader.GetUtf8String().ToArray());
+					}
+				}
+
+				if (message is not null)
+				{
+					if (time is not null)
+					{
+						return $"[{time:yyyy.MM.dd-HH.mm.ss:fff}][  0]{message}"; // Structured logs currently don't contain frame number
+					}
+					else
+					{
+						return $"{message}";
+					}
+				}
+			}
+			return String.Empty;
+		}
+
 		/// <summary>
 		/// Count the number of lines in the message field of a log event
 		/// </summary>
