Horde: Tweaks to handling of empty directories when extracting data from bundles.

- Output progress stats at the end of the extraction process, and not after the first update.
- Fix windowing for partitioning into multiple threads when bundles contains empty files/directories.

diff --git a/Engine/Source/Programs/Shared/EpicGames.Horde/Storage/Nodes/DirectoryNode.cs b/Engine/Source/Programs/Shared/EpicGames.Horde/Storage/Nodes/DirectoryNode.cs
--- a/Engine/Source/Programs/Shared/EpicGames.Horde/Storage/Nodes/DirectoryNode.cs
+++ b/Engine/Source/Programs/Shared/EpicGames.Horde/Storage/Nodes/DirectoryNode.cs
@@ -63,7 +63,7 @@ namespace EpicGames.Horde.Storage.Nodes
 			{
 				Count += count;
 				Size += size;
-				if (_timer.Elapsed > TimeSpan.FromSeconds(10.0) || Count == count)
+				if (_timer.Elapsed > TimeSpan.FromSeconds(10.0))
 				{
 					Rate = (Size - _lastTotalSize) / _timer.Elapsed.TotalSeconds;
 					_lastTotalSize = Size;
@@ -808,6 +808,7 @@ namespace EpicGames.Horde.Storage.Nodes
 			finally
 			{
 				await Task.WhenAll(tasks);
+				copyStats?.Flush();
 			}
 		}
 
@@ -829,8 +830,8 @@ namespace EpicGames.Horde.Storage.Nodes
 
 			foreach (DirectoryEntry directoryEntry in _nameToDirectoryEntry.Values)
 			{
-				// Traverse into any directory that overlaps with the window (window starts before end of the directory, and window ends at or beyond the start of the directory)
-				if (windowOffset < directoryEntry.Length && windowOffset + windowLength >= 0)
+				// Traverse into any directory that overlaps with the window (window starts before this directory, and window ends after the start of the directory)
+				if (windowOffset <= 0 && windowOffset + windowLength > 0)
 				{
 					DirectoryInfo subDirectoryInfo = directoryInfo.CreateSubdirectory(directoryEntry.Name.ToString());
 					DirectoryNode subDirectoryNode = await directoryEntry.ExpandAsync(cancellationToken);
