Support using the Description property in [CommandLine] attributes to document an option, as well as the standalone [Description] attribute.

diff --git a/Engine/Source/Programs/Shared/EpicGames.Core/CommandLineArguments.cs b/Engine/Source/Programs/Shared/EpicGames.Core/CommandLineArguments.cs
--- a/Engine/Source/Programs/Shared/EpicGames.Core/CommandLineArguments.cs
+++ b/Engine/Source/Programs/Shared/EpicGames.Core/CommandLineArguments.cs
@@ -882,6 +882,14 @@ namespace EpicGames.Core
 					}
 					descriptionBuilder.Append(attribute.Description);
 				}
+				foreach (CommandLineAttribute attribute in target.Member.GetCustomAttributes<CommandLineAttribute>().Where(x => x.Description != null))
+				{
+					if (descriptionBuilder.Length > 0)
+					{
+						descriptionBuilder.Append('\n');
+					}
+					descriptionBuilder.Append(attribute.Description);
+				}
 
 				string description = descriptionBuilder.ToString();
 				if (description.Length == 0)
