Horde: Revert incorrect change to partition directory fetches on different tasks.

diff --git a/Engine/Source/Programs/Shared/EpicGames.Horde/Storage/Nodes/DirectoryNode.cs b/Engine/Source/Programs/Shared/EpicGames.Horde/Storage/Nodes/DirectoryNode.cs
--- a/Engine/Source/Programs/Shared/EpicGames.Horde/Storage/Nodes/DirectoryNode.cs
+++ b/Engine/Source/Programs/Shared/EpicGames.Horde/Storage/Nodes/DirectoryNode.cs
@@ -830,8 +830,8 @@ namespace EpicGames.Horde.Storage.Nodes
 
 			foreach (DirectoryEntry directoryEntry in _nameToDirectoryEntry.Values)
 			{
-				// Traverse into any directory that overlaps with the window (window starts before this directory, and window ends after the start of the directory)
-				if (windowOffset <= 0 && windowOffset + windowLength > 0)
+				// Traverse into any directory that overlaps with the window (window starts before end of the directory, and window ends at or beyond the start of the directory)
+				if (windowOffset < directoryEntry.Length && windowOffset + windowLength >= 0)
 				{
 					DirectoryInfo subDirectoryInfo = directoryInfo.CreateSubdirectory(directoryEntry.Name.ToString());
 					DirectoryNode subDirectoryNode = await directoryEntry.ExpandAsync(cancellationToken);
