Horde: Add TcpSslTranport for compute communication

diff --git a/Engine/Source/Programs/Shared/EpicGames.Horde.Tests/ComputeSocketTests.cs b/Engine/Source/Programs/Shared/EpicGames.Horde.Tests/ComputeSocketTests.cs
--- a/Engine/Source/Programs/Shared/EpicGames.Horde.Tests/ComputeSocketTests.cs
+++ b/Engine/Source/Programs/Shared/EpicGames.Horde.Tests/ComputeSocketTests.cs
@@ -82,23 +82,48 @@ namespace EpicGames.Horde.Tests
 		[TestMethod]
 		public async Task TestAgentMessageLoopTcpAsync()
 		{
-			const int Port = 9990;
-			TcpListener listener = new TcpListener(IPAddress.Loopback, Port);
-			listener.Start();
-
-			using Socket clientSocket = new Socket(SocketType.Stream, ProtocolType.Tcp);
-			Task clientConnectTask = clientSocket.ConnectAsync(IPAddress.Loopback, Port, CancellationToken.None).AsTask();
+			using CancellationTokenSource cts = new (5000);
+			(Socket clientSocket, Socket serverSocket) = await CreateSocketsAsync(cts.Token);
 
-			using Socket serverSocket = await listener.AcceptSocketAsync(CancellationToken.None);
-			await clientConnectTask;
+			await using RemoteComputeSocket localSocket = new (new TcpTransport(clientSocket), new TestLogger());
+			await using RemoteComputeSocket agentSocket = new (new TcpTransport(serverSocket), new TestLogger());
 
-			await using RemoteComputeSocket localSocket = new RemoteComputeSocket(new TcpTransport(clientSocket), new TestLogger());
-			await using RemoteComputeSocket agentSocket = new RemoteComputeSocket(new TcpTransport(serverSocket), new TestLogger());
+			await RunAgentTestsAsync(localSocket, agentSocket, cts.Token);
+		}
+		
+		[TestMethod]
+		public async Task TestAgentMessageLoopTcpSslAsync()
+		{
+			using CancellationTokenSource cts = new (5000);
+			(Socket clientSocket, Socket serverSocket) = await CreateSocketsAsync(cts.Token);
+			
+			byte[] certData = TcpSslTransport.GenerateCert();
+			using TcpSslTransport clientTransport = new (clientSocket, certData, false);
+			using TcpSslTransport serverTransport = new (serverSocket, certData, true);
+			
+			Task t1 = clientTransport.AuthenticateAsync(cts.Token);
+			Task t2 = serverTransport.AuthenticateAsync(cts.Token);
+			await t2;
+			await t1;
+			
+			await using RemoteComputeSocket localSocket = new (clientTransport, new TestLogger());
+			await using RemoteComputeSocket agentSocket = new (serverTransport, new TestLogger());
+			await RunAgentTestsAsync(localSocket, agentSocket, cts.Token);
+		}
 
-			await RunAgentTestsAsync(localSocket, agentSocket);
+		private static async Task<(Socket client, Socket server)> CreateSocketsAsync(CancellationToken cancellationToken)
+		{
+			int port = GetAvailablePort();
+			TcpListener listener = new (IPAddress.Loopback, port);
+			listener.Start();
+			Socket clientSocket = new (SocketType.Stream, ProtocolType.Tcp);
+			Task clientConnectTask = clientSocket.ConnectAsync(IPAddress.Loopback, port, cancellationToken).AsTask();
+			Socket serverSocket = await listener.AcceptSocketAsync(cancellationToken);
+			await clientConnectTask;
+			return (clientSocket, serverSocket);
 		}
 
-		static async Task RunAgentTestsAsync(RemoteComputeSocket localSocket, RemoteComputeSocket agentSocket)
+		static async Task RunAgentTestsAsync(RemoteComputeSocket localSocket, RemoteComputeSocket agentSocket, CancellationToken cancellationToken = default)
 		{
 			DirectoryReference tempDir = new DirectoryReference("test-temp");
 			await using (BackgroundTask agentTask = BackgroundTask.StartNew(ctx => RunAgentAsync(agentSocket, tempDir, ctx)))
@@ -106,17 +131,17 @@ namespace EpicGames.Horde.Tests
 				const int PrimaryChannelId = 0;
 				using (AgentMessageChannel channel = localSocket.CreateAgentMessageChannel(PrimaryChannelId, 4 * 1024 * 1024))
 				{
-					await channel.WaitForAttachAsync();
+					await channel.WaitForAttachAsync(cancellationToken);
 
-					await channel.PingAsync();
-					using (AgentMessage message = await channel.ReceiveAsync(CancellationToken.None))
+					await channel.PingAsync(cancellationToken);
+					using (AgentMessage message = await channel.ReceiveAsync(cancellationToken))
 					{
 						Assert.AreEqual(AgentMessageType.Ping, message.Type);
 						Assert.IsTrue(message.Data.Span.SequenceEqual(ReadOnlySpan<byte>.Empty));
 					}
 
-					await channel.SendXorRequestAsync(new byte[] { 1, 2, 3 }, 44);
-					using (AgentMessage message = await channel.ReceiveAsync(CancellationToken.None))
+					await channel.SendXorRequestAsync(new byte[] { 1, 2, 3 }, 44, cancellationToken);
+					using (AgentMessage message = await channel.ReceiveAsync(cancellationToken))
 					{
 						Assert.AreEqual(AgentMessageType.XorResponse, message.Type);
 						Assert.IsTrue(message.Data.Span.SequenceEqual(new byte[] { 1 ^ 44, 2 ^ 44, 3 ^ 44 }));
@@ -125,18 +150,18 @@ namespace EpicGames.Horde.Tests
 					const int SecondaryChannelId = 1;
 					using (AgentMessageChannel channel2 = localSocket.CreateAgentMessageChannel(SecondaryChannelId, 4 * 1024 * 1024))
 					{
-						await channel.ForkAsync(SecondaryChannelId, 4 * 1024 * 1024);
+						await channel.ForkAsync(SecondaryChannelId, 4 * 1024 * 1024, cancellationToken);
 
-						await channel2.WaitForAttachAsync();
+						await channel2.WaitForAttachAsync(cancellationToken);
 
-						await channel2.SendXorRequestAsync(new byte[] { 1, 2, 3 }, 44);
-						using (AgentMessage message = await channel2.ReceiveAsync(CancellationToken.None))
+						await channel2.SendXorRequestAsync(new byte[] { 1, 2, 3 }, 44, cancellationToken);
+						using (AgentMessage message = await channel2.ReceiveAsync(cancellationToken))
 						{
 							Assert.AreEqual(AgentMessageType.XorResponse, message.Type);
 							Assert.IsTrue(message.Data.Span.SequenceEqual(new byte[] { 1 ^ 44, 2 ^ 44, 3 ^ 44 }));
 						}
 
-						await channel2.CloseAsync();
+						await channel2.CloseAsync(cancellationToken);
 					}
 
 					using MemoryStorageClient memoryStorage = new MemoryStorageClient();
@@ -153,27 +178,27 @@ namespace EpicGames.Horde.Tests
 						byte[] data = Encoding.UTF8.GetBytes("Hello world");
 
 						using ChunkedDataWriter writer = new ChunkedDataWriter(treeWriter, new ChunkingOptions());
-						ChunkedData chunkedData = await writer.CreateAsync(data, CancellationToken.None);
+						ChunkedData chunkedData = await writer.CreateAsync(data, cancellationToken);
 
 						DirectoryNode directory = new DirectoryNode();
 						directory.AddFile("hello.txt", FileEntryFlags.None, data.Length, chunkedData);
 
-						HashedNodeRef<DirectoryNode> directoryRef = await treeWriter.WriteHashedNodeAsync(directory);
+						HashedNodeRef<DirectoryNode> directoryRef = await treeWriter.WriteHashedNodeAsync(directory, cancellationToken);
 
 						DirectoryNode root = new DirectoryNode();
 						root.AddDirectory(new DirectoryEntry("subdir", directory.Length, directoryRef));
 
-						IBlobHandle handle = await treeWriter.FlushAsync(root);
-						await channel.UploadFilesAsync("", handle.GetLocator(), storage);
+						IBlobHandle handle = await treeWriter.FlushAsync(root, cancellationToken);
+						await channel.UploadFilesAsync("", handle.GetLocator(), storage, cancellationToken);
 
 						Assert.IsTrue(FileReference.Exists(file));
-						byte[] readData = await FileReference.ReadAllBytesAsync(file);
+						byte[] readData = await FileReference.ReadAllBytesAsync(file, cancellationToken);
 						Assert.IsTrue(readData.SequenceEqual(data));
 
-						await channel.DeleteFilesAsync(new[] { "subdir/hello.txt" }, CancellationToken.None);
+						await channel.DeleteFilesAsync(new[] { "subdir/hello.txt" }, cancellationToken);
 
-						await channel.PingAsync();
-						using (AgentMessage message = await channel.ReceiveAsync(CancellationToken.None))
+						await channel.PingAsync(cancellationToken);
+						using (AgentMessage message = await channel.ReceiveAsync(cancellationToken))
 						{
 							Assert.AreEqual(AgentMessageType.Ping, message.Type);
 							Assert.IsTrue(message.Data.Span.SequenceEqual(ReadOnlySpan<byte>.Empty));
@@ -187,11 +212,28 @@ namespace EpicGames.Horde.Tests
 			await localSocket.CloseAsync(CancellationToken.None);
 			await agentSocket.CloseAsync(CancellationToken.None);
 		}
+		
+		static int GetAvailablePort()
+		{
+			TcpListener listener = new(IPAddress.Loopback, 0);
+			listener.Start();
+			int port = ((IPEndPoint)listener.LocalEndpoint).Port;
+			listener.Stop();
+			return port;
+		}
 
 		static async Task RunAgentAsync(ComputeSocket socket, DirectoryReference tempDir, CancellationToken cancellationToken)
 		{
-			AgentMessageHandler handler = new AgentMessageHandler(tempDir, null, true, null, NullLogger.Instance);
-			await handler.RunAsync(socket, cancellationToken);
+			try
+			{
+				AgentMessageHandler handler = new AgentMessageHandler(tempDir, null, true, null, NullLogger.Instance);
+				await handler.RunAsync(socket, cancellationToken);
+			}
+			catch (Exception e)
+			{
+				Console.WriteLine("Exception when running agent:\n" + e);
+				throw;
+			}
 		}
 	}
 }
diff --git a/Engine/Source/Programs/Shared/EpicGames.Horde/Compute/Transports/TcpSslTransport.cs b/Engine/Source/Programs/Shared/EpicGames.Horde/Compute/Transports/TcpSslTransport.cs
new file mode 100644
--- /dev/null
+++ b/Engine/Source/Programs/Shared/EpicGames.Horde/Compute/Transports/TcpSslTransport.cs
@@ -0,0 +1,146 @@
+// Copyright Epic Games, Inc. All Rights Reserved.
+
+using System;
+using System.Buffers;
+using System.Net.Security;
+using System.Net.Sockets;
+using System.Security.Cryptography;
+using System.Security.Cryptography.X509Certificates;
+using System.Threading;
+using System.Threading.Tasks;
+
+namespace EpicGames.Horde.Compute.Transports;
+
+/// <summary>
+/// Implementation of <see cref="ComputeTransport"/> for communicating over a socket using SSL/TLS
+/// </summary>
+public sealed class TcpSslTransport : ComputeTransport, IDisposable
+{
+	private readonly Socket _socket;
+	private readonly X509Certificate2 _cert;
+	private readonly bool _isServer;
+	private readonly NetworkStream _networkStream;
+	private readonly SslStream _sslStream;
+
+	/// <summary>
+	/// Constructor
+	/// </summary>
+	/// <param name="socket">Socket to communicate over</param>
+	/// <param name="certData">Certificate used for auth on both server and client</param>
+	/// <param name="isServer">Whether socket is acting as a server or client</param>
+	public TcpSslTransport(Socket socket, byte[] certData, bool isServer) : this(socket, new X509Certificate2(certData), isServer)
+	{
+	}
+	
+	/// <summary>
+	/// Constructor
+	/// </summary>
+	/// <param name="socket">Socket to communicate over</param>
+	/// <param name="cert">Certificate used for auth on both server and client</param>
+	/// <param name="isServer">Whether socket is acting as a server or client</param>
+	public TcpSslTransport(Socket socket, X509Certificate2 cert, bool isServer)
+	{
+		_socket = socket;
+		_cert = cert;
+		_isServer = isServer;
+		_networkStream = new NetworkStream(socket);
+		_sslStream = new SslStream(_networkStream, false);
+	}
+	
+	/// <inheritdoc/>
+	public void Dispose()
+	{
+		_networkStream.Dispose();
+		_sslStream.Dispose();
+	}
+
+	/// <summary>
+	/// Perform SSL authentication
+	/// </summary>
+	public async Task AuthenticateAsync(CancellationToken cancellationToken)
+	{
+		if (_isServer)
+		{
+			await AuthenticateAsServerAsync(cancellationToken);
+		}
+		else
+		{
+			await AuthenticateAsClientAsync();
+		}
+	}
+
+	private async Task AuthenticateAsClientAsync()
+	{
+		SslClientAuthenticationOptions options = new ()
+		{
+			TargetHost = "horde", // Since cert is self-signed and not CA-based, this is ignored
+			ClientCertificates = new X509CertificateCollection { _cert },
+			CertificateRevocationCheckMode = X509RevocationMode.Offline,
+			EncryptionPolicy = EncryptionPolicy.RequireEncryption,
+			RemoteCertificateValidationCallback = ValidateCert
+		};
+		
+		await _sslStream.AuthenticateAsClientAsync(options);
+	}
+	
+	private async Task AuthenticateAsServerAsync(CancellationToken cancellationToken)
+	{
+		SslServerAuthenticationOptions opt = new()
+		{
+			EncryptionPolicy = EncryptionPolicy.RequireEncryption,
+			ServerCertificate = _cert,
+			AllowRenegotiation = false,
+			ClientCertificateRequired = true,
+			RemoteCertificateValidationCallback = ValidateCert,
+		};
+		await _sslStream.AuthenticateAsServerAsync(opt, cancellationToken);
+	}
+	
+	/// <summary>
+	/// Checks the certificate returned by the server is indeed the correct one
+	/// </summary>
+	/// <param name="sender"></param>
+	/// <param name="certificate"></param>
+	/// <param name="chain"></param>
+	/// <param name="sslPolicyErrors"></param>
+	/// <returns>True if it matches</returns>
+	private bool ValidateCert(object sender, X509Certificate? certificate, X509Chain? chain, SslPolicyErrors sslPolicyErrors)
+	{
+		return certificate != null && certificate.GetCertHashString() == _cert.GetCertHashString();
+	}
+
+	/// <inheritdoc/>
+	public override async ValueTask<int> RecvAsync(Memory<byte> buffer, CancellationToken cancellationToken)
+	{
+		return await _sslStream.ReadAsync(buffer, cancellationToken);
+	}
+
+	/// <inheritdoc/>
+	public override async ValueTask SendAsync(ReadOnlySequence<byte> buffer, CancellationToken cancellationToken)
+	{
+		foreach (ReadOnlyMemory<byte> memory in buffer)
+		{
+			await _sslStream.WriteAsync(memory, cancellationToken);
+		}
+		await _sslStream.FlushAsync(cancellationToken);
+	}
+
+	/// <inheritdoc/>
+	public override ValueTask MarkCompleteAsync(CancellationToken cancelationToken)
+	{
+		_socket.Shutdown(SocketShutdown.Send);
+		return new ValueTask();
+	}
+	
+	/// <summary>
+	/// Generate a self-signed certificate to be used for communicating between client and server of this transport
+	/// </summary>
+	/// <returns>A X509 certificate serialized as bytes</returns>
+	public static byte[] GenerateCert()
+	{
+		using RSA rsa = RSA.Create(2048);
+		CertificateRequest req = new ("cn=horde", rsa, HashAlgorithmName.SHA256, RSASignaturePadding.Pkcs1);
+		X509Certificate2 cert = req.CreateSelfSigned(DateTimeOffset.UtcNow - TimeSpan.FromSeconds(10), DateTimeOffset.UtcNow.AddHours(24));
+		return cert.Export(X509ContentType.Pkcs12); // Note: Need to reimport this to use immediately, otherwise key is ephemeral
+	}
+}
