Horde: Don't wrap storage cancellation exceptions in a custom exception type.

diff --git a/Engine/Source/Programs/Shared/EpicGames.Horde/Storage/Nodes/FileEntry.cs b/Engine/Source/Programs/Shared/EpicGames.Horde/Storage/Nodes/FileEntry.cs
--- a/Engine/Source/Programs/Shared/EpicGames.Horde/Storage/Nodes/FileEntry.cs
+++ b/Engine/Source/Programs/Shared/EpicGames.Horde/Storage/Nodes/FileEntry.cs
@@ -169,7 +169,7 @@ namespace EpicGames.Horde.Storage.Nodes
 					}
 				}
 			}
-			catch (Exception ex)
+			catch (Exception ex) when (ex is not OperationCanceledException)
 			{
 				throw new Exception($"Unable to extract file {file.FullName}", ex);
 			}
