EpicGames.OIDC: Update several nuget packages

#rnx

diff --git a/Engine/Source/Programs/Shared/EpicGames.OIDC/EpicGames.OIDC.csproj b/Engine/Source/Programs/Shared/EpicGames.OIDC/EpicGames.OIDC.csproj
--- a/Engine/Source/Programs/Shared/EpicGames.OIDC/EpicGames.OIDC.csproj
+++ b/Engine/Source/Programs/Shared/EpicGames.OIDC/EpicGames.OIDC.csproj
@@ -37,9 +37,9 @@
   
   <ItemGroup>
     <PackageReference Include="IdentityModel.OidcClient" Version="4.0.0" />
-    <PackageReference Include="Microsoft.Extensions.Configuration" Version="5.0.0" />
-    <PackageReference Include="Microsoft.Extensions.Configuration.Binder" Version="5.0.0" />
-    <PackageReference Include="Microsoft.Extensions.Configuration.Json" Version="5.0.0" />
+    <PackageReference Include="Microsoft.Extensions.Configuration" Version="6.0.0" />
+    <PackageReference Include="Microsoft.Extensions.Configuration.Binder" Version="6.0.0" />
+    <PackageReference Include="Microsoft.Extensions.Configuration.Json" Version="6.0.0" />
     <PackageReference Include="Microsoft.Win32.Registry" Version="5.0.0" />
   </ItemGroup>
 
