Horde: Make ComputeTransport disposable

To make cleanup easier when multiple transports are stacked.

diff --git a/Engine/Source/Programs/Horde/Horde.Agent/Commands/Compute/ComputeWorkerCommand.cs b/Engine/Source/Programs/Horde/Horde.Agent/Commands/Compute/ComputeWorkerCommand.cs
--- a/Engine/Source/Programs/Horde/Horde.Agent/Commands/Compute/ComputeWorkerCommand.cs
+++ b/Engine/Source/Programs/Horde/Horde.Agent/Commands/Compute/ComputeWorkerCommand.cs
@@ -27,7 +27,8 @@ namespace Horde.Agent.Commands.Compute
 			using Socket tcpSocket = new Socket(SocketType.Stream, ProtocolType.IP);
 			await tcpSocket.ConnectAsync(IPAddress.Loopback, Port);
 
-			await using (RemoteComputeSocket socket = new RemoteComputeSocket(new TcpTransport(tcpSocket), logger))
+			await using TcpTransport transport = new TcpTransport(tcpSocket);
+			await using (RemoteComputeSocket socket = new RemoteComputeSocket(transport, logger))
 			{
 				logger.LogInformation("Running worker...");
 				await RunWorkerAsync(socket, logger, CancellationToken.None);
diff --git a/Engine/Source/Programs/Horde/Horde.Agent/Leases/Handlers/ComputeHandler.cs b/Engine/Source/Programs/Horde/Horde.Agent/Leases/Handlers/ComputeHandler.cs
--- a/Engine/Source/Programs/Horde/Horde.Agent/Leases/Handlers/ComputeHandler.cs
+++ b/Engine/Source/Programs/Horde/Horde.Agent/Leases/Handlers/ComputeHandler.cs
@@ -35,6 +35,12 @@ namespace Horde.Agent.Leases.Handlers
 			public TimeSpan TimeSinceActivity => TimeSpan.FromTicks((long)((Stopwatch.GetTimestamp() - Interlocked.CompareExchange(ref _lastPingTicks, 0, 0)) * s_ticksToSystemTicks));
 
 			public override ValueTask MarkCompleteAsync(CancellationToken cancellationToken) => _inner.MarkCompleteAsync(cancellationToken);
+			
+			/// <inheritdoc/>
+			public override ValueTask DisposeAsync()
+			{
+				return _inner.DisposeAsync();
+			}
 
 			public override async ValueTask<int> RecvAsync(Memory<byte> buffer, CancellationToken cancellationToken)
 			{
@@ -123,7 +129,7 @@ namespace Horde.Agent.Leases.Handlers
 
 				logger.LogInformation("Matched connection for {Nonce}", StringUtils.FormatHexString(computeTask.Nonce.Span));
 
-				TcpTransportWithTimeout transport = new TcpTransportWithTimeout(tcpClient.Client);
+				await using TcpTransportWithTimeout transport = new TcpTransportWithTimeout(tcpClient.Client);
 				using (CancellationTokenSource cts = CancellationTokenSource.CreateLinkedTokenSource(cancellationToken))
 				{
 					await using BackgroundTask timeoutTask = BackgroundTask.StartNew(ctx => TickTimeoutAsync(transport, cts, logger, ctx));
diff --git a/Engine/Source/Programs/Horde/Horde/Commands/Compute/ComputeWorker.cs b/Engine/Source/Programs/Horde/Horde/Commands/Compute/ComputeWorker.cs
--- a/Engine/Source/Programs/Horde/Horde/Commands/Compute/ComputeWorker.cs
+++ b/Engine/Source/Programs/Horde/Horde/Commands/Compute/ComputeWorker.cs
@@ -27,7 +27,8 @@ namespace Horde.Commands.Compute
 			using Socket tcpSocket = new Socket(SocketType.Stream, ProtocolType.IP);
 			await tcpSocket.ConnectAsync(IPAddress.Loopback, Port);
 
-			await using (RemoteComputeSocket socket = new RemoteComputeSocket(new TcpTransport(tcpSocket), logger))
+			await using TcpTransport transport = new TcpTransport(tcpSocket);
+			await using (RemoteComputeSocket socket = new RemoteComputeSocket(transport, logger))
 			{
 				logger.LogInformation("Running worker...");
 				await RunWorkerAsync(socket, logger, CancellationToken.None);
diff --git a/Engine/Source/Programs/Shared/EpicGames.Horde.Tests/BufferTests.cs b/Engine/Source/Programs/Shared/EpicGames.Horde.Tests/BufferTests.cs
--- a/Engine/Source/Programs/Shared/EpicGames.Horde.Tests/BufferTests.cs
+++ b/Engine/Source/Programs/Shared/EpicGames.Horde.Tests/BufferTests.cs
@@ -104,8 +104,10 @@ namespace EpicGames.Horde.Tests
 			Pipe sourceToTargetPipe = new Pipe();
 			Pipe targetToSourcePipe = new Pipe();
 
-			await using RemoteComputeSocket producerSocket = new RemoteComputeSocket(new PipeTransport(targetToSourcePipe.Reader, sourceToTargetPipe.Writer), NullLogger.Instance);
-			await using RemoteComputeSocket consumerSocket = new RemoteComputeSocket(new PipeTransport(sourceToTargetPipe.Reader, targetToSourcePipe.Writer), NullLogger.Instance);
+			await using PipeTransport producerTransport = new (targetToSourcePipe.Reader, sourceToTargetPipe.Writer);
+			await using PipeTransport consumerTransport = new (sourceToTargetPipe.Reader, targetToSourcePipe.Writer);
+			await using RemoteComputeSocket producerSocket = new (producerTransport, NullLogger.Instance);
+			await using RemoteComputeSocket consumerSocket = new (consumerTransport, NullLogger.Instance);
 
 			using ComputeBuffer consumerBuffer = createBuffer(Length);
 			consumerSocket.AttachRecvBuffer(ChannelId, consumerBuffer);
@@ -165,8 +167,10 @@ namespace EpicGames.Horde.Tests
 		{
 			Pipe recvPipe = new Pipe();
 			Pipe sendPipe = new Pipe();
-			await using RemoteComputeSocket localSocket = new RemoteComputeSocket(new PipeTransport(sendPipe.Reader, recvPipe.Writer), NullLogger.Instance);
-			await using RemoteComputeSocket remoteSocket = new RemoteComputeSocket(new PipeTransport(recvPipe.Reader, sendPipe.Writer), NullLogger.Instance);
+			await using PipeTransport localTransport = new (sendPipe.Reader, recvPipe.Writer);
+			await using PipeTransport remoteTransport = new (recvPipe.Reader, sendPipe.Writer);
+			await using RemoteComputeSocket localSocket = new (localTransport, NullLogger.Instance);
+			await using RemoteComputeSocket remoteSocket = new (remoteTransport, NullLogger.Instance);
 
 			using (PooledBuffer remoteBuffer = new PooledBuffer(1024))
 			{
diff --git a/Engine/Source/Programs/Shared/EpicGames.Horde.Tests/ComputeSocketTests.cs b/Engine/Source/Programs/Shared/EpicGames.Horde.Tests/ComputeSocketTests.cs
--- a/Engine/Source/Programs/Shared/EpicGames.Horde.Tests/ComputeSocketTests.cs
+++ b/Engine/Source/Programs/Shared/EpicGames.Horde.Tests/ComputeSocketTests.cs
@@ -73,8 +73,10 @@ namespace EpicGames.Horde.Tests
 		{
 			Pipe recvPipe = new Pipe();
 			Pipe sendPipe = new Pipe();
-			await using RemoteComputeSocket localSocket = new RemoteComputeSocket(new PipeTransport(sendPipe.Reader, recvPipe.Writer), new TestLogger());
-			await using RemoteComputeSocket agentSocket = new RemoteComputeSocket(new PipeTransport(recvPipe.Reader, sendPipe.Writer), new TestLogger());
+			await using PipeTransport localTransport = new (sendPipe.Reader, recvPipe.Writer);
+			await using PipeTransport agentTransport = new (recvPipe.Reader, sendPipe.Writer);
+			await using RemoteComputeSocket localSocket = new (localTransport, new TestLogger());
+			await using RemoteComputeSocket agentSocket = new (agentTransport, new TestLogger());
 
 			await RunAgentTestsAsync(localSocket, agentSocket);
 		}
@@ -85,8 +87,10 @@ namespace EpicGames.Horde.Tests
 			using CancellationTokenSource cts = new (5000);
 			(Socket clientSocket, Socket serverSocket) = await CreateSocketsAsync(cts.Token);
 
-			await using RemoteComputeSocket localSocket = new (new TcpTransport(clientSocket), new TestLogger());
-			await using RemoteComputeSocket agentSocket = new (new TcpTransport(serverSocket), new TestLogger());
+			await using TcpTransport clientTransport = new (clientSocket);
+			await using TcpTransport serverTransport = new (serverSocket);
+			await using RemoteComputeSocket localSocket = new (clientTransport, new TestLogger());
+			await using RemoteComputeSocket agentSocket = new (serverTransport, new TestLogger());
 
 			await RunAgentTestsAsync(localSocket, agentSocket, cts.Token);
 		}
@@ -98,8 +102,8 @@ namespace EpicGames.Horde.Tests
 			(Socket clientSocket, Socket serverSocket) = await CreateSocketsAsync(cts.Token);
 			
 			byte[] certData = TcpSslTransport.GenerateCert();
-			using TcpSslTransport clientTransport = new (clientSocket, certData, false);
-			using TcpSslTransport serverTransport = new (serverSocket, certData, true);
+			await using TcpSslTransport clientTransport = new (clientSocket, certData, false);
+			await using TcpSslTransport serverTransport = new (serverSocket, certData, true);
 			
 			Task t1 = clientTransport.AuthenticateAsync(cts.Token);
 			Task t2 = serverTransport.AuthenticateAsync(cts.Token);
diff --git a/Engine/Source/Programs/Shared/EpicGames.Horde.Tests/TransportTests.cs b/Engine/Source/Programs/Shared/EpicGames.Horde.Tests/TransportTests.cs
--- a/Engine/Source/Programs/Shared/EpicGames.Horde.Tests/TransportTests.cs
+++ b/Engine/Source/Programs/Shared/EpicGames.Horde.Tests/TransportTests.cs
@@ -24,7 +24,7 @@ namespace EpicGames.Horde.Tests
 			byte[] payload;
 			{
 				using MemoryStream memoryStream = new MemoryStream();
-				StreamTransport streamTransport = new StreamTransport(memoryStream);
+				await using StreamTransport streamTransport = new StreamTransport(memoryStream);
 				await streamTransport.SendAsync(input, CancellationToken.None);
 				payload = memoryStream.ToArray();
 			}
@@ -32,7 +32,7 @@ namespace EpicGames.Horde.Tests
 			byte[] output = new byte[input.Length];
 			{
 				using MemoryStream memoryStream = new MemoryStream(payload);
-				StreamTransport streamTransport = new StreamTransport(memoryStream);
+				await using StreamTransport streamTransport = new StreamTransport(memoryStream);
 				await CopyToAsync(streamTransport, output);
 			}
 
@@ -60,7 +60,7 @@ namespace EpicGames.Horde.Tests
 
 			{
 				using MemoryStream memoryStream = new MemoryStream();
-				StreamTransport streamTransport = new StreamTransport(memoryStream);
+				await using StreamTransport streamTransport = new StreamTransport(memoryStream);
 				await using AesTransport aesTransport = new AesTransport(streamTransport, key, nonce);
 				await aesTransport.SendAsync(input, CancellationToken.None);
 				encrypted = memoryStream.ToArray();
@@ -69,7 +69,7 @@ namespace EpicGames.Horde.Tests
 			byte[] output = new byte[input.Length];
 			{
 				using MemoryStream memoryStream = new MemoryStream(encrypted);
-				StreamTransport streamTransport = new StreamTransport(memoryStream);
+				await using StreamTransport streamTransport = new StreamTransport(memoryStream);
 				await using AesTransport aesTransport = new AesTransport(streamTransport, key, nonce);
 
 				for (int offset = 0; offset < output.Length;)
diff --git a/Engine/Source/Programs/Shared/EpicGames.Horde/Compute/Clients/AgentComputeClient.cs b/Engine/Source/Programs/Shared/EpicGames.Horde/Compute/Clients/AgentComputeClient.cs
--- a/Engine/Source/Programs/Shared/EpicGames.Horde/Compute/Clients/AgentComputeClient.cs
+++ b/Engine/Source/Programs/Shared/EpicGames.Horde/Compute/Clients/AgentComputeClient.cs
@@ -106,7 +106,8 @@ namespace EpicGames.Horde.Compute.Clients
 			await using BackgroundTask agentTask = BackgroundTask.StartNew(ctx => RunAgentAsync(_hordeAgentAssembly, _port, logger, ctx));
 			using Socket tcpSocket = await listener.AcceptAsync(cancellationToken);
 
-			await using RemoteComputeSocket socket = new RemoteComputeSocket(new TcpTransport(tcpSocket), _logger);
+			await using TcpTransport tcpTransport = new (tcpSocket);
+			await using RemoteComputeSocket socket = new (tcpTransport, _logger);
 			yield return socket;
 
 			await socket.CloseAsync(cancellationToken);
diff --git a/Engine/Source/Programs/Shared/EpicGames.Horde/Compute/Clients/LocalComputeClient.cs b/Engine/Source/Programs/Shared/EpicGames.Horde/Compute/Clients/LocalComputeClient.cs
--- a/Engine/Source/Programs/Shared/EpicGames.Horde/Compute/Clients/LocalComputeClient.cs
+++ b/Engine/Source/Programs/Shared/EpicGames.Horde/Compute/Clients/LocalComputeClient.cs
@@ -76,13 +76,11 @@ namespace EpicGames.Horde.Compute.Clients
 		static async Task RunListenerAsync(Socket listener, DirectoryReference sandboxDir, bool executeInProcess, ILogger logger, CancellationToken cancellationToken)
 		{
 			using Socket tcpSocket = await listener.AcceptAsync(cancellationToken);
-
-			await using (RemoteComputeSocket socket = new RemoteComputeSocket(new TcpTransport(tcpSocket), logger))
-			{
-				AgentMessageHandler worker = new AgentMessageHandler(sandboxDir, null, executeInProcess, null, logger);
-				await worker.RunAsync(socket, cancellationToken);
-				await socket.CloseAsync(cancellationToken);
-			}
+			await using TcpTransport tcpTransport = new (tcpSocket);
+			await using RemoteComputeSocket socket = new (tcpTransport, logger);
+			AgentMessageHandler worker = new (sandboxDir, null, executeInProcess, null, logger);
+			await worker.RunAsync(socket, cancellationToken);
+			await socket.CloseAsync(cancellationToken);
 		}
 
 		/// <inheritdoc/>
diff --git a/Engine/Source/Programs/Shared/EpicGames.Horde/Compute/Clients/ServerComputeClient.cs b/Engine/Source/Programs/Shared/EpicGames.Horde/Compute/Clients/ServerComputeClient.cs
--- a/Engine/Source/Programs/Shared/EpicGames.Horde/Compute/Clients/ServerComputeClient.cs
+++ b/Engine/Source/Programs/Shared/EpicGames.Horde/Compute/Clients/ServerComputeClient.cs
@@ -296,7 +296,8 @@ namespace EpicGames.Horde.Compute.Clients
 			// Pass the rest of the call over to the handler
 			byte[] key = StringUtils.ParseHexString(response.Key);
 
-			await using RemoteComputeSocket computeSocket = new RemoteComputeSocket(new TcpTransport(socket), workerLogger);
+			await using TcpTransport transport = new TcpTransport(socket);
+			await using RemoteComputeSocket computeSocket = new RemoteComputeSocket(transport, workerLogger);
 			yield return new LeaseInfo(response.Properties, response.AssignedResources, computeSocket, response.Ip, response.ConnectionMode, response.Ports);
 		}
 
diff --git a/Engine/Source/Programs/Shared/EpicGames.Horde/Compute/ComputeTransport.cs b/Engine/Source/Programs/Shared/EpicGames.Horde/Compute/ComputeTransport.cs
--- a/Engine/Source/Programs/Shared/EpicGames.Horde/Compute/ComputeTransport.cs
+++ b/Engine/Source/Programs/Shared/EpicGames.Horde/Compute/ComputeTransport.cs
@@ -11,7 +11,7 @@ namespace EpicGames.Horde.Compute
 	/// <summary>
 	/// Low-level interface for transferring data
 	/// </summary>
-	public abstract class ComputeTransport
+	public abstract class ComputeTransport : IAsyncDisposable
 	{
 		/// <summary>
 		/// Writes data to the underlying transport
@@ -32,6 +32,9 @@ namespace EpicGames.Horde.Compute
 		/// </summary>
 		/// <param name="cancellationToken">Cancellation token for the operation</param>
 		public abstract ValueTask MarkCompleteAsync(CancellationToken cancellationToken);
+		
+		/// <inheritdoc/>
+		public abstract ValueTask DisposeAsync();
 
 		/// <summary>
 		/// Fill the given buffer with data
diff --git a/Engine/Source/Programs/Shared/EpicGames.Horde/Compute/Transports/AesTransport.cs b/Engine/Source/Programs/Shared/EpicGames.Horde/Compute/Transports/AesTransport.cs
--- a/Engine/Source/Programs/Shared/EpicGames.Horde/Compute/Transports/AesTransport.cs
+++ b/Engine/Source/Programs/Shared/EpicGames.Horde/Compute/Transports/AesTransport.cs
@@ -15,7 +15,7 @@ namespace EpicGames.Horde.Compute.Transports
 	/// Transport layer that adds AES encryption on top of an underlying transport implementation. Key must be exchanged separately
 	/// (eg. via the HTTPS request to negotiate a lease with the server).
 	/// </summary>
-	public sealed class AesTransport : ComputeTransport, IAsyncDisposable
+	public sealed class AesTransport : ComputeTransport
 	{
 		/// <summary>
 		/// Length of the required encrption key. 
@@ -31,6 +31,7 @@ namespace EpicGames.Horde.Compute.Transports
 		const int FooterLength = 16; // 16-byte auth tag for encryption
 
 		readonly ComputeTransport _inner;
+		readonly bool _leaveInnerOpen;
 		readonly Pipe _readPipe;
 		readonly AesGcm _aesGcm;
 
@@ -55,7 +56,8 @@ namespace EpicGames.Horde.Compute.Transports
 		/// <param name="inner">The underlying transport implementation</param>
 		/// <param name="key">AES encryption key (256 bits / 32 bytes)</param>
 		/// <param name="nonce">Cryptographic nonce to identify the connection. Must be longer than <see cref="NonceLength"/>.</param>
-		public AesTransport(ComputeTransport inner, ReadOnlySpan<byte> key, ReadOnlySpan<byte> nonce)
+		/// <param name="leaveInnerOpen">Whether inner compute transport should be disposed</param>
+		public AesTransport(ComputeTransport inner, ReadOnlySpan<byte> key, ReadOnlySpan<byte> nonce, bool leaveInnerOpen = false)
 		{
 			if (key.Length != KeyLength)
 			{
@@ -71,19 +73,23 @@ namespace EpicGames.Horde.Compute.Transports
 			_aesGcm = new AesGcm(key);
 			_readNonce = nonce.Slice(0, NonceLength).ToArray();
 			_writeNonce = nonce.Slice(0, NonceLength).ToArray();
-
+			_leaveInnerOpen = leaveInnerOpen;
 			_backgroundReadTask = BackgroundTask.StartNew(BackgroundReadAsync);
 
 			_writeBufferOwner = MemoryPool<byte>.Shared.Rent(WritePacketSize * 2);
 		}
 
 		/// <inheritdoc/>
-		public async ValueTask DisposeAsync()
+		public override async ValueTask DisposeAsync()
 		{
 			await _lastWriteTask;
 			await _backgroundReadTask.DisposeAsync();
 			_writeBufferOwner.Dispose();
 			_aesGcm.Dispose();
+			if (!_leaveInnerOpen)
+			{
+				await _inner.DisposeAsync();
+			}
 		}
 
 		/// <summary>
diff --git a/Engine/Source/Programs/Shared/EpicGames.Horde/Compute/Transports/IdleTimeoutTransport.cs b/Engine/Source/Programs/Shared/EpicGames.Horde/Compute/Transports/IdleTimeoutTransport.cs
--- a/Engine/Source/Programs/Shared/EpicGames.Horde/Compute/Transports/IdleTimeoutTransport.cs
+++ b/Engine/Source/Programs/Shared/EpicGames.Horde/Compute/Transports/IdleTimeoutTransport.cs
@@ -35,6 +35,13 @@ public class IdleTimeoutTransport : ComputeTransport
 		NoDataTimeout = noDataTimeout ?? NoDataTimeout;
 		_lastPingTicks = Stopwatch.GetTimestamp();
 	}
+	
+	/// <inheritdoc/>
+	public override ValueTask DisposeAsync()
+	{
+		GC.SuppressFinalize(this);
+		return ValueTask.CompletedTask;
+	}
 
 	/// <summary>
 	/// Time since last send or receive completed
diff --git a/Engine/Source/Programs/Shared/EpicGames.Horde/Compute/Transports/PipeTransport.cs b/Engine/Source/Programs/Shared/EpicGames.Horde/Compute/Transports/PipeTransport.cs
--- a/Engine/Source/Programs/Shared/EpicGames.Horde/Compute/Transports/PipeTransport.cs
+++ b/Engine/Source/Programs/Shared/EpicGames.Horde/Compute/Transports/PipeTransport.cs
@@ -33,6 +33,13 @@ namespace EpicGames.Horde.Compute.Transports
 			_reader = reader;
 			_writer = writer;
 		}
+		
+		/// <inheritdoc/>
+		public override ValueTask DisposeAsync()
+		{
+			GC.SuppressFinalize(this);
+			return ValueTask.CompletedTask;
+		}
 
 		/// <inheritdoc/>
 		public override async ValueTask<int> RecvAsync(Memory<byte> buffer, CancellationToken cancellationToken)
diff --git a/Engine/Source/Programs/Shared/EpicGames.Horde/Compute/Transports/StreamTransport.cs b/Engine/Source/Programs/Shared/EpicGames.Horde/Compute/Transports/StreamTransport.cs
--- a/Engine/Source/Programs/Shared/EpicGames.Horde/Compute/Transports/StreamTransport.cs
+++ b/Engine/Source/Programs/Shared/EpicGames.Horde/Compute/Transports/StreamTransport.cs
@@ -14,15 +14,29 @@ namespace EpicGames.Horde.Compute.Transports
 	class StreamTransport : ComputeTransport
 	{
 		readonly Stream _stream;
-
-		/// <inheritdoc/>
+		readonly bool _leaveOpen;
 		public long Position { get; private set; }
 
 		/// <summary>
 		/// Constructor
 		/// </summary>
 		/// <param name="stream">Stream to use for the transferring data</param>
-		public StreamTransport(Stream stream) => _stream = stream;
+		/// <param name="leaveOpen">Whether to leave the inner stream open when disposing</param>
+		public StreamTransport(Stream stream, bool leaveOpen = false)
+		{
+			_stream = stream;
+			_leaveOpen = leaveOpen;
+		}
+
+		/// <inheritdoc/>
+		public override async ValueTask DisposeAsync()
+		{
+			if (!_leaveOpen)
+			{
+				await _stream.DisposeAsync();
+			}
+			GC.SuppressFinalize(this);
+		}
 
 		/// <inheritdoc/>
 		public override async ValueTask<int> RecvAsync(Memory<byte> buffer, CancellationToken cancellationToken)
diff --git a/Engine/Source/Programs/Shared/EpicGames.Horde/Compute/Transports/TcpSslTransport.cs b/Engine/Source/Programs/Shared/EpicGames.Horde/Compute/Transports/TcpSslTransport.cs
--- a/Engine/Source/Programs/Shared/EpicGames.Horde/Compute/Transports/TcpSslTransport.cs
+++ b/Engine/Source/Programs/Shared/EpicGames.Horde/Compute/Transports/TcpSslTransport.cs
@@ -14,7 +14,7 @@ namespace EpicGames.Horde.Compute.Transports;
 /// <summary>
 /// Implementation of <see cref="ComputeTransport"/> for communicating over a socket using SSL/TLS
 /// </summary>
-public sealed class TcpSslTransport : ComputeTransport, IDisposable
+public sealed class TcpSslTransport : ComputeTransport
 {
 	private readonly Socket _socket;
 	private readonly X509Certificate2 _cert;
@@ -48,10 +48,11 @@ public sealed class TcpSslTransport : ComputeTransport, IDisposable
 	}
 	
 	/// <inheritdoc/>
-	public void Dispose()
+	public override ValueTask DisposeAsync()
 	{
 		_networkStream.Dispose();
 		_sslStream.Dispose();
+		return ValueTask.CompletedTask;
 	}
 
 	/// <summary>
diff --git a/Engine/Source/Programs/Shared/EpicGames.Horde/Compute/Transports/TcpTransport.cs b/Engine/Source/Programs/Shared/EpicGames.Horde/Compute/Transports/TcpTransport.cs
--- a/Engine/Source/Programs/Shared/EpicGames.Horde/Compute/Transports/TcpTransport.cs
+++ b/Engine/Source/Programs/Shared/EpicGames.Horde/Compute/Transports/TcpTransport.cs
@@ -24,6 +24,12 @@ namespace EpicGames.Horde.Compute.Transports
 		/// </summary>
 		/// <param name="socket">Socket to communicate over</param>
 		public TcpTransport(Socket socket) => _socket = socket;
+		
+		/// <inheritdoc/>
+		public override ValueTask DisposeAsync()
+		{
+			return ValueTask.CompletedTask;
+		}
 
 		/// <inheritdoc/>
 		public override async ValueTask<int> RecvAsync(Memory<byte> buffer, CancellationToken cancellationToken)
