Horde: Log errors in ComputeSocket's background send/recv

diff --git a/Engine/Source/Programs/Shared/EpicGames.Horde/Compute/ComputeSocket.cs b/Engine/Source/Programs/Shared/EpicGames.Horde/Compute/ComputeSocket.cs
--- a/Engine/Source/Programs/Shared/EpicGames.Horde/Compute/ComputeSocket.cs
+++ b/Engine/Source/Programs/Shared/EpicGames.Horde/Compute/ComputeSocket.cs
@@ -433,72 +433,80 @@ namespace EpicGames.Horde.Compute
 
 		async Task RunRecvTaskAsync(ComputeTransport transport, CancellationToken cancellationToken)
 		{
-			_logger.LogDebug("Started socket reader");
-
-			List<Task> detachTasks = new List<Task>();
-
-			byte[] header = new byte[8];
 			try
 			{
-				Memory<byte> last = Memory<byte>.Empty;
+				_logger.LogDebug("Started socket reader");
 
-				// Process messages from the remote
-				for (; ; )
+				List<Task> detachTasks = new List<Task>();
+
+				byte[] header = new byte[8];
+				try
 				{
-					detachTasks.RemoveCompleteTasks();
+					Memory<byte> last = Memory<byte>.Empty;
 
-					// Read the next packet header
-					if (!await transport.RecvOptionalAsync(header, cancellationToken))
+					// Process messages from the remote
+					for (; ; )
 					{
-						_logger.LogDebug("End of socket");
-						break;
-					}
+						detachTasks.RemoveCompleteTasks();
 
-					// Parse the target buffer and packet size
-					int id = BinaryPrimitives.ReadInt32LittleEndian(header);
-					int size = BinaryPrimitives.ReadInt32LittleEndian(header.AsSpan(4));
+						// Read the next packet header
+						if (!await transport.RecvOptionalAsync(header, cancellationToken))
+						{
+							_logger.LogDebug("End of socket");
+							break;
+						}
 
-					// Dispatch it to the correct place
-					if (size >= 0)
-					{
-						await ReadPacketAsync(transport, id, size, cancellationToken);
-					}
-					else if (size == (int)ControlMessageType.Detach)
-					{
-						detachTasks.Add(DetachRecvBufferAsync(id, cancellationToken));
-					}
-					else if (size == (int)ControlMessageType.KeepAlive)
-					{
-						_logger.LogDebug("Received ping message");
+						// Parse the target buffer and packet size
+						int id = BinaryPrimitives.ReadInt32LittleEndian(header);
+						int size = BinaryPrimitives.ReadInt32LittleEndian(header.AsSpan(4));
+
+						// Dispatch it to the correct place
+						if (size >= 0)
+						{
+							await ReadPacketAsync(transport, id, size, cancellationToken);
+						}
+						else if (size == (int)ControlMessageType.Detach)
+						{
+							detachTasks.Add(DetachRecvBufferAsync(id, cancellationToken));
+						}
+						else if (size == (int)ControlMessageType.KeepAlive)
+						{
+							_logger.LogDebug("Received ping message");
+						}
+						else
+						{
+							_logger.LogDebug("Unrecognized control message: {Message}", size);
+						}
 					}
-					else
+				}
+				catch (OperationCanceledException)
+				{
+				}
+
+				// Mark all buffers as complete
+				lock (_lockObject)
+				{
+					_complete = true;
+					foreach (int channelIdx in _recvBuffers.Keys)
 					{
-						_logger.LogDebug("Unrecognized control message: {Message}", size);
+						detachTasks.Add(DetachRecvBufferAsync(channelIdx, cancellationToken));
 					}
 				}
-			}
-			catch (OperationCanceledException)
-			{
-			}
 
-			// Mark all buffers as complete
-			lock (_lockObject)
-			{
-				_complete = true;
-				foreach (int channelIdx in _recvBuffers.Keys)
+				// Wait for all the detach tasks to finish
+				if (detachTasks.Count > 0)
 				{
-					detachTasks.Add(DetachRecvBufferAsync(channelIdx, cancellationToken));
+					_logger.LogDebug("Waiting for detach tasks to complete...");
+					await Task.WhenAll(detachTasks).WaitAsync(cancellationToken);
 				}
-			}
 
-			// Wait for all the detach tasks to finish
-			if (detachTasks.Count > 0)
+				_logger.LogDebug("Closing reader");
+			}
+			catch (Exception e)
 			{
-				_logger.LogDebug("Waiting for detach tasks to complete...");
-				await Task.WhenAll(detachTasks).WaitAsync(cancellationToken);
+				_logger.LogError(e, "Error in background receive");
+				throw;
 			}
-
-			_logger.LogDebug("Closing reader");
 		}
 
 		async Task ReadPacketAsync(ComputeTransport transport, int id, int size, CancellationToken cancellationToken)
@@ -775,21 +783,29 @@ namespace EpicGames.Horde.Compute
 
 		async Task SendFromBufferAsync(int channelId, SendBuffer sendBuffer, CancellationToken cancellationToken)
 		{
-			ComputeBufferReader reader = sendBuffer._reader!;
-			while (!cancellationToken.IsCancellationRequested)
+			try
 			{
-				ReadOnlyMemory<byte> memory = reader.GetReadBuffer();
-				if (memory.Length > 0)
-				{
-					await SendAsync(channelId, memory, cancellationToken);
-					reader.AdvanceReadPosition(memory.Length);
-				}
-				if (reader.IsComplete)
+				ComputeBufferReader reader = sendBuffer._reader!;
+				while (!cancellationToken.IsCancellationRequested)
 				{
-					await MarkCompleteAsync(channelId, cancellationToken);
-					break;
+					ReadOnlyMemory<byte> memory = reader.GetReadBuffer();
+					if (memory.Length > 0)
+					{
+						await SendAsync(channelId, memory, cancellationToken);
+						reader.AdvanceReadPosition(memory.Length);
+					}
+					if (reader.IsComplete)
+					{
+						await MarkCompleteAsync(channelId, cancellationToken);
+						break;
+					}
+					await reader.WaitToReadAsync(1, cancellationToken);
 				}
-				await reader.WaitToReadAsync(1, cancellationToken);
+			}
+			catch (Exception e)
+			{
+				_logger.LogError(e, "Error in background send");
+				throw;
 			}
 		}
 	}
