Horde: Allow specifying encryption mode in compute client

Currently defaults to no encryption (UBA uses separate encryption)

diff --git a/Engine/Source/Programs/Shared/EpicGames.Horde/Compute/Clients/ServerComputeClient.cs b/Engine/Source/Programs/Shared/EpicGames.Horde/Compute/Clients/ServerComputeClient.cs
--- a/Engine/Source/Programs/Shared/EpicGames.Horde/Compute/Clients/ServerComputeClient.cs
+++ b/Engine/Source/Programs/Shared/EpicGames.Horde/Compute/Clients/ServerComputeClient.cs
@@ -266,7 +266,8 @@ namespace EpicGames.Horde.Compute.Clients
 			// Connect to the remote machine
 			using Socket socket = new Socket(SocketType.Stream, ProtocolType.Tcp);
 			
-			workerLogger.LogDebug("Connecting to {AgentId} at {AgentAddress} ({ConnectionType} via {ConnectionAddress}) with nonce {Nonce}...", response.AgentId, agentAddress, response.ConnectionMode, response.ConnectionAddress ?? "None", response.Nonce);
+			workerLogger.LogDebug("Connecting to {AgentId} at {AgentAddress} ({ConnectionType} via {ConnectionAddress}) with nonce {Nonce} and encryption {Encryption}...",
+				response.AgentId, agentAddress, response.ConnectionMode, response.ConnectionAddress ?? "None", response.Nonce, response.Encryption);
 			switch (response.ConnectionMode)
 			{
 				case ConnectionMode.Direct:
@@ -293,14 +294,32 @@ namespace EpicGames.Horde.Compute.Clients
 			await socket.SendMessageAsync(nonce, SocketFlags.None, cancellationToken);
 			workerLogger.LogInformation("Connected to {AgentId} ({Ip}) under lease {LeaseId}", response.AgentId, response.Ip, response.LeaseId);
 
-			// Pass the rest of the call over to the handler
-			byte[] key = StringUtils.ParseHexString(response.Key);
-
-			await using TcpTransport transport = new TcpTransport(socket);
-			await using RemoteComputeSocket computeSocket = new RemoteComputeSocket(transport, workerLogger);
+			await using ComputeTransport transport = await CreateTransportAsync(socket, response, cancellationToken);
+			await using RemoteComputeSocket computeSocket = new (transport, workerLogger);
 			yield return new LeaseInfo(response.Properties, response.AssignedResources, computeSocket, response.Ip, response.ConnectionMode, response.Ports);
 		}
 
+		private static async Task<ComputeTransport> CreateTransportAsync(Socket socket, AssignComputeResponse response, CancellationToken cancellationToken)
+		{
+			switch (response.Encryption)
+			{
+				case Encryption.Ssl:
+					TcpSslTransport sslTransport = new(socket, StringUtils.ParseHexString(response.Certificate), false);
+					await sslTransport.AuthenticateAsync(cancellationToken);
+					return sslTransport;
+				
+				case Encryption.Aes:
+#pragma warning disable CA2000 // Dispose objects before losing scope
+					TcpTransport tcpTransport = new (socket);
+					return new AesTransport(tcpTransport, StringUtils.ParseHexString(response.Key), StringUtils.ParseHexString(response.Nonce));
+#pragma warning restore CA2000 // Restore CA2000
+
+				case Encryption.None:
+				default:
+					return new TcpTransport(socket);
+			}
+		}
+
 		private static (string host, int port) ParseHostPort(string address)
 		{
 			try
