Horde: Change log level to info for background recv/send logs in ComputeSocket

While log level error is technically correct, this is failing builds in Horde as it reacts to that log level inside jobs.

diff --git a/Engine/Source/Programs/Shared/EpicGames.Horde/Compute/ComputeSocket.cs b/Engine/Source/Programs/Shared/EpicGames.Horde/Compute/ComputeSocket.cs
--- a/Engine/Source/Programs/Shared/EpicGames.Horde/Compute/ComputeSocket.cs
+++ b/Engine/Source/Programs/Shared/EpicGames.Horde/Compute/ComputeSocket.cs
@@ -504,7 +504,7 @@ namespace EpicGames.Horde.Compute
 			}
 			catch (Exception e)
 			{
-				_logger.LogError(e, "Error in background receive");
+				_logger.LogInformation(e, "Error in background receive");
 				throw;
 			}
 		}
@@ -804,7 +804,7 @@ namespace EpicGames.Horde.Compute
 			}
 			catch (Exception e)
 			{
-				_logger.LogError(e, "Error in background send");
+				_logger.LogInformation(e, "Error in background send");
 				throw;
 			}
 		}
