Horde: Improve error handling when exceptions are thrown on forked compute message channels.

* Cancellation token source shared by all channel tasks is cancelled whenever any channel throws an exception.
* Exception callstack is captured at source, logged, then put into a shared location as an atomic, first-writer-wins operation.
* Once all tasks have joined with the main task, the exception is re-thrown with the original callstack.

diff --git a/Engine/Source/Programs/Shared/EpicGames.Horde/Compute/AgentMessageHandler.cs b/Engine/Source/Programs/Shared/EpicGames.Horde/Compute/AgentMessageHandler.cs
--- a/Engine/Source/Programs/Shared/EpicGames.Horde/Compute/AgentMessageHandler.cs
+++ b/Engine/Source/Programs/Shared/EpicGames.Horde/Compute/AgentMessageHandler.cs
@@ -14,6 +14,7 @@ using EpicGames.Horde.Storage;
 using System.IO;
 using System.Linq;
 using EpicGames.Horde.Storage.Clients;
+using System.Runtime.ExceptionServices;
 
 namespace EpicGames.Horde.Compute
 {
@@ -52,16 +53,36 @@ namespace EpicGames.Horde.Compute
 		/// <param name="cancellationToken">Cancellation token for the operation</param>
 		public async Task RunAsync(ComputeSocket socket, CancellationToken cancellationToken)
 		{
-			await RunAsync(socket, 0, 4 * 1024 * 1024, cancellationToken);
+			// Since we allow forking message channels, we want to ensure that errors on one channel are propagated back here, and terminate the whole connection. 
+			// To do that, we take first exception thrown and rethrow it with the original callstack here, while also forcing all other tasks to terminate via a 
+			// shared cancellation token.
+			using CancellationTokenSource cancellationSource = CancellationTokenSource.CreateLinkedTokenSource(cancellationToken);
+			ExceptionDispatchInfo? exceptionInfo = null;
+
+			void PostException(Exception ex)
+			{
+				// Capture stack from call site
+				Interlocked.CompareExchange(ref exceptionInfo, ExceptionDispatchInfo.Capture(ex), null);
+				cancellationSource.Cancel();
+			}
+
+			await RunAsync(socket, 0, 4 * 1024 * 1024, PostException, cancellationSource.Token);
+
+			// Throw the regular cancellation exception if requested
+			cancellationToken.ThrowIfCancellationRequested();
+
+			// Otherwise throw any exception posted by a child task
+			exceptionInfo?.Throw();
 		}
 
-		async Task RunAsync(ComputeSocket socket, int channelId, int bufferSize, CancellationToken cancellationToken)
+		async Task RunAsync(ComputeSocket socket, int channelId, int bufferSize, Action<Exception> postException, CancellationToken cancellationToken)
 		{
-			using (AgentMessageChannel channel = socket.CreateAgentMessageChannel(channelId, bufferSize))
+			List<Task> childTasks = new List<Task>();
+			try
 			{
+				using AgentMessageChannel channel = socket.CreateAgentMessageChannel(channelId, bufferSize);
 				await channel.AttachAsync(cancellationToken);
 
-				List<Task> childTasks = new List<Task>();
 				for (; ; )
 				{
 					using AgentMessage message = await channel.ReceiveAsync(cancellationToken);
@@ -70,7 +91,6 @@ namespace EpicGames.Horde.Compute
 					switch (message.Type)
 					{
 						case AgentMessageType.None:
-							await Task.WhenAll(childTasks);
 							return;
 						case AgentMessageType.Ping:
 							await channel.PingAsync(cancellationToken);
@@ -78,7 +98,7 @@ namespace EpicGames.Horde.Compute
 						case AgentMessageType.Fork:
 							{
 								ForkMessage fork = message.ParseForkMessage();
-								childTasks.Add(Task.Run(() => RunAsync(socket, fork.ChannelId, fork.BufferSize, cancellationToken), cancellationToken));
+								childTasks.Add(Task.Run(() => RunAsync(socket, fork.ChannelId, fork.BufferSize, postException, cancellationToken), cancellationToken));
 							}
 							break;
 						case AgentMessageType.WriteFiles:
@@ -116,6 +136,20 @@ namespace EpicGames.Horde.Compute
 					}
 				}
 			}
+			catch (OperationCanceledException ex)
+			{
+				// Ignore cancellations; we will re-throw from the root RunAsync() method.
+				_logger.LogDebug(ex, "Compute Channel {ChannelId}: Cancelled.", channelId);
+			}
+			catch (Exception ex)
+			{
+				_logger.LogInformation(ex, "Compute Channel {ChannelId}: Exception:", channelId, ex.Message);
+				postException(ex);
+			}
+			finally
+			{
+				await Task.WhenAll(childTasks);
+			}
 		}
 
 		static async ValueTask RunXorAsync(AgentMessageChannel channel, ReadOnlyMemory<byte> source, byte value, CancellationToken cancellationToken)
