Horde: Enable new bundle serialization format by default.

diff --git a/Engine/Source/Programs/Shared/EpicGames.Horde/Storage/Bundles/BundleOptions.cs b/Engine/Source/Programs/Shared/EpicGames.Horde/Storage/Bundles/BundleOptions.cs
--- a/Engine/Source/Programs/Shared/EpicGames.Horde/Storage/Bundles/BundleOptions.cs
+++ b/Engine/Source/Programs/Shared/EpicGames.Horde/Storage/Bundles/BundleOptions.cs
@@ -15,7 +15,7 @@ namespace EpicGames.Horde.Storage.Bundles
 		/// <summary>
 		/// Maximum version number of bundles to write
 		/// </summary>
-		public BundleVersion MaxVersion { get; set; } = BundleVersion.LatestV1;
+		public BundleVersion MaxVersion { get; set; } = BundleVersion.LatestV2;
 
 		/// <summary>
 		/// Maximum payload size fo a blob
