Horde: Fix static analysis warnings.

diff --git a/Engine/Source/Programs/Shared/EpicGames.Horde/Compute/AgentMessageHandler.cs b/Engine/Source/Programs/Shared/EpicGames.Horde/Compute/AgentMessageHandler.cs
--- a/Engine/Source/Programs/Shared/EpicGames.Horde/Compute/AgentMessageHandler.cs
+++ b/Engine/Source/Programs/Shared/EpicGames.Horde/Compute/AgentMessageHandler.cs
@@ -72,7 +72,9 @@ namespace EpicGames.Horde.Compute
 			cancellationToken.ThrowIfCancellationRequested();
 
 			// Otherwise throw any exception posted by a child task
+#pragma warning disable CA1508 // Static analyzer doesn't understand how this can be non-null
 			exceptionInfo?.Throw();
+#pragma warning restore CA1508
 		}
 
 		async Task RunAsync(ComputeSocket socket, int channelId, int bufferSize, Action<Exception> postException, CancellationToken cancellationToken)
@@ -143,7 +145,7 @@ namespace EpicGames.Horde.Compute
 			}
 			catch (Exception ex)
 			{
-				_logger.LogInformation(ex, "Compute Channel {ChannelId}: Exception:", channelId, ex.Message);
+				_logger.LogInformation(ex, "Compute Channel {ChannelId}: Exception: {Message}", channelId, ex.Message);
 				postException(ex);
 			}
 			finally
