[Uba]
* Added c# path for RegisterDisallowedPath

diff --git a/Engine/Source/Programs/Shared/EpicGames.UBA/IStorageServer.cs b/Engine/Source/Programs/Shared/EpicGames.UBA/IStorageServer.cs
--- a/Engine/Source/Programs/Shared/EpicGames.UBA/IStorageServer.cs
+++ b/Engine/Source/Programs/Shared/EpicGames.UBA/IStorageServer.cs
@@ -55,6 +55,11 @@ namespace EpicGames.UBA
 		/// </summary>
 		public abstract void SaveCasTable();
 
+		/// <summary>
+		/// Register disallowed paths for clients to download
+		/// </summary>
+		public abstract void RegisterDisallowedPath(string path);
+
 		/// <summary>
 		/// Create a IStorageServer object
 		/// </summary>
diff --git a/Engine/Source/Programs/Shared/EpicGames.UBA/Impl/StorageServerImpl.cs b/Engine/Source/Programs/Shared/EpicGames.UBA/Impl/StorageServerImpl.cs
--- a/Engine/Source/Programs/Shared/EpicGames.UBA/Impl/StorageServerImpl.cs
+++ b/Engine/Source/Programs/Shared/EpicGames.UBA/Impl/StorageServerImpl.cs
@@ -20,6 +20,9 @@ namespace EpicGames.UBA
 
 		[DllImport("UbaHost", CharSet = CharSet.Auto)]
 		static extern void StorageServer_SaveCasTable(IntPtr server);
+
+		[DllImport("UbaHost", CharSet = CharSet.Auto)]
+		static extern void StorageServer_RegisterDisallowedPath(IntPtr server, string path);
 		#endregion
 
 		public StorageServerImpl(IServer server, ILogger logger, StorageServerCreateInfo info)
@@ -56,6 +59,8 @@ namespace EpicGames.UBA
 		public IntPtr GetHandle() => _handle;
 
 		public void SaveCasTable() => StorageServer_SaveCasTable(_handle);
+
+		public void RegisterDisallowedPath(string path) => StorageServer_RegisterDisallowedPath(_handle, path);
 		#endregion
 	}
 }
