Horde: Fix dispose order in DirectoryNodeZipStream.

diff --git a/Engine/Source/Programs/Shared/EpicGames.Horde/Storage/Nodes/DirectoryNode.cs b/Engine/Source/Programs/Shared/EpicGames.Horde/Storage/Nodes/DirectoryNode.cs
--- a/Engine/Source/Programs/Shared/EpicGames.Horde/Storage/Nodes/DirectoryNode.cs
+++ b/Engine/Source/Programs/Shared/EpicGames.Horde/Storage/Nodes/DirectoryNode.cs
@@ -1034,17 +1034,20 @@ namespace EpicGames.Horde.Storage.Nodes
 		/// <inheritdoc/>
 		protected override void Dispose(bool disposing)
 		{
-			base.Dispose(disposing);
+			if (disposing)
+			{
+				_backgroundTask.DisposeAsync().AsTask().Wait();
+			}
 
-			_backgroundTask.DisposeAsync().AsTask().Wait();
+			base.Dispose(disposing);
 		}
 
 		/// <inheritdoc/>
 		public override async ValueTask DisposeAsync()
 		{
-			await base.DisposeAsync();
-
 			await _backgroundTask.DisposeAsync();
+
+			await base.DisposeAsync();
 		}
 
 		/// <inheritdoc/>
