Change the verbosity of messages related to object references from UHT so that the "AllowAndLog" behavior goes from Trace level messages (which are filtered out by default) to Info level messages (which go to console and log file by default).

diff --git a/Engine/Source/Programs/Shared/EpicGames.UHT/Types/Properties/UhtObjectPropertyBase.cs b/Engine/Source/Programs/Shared/EpicGames.UHT/Types/Properties/UhtObjectPropertyBase.cs
--- a/Engine/Source/Programs/Shared/EpicGames.UHT/Types/Properties/UhtObjectPropertyBase.cs
+++ b/Engine/Source/Programs/Shared/EpicGames.UHT/Types/Properties/UhtObjectPropertyBase.cs
@@ -273,11 +273,11 @@ namespace EpicGames.UHT.Types
 				case UhtIssueBehavior.AllowAndLog:
 					if (!String.IsNullOrEmpty(alternativeTypeDesc))
 					{
-						this.LogTrace($"{pointerTypeDesc} usage in member declaration detected [[[{Outer.EngineName}]]].  Consider {alternativeTypeDesc} as an alternative.");
+						this.LogInfo($"{pointerTypeDesc} usage in member declaration detected [[[{Outer.EngineName}]]].  Consider {alternativeTypeDesc} as an alternative.");
 					}
 					else
 					{
-						this.LogTrace($"{pointerTypeDesc} usage in member declaration detected [[[{Outer.EngineName}]]].");
+						this.LogInfo($"{pointerTypeDesc} usage in member declaration detected [[[{Outer.EngineName}]]].");
 					}
 					break;
 
