Horde: Change asserts in BundleWriter to run in release builds.

diff --git a/Engine/Source/Programs/Shared/EpicGames.Horde/Storage/Bundles/V2/BundleWriter.cs b/Engine/Source/Programs/Shared/EpicGames.Horde/Storage/Bundles/V2/BundleWriter.cs
--- a/Engine/Source/Programs/Shared/EpicGames.Horde/Storage/Bundles/V2/BundleWriter.cs
+++ b/Engine/Source/Programs/Shared/EpicGames.Horde/Storage/Bundles/V2/BundleWriter.cs
@@ -3,8 +3,10 @@
 using System;
 using System.Collections.Generic;
 using System.Diagnostics;
+using System.Diagnostics.CodeAnalysis;
 using System.IO;
 using System.Linq;
+using System.Runtime.CompilerServices;
 using System.Threading;
 using System.Threading.Tasks;
 using EpicGames.Core;
@@ -155,20 +157,20 @@ namespace EpicGames.Horde.Storage.Bundles.V2
 
 			public Memory<byte> GetOutputBuffer(int usedSize, int desiredSize)
 			{
-				Debug.Assert(_packetWriter != null);
+				RuntimeAssert(_packetWriter != null);
 				return _packetWriter!.GetOutputBuffer(usedSize, desiredSize);
 			}
 
 			public BlobData GetPendingExport(int exportIdx)
 			{
-				Debug.Assert(_packetWriter != null);
+				RuntimeAssert(_packetWriter != null);
 				return _packetWriter.GetExport(exportIdx);
 			}
 
 			public PendingExportHandle CompleteExport(BlobType type, int size, IReadOnlyList<IBlobHandle> references, IReadOnlyList<AliasInfo> aliases)
 			{
-				Debug.Assert(_packetWriter != null);
-				Debug.Assert(_packetHandle != null);
+				RuntimeAssert(_packetWriter != null);
+				RuntimeAssert(_packetHandle != null);
 
 				int exportIdx = _packetWriter.CompleteExport(size, type, references);
 				PendingExportHandle exportHandle = new PendingExportHandle(_packetHandle, exportIdx);
@@ -190,8 +192,8 @@ namespace EpicGames.Horde.Storage.Bundles.V2
 
 			void StartPacket()
 			{
-				Debug.Assert(_packetHandle == null);
-				Debug.Assert(_packetWriter == null);
+				RuntimeAssert(_packetHandle == null);
+				RuntimeAssert(_packetWriter == null);
 
 				_packetHandle = new PendingPacketHandle(this);
 				_packetWriter = new PacketWriter(this, _packetHandle, _cache.Allocator, _lockObject);
@@ -199,10 +201,10 @@ namespace EpicGames.Horde.Storage.Bundles.V2
 
 			void FinishPacket()
 			{
-				Debug.Assert(_packetHandle != null);
-				Debug.Assert(_packetWriter != null);
-				Debug.Assert(_bundleReferences != null);
-				Debug.Assert(_encodedPacketWriter != null);
+				RuntimeAssert(_packetHandle != null);
+				RuntimeAssert(_packetWriter != null);
+				RuntimeAssert(_bundleReferences != null);
+				RuntimeAssert(_encodedPacketWriter != null);
 
 				if (_packetWriter.GetExportCount() > 0)
 				{
@@ -378,5 +380,16 @@ namespace EpicGames.Horde.Storage.Bundles.V2
 			}
 			return exportHandle;
 		}
+
+		/// <summary>
+		/// Helper method to check a precondition is valid at runtime, regardless of build configuration.
+		/// </summary>
+		static void RuntimeAssert([DoesNotReturnIf(false)] bool condition, [CallerArgumentExpression("condition")] string? message = null)
+		{
+			if (!condition)
+			{
+				throw new InvalidOperationException($"Condition failed: {message}");
+			}
+		}
 	}
 }
