Horde: Fix exception flushing a bundle writer that has already been fully flushed.

diff --git a/Engine/Source/Programs/Shared/EpicGames.Horde/Storage/Bundles/V2/BundleWriter.cs b/Engine/Source/Programs/Shared/EpicGames.Horde/Storage/Bundles/V2/BundleWriter.cs
--- a/Engine/Source/Programs/Shared/EpicGames.Horde/Storage/Bundles/V2/BundleWriter.cs
+++ b/Engine/Source/Programs/Shared/EpicGames.Horde/Storage/Bundles/V2/BundleWriter.cs
@@ -246,7 +246,10 @@ namespace EpicGames.Horde.Storage.Bundles.V2
 					return;
 				}
 
-				FinishPacket();
+				if (_packetWriter != null)
+				{
+					FinishPacket();
+				}
 
 				if (_encodedPacketWriter.Length == 0)
 				{
