Horde: Add methods for tools endpoint to HordeHttpClient.

diff --git a/Engine/Source/Programs/Shared/EpicGames.Horde/HordeHttpClient.cs b/Engine/Source/Programs/Shared/EpicGames.Horde/HordeHttpClient.cs
--- a/Engine/Source/Programs/Shared/EpicGames.Horde/HordeHttpClient.cs
+++ b/Engine/Source/Programs/Shared/EpicGames.Horde/HordeHttpClient.cs
@@ -2,6 +2,7 @@
 
 using System;
 using System.Collections.Generic;
+using System.IO;
 using System.Net.Http;
 using System.Net.Http.Headers;
 using System.Net.Http.Json;
@@ -17,6 +18,7 @@ using EpicGames.Horde.Projects;
 using EpicGames.Horde.Secrets;
 using EpicGames.Horde.Server;
 using EpicGames.Horde.Storage.Clients;
+using EpicGames.Horde.Tools;
 using Microsoft.Extensions.DependencyInjection;
 using Microsoft.Extensions.Logging;
 using Polly;
@@ -24,6 +26,8 @@ using Polly.Extensions.Http;
 using Polly.Retry;
 using Polly.Timeout;
 
+#pragma warning disable CA2234
+
 namespace EpicGames.Horde
 {
 	/// <summary>
@@ -51,6 +55,11 @@ namespace EpicGames.Horde
 		static readonly JsonSerializerOptions s_jsonSerializerOptions = CreateJsonSerializerOptions();
 		internal static JsonSerializerOptions JsonSerializerOptions => s_jsonSerializerOptions;
 
+		/// <summary>
+		/// Base address for the Horde server
+		/// </summary>
+		public Uri BaseUrl => _httpClient.BaseAddress ?? throw new InvalidOperationException("Expected Horde server base address to be configured");
+
 		/// <summary>
 		/// Constructor
 		/// </summary>
@@ -256,6 +265,42 @@ namespace EpicGames.Horde
 
 		#endregion
 
+		#region Tools
+
+		/// <summary>
+		/// Enumerates all the available tools.
+		/// </summary>
+		public Task<GetToolsSummaryResponse> GetToolsAsync(CancellationToken cancellationToken = default)
+		{
+			return GetAsync<GetToolsSummaryResponse>(_httpClient, "api/v1/tools", cancellationToken);
+		}
+
+		/// <summary>
+		/// Gets information about a particular tool
+		/// </summary>
+		public Task<GetToolResponse> GetToolAsync(ToolId id, CancellationToken cancellationToken = default)
+		{
+			return GetAsync<GetToolResponse>(_httpClient, $"api/v1/tools/{id}", cancellationToken);
+		}
+
+		/// <summary>
+		/// Gets information about a particular deployment
+		/// </summary>
+		public Task<GetToolDeploymentResponse> GetToolDeploymentAsync(ToolId id, ToolDeploymentId deploymentId, CancellationToken cancellationToken = default)
+		{
+			return GetAsync<GetToolDeploymentResponse>(_httpClient, $"api/v1/tools/{id}/deployments/{deploymentId}", cancellationToken);
+		}
+
+		/// <summary>
+		/// Gets a zip stream for a particular deployment
+		/// </summary>
+		public async Task<Stream> GetToolDeploymentZipAsync(ToolId id, ToolDeploymentId deploymentId, CancellationToken cancellationToken = default)
+		{
+			return await _httpClient.GetStreamAsync($"api/v1/tools/{id}/deployments/{deploymentId}?action=download", cancellationToken);
+		}
+
+		#endregion
+
 		#region Utility Methods
 
 		/// <summary>
