EpicGames.UBA: Support nuget pack

#rnx

diff --git a/Engine/Source/Programs/Shared/EpicGames.UBA/EpicGames.UBA.csproj b/Engine/Source/Programs/Shared/EpicGames.UBA/EpicGames.UBA.csproj
--- a/Engine/Source/Programs/Shared/EpicGames.UBA/EpicGames.UBA.csproj
+++ b/Engine/Source/Programs/Shared/EpicGames.UBA/EpicGames.UBA.csproj
@@ -17,6 +17,7 @@
     <EnforceCodeStyleInBuild>False</EnforceCodeStyleInBuild>
     <ProduceReferenceAssemblyInOutDir>true</ProduceReferenceAssemblyInOutDir>
     <RunAnalyzersDuringBuild>false</RunAnalyzersDuringBuild>
+    <PackageReadmeFile>README.md</PackageReadmeFile>
   </PropertyGroup>
   <PropertyGroup Condition=" '$(Configuration)|$(Platform)' == 'Debug|AnyCPU' ">
     <DocumentationFile></DocumentationFile>
@@ -34,6 +35,7 @@
     <Compile Include="..\MetaData.cs">
       <Link>Properties\MetaData.cs</Link>
     </Compile>
+    <None Include="docs\README.md" Pack="true" PackagePath="\" />
   </ItemGroup>
 
   <Import Project="Library.props" />
diff --git a/Engine/Source/Programs/Shared/EpicGames.UBA/docs/README.md b/Engine/Source/Programs/Shared/EpicGames.UBA/docs/README.md
new file mode 100644
--- /dev/null
+++ b/Engine/Source/Programs/Shared/EpicGames.UBA/docs/README.md
@@ -0,0 +1 @@
+EpicGames.UBA
\ No newline at end of file
