Horde: Include the current deployment identifier in tool responses.

diff --git a/Engine/Source/Programs/Horde/Horde.Server/Tools/ToolsController.cs b/Engine/Source/Programs/Horde/Horde.Server/Tools/ToolsController.cs
--- a/Engine/Source/Programs/Horde/Horde.Server/Tools/ToolsController.cs
+++ b/Engine/Source/Programs/Horde/Horde.Server/Tools/ToolsController.cs
@@ -213,8 +213,8 @@ namespace Horde.Server.Tools
 
 		static GetToolSummaryResponse CreateGetToolSummaryResponse(ITool tool)
 		{
-			string? latestVersion = (tool.Deployments.Count > 0) ? tool.Deployments[^1].Version : null;
-			return new GetToolSummaryResponse(tool.Id, tool.Config.Name, tool.Config.Description, latestVersion);
+			IToolDeployment? deployment = (tool.Deployments.Count == 0) ? null : tool.Deployments[^1];
+			return new GetToolSummaryResponse(tool.Id, tool.Config.Name, tool.Config.Description, deployment?.Version, deployment?.Id);
 		}
 
 		/// <summary>
diff --git a/Engine/Source/Programs/Shared/EpicGames.Horde/Tools/ToolMessages.cs b/Engine/Source/Programs/Shared/EpicGames.Horde/Tools/ToolMessages.cs
--- a/Engine/Source/Programs/Shared/EpicGames.Horde/Tools/ToolMessages.cs
+++ b/Engine/Source/Programs/Shared/EpicGames.Horde/Tools/ToolMessages.cs
@@ -23,8 +23,9 @@ namespace EpicGames.Horde.Tools
 	/// <param name="Id">Unique identifier for the tool</param>
 	/// <param name="Name">Name of the tool</param>
 	/// <param name="Description">Description for the tool</param>
-	/// <param name="Version">Version number of the latest deployment of this tool</param>
-	public record class GetToolSummaryResponse(ToolId Id, string Name, string Description, string? Version);
+	/// <param name="Version">Version number of the current deployment of this tool</param>
+	/// <param name="DeploymentId">Identifier for the current deployment</param>
+	public record class GetToolSummaryResponse(ToolId Id, string Name, string Description, string? Version, ToolDeploymentId? DeploymentId);
 
 	/// <summary>
 	/// Response when querying all tools
