Horde: Add a flag for whether to show tools inside UGS.

diff --git a/Engine/Source/Programs/Horde/Horde.Server/Tools/ToolConfig.cs b/Engine/Source/Programs/Horde/Horde.Server/Tools/ToolConfig.cs
--- a/Engine/Source/Programs/Horde/Horde.Server/Tools/ToolConfig.cs
+++ b/Engine/Source/Programs/Horde/Horde.Server/Tools/ToolConfig.cs
@@ -47,6 +47,11 @@ namespace Horde.Server.Tools
 		/// </summary>
 		public bool Public { get; set; }
 
+		/// <summary>
+		/// Whether to show this tool for download in the UGS tools menu
+		/// </summary>
+		public bool ShowInUgs { get; set; }
+
 		/// <summary>
 		/// Default namespace for new deployments of this tool
 		/// </summary>
diff --git a/Engine/Source/Programs/Horde/Horde.Server/Tools/ToolsController.cs b/Engine/Source/Programs/Horde/Horde.Server/Tools/ToolsController.cs
--- a/Engine/Source/Programs/Horde/Horde.Server/Tools/ToolsController.cs
+++ b/Engine/Source/Programs/Horde/Horde.Server/Tools/ToolsController.cs
@@ -214,7 +214,7 @@ namespace Horde.Server.Tools
 		static GetToolSummaryResponse CreateGetToolSummaryResponse(ITool tool)
 		{
 			IToolDeployment? deployment = (tool.Deployments.Count == 0) ? null : tool.Deployments[^1];
-			return new GetToolSummaryResponse(tool.Id, tool.Config.Name, tool.Config.Description, deployment?.Version, deployment?.Id);
+			return new GetToolSummaryResponse(tool.Id, tool.Config.Name, tool.Config.Description, deployment?.Version, deployment?.Id, tool.Config.ShowInUgs);
 		}
 
 		/// <summary>
@@ -258,7 +258,7 @@ namespace Horde.Server.Tools
 
 		static GetToolResponse CreateGetToolResponse(ITool tool, List<GetToolDeploymentResponse> deployments)
 		{
-			return new GetToolResponse(tool.Id, tool.Config.Name, tool.Config.Description, deployments, tool.Config.Public);
+			return new GetToolResponse(tool.Id, tool.Config.Name, tool.Config.Description, deployments, tool.Config.Public, tool.Config.ShowInUgs);
 		}
 
 		/// <summary>
diff --git a/Engine/Source/Programs/Shared/EpicGames.Horde/Tools/ToolMessages.cs b/Engine/Source/Programs/Shared/EpicGames.Horde/Tools/ToolMessages.cs
--- a/Engine/Source/Programs/Shared/EpicGames.Horde/Tools/ToolMessages.cs
+++ b/Engine/Source/Programs/Shared/EpicGames.Horde/Tools/ToolMessages.cs
@@ -15,7 +15,8 @@ namespace EpicGames.Horde.Tools
 	/// <param name="Description">Description for the tool</param>
 	/// <param name="Deployments">Current deployments of this tool, sorted by time.</param>
 	/// <param name="Public">Whether this tool should be exposed for download on a public endpoint without authentication</param>
-	public record class GetToolResponse(ToolId Id, string Name, string Description, List<GetToolDeploymentResponse> Deployments, bool Public);
+	/// <param name="ShowInUgs">Whether to show this tool for download inside UGS</param>
+	public record class GetToolResponse(ToolId Id, string Name, string Description, List<GetToolDeploymentResponse> Deployments, bool Public, bool ShowInUgs);
 
 	/// <summary>
 	/// Summary for a particular tool.
@@ -25,7 +26,8 @@ namespace EpicGames.Horde.Tools
 	/// <param name="Description">Description for the tool</param>
 	/// <param name="Version">Version number of the current deployment of this tool</param>
 	/// <param name="DeploymentId">Identifier for the current deployment</param>
-	public record class GetToolSummaryResponse(ToolId Id, string Name, string Description, string? Version, ToolDeploymentId? DeploymentId);
+	/// <param name="ShowInUgs">Whether to show this tool for download inside UGS</param>
+	public record class GetToolSummaryResponse(ToolId Id, string Name, string Description, string? Version, ToolDeploymentId? DeploymentId, bool ShowInUgs);
 
 	/// <summary>
 	/// Response when querying all tools
