Horde: Add commands for downloading/uploading tools from command line tool.

diff --git a/Engine/Source/Programs/Horde/Horde.Server/Tools/ToolsController.cs b/Engine/Source/Programs/Horde/Horde.Server/Tools/ToolsController.cs
--- a/Engine/Source/Programs/Horde/Horde.Server/Tools/ToolsController.cs
+++ b/Engine/Source/Programs/Horde/Horde.Server/Tools/ToolsController.cs
@@ -73,7 +73,7 @@ namespace Horde.Server.Tools
 		/// <returns>Information about the registered agent</returns>
 		[HttpPost]
 		[Route("/api/v1/tools/{id}/deployments")]
-		public async Task<ActionResult<CreateDeploymentResponse>> CreateDeploymentAsync(ToolId id, [FromForm] ToolDeploymentConfig options, [FromForm] IFormFile file, CancellationToken cancellationToken)
+		public async Task<ActionResult<CreateToolDeploymentResponse>> CreateDeploymentAsync(ToolId id, [FromForm] ToolDeploymentConfig options, [FromForm] IFormFile file, CancellationToken cancellationToken)
 		{
 			ITool? tool = await _toolCollection.GetAsync(id, _globalConfig.Value);
 
@@ -95,7 +95,7 @@ namespace Horde.Server.Tools
 				}
 			}
 
-			return new CreateDeploymentResponse(tool.Deployments[^1].Id);
+			return new CreateToolDeploymentResponse(tool.Deployments[^1].Id);
 		}
 
 		/// <summary>
@@ -103,7 +103,7 @@ namespace Horde.Server.Tools
 		/// </summary>
 		[HttpPost]
 		[Route("/api/v2/tools/{id}/deployments")]
-		public async Task<ActionResult<CreateDeploymentResponse>> CreateDeploymentAsync(ToolId id, CreateDeploymentRequest request, CancellationToken cancellationToken)
+		public async Task<ActionResult<CreateToolDeploymentResponse>> CreateDeploymentAsync(ToolId id, CreateToolDeploymentRequest request, CancellationToken cancellationToken)
 		{
 			ITool? tool = await _toolCollection.GetAsync(id, _globalConfig.Value);
 
@@ -124,7 +124,7 @@ namespace Horde.Server.Tools
 				return NotFound(id);
 			}
 
-			return new CreateDeploymentResponse(tool.Deployments[^1].Id);
+			return new CreateToolDeploymentResponse(tool.Deployments[^1].Id);
 		}
 
 		/// <summary>
diff --git a/Engine/Source/Programs/Horde/Horde/Commands/Tools/ToolDownload.cs b/Engine/Source/Programs/Horde/Horde/Commands/Tools/ToolDownload.cs
new file mode 100644
--- /dev/null
+++ b/Engine/Source/Programs/Horde/Horde/Commands/Tools/ToolDownload.cs
@@ -0,0 +1,51 @@
+// Copyright Epic Games, Inc. All Rights Reserved.
+
+using System.ComponentModel;
+using System.IO.Compression;
+using EpicGames.Core;
+using EpicGames.Horde;
+using EpicGames.Horde.Tools;
+using Microsoft.Extensions.Logging;
+using Microsoft.Extensions.Options;
+
+namespace Horde.Commands
+{
+	[Command("tool", "download", "Downloads a tool from the server to a local directory")]
+	class ToolDownload : Command
+	{
+		[CommandLine("-Id=", Required = true)]
+		[Description("The tool to download")]
+		public ToolId Id { get; set; }
+
+		[CommandLine("-Deployment=")]
+		[Description("Deployment to fetch")]
+		public ToolDeploymentId? DeploymentId { get; set; }
+
+		[CommandLine("-OutputDir=", Required = true)]
+		[Description("Output directory for the downloaded tool data")]
+		public DirectoryReference OutputDir { get; set; } = null!;
+
+		readonly HordeHttpClient _hordeHttpClient;
+
+		public ToolDownload(HordeHttpClient httpClient)
+		{
+			_hordeHttpClient = httpClient;
+		}
+
+		/// <inheritdoc/>
+		public override async Task<int> ExecuteAsync(ILogger logger)
+		{
+			DirectoryReference.CreateDirectory(OutputDir);
+
+			FileReference zipFile = FileReference.Combine(OutputDir, "temp.zip");
+			using (FileStream zipOutputStream = FileReference.Open(zipFile, FileMode.Create, FileAccess.Write))
+			{
+				await using Stream stream = await _hordeHttpClient.GetToolDeploymentZipAsync(Id, DeploymentId, CancellationToken.None);
+				await stream.CopyToAsync(zipOutputStream);
+			}
+
+			ZipFile.ExtractToDirectory(zipFile.FullName, OutputDir.FullName, true);
+			return 0;
+		}
+	}
+}
diff --git a/Engine/Source/Programs/Horde/Horde/Commands/Tools/ToolList.cs b/Engine/Source/Programs/Horde/Horde/Commands/Tools/ToolList.cs
new file mode 100644
--- /dev/null
+++ b/Engine/Source/Programs/Horde/Horde/Commands/Tools/ToolList.cs
@@ -0,0 +1,33 @@
+// Copyright Epic Games, Inc. All Rights Reserved.
+
+using System.ComponentModel;
+using EpicGames.Core;
+using EpicGames.Horde;
+using EpicGames.Horde.Tools;
+using Microsoft.Extensions.Logging;
+using Microsoft.Extensions.Options;
+
+namespace Horde.Commands
+{
+	[Command("tool", "list", "List all tools available for download")]
+	class ToolList : Command
+	{
+		readonly HordeHttpClient _hordeHttpClient;
+
+		public ToolList(HordeHttpClient httpClient)
+		{
+			_hordeHttpClient = httpClient;
+		}
+
+		/// <inheritdoc/>
+		public override async Task<int> ExecuteAsync(ILogger logger)
+		{
+			GetToolsSummaryResponse response = await _hordeHttpClient.GetToolsAsync();
+			foreach (GetToolSummaryResponse tool in response.Tools)
+			{
+				logger.LogInformation("  {ToolId,-20} {Deployment,-30} {Version,-15}", tool.Id, tool.DeploymentId, tool.Version);
+			}
+			return 0;
+		}
+	}
+}
diff --git a/Engine/Source/Programs/Horde/Horde/Commands/Tools/ToolUpload.cs b/Engine/Source/Programs/Horde/Horde/Commands/Tools/ToolUpload.cs
new file mode 100644
--- /dev/null
+++ b/Engine/Source/Programs/Horde/Horde/Commands/Tools/ToolUpload.cs
@@ -0,0 +1,59 @@
+// Copyright Epic Games, Inc. All Rights Reserved.
+
+using System.ComponentModel;
+using EpicGames.Core;
+using EpicGames.Horde;
+using EpicGames.Horde.Storage;
+using EpicGames.Horde.Storage.Backends;
+using EpicGames.Horde.Storage.Clients;
+using EpicGames.Horde.Storage.Nodes;
+using EpicGames.Horde.Tools;
+using Microsoft.Extensions.Logging;
+using Microsoft.Extensions.Options;
+
+namespace Horde.Commands
+{
+	[Command("tool", "upload", "Uploads a tool from a local directory")]
+	class ToolUpload : Command
+	{
+		[CommandLine("-Id=")]
+		[Description("Identifier for the tool to upload")]
+		public ToolId ToolId { get; set; }
+
+		[CommandLine("-Version=")]
+		[Description("Optional version number for the new upload")]
+		public string? Version { get; set; }
+
+		[CommandLine("-InputDir=")]
+		[Description("Directory containing files to upload for the tool")]
+		public DirectoryReference InputDir { get; set; } = null!;
+
+		readonly HordeHttpClient _hordeHttpClient;
+		readonly HttpStorageClientFactory _httpStorageClientFactory;
+
+		public ToolUpload(HordeHttpClient httpClient, HttpStorageClientFactory httpStorageClientFactory)
+		{
+			_hordeHttpClient = httpClient;
+			_httpStorageClientFactory = httpStorageClientFactory;
+		}
+
+		/// <inheritdoc/>
+		public override async Task<int> ExecuteAsync(ILogger logger)
+		{
+			using IStorageClient storageClient = _httpStorageClientFactory.CreateClientWithPath($"api/v1/tools/{ToolId}");
+
+			IBlobHandle handle;
+			await using (IStorageWriter writer = storageClient.CreateWriter())
+			{
+				DirectoryNode sandbox = new DirectoryNode();
+				await sandbox.CopyFromDirectoryAsync(InputDir.ToDirectoryInfo(), new ChunkingOptions(), writer, null);
+				handle = await writer.FlushAsync(sandbox);
+			}
+
+			ToolDeploymentId deploymentId = await _hordeHttpClient.CreateToolDeploymentAsync(ToolId, Version, null, null, handle.GetLocator());
+			logger.LogInformation("Created deployment {DeploymentId}", deploymentId);
+
+			return 0;
+		}
+	}
+}
diff --git a/Engine/Source/Programs/Shared/EpicGames.Horde/HordeHttpClient.cs b/Engine/Source/Programs/Shared/EpicGames.Horde/HordeHttpClient.cs
--- a/Engine/Source/Programs/Shared/EpicGames.Horde/HordeHttpClient.cs
+++ b/Engine/Source/Programs/Shared/EpicGames.Horde/HordeHttpClient.cs
@@ -17,6 +17,7 @@ using EpicGames.Horde.Dashboard;
 using EpicGames.Horde.Projects;
 using EpicGames.Horde.Secrets;
 using EpicGames.Horde.Server;
+using EpicGames.Horde.Storage;
 using EpicGames.Horde.Storage.Clients;
 using EpicGames.Horde.Tools;
 using Microsoft.Extensions.DependencyInjection;
@@ -294,9 +295,32 @@ namespace EpicGames.Horde
 		/// <summary>
 		/// Gets a zip stream for a particular deployment
 		/// </summary>
-		public async Task<Stream> GetToolDeploymentZipAsync(ToolId id, ToolDeploymentId deploymentId, CancellationToken cancellationToken = default)
+		public async Task<Stream> GetToolDeploymentZipAsync(ToolId id, ToolDeploymentId? deploymentId, CancellationToken cancellationToken = default)
 		{
-			return await _httpClient.GetStreamAsync($"api/v1/tools/{id}/deployments/{deploymentId}?action=download", cancellationToken);
+			if (deploymentId == null)
+			{
+				return await _httpClient.GetStreamAsync($"api/v1/tools/{id}?action=zip", cancellationToken);
+			}
+			else
+			{
+				return await _httpClient.GetStreamAsync($"api/v1/tools/{id}/deployments/{deploymentId}?action=zip", cancellationToken);
+			}
+		}
+
+		/// <summary>
+		/// Creates a new tool deployment
+		/// </summary>
+		/// <param name="id">Id for the tool</param>
+		/// <param name="version">Version string for the new deployment</param>
+		/// <param name="duration">Duration over which to deploy the tool</param>
+		/// <param name="createPaused">Whether to create the deployment, but do not start rolling it out yet</param>
+		/// <param name="node">Location of a directory node describing the deployment</param>
+		/// <param name="cancellationToken">Cancellation token for the operation</param>
+		public async Task<ToolDeploymentId> CreateToolDeploymentAsync(ToolId id, string? version, double? duration, bool? createPaused, BlobLocator node, CancellationToken cancellationToken = default)
+		{
+			CreateToolDeploymentRequest request = new CreateToolDeploymentRequest(version ?? String.Empty, duration, createPaused, node);
+			CreateToolDeploymentResponse response = await PostAsync<CreateToolDeploymentResponse, CreateToolDeploymentRequest>(_httpClient, $"api/v2/tools/{id}/deployments", request, cancellationToken);
+			return response.Id;
 		}
 
 		#endregion
diff --git a/Engine/Source/Programs/Shared/EpicGames.Horde/Tools/ToolMessages.cs b/Engine/Source/Programs/Shared/EpicGames.Horde/Tools/ToolMessages.cs
--- a/Engine/Source/Programs/Shared/EpicGames.Horde/Tools/ToolMessages.cs
+++ b/Engine/Source/Programs/Shared/EpicGames.Horde/Tools/ToolMessages.cs
@@ -55,13 +55,13 @@ namespace EpicGames.Horde.Tools
 	/// <param name="Duration">Number of minutes over which to do the deployment</param>
 	/// <param name="CreatePaused">Whether to create the deployment in a paused state</param>
 	/// <param name="Node">Handle to the root node</param>
-	public record CreateDeploymentRequest(string Version, double? Duration, bool? CreatePaused, BlobLocator Node);
+	public record CreateToolDeploymentRequest(string Version, double? Duration, bool? CreatePaused, BlobLocator Node);
 
 	/// <summary>
 	/// Response from creating a deployment
 	/// </summary>
 	/// <param name="Id">Identifier for the created deployment</param>
-	public record CreateDeploymentResponse(ToolDeploymentId Id);
+	public record CreateToolDeploymentResponse(ToolDeploymentId Id);
 
 	/// <summary>
 	/// Current state of a tool's deployment
