UnrealBuildTool: UBA disallowed paths

#rnx

diff --git a/Engine/Source/Programs/Shared/EpicGames.UBA/Utils.cs b/Engine/Source/Programs/Shared/EpicGames.UBA/Utils.cs
--- a/Engine/Source/Programs/Shared/EpicGames.UBA/Utils.cs
+++ b/Engine/Source/Programs/Shared/EpicGames.UBA/Utils.cs
@@ -1,6 +1,7 @@
 ﻿// Copyright Epic Games, Inc. All Rights Reserved.
 
 using System;
+using System.Collections.Generic;
 using System.IO;
 using System.Reflection;
 using System.Runtime.InteropServices;
@@ -12,12 +13,31 @@ namespace EpicGames.UBA
 	/// </summary>
 	public static class Utils
 	{
+
 		/// <summary>
 		/// Is UBA available?
 		/// </summary>
 		public static bool IsAvailable() => s_available.Value;
 		static readonly Lazy<bool> s_available = new Lazy<bool>(() => File.Exists(GetLibraryPath()));
 
+		/// <summary>
+		/// Paths that are not allowed to be transferred over the network for UBA remote agents.
+		/// </summary>
+		/// <returns>enumerable of disallowed paths</returns>
+		public static IEnumerable<string> DisallowedPaths() => s_disallowedPaths;
+		static readonly SortedSet<string> s_disallowedPaths = new();
+
+		/// <summary>
+		/// Registers a path that is not allowed to be transferred over the network for UBA remote agents.
+		/// </summary>
+		/// <param name="paths">The paths to add to thie disallowed list</param>
+		public static void RegisterDisallowedPaths(params string[] paths) => s_disallowedPaths.UnionWith(paths);
+
+		/// <summary>
+		/// Get the path to the p/invoke library that would be loaded 
+		/// </summary>
+		/// <returns>The path to the library</returns>
+		/// <exception cref="PlatformNotSupportedException">If the operating system is not supported</exception>
 		[System.Diagnostics.CodeAnalysis.SuppressMessage("Globalization", "CA1308:Normalize strings to uppercase", Justification = "folder path is lowercase")]
 		static string GetLibraryPath()
 		{
diff --git a/Engine/Source/Programs/UnrealBuildTool/ToolChain/AppleToolChain.cs b/Engine/Source/Programs/UnrealBuildTool/ToolChain/AppleToolChain.cs
--- a/Engine/Source/Programs/UnrealBuildTool/ToolChain/AppleToolChain.cs
+++ b/Engine/Source/Programs/UnrealBuildTool/ToolChain/AppleToolChain.cs
@@ -85,6 +85,9 @@ namespace UnrealBuildTool
 
 			// convert to float for easy comparison
 			SDKVersionFloat = Single.Parse(SDKVersion, System.Globalization.CultureInfo.InvariantCulture);
+
+			// Xcode sdk files are not allowed to be transferred over the network, ensure UBA will not do so
+			EpicGames.UBA.Utils.RegisterDisallowedPaths(XcodeDeveloperDir.FullName);
 		}
 
 		/// <summary>
