UBA: Set disallowed paths for the session

#rnx

diff --git a/Engine/Source/Programs/Shared/EpicGames.UBA/Impl/StorageServerImpl.cs b/Engine/Source/Programs/Shared/EpicGames.UBA/Impl/StorageServerImpl.cs
--- a/Engine/Source/Programs/Shared/EpicGames.UBA/Impl/StorageServerImpl.cs
+++ b/Engine/Source/Programs/Shared/EpicGames.UBA/Impl/StorageServerImpl.cs
@@ -1,6 +1,7 @@
 // Copyright Epic Games, Inc. All Rights Reserved.
 
 using System;
+using System.Linq;
 using System.Runtime.InteropServices;
 
 namespace EpicGames.UBA
@@ -30,6 +31,7 @@ namespace EpicGames.UBA
 			_server = server;
 			_logger = logger;
 			_handle = CreateStorageServer(_server.GetHandle(), info.RootDirectory, info.CapacityBytes, info.StoreCompressed, _logger.GetHandle(), info.Zone);
+			Utils.DisallowedPaths().ToList().ForEach(x => RegisterDisallowedPath(x));
 		}
 
 		#region IDisposable
