Horde: Move JobId type into EpicGames.Horde.

diff --git a/Engine/Source/Programs/Horde/Horde.Server.Tests/DeviceServiceTests.cs b/Engine/Source/Programs/Horde/Horde.Server.Tests/DeviceServiceTests.cs
--- a/Engine/Source/Programs/Horde/Horde.Server.Tests/DeviceServiceTests.cs
+++ b/Engine/Source/Programs/Horde/Horde.Server.Tests/DeviceServiceTests.cs
@@ -22,6 +22,7 @@ using EpicGames.Horde.Jobs.Templates;
 using EpicGames.Horde.Streams;
 using Horde.Server.Jobs.Templates;
 using Moq;
+using EpicGames.Horde.Jobs;
 
 namespace Horde.Server.Tests
 {
diff --git a/Engine/Source/Programs/Horde/Horde.Server.Tests/Fixture.cs b/Engine/Source/Programs/Horde/Horde.Server.Tests/Fixture.cs
--- a/Engine/Source/Programs/Horde/Horde.Server.Tests/Fixture.cs
+++ b/Engine/Source/Programs/Horde/Horde.Server.Tests/Fixture.cs
@@ -5,6 +5,7 @@ using System.IO;
 using System.Text;
 using System.Threading.Tasks;
 using EpicGames.Core;
+using EpicGames.Horde.Jobs;
 using EpicGames.Horde.Jobs.Templates;
 using EpicGames.Horde.Projects;
 using EpicGames.Horde.Streams;
diff --git a/Engine/Source/Programs/Horde/Horde.Server.Tests/IssueServiceTests.cs b/Engine/Source/Programs/Horde/Horde.Server.Tests/IssueServiceTests.cs
--- a/Engine/Source/Programs/Horde/Horde.Server.Tests/IssueServiceTests.cs
+++ b/Engine/Source/Programs/Horde/Horde.Server.Tests/IssueServiceTests.cs
@@ -33,6 +33,7 @@ using EpicGames.Horde.Projects;
 using EpicGames.Horde.Users;
 using System.Buffers;
 using EpicGames.Horde.Streams;
+using EpicGames.Horde.Jobs;
 
 namespace Horde.Server.Tests
 {
diff --git a/Engine/Source/Programs/Horde/Horde.Server.Tests/JobsControllerDbTest.cs b/Engine/Source/Programs/Horde/Horde.Server.Tests/JobsControllerDbTest.cs
--- a/Engine/Source/Programs/Horde/Horde.Server.Tests/JobsControllerDbTest.cs
+++ b/Engine/Source/Programs/Horde/Horde.Server.Tests/JobsControllerDbTest.cs
@@ -3,6 +3,7 @@
 using System.Collections.Generic;
 using System.Threading.Tasks;
 using EpicGames.Core;
+using EpicGames.Horde.Jobs;
 using Horde.Server.Jobs;
 using Microsoft.AspNetCore.Mvc;
 using Microsoft.VisualStudio.TestTools.UnitTesting;
diff --git a/Engine/Source/Programs/Horde/Horde.Server.Tests/LogFileServiceTests.cs b/Engine/Source/Programs/Horde/Horde.Server.Tests/LogFileServiceTests.cs
--- a/Engine/Source/Programs/Horde/Horde.Server.Tests/LogFileServiceTests.cs
+++ b/Engine/Source/Programs/Horde/Horde.Server.Tests/LogFileServiceTests.cs
@@ -9,6 +9,7 @@ using System.Text;
 using System.Threading;
 using System.Threading.Tasks;
 using EpicGames.Core;
+using EpicGames.Horde.Jobs;
 using EpicGames.Horde.Storage;
 using EpicGames.Horde.Storage.Backends;
 using Horde.Server.Agents.Sessions;
diff --git a/Engine/Source/Programs/Horde/Horde.Server.Tests/LogIndexingTests.cs b/Engine/Source/Programs/Horde/Horde.Server.Tests/LogIndexingTests.cs
--- a/Engine/Source/Programs/Horde/Horde.Server.Tests/LogIndexingTests.cs
+++ b/Engine/Source/Programs/Horde/Horde.Server.Tests/LogIndexingTests.cs
@@ -19,6 +19,7 @@ using Microsoft.Extensions.Logging;
 using Microsoft.VisualStudio.TestTools.UnitTesting;
 using EpicGames.Horde.Logs;
 using OpenTelemetry.Trace;
+using EpicGames.Horde.Jobs;
 
 namespace Horde.Server.Tests
 {
diff --git a/Engine/Source/Programs/Horde/Horde.Server.Tests/Models/Job/GetLabelStatesTests.cs b/Engine/Source/Programs/Horde/Horde.Server.Tests/Models/Job/GetLabelStatesTests.cs
--- a/Engine/Source/Programs/Horde/Horde.Server.Tests/Models/Job/GetLabelStatesTests.cs
+++ b/Engine/Source/Programs/Horde/Horde.Server.Tests/Models/Job/GetLabelStatesTests.cs
@@ -1,6 +1,7 @@
 // Copyright Epic Games, Inc. All Rights Reserved.
 
 using System.Collections.Generic;
+using EpicGames.Horde.Jobs;
 using Horde.Server.Jobs;
 using Horde.Server.Jobs.Graphs;
 using HordeCommon;
diff --git a/Engine/Source/Programs/Horde/Horde.Server.Tests/NotificationServiceTests.cs b/Engine/Source/Programs/Horde/Horde.Server.Tests/NotificationServiceTests.cs
--- a/Engine/Source/Programs/Horde/Horde.Server.Tests/NotificationServiceTests.cs
+++ b/Engine/Source/Programs/Horde/Horde.Server.Tests/NotificationServiceTests.cs
@@ -19,6 +19,7 @@ using Horde.Server.Logs;
 using Microsoft.Extensions.Caching.Memory;
 using Horde.Server.Agents;
 using EpicGames.Horde.Streams;
+using EpicGames.Horde.Jobs;
 
 namespace Horde.Server.Tests
 {
diff --git a/Engine/Source/Programs/Horde/Horde.Server.Tests/SchedulerTests.cs b/Engine/Source/Programs/Horde/Horde.Server.Tests/SchedulerTests.cs
--- a/Engine/Source/Programs/Horde/Horde.Server.Tests/SchedulerTests.cs
+++ b/Engine/Source/Programs/Horde/Horde.Server.Tests/SchedulerTests.cs
@@ -20,6 +20,7 @@ using Horde.Server.Perforce;
 using EpicGames.Horde.Jobs.Templates;
 using EpicGames.Horde.Projects;
 using EpicGames.Horde.Streams;
+using EpicGames.Horde.Jobs;
 
 namespace Horde.Server.Tests
 {
diff --git a/Engine/Source/Programs/Horde/Horde.Server.Tests/TestDataTests.cs b/Engine/Source/Programs/Horde/Horde.Server.Tests/TestDataTests.cs
--- a/Engine/Source/Programs/Horde/Horde.Server.Tests/TestDataTests.cs
+++ b/Engine/Source/Programs/Horde/Horde.Server.Tests/TestDataTests.cs
@@ -24,6 +24,7 @@ using EpicGames.Horde;
 using EpicGames.Horde.Jobs.Templates;
 using EpicGames.Horde.Projects;
 using EpicGames.Horde.Streams;
+using EpicGames.Horde.Jobs;
 
 namespace Horde.Server.Tests
 {
diff --git a/Engine/Source/Programs/Horde/Horde.Server/Compute/ComputeService.cs b/Engine/Source/Programs/Horde/Horde.Server/Compute/ComputeService.cs
--- a/Engine/Source/Programs/Horde/Horde.Server/Compute/ComputeService.cs
+++ b/Engine/Source/Programs/Horde/Horde.Server/Compute/ComputeService.cs
@@ -14,6 +14,7 @@ using EpicGames.Horde.Agents.Leases;
 using EpicGames.Horde.Compute;
 using EpicGames.Horde.Compute.Clients;
 using EpicGames.Horde.Compute.Transports;
+using EpicGames.Horde.Jobs;
 using Google.Protobuf;
 using Google.Protobuf.WellKnownTypes;
 using Horde.Common.Rpc;
diff --git a/Engine/Source/Programs/Horde/Horde.Server/Devices/DeviceCollection.cs b/Engine/Source/Programs/Horde/Horde.Server/Devices/DeviceCollection.cs
--- a/Engine/Source/Programs/Horde/Horde.Server/Devices/DeviceCollection.cs
+++ b/Engine/Source/Programs/Horde/Horde.Server/Devices/DeviceCollection.cs
@@ -6,6 +6,7 @@ using System.Linq;
 using System.Linq.Expressions;
 using System.Threading.Tasks;
 using EpicGames.Horde.Devices;
+using EpicGames.Horde.Jobs;
 using EpicGames.Horde.Streams;
 using EpicGames.Horde.Users;
 using Horde.Server.Jobs;
diff --git a/Engine/Source/Programs/Horde/Horde.Server/Devices/DeviceService.cs b/Engine/Source/Programs/Horde/Horde.Server/Devices/DeviceService.cs
--- a/Engine/Source/Programs/Horde/Horde.Server/Devices/DeviceService.cs
+++ b/Engine/Source/Programs/Horde/Horde.Server/Devices/DeviceService.cs
@@ -8,6 +8,7 @@ using System.Threading;
 using System.Threading.Tasks;
 using EpicGames.Core;
 using EpicGames.Horde.Devices;
+using EpicGames.Horde.Jobs;
 using EpicGames.Horde.Projects;
 using EpicGames.Horde.Users;
 using Horde.Server.Acls;
diff --git a/Engine/Source/Programs/Horde/Horde.Server/Devices/DevicesController.cs b/Engine/Source/Programs/Horde/Horde.Server/Devices/DevicesController.cs
--- a/Engine/Source/Programs/Horde/Horde.Server/Devices/DevicesController.cs
+++ b/Engine/Source/Programs/Horde/Horde.Server/Devices/DevicesController.cs
@@ -6,6 +6,7 @@ using System.Linq;
 using System.Text.Json;
 using System.Threading.Tasks;
 using EpicGames.Horde.Devices;
+using EpicGames.Horde.Jobs;
 using Horde.Server.Jobs;
 using Horde.Server.Server;
 using Horde.Server.Users;
diff --git a/Engine/Source/Programs/Horde/Horde.Server/Devices/IDevice.cs b/Engine/Source/Programs/Horde/Horde.Server/Devices/IDevice.cs
--- a/Engine/Source/Programs/Horde/Horde.Server/Devices/IDevice.cs
+++ b/Engine/Source/Programs/Horde/Horde.Server/Devices/IDevice.cs
@@ -3,6 +3,7 @@
 using System;
 using System.Collections.Generic;
 using EpicGames.Horde.Devices;
+using EpicGames.Horde.Jobs;
 using EpicGames.Horde.Projects;
 using EpicGames.Horde.Streams;
 using EpicGames.Horde.Users;
diff --git a/Engine/Source/Programs/Horde/Horde.Server/Devices/IDeviceCollection.cs b/Engine/Source/Programs/Horde/Horde.Server/Devices/IDeviceCollection.cs
--- a/Engine/Source/Programs/Horde/Horde.Server/Devices/IDeviceCollection.cs
+++ b/Engine/Source/Programs/Horde/Horde.Server/Devices/IDeviceCollection.cs
@@ -4,6 +4,7 @@ using System;
 using System.Collections.Generic;
 using System.Threading.Tasks;
 using EpicGames.Horde.Devices;
+using EpicGames.Horde.Jobs;
 using EpicGames.Horde.Users;
 using Horde.Server.Jobs;
 using MongoDB.Bson;
diff --git a/Engine/Source/Programs/Horde/Horde.Server/Issues/IIssueCollection.cs b/Engine/Source/Programs/Horde/Horde.Server/Issues/IIssueCollection.cs
--- a/Engine/Source/Programs/Horde/Horde.Server/Issues/IIssueCollection.cs
+++ b/Engine/Source/Programs/Horde/Horde.Server/Issues/IIssueCollection.cs
@@ -14,6 +14,7 @@ using EpicGames.Horde.Issues;
 using EpicGames.Horde.Users;
 using EpicGames.Horde.Streams;
 using EpicGames.Horde.Jobs.Templates;
+using EpicGames.Horde.Jobs;
 
 namespace Horde.Server.Issues
 {
diff --git a/Engine/Source/Programs/Horde/Horde.Server/Issues/IIssueStep.cs b/Engine/Source/Programs/Horde/Horde.Server/Issues/IIssueStep.cs
--- a/Engine/Source/Programs/Horde/Horde.Server/Issues/IIssueStep.cs
+++ b/Engine/Source/Programs/Horde/Horde.Server/Issues/IIssueStep.cs
@@ -1,6 +1,7 @@
 // Copyright Epic Games, Inc. All Rights Reserved.
 
 using System;
+using EpicGames.Horde.Jobs;
 using Horde.Server.Jobs;
 using Horde.Server.Jobs.Graphs;
 using Horde.Server.Logs;
diff --git a/Engine/Source/Programs/Horde/Horde.Server/Issues/IssueCollection.cs b/Engine/Source/Programs/Horde/Horde.Server/Issues/IssueCollection.cs
--- a/Engine/Source/Programs/Horde/Horde.Server/Issues/IssueCollection.cs
+++ b/Engine/Source/Programs/Horde/Horde.Server/Issues/IssueCollection.cs
@@ -9,6 +9,7 @@ using System.Linq;
 using System.Threading.Tasks;
 using EpicGames.Core;
 using EpicGames.Horde.Issues;
+using EpicGames.Horde.Jobs;
 using EpicGames.Horde.Jobs.Templates;
 using EpicGames.Horde.Streams;
 using EpicGames.Horde.Users;
diff --git a/Engine/Source/Programs/Horde/Horde.Server/Issues/IssueService.cs b/Engine/Source/Programs/Horde/Horde.Server/Issues/IssueService.cs
--- a/Engine/Source/Programs/Horde/Horde.Server/Issues/IssueService.cs
+++ b/Engine/Source/Programs/Horde/Horde.Server/Issues/IssueService.cs
@@ -29,6 +29,7 @@ using Microsoft.Extensions.Options;
 using MongoDB.Bson;
 using MongoDB.Bson.Serialization;
 using OpenTelemetry.Trace;
+using EpicGames.Horde.Jobs;
 
 namespace Horde.Server.Issues
 {
diff --git a/Engine/Source/Programs/Horde/Horde.Server/Issues/IssuesController.cs b/Engine/Source/Programs/Horde/Horde.Server/Issues/IssuesController.cs
--- a/Engine/Source/Programs/Horde/Horde.Server/Issues/IssuesController.cs
+++ b/Engine/Source/Programs/Horde/Horde.Server/Issues/IssuesController.cs
@@ -23,6 +23,7 @@ using Microsoft.Extensions.Options;
 using Horde.Server.Server;
 using EpicGames.Horde.Streams;
 using EpicGames.Horde.Users;
+using EpicGames.Horde.Jobs;
 
 namespace Horde.Server.Issues
 {
diff --git a/Engine/Source/Programs/Horde/Horde.Server/Jobs/Artifacts/ArtifactCollectionV1.cs b/Engine/Source/Programs/Horde/Horde.Server/Jobs/Artifacts/ArtifactCollectionV1.cs
--- a/Engine/Source/Programs/Horde/Horde.Server/Jobs/Artifacts/ArtifactCollectionV1.cs
+++ b/Engine/Source/Programs/Horde/Horde.Server/Jobs/Artifacts/ArtifactCollectionV1.cs
@@ -11,6 +11,7 @@ using MongoDB.Bson;
 using MongoDB.Bson.Serialization.Attributes;
 using MongoDB.Driver;
 using EpicGames.Horde.Storage;
+using EpicGames.Horde.Jobs;
 
 namespace Horde.Server.Jobs.Artifacts
 {
diff --git a/Engine/Source/Programs/Horde/Horde.Server/Jobs/Artifacts/ArtifactMessagesV1.cs b/Engine/Source/Programs/Horde/Horde.Server/Jobs/Artifacts/ArtifactMessagesV1.cs
--- a/Engine/Source/Programs/Horde/Horde.Server/Jobs/Artifacts/ArtifactMessagesV1.cs
+++ b/Engine/Source/Programs/Horde/Horde.Server/Jobs/Artifacts/ArtifactMessagesV1.cs
@@ -1,6 +1,7 @@
 // Copyright Epic Games, Inc. All Rights Reserved.
 
 using System.Collections.Generic;
+using EpicGames.Horde.Jobs;
 
 namespace Horde.Server.Jobs.Artifacts
 {
diff --git a/Engine/Source/Programs/Horde/Horde.Server/Jobs/Artifacts/ArtifactsControllerV1.cs b/Engine/Source/Programs/Horde/Horde.Server/Jobs/Artifacts/ArtifactsControllerV1.cs
--- a/Engine/Source/Programs/Horde/Horde.Server/Jobs/Artifacts/ArtifactsControllerV1.cs
+++ b/Engine/Source/Programs/Horde/Horde.Server/Jobs/Artifacts/ArtifactsControllerV1.cs
@@ -8,6 +8,7 @@ using System.IO.Compression;
 using System.Linq;
 using System.Security.Claims;
 using System.Threading.Tasks;
+using EpicGames.Horde.Jobs;
 using Horde.Server.Acls;
 using Horde.Server.Artifacts;
 using Horde.Server.Jobs.Graphs;
diff --git a/Engine/Source/Programs/Horde/Horde.Server/Jobs/Artifacts/IArtifactCollectionV1.cs b/Engine/Source/Programs/Horde/Horde.Server/Jobs/Artifacts/IArtifactCollectionV1.cs
--- a/Engine/Source/Programs/Horde/Horde.Server/Jobs/Artifacts/IArtifactCollectionV1.cs
+++ b/Engine/Source/Programs/Horde/Horde.Server/Jobs/Artifacts/IArtifactCollectionV1.cs
@@ -2,6 +2,7 @@
 
 using System.Collections.Generic;
 using System.Threading.Tasks;
+using EpicGames.Horde.Jobs;
 using MongoDB.Bson;
 
 namespace Horde.Server.Jobs.Artifacts
diff --git a/Engine/Source/Programs/Horde/Horde.Server/Jobs/Artifacts/IArtifactV1.cs b/Engine/Source/Programs/Horde/Horde.Server/Jobs/Artifacts/IArtifactV1.cs
--- a/Engine/Source/Programs/Horde/Horde.Server/Jobs/Artifacts/IArtifactV1.cs
+++ b/Engine/Source/Programs/Horde/Horde.Server/Jobs/Artifacts/IArtifactV1.cs
@@ -1,5 +1,6 @@
 // Copyright Epic Games, Inc. All Rights Reserved.
 
+using EpicGames.Horde.Jobs;
 using MongoDB.Bson;
 
 namespace Horde.Server.Jobs.Artifacts
diff --git a/Engine/Source/Programs/Horde/Horde.Server/Jobs/Bisect/BisectController.cs b/Engine/Source/Programs/Horde/Horde.Server/Jobs/Bisect/BisectController.cs
--- a/Engine/Source/Programs/Horde/Horde.Server/Jobs/Bisect/BisectController.cs
+++ b/Engine/Source/Programs/Horde/Horde.Server/Jobs/Bisect/BisectController.cs
@@ -21,6 +21,7 @@ using EpicGames.Horde.Streams;
 using EpicGames.Horde.Users;
 using OpenTelemetry.Trace;
 using Microsoft.Extensions.Logging;
+using EpicGames.Horde.Jobs;
 
 namespace Horde.Server.Jobs.Bisect
 {
diff --git a/Engine/Source/Programs/Horde/Horde.Server/Jobs/Bisect/BisectTaskCollection.cs b/Engine/Source/Programs/Horde/Horde.Server/Jobs/Bisect/BisectTaskCollection.cs
--- a/Engine/Source/Programs/Horde/Horde.Server/Jobs/Bisect/BisectTaskCollection.cs
+++ b/Engine/Source/Programs/Horde/Horde.Server/Jobs/Bisect/BisectTaskCollection.cs
@@ -16,6 +16,7 @@ using EpicGames.Horde.Jobs.Templates;
 using EpicGames.Horde.Streams;
 using EpicGames.Horde.Users;
 using OpenTelemetry.Trace;
+using EpicGames.Horde.Jobs;
 
 namespace Horde.Server.Jobs.Bisect
 {
diff --git a/Engine/Source/Programs/Horde/Horde.Server/Jobs/Bisect/IBisectTask.cs b/Engine/Source/Programs/Horde/Horde.Server/Jobs/Bisect/IBisectTask.cs
--- a/Engine/Source/Programs/Horde/Horde.Server/Jobs/Bisect/IBisectTask.cs
+++ b/Engine/Source/Programs/Horde/Horde.Server/Jobs/Bisect/IBisectTask.cs
@@ -6,6 +6,7 @@ using HordeCommon;
 using EpicGames.Horde.Jobs.Templates;
 using EpicGames.Horde.Streams;
 using EpicGames.Horde.Users;
+using EpicGames.Horde.Jobs;
 
 namespace Horde.Server.Jobs.Bisect
 {
diff --git a/Engine/Source/Programs/Horde/Horde.Server/Jobs/Bisect/IBisectTaskCollection.cs b/Engine/Source/Programs/Horde/Horde.Server/Jobs/Bisect/IBisectTaskCollection.cs
--- a/Engine/Source/Programs/Horde/Horde.Server/Jobs/Bisect/IBisectTaskCollection.cs
+++ b/Engine/Source/Programs/Horde/Horde.Server/Jobs/Bisect/IBisectTaskCollection.cs
@@ -9,6 +9,7 @@ using HordeCommon;
 using EpicGames.Horde.Users;
 using Microsoft.Extensions.Logging;
 using Horde.Server.Jobs.Graphs;
+using EpicGames.Horde.Jobs;
 
 namespace Horde.Server.Jobs.Bisect
 {
diff --git a/Engine/Source/Programs/Horde/Horde.Server/Jobs/IJob.cs b/Engine/Source/Programs/Horde/Horde.Server/Jobs/IJob.cs
--- a/Engine/Source/Programs/Horde/Horde.Server/Jobs/IJob.cs
+++ b/Engine/Source/Programs/Horde/Horde.Server/Jobs/IJob.cs
@@ -25,6 +25,7 @@ using HordeCommon.Rpc.Tasks;
 using MongoDB.Bson;
 using MongoDB.Driver;
 using Horde.Server.Jobs.Bisect;
+using EpicGames.Horde.Jobs;
 
 namespace Horde.Server.Jobs
 {
diff --git a/Engine/Source/Programs/Horde/Horde.Server/Jobs/IJobCollection.cs b/Engine/Source/Programs/Horde/Horde.Server/Jobs/IJobCollection.cs
--- a/Engine/Source/Programs/Horde/Horde.Server/Jobs/IJobCollection.cs
+++ b/Engine/Source/Programs/Horde/Horde.Server/Jobs/IJobCollection.cs
@@ -20,6 +20,7 @@ using MongoDB.Bson;
 using Horde.Server.Jobs.Bisect;
 using System.Threading;
 using Horde.Server.Acls;
+using EpicGames.Horde.Jobs;
 
 namespace Horde.Server.Jobs
 {
diff --git a/Engine/Source/Programs/Horde/Horde.Server/Jobs/IJobStepRef.cs b/Engine/Source/Programs/Horde/Horde.Server/Jobs/IJobStepRef.cs
--- a/Engine/Source/Programs/Horde/Horde.Server/Jobs/IJobStepRef.cs
+++ b/Engine/Source/Programs/Horde/Horde.Server/Jobs/IJobStepRef.cs
@@ -11,6 +11,7 @@ using MongoDB.Bson.Serialization.Attributes;
 using EpicGames.Horde.Agents;
 using EpicGames.Horde.Streams;
 using EpicGames.Horde.Jobs.Templates;
+using EpicGames.Horde.Jobs;
 
 namespace Horde.Server.Jobs
 {
diff --git a/Engine/Source/Programs/Horde/Horde.Server/Jobs/IJobStepRefCollection.cs b/Engine/Source/Programs/Horde/Horde.Server/Jobs/IJobStepRefCollection.cs
--- a/Engine/Source/Programs/Horde/Horde.Server/Jobs/IJobStepRefCollection.cs
+++ b/Engine/Source/Programs/Horde/Horde.Server/Jobs/IJobStepRefCollection.cs
@@ -12,6 +12,7 @@ using Microsoft.Extensions.Logging;
 using EpicGames.Horde.Agents;
 using EpicGames.Horde.Jobs.Templates;
 using EpicGames.Horde.Streams;
+using EpicGames.Horde.Jobs;
 
 namespace Horde.Server.Jobs
 {
diff --git a/Engine/Source/Programs/Horde/Horde.Server/Jobs/JobCollection.cs b/Engine/Source/Programs/Horde/Horde.Server/Jobs/JobCollection.cs
--- a/Engine/Source/Programs/Horde/Horde.Server/Jobs/JobCollection.cs
+++ b/Engine/Source/Programs/Horde/Horde.Server/Jobs/JobCollection.cs
@@ -34,6 +34,7 @@ using MongoDB.Bson.Serialization.Options;
 using MongoDB.Driver;
 using OpenTelemetry.Trace;
 using EpicGames.Horde.Jobs.Templates;
+using EpicGames.Horde.Jobs;
 
 namespace Horde.Server.Jobs
 {
diff --git a/Engine/Source/Programs/Horde/Horde.Server/Jobs/JobId.cs b/Engine/Source/Programs/Horde/Horde.Server/Jobs/JobId.cs
deleted file mode 100644
diff --git a/Engine/Source/Programs/Horde/Horde.Server/Jobs/JobIdUtils.cs b/Engine/Source/Programs/Horde/Horde.Server/Jobs/JobIdUtils.cs
--- a/Engine/Source/Programs/Horde/Horde.Server/Jobs/JobIdUtils.cs
+++ b/Engine/Source/Programs/Horde/Horde.Server/Jobs/JobIdUtils.cs
@@ -1,5 +1,6 @@
 // Copyright Epic Games, Inc. All Rights Reserved.
 
+using EpicGames.Horde.Jobs;
 using Horde.Server.Utilities;
 
 namespace Horde.Server.Jobs
diff --git a/Engine/Source/Programs/Horde/Horde.Server/Jobs/JobRpcService.cs b/Engine/Source/Programs/Horde/Horde.Server/Jobs/JobRpcService.cs
--- a/Engine/Source/Programs/Horde/Horde.Server/Jobs/JobRpcService.cs
+++ b/Engine/Source/Programs/Horde/Horde.Server/Jobs/JobRpcService.cs
@@ -39,6 +39,7 @@ using Microsoft.Extensions.Options;
 using MongoDB.Bson;
 using MongoDB.Bson.Serialization;
 using EpicGames.Horde;
+using EpicGames.Horde.Jobs;
 
 namespace Horde.Server.Jobs
 {
diff --git a/Engine/Source/Programs/Horde/Horde.Server/Jobs/JobService.cs b/Engine/Source/Programs/Horde/Horde.Server/Jobs/JobService.cs
--- a/Engine/Source/Programs/Horde/Horde.Server/Jobs/JobService.cs
+++ b/Engine/Source/Programs/Horde/Horde.Server/Jobs/JobService.cs
@@ -11,6 +11,7 @@ using System.Threading.Tasks;
 using EpicGames.Core;
 using EpicGames.Horde.Agents;
 using EpicGames.Horde.Agents.Leases;
+using EpicGames.Horde.Jobs;
 using EpicGames.Horde.Jobs.Templates;
 using EpicGames.Horde.Streams;
 using EpicGames.Horde.Users;
diff --git a/Engine/Source/Programs/Horde/Horde.Server/Jobs/JobStepBatchId.cs b/Engine/Source/Programs/Horde/Horde.Server/Jobs/JobStepBatchId.cs
deleted file mode 100644
diff --git a/Engine/Source/Programs/Horde/Horde.Server/Jobs/JobStepId.cs b/Engine/Source/Programs/Horde/Horde.Server/Jobs/JobStepId.cs
deleted file mode 100644
diff --git a/Engine/Source/Programs/Horde/Horde.Server/Jobs/JobStepRefCollection.cs b/Engine/Source/Programs/Horde/Horde.Server/Jobs/JobStepRefCollection.cs
--- a/Engine/Source/Programs/Horde/Horde.Server/Jobs/JobStepRefCollection.cs
+++ b/Engine/Source/Programs/Horde/Horde.Server/Jobs/JobStepRefCollection.cs
@@ -16,6 +16,7 @@ using MongoDB.Driver;
 using EpicGames.Horde.Agents;
 using EpicGames.Horde.Streams;
 using EpicGames.Horde.Jobs.Templates;
+using EpicGames.Horde.Jobs;
 
 namespace Horde.Server.Jobs
 {
diff --git a/Engine/Source/Programs/Horde/Horde.Server/Jobs/JobTaskSource.cs b/Engine/Source/Programs/Horde/Horde.Server/Jobs/JobTaskSource.cs
--- a/Engine/Source/Programs/Horde/Horde.Server/Jobs/JobTaskSource.cs
+++ b/Engine/Source/Programs/Horde/Horde.Server/Jobs/JobTaskSource.cs
@@ -36,6 +36,7 @@ using EpicGames.Horde.Agents.Leases;
 using Horde.Server.Jobs.Bisect;
 using EpicGames.Horde.Agents;
 using EpicGames.Horde.Storage;
+using EpicGames.Horde.Jobs;
 
 namespace Horde.Server.Jobs
 {
diff --git a/Engine/Source/Programs/Horde/Horde.Server/Jobs/JobsController.cs b/Engine/Source/Programs/Horde/Horde.Server/Jobs/JobsController.cs
--- a/Engine/Source/Programs/Horde/Horde.Server/Jobs/JobsController.cs
+++ b/Engine/Source/Programs/Horde/Horde.Server/Jobs/JobsController.cs
@@ -6,6 +6,7 @@ using System.Linq;
 using System.Threading;
 using System.Threading.Tasks;
 using EpicGames.Core;
+using EpicGames.Horde.Jobs;
 using EpicGames.Horde.Jobs.Templates;
 using EpicGames.Horde.Streams;
 using EpicGames.Horde.Users;
diff --git a/Engine/Source/Programs/Horde/Horde.Server/Jobs/Schedules/ScheduleService.cs b/Engine/Source/Programs/Horde/Horde.Server/Jobs/Schedules/ScheduleService.cs
--- a/Engine/Source/Programs/Horde/Horde.Server/Jobs/Schedules/ScheduleService.cs
+++ b/Engine/Source/Programs/Horde/Horde.Server/Jobs/Schedules/ScheduleService.cs
@@ -8,6 +8,7 @@ using System.Text.RegularExpressions;
 using System.Threading;
 using System.Threading.Tasks;
 using EpicGames.Core;
+using EpicGames.Horde.Jobs;
 using EpicGames.Horde.Jobs.Templates;
 using EpicGames.Horde.Streams;
 using EpicGames.Redis;
diff --git a/Engine/Source/Programs/Horde/Horde.Server/Jobs/Schedules/Schedules.cs b/Engine/Source/Programs/Horde/Horde.Server/Jobs/Schedules/Schedules.cs
--- a/Engine/Source/Programs/Horde/Horde.Server/Jobs/Schedules/Schedules.cs
+++ b/Engine/Source/Programs/Horde/Horde.Server/Jobs/Schedules/Schedules.cs
@@ -2,6 +2,7 @@
 
 using System;
 using System.Collections.Generic;
+using EpicGames.Horde.Jobs;
 using Horde.Server.Perforce;
 using Horde.Server.Streams;
 
diff --git a/Engine/Source/Programs/Horde/Horde.Server/Jobs/SubResourceId.cs b/Engine/Source/Programs/Horde/Horde.Server/Jobs/SubResourceId.cs
deleted file mode 100644
diff --git a/Engine/Source/Programs/Horde/Horde.Server/Jobs/SubResourceIdConverter.cs b/Engine/Source/Programs/Horde/Horde.Server/Jobs/SubResourceIdConverter.cs
deleted file mode 100644
diff --git a/Engine/Source/Programs/Horde/Horde.Server/Jobs/TestData/ITestData.cs b/Engine/Source/Programs/Horde/Horde.Server/Jobs/TestData/ITestData.cs
--- a/Engine/Source/Programs/Horde/Horde.Server/Jobs/TestData/ITestData.cs
+++ b/Engine/Source/Programs/Horde/Horde.Server/Jobs/TestData/ITestData.cs
@@ -5,6 +5,7 @@ using System;
 using EpicGames.Horde.Jobs.Templates;
 using MongoDB.Bson;
 using EpicGames.Horde.Streams;
+using EpicGames.Horde.Jobs;
 
 namespace Horde.Server.Jobs.TestData
 {
diff --git a/Engine/Source/Programs/Horde/Horde.Server/Jobs/TestData/ITestDataCollection.cs b/Engine/Source/Programs/Horde/Horde.Server/Jobs/TestData/ITestDataCollection.cs
--- a/Engine/Source/Programs/Horde/Horde.Server/Jobs/TestData/ITestDataCollection.cs
+++ b/Engine/Source/Programs/Horde/Horde.Server/Jobs/TestData/ITestDataCollection.cs
@@ -3,6 +3,7 @@
 using System;
 using System.Collections.Generic;
 using System.Threading.Tasks;
+using EpicGames.Horde.Jobs;
 using EpicGames.Horde.Streams;
 using MongoDB.Bson;
 
diff --git a/Engine/Source/Programs/Horde/Horde.Server/Jobs/TestData/TestDataCollection.cs b/Engine/Source/Programs/Horde/Horde.Server/Jobs/TestData/TestDataCollection.cs
--- a/Engine/Source/Programs/Horde/Horde.Server/Jobs/TestData/TestDataCollection.cs
+++ b/Engine/Source/Programs/Horde/Horde.Server/Jobs/TestData/TestDataCollection.cs
@@ -4,6 +4,7 @@ using System;
 using System.Collections.Generic;
 using System.Linq;
 using System.Threading.Tasks;
+using EpicGames.Horde.Jobs;
 using EpicGames.Horde.Jobs.Templates;
 using EpicGames.Horde.Streams;
 using Horde.Server.Server;
diff --git a/Engine/Source/Programs/Horde/Horde.Server/Jobs/TestData/TestDataController.cs b/Engine/Source/Programs/Horde/Horde.Server/Jobs/TestData/TestDataController.cs
--- a/Engine/Source/Programs/Horde/Horde.Server/Jobs/TestData/TestDataController.cs
+++ b/Engine/Source/Programs/Horde/Horde.Server/Jobs/TestData/TestDataController.cs
@@ -14,6 +14,7 @@ using EpicGames.Core;
 using Microsoft.Extensions.Options;
 using Horde.Server.Server;
 using EpicGames.Horde.Streams;
+using EpicGames.Horde.Jobs;
 
 namespace Horde.Server.Jobs.TestData
 {
diff --git a/Engine/Source/Programs/Horde/Horde.Server/Jobs/TestData/TestDataMessages.cs b/Engine/Source/Programs/Horde/Horde.Server/Jobs/TestData/TestDataMessages.cs
--- a/Engine/Source/Programs/Horde/Horde.Server/Jobs/TestData/TestDataMessages.cs
+++ b/Engine/Source/Programs/Horde/Horde.Server/Jobs/TestData/TestDataMessages.cs
@@ -6,6 +6,7 @@ using System.ComponentModel.DataAnnotations;
 using System.Linq;
 using MongoDB.Bson.Serialization;
 using EpicGames.Horde.Streams;
+using EpicGames.Horde.Jobs;
 
 namespace Horde.Server.Jobs.TestData
 {
diff --git a/Engine/Source/Programs/Horde/Horde.Server/Jobs/Timing/IJobTimingCollection.cs b/Engine/Source/Programs/Horde/Horde.Server/Jobs/Timing/IJobTimingCollection.cs
--- a/Engine/Source/Programs/Horde/Horde.Server/Jobs/Timing/IJobTimingCollection.cs
+++ b/Engine/Source/Programs/Horde/Horde.Server/Jobs/Timing/IJobTimingCollection.cs
@@ -2,6 +2,7 @@
 
 using System.Collections.Generic;
 using System.Threading.Tasks;
+using EpicGames.Horde.Jobs;
 using MongoDB.Bson.Serialization.Attributes;
 
 namespace Horde.Server.Jobs.Timing
diff --git a/Engine/Source/Programs/Horde/Horde.Server/Jobs/Timing/JobTimingCollection.cs b/Engine/Source/Programs/Horde/Horde.Server/Jobs/Timing/JobTimingCollection.cs
--- a/Engine/Source/Programs/Horde/Horde.Server/Jobs/Timing/JobTimingCollection.cs
+++ b/Engine/Source/Programs/Horde/Horde.Server/Jobs/Timing/JobTimingCollection.cs
@@ -4,6 +4,7 @@ using System.Collections.Generic;
 using System.Diagnostics.CodeAnalysis;
 using System.Linq;
 using System.Threading.Tasks;
+using EpicGames.Horde.Jobs;
 using Horde.Server.Server;
 using Horde.Server.Utilities;
 using MongoDB.Bson.Serialization.Attributes;
diff --git a/Engine/Source/Programs/Horde/Horde.Server/Logs/ILogFile.cs b/Engine/Source/Programs/Horde/Horde.Server/Logs/ILogFile.cs
--- a/Engine/Source/Programs/Horde/Horde.Server/Logs/ILogFile.cs
+++ b/Engine/Source/Programs/Horde/Horde.Server/Logs/ILogFile.cs
@@ -4,6 +4,7 @@ using System;
 using System.Collections.Generic;
 using EpicGames.Core;
 using EpicGames.Horde.Agents.Leases;
+using EpicGames.Horde.Jobs;
 using EpicGames.Horde.Storage;
 using Horde.Server.Agents.Sessions;
 using Horde.Server.Jobs;
diff --git a/Engine/Source/Programs/Horde/Horde.Server/Logs/ILogFileCollection.cs b/Engine/Source/Programs/Horde/Horde.Server/Logs/ILogFileCollection.cs
--- a/Engine/Source/Programs/Horde/Horde.Server/Logs/ILogFileCollection.cs
+++ b/Engine/Source/Programs/Horde/Horde.Server/Logs/ILogFileCollection.cs
@@ -6,6 +6,7 @@ using System.Threading.Tasks;
 using Horde.Server.Agents.Sessions;
 using Horde.Server.Jobs;
 using EpicGames.Horde.Agents.Leases;
+using EpicGames.Horde.Jobs;
 
 namespace Horde.Server.Logs
 {
diff --git a/Engine/Source/Programs/Horde/Horde.Server/Logs/LogFileCollection.cs b/Engine/Source/Programs/Horde/Horde.Server/Logs/LogFileCollection.cs
--- a/Engine/Source/Programs/Horde/Horde.Server/Logs/LogFileCollection.cs
+++ b/Engine/Source/Programs/Horde/Horde.Server/Logs/LogFileCollection.cs
@@ -6,6 +6,7 @@ using System.Net;
 using System.Threading;
 using System.Threading.Tasks;
 using EpicGames.Horde.Agents.Leases;
+using EpicGames.Horde.Jobs;
 using EpicGames.Horde.Storage;
 using Horde.Server.Agents.Sessions;
 using Horde.Server.Jobs;
diff --git a/Engine/Source/Programs/Horde/Horde.Server/Logs/LogFileService.cs b/Engine/Source/Programs/Horde/Horde.Server/Logs/LogFileService.cs
--- a/Engine/Source/Programs/Horde/Horde.Server/Logs/LogFileService.cs
+++ b/Engine/Source/Programs/Horde/Horde.Server/Logs/LogFileService.cs
@@ -28,6 +28,7 @@ using Microsoft.Extensions.Options;
 using MongoDB.Bson;
 using OpenTelemetry.Trace;
 using EpicGames.Horde.Agents.Leases;
+using EpicGames.Horde.Jobs;
 
 namespace Horde.Server.Logs
 {
diff --git a/Engine/Source/Programs/Horde/Horde.Server/Logs/LogsController.cs b/Engine/Source/Programs/Horde/Horde.Server/Logs/LogsController.cs
--- a/Engine/Source/Programs/Horde/Horde.Server/Logs/LogsController.cs
+++ b/Engine/Source/Programs/Horde/Horde.Server/Logs/LogsController.cs
@@ -9,6 +9,7 @@ using System.Text;
 using System.Threading;
 using System.Threading.Tasks;
 using EpicGames.Core;
+using EpicGames.Horde.Jobs;
 using EpicGames.Horde.Logs;
 using EpicGames.Horde.Storage;
 using Horde.Server.Acls;
diff --git a/Engine/Source/Programs/Horde/Horde.Server/Notifications/INotificationService.cs b/Engine/Source/Programs/Horde/Horde.Server/Notifications/INotificationService.cs
--- a/Engine/Source/Programs/Horde/Horde.Server/Notifications/INotificationService.cs
+++ b/Engine/Source/Programs/Horde/Horde.Server/Notifications/INotificationService.cs
@@ -13,6 +13,7 @@ using Horde.Server.Users;
 using MongoDB.Bson;
 using Horde.Server.Streams;
 using Horde.Server.Agents;
+using EpicGames.Horde.Jobs;
 
 namespace Horde.Server.Notifications
 {
diff --git a/Engine/Source/Programs/Horde/Horde.Server/Notifications/NotificationService.cs b/Engine/Source/Programs/Horde/Horde.Server/Notifications/NotificationService.cs
--- a/Engine/Source/Programs/Horde/Horde.Server/Notifications/NotificationService.cs
+++ b/Engine/Source/Programs/Horde/Horde.Server/Notifications/NotificationService.cs
@@ -31,6 +31,7 @@ using Microsoft.Extensions.Options;
 using MongoDB.Bson;
 using StackExchange.Redis;
 using JsonSerializer = System.Text.Json.JsonSerializer;
+using EpicGames.Horde.Jobs;
 
 namespace Horde.Server.Notifications
 {
diff --git a/Engine/Source/Programs/Horde/Horde.Server/Notifications/Sinks/SlackNotificationSink.cs b/Engine/Source/Programs/Horde/Horde.Server/Notifications/Sinks/SlackNotificationSink.cs
--- a/Engine/Source/Programs/Horde/Horde.Server/Notifications/Sinks/SlackNotificationSink.cs
+++ b/Engine/Source/Programs/Horde/Horde.Server/Notifications/Sinks/SlackNotificationSink.cs
@@ -47,6 +47,7 @@ using MongoDB.Driver;
 using EpicGames.Horde;
 using Horde.Server.Agents;
 using EpicGames.Horde.Jobs.Templates;
+using EpicGames.Horde.Jobs;
 
 namespace Horde.Server.Notifications.Sinks
 {
diff --git a/Engine/Source/Programs/Horde/Horde.Server/Server/DebugController.cs b/Engine/Source/Programs/Horde/Horde.Server/Server/DebugController.cs
--- a/Engine/Source/Programs/Horde/Horde.Server/Server/DebugController.cs
+++ b/Engine/Source/Programs/Horde/Horde.Server/Server/DebugController.cs
@@ -16,6 +16,7 @@ using System.Web;
 using EpicGames.Core;
 using EpicGames.Horde.Agents.Leases;
 using EpicGames.Horde.Compute;
+using EpicGames.Horde.Jobs;
 using Google.Protobuf;
 using Horde.Common.Rpc;
 using Horde.Server.Acls;
diff --git a/Engine/Source/Programs/Horde/Horde.Server/Startup.cs b/Engine/Source/Programs/Horde/Horde.Server/Startup.cs
--- a/Engine/Source/Programs/Horde/Horde.Server/Startup.cs
+++ b/Engine/Source/Programs/Horde/Horde.Server/Startup.cs
@@ -115,6 +115,7 @@ using Horde.Server.Aws;
 using StackExchange.Redis;
 using Horde.Server.Telemetry.Sinks;
 using Horde.Server.Telemetry.Metrics;
+using EpicGames.Horde.Jobs;
 
 namespace Horde.Server
 {
@@ -1104,6 +1105,40 @@ namespace Horde.Server
 			public override void Serialize(BsonSerializationContext context, BsonSerializationArgs args, TValue value) => context.Writer.WriteInt32(_converter.ToSubResourceId(value).Value);
 		}
 
+		/// <summary>
+		/// Serializer for subresource ids
+		/// </summary>
+		public class SubResourceIdSerializer : IBsonSerializer<SubResourceId>
+		{
+			/// <inheritdoc/>
+			public Type ValueType => typeof(SubResourceId);
+
+			/// <inheritdoc/>
+			public object Deserialize(BsonDeserializationContext context, BsonDeserializationArgs args)
+			{
+				return new SubResourceId((ushort)context.Reader.ReadInt32());
+			}
+
+			/// <inheritdoc/>
+			public void Serialize(BsonSerializationContext context, BsonSerializationArgs args, object value)
+			{
+				SubResourceId id = (SubResourceId)value;
+				context.Writer.WriteInt32((int)id.Value);
+			}
+
+			/// <inheritdoc/>
+			public void Serialize(BsonSerializationContext context, BsonSerializationArgs args, SubResourceId id)
+			{
+				context.Writer.WriteInt32((int)id.Value);
+			}
+
+			/// <inheritdoc/>
+			SubResourceId IBsonSerializer<SubResourceId>.Deserialize(BsonDeserializationContext context, BsonDeserializationArgs vrgs)
+			{
+				return new SubResourceId((ushort)context.Reader.ReadInt32());
+			}
+		}
+
 		sealed class SubResourceIdBsonSerializationProvider : BsonSerializationProviderBase
 		{
 			/// <inheritdoc/>
@@ -1139,6 +1174,7 @@ namespace Horde.Server
 				BsonSerializer.RegisterSerializer(new AclActionBsonSerializer());
 				BsonSerializer.RegisterSerializer(new AclScopeNameBsonSerializer());
 				BsonSerializer.RegisterSerializer(new ConditionSerializer());
+				BsonSerializer.RegisterSerializer(new SubResourceIdSerializer());
 				BsonSerializer.RegisterSerializationProvider(new BsonSerializationProvider());
 				BsonSerializer.RegisterSerializationProvider(new StringIdBsonSerializationProvider());
 				BsonSerializer.RegisterSerializationProvider(new BinaryIdBsonSerializationProvider());
diff --git a/Engine/Source/Programs/Horde/Horde.Server/Streams/IStream.cs b/Engine/Source/Programs/Horde/Horde.Server/Streams/IStream.cs
--- a/Engine/Source/Programs/Horde/Horde.Server/Streams/IStream.cs
+++ b/Engine/Source/Programs/Horde/Horde.Server/Streams/IStream.cs
@@ -3,6 +3,7 @@
 using System;
 using System.Collections.Generic;
 using System.Linq;
+using EpicGames.Horde.Jobs;
 using EpicGames.Horde.Jobs.Templates;
 using EpicGames.Horde.Streams;
 using EpicGames.Horde.Users;
diff --git a/Engine/Source/Programs/Horde/Horde.Server/Streams/IStreamCollection.cs b/Engine/Source/Programs/Horde/Horde.Server/Streams/IStreamCollection.cs
--- a/Engine/Source/Programs/Horde/Horde.Server/Streams/IStreamCollection.cs
+++ b/Engine/Source/Programs/Horde/Horde.Server/Streams/IStreamCollection.cs
@@ -4,6 +4,7 @@ using System;
 using System.Collections.Generic;
 using System.Linq;
 using System.Threading.Tasks;
+using EpicGames.Horde.Jobs;
 using EpicGames.Horde.Jobs.Templates;
 using Horde.Server.Jobs;
 
diff --git a/Engine/Source/Programs/Horde/Horde.Server/Streams/StreamCollection.cs b/Engine/Source/Programs/Horde/Horde.Server/Streams/StreamCollection.cs
--- a/Engine/Source/Programs/Horde/Horde.Server/Streams/StreamCollection.cs
+++ b/Engine/Source/Programs/Horde/Horde.Server/Streams/StreamCollection.cs
@@ -6,6 +6,7 @@ using System.IO;
 using System.Linq;
 using System.Threading.Tasks;
 using EpicGames.Core;
+using EpicGames.Horde.Jobs;
 using EpicGames.Horde.Jobs.Templates;
 using EpicGames.Horde.Streams;
 using EpicGames.Horde.Users;
diff --git a/Engine/Source/Programs/Horde/Horde.Server/Streams/StreamMessages.cs b/Engine/Source/Programs/Horde/Horde.Server/Streams/StreamMessages.cs
--- a/Engine/Source/Programs/Horde/Horde.Server/Streams/StreamMessages.cs
+++ b/Engine/Source/Programs/Horde/Horde.Server/Streams/StreamMessages.cs
@@ -15,6 +15,7 @@ using Horde.Server.Issues;
 using Horde.Server.Perforce;
 using Horde.Server.Projects;
 using Horde.Server.Jobs;
+using EpicGames.Horde.Jobs;
 
 namespace Horde.Server.Streams
 {
diff --git a/Engine/Source/Programs/Horde/Horde.Server/Tasks/ConformTaskSource.cs b/Engine/Source/Programs/Horde/Horde.Server/Tasks/ConformTaskSource.cs
--- a/Engine/Source/Programs/Horde/Horde.Server/Tasks/ConformTaskSource.cs
+++ b/Engine/Source/Programs/Horde/Horde.Server/Tasks/ConformTaskSource.cs
@@ -21,6 +21,7 @@ using Microsoft.Extensions.Logging;
 using Microsoft.Extensions.Options;
 using EpicGames.Horde.Agents;
 using EpicGames.Horde.Agents.Leases;
+using EpicGames.Horde.Jobs;
 
 namespace Horde.Server.Tasks
 {
diff --git a/Engine/Source/Programs/Horde/Horde.Server/Tasks/RestartTaskSource.cs b/Engine/Source/Programs/Horde/Horde.Server/Tasks/RestartTaskSource.cs
--- a/Engine/Source/Programs/Horde/Horde.Server/Tasks/RestartTaskSource.cs
+++ b/Engine/Source/Programs/Horde/Horde.Server/Tasks/RestartTaskSource.cs
@@ -11,6 +11,7 @@ using Horde.Server.Utilities;
 using HordeCommon;
 using HordeCommon.Rpc.Tasks;
 using EpicGames.Horde.Agents.Leases;
+using EpicGames.Horde.Jobs;
 
 namespace Horde.Server.Tasks
 {
diff --git a/Engine/Source/Programs/Horde/Horde.Server/Tasks/ShutdownTaskSource.cs b/Engine/Source/Programs/Horde/Horde.Server/Tasks/ShutdownTaskSource.cs
--- a/Engine/Source/Programs/Horde/Horde.Server/Tasks/ShutdownTaskSource.cs
+++ b/Engine/Source/Programs/Horde/Horde.Server/Tasks/ShutdownTaskSource.cs
@@ -11,6 +11,7 @@ using Horde.Server.Utilities;
 using HordeCommon;
 using HordeCommon.Rpc.Tasks;
 using EpicGames.Horde.Agents.Leases;
+using EpicGames.Horde.Jobs;
 
 namespace Horde.Server.Tasks
 {
diff --git a/Engine/Source/Programs/Horde/Horde.Server/Tasks/UpgradeTaskSource.cs b/Engine/Source/Programs/Horde/Horde.Server/Tasks/UpgradeTaskSource.cs
--- a/Engine/Source/Programs/Horde/Horde.Server/Tasks/UpgradeTaskSource.cs
+++ b/Engine/Source/Programs/Horde/Horde.Server/Tasks/UpgradeTaskSource.cs
@@ -18,6 +18,7 @@ using HordeCommon.Rpc.Tasks;
 using Microsoft.Extensions.Options;
 using EpicGames.Horde.Agents.Leases;
 using EpicGames.Horde.Tools;
+using EpicGames.Horde.Jobs;
 
 namespace Horde.Server.Tasks
 {
diff --git a/Engine/Source/Programs/Horde/Horde.Server/Telemetry/Sinks/ClickHouseTelemetrySink.cs b/Engine/Source/Programs/Horde/Horde.Server/Telemetry/Sinks/ClickHouseTelemetrySink.cs
--- a/Engine/Source/Programs/Horde/Horde.Server/Telemetry/Sinks/ClickHouseTelemetrySink.cs
+++ b/Engine/Source/Programs/Horde/Horde.Server/Telemetry/Sinks/ClickHouseTelemetrySink.cs
@@ -7,6 +7,7 @@ using System.Net.Http;
 using System.Text;
 using System.Threading;
 using System.Threading.Tasks;
+using EpicGames.Horde.Jobs;
 using Horde.Server.Jobs;
 using HordeCommon.Rpc.Messages.Telemetry;
 using Microsoft.Extensions.Logging;
diff --git a/Engine/Source/Programs/Horde/Horde.Server/Users/IUser.cs b/Engine/Source/Programs/Horde/Horde.Server/Users/IUser.cs
--- a/Engine/Source/Programs/Horde/Horde.Server/Users/IUser.cs
+++ b/Engine/Source/Programs/Horde/Horde.Server/Users/IUser.cs
@@ -4,6 +4,7 @@ using System;
 using System.Collections.Generic;
 using System.Linq;
 using EpicGames.Core;
+using EpicGames.Horde.Jobs;
 using EpicGames.Horde.Jobs.Templates;
 using EpicGames.Horde.Streams;
 using EpicGames.Horde.Users;
diff --git a/Engine/Source/Programs/Horde/Horde.Server/Users/IUserCollection.cs b/Engine/Source/Programs/Horde/Horde.Server/Users/IUserCollection.cs
--- a/Engine/Source/Programs/Horde/Horde.Server/Users/IUserCollection.cs
+++ b/Engine/Source/Programs/Horde/Horde.Server/Users/IUserCollection.cs
@@ -8,6 +8,7 @@ using Horde.Server.Utilities;
 using MongoDB.Bson;
 using Horde.Server.Jobs.Bisect;
 using EpicGames.Horde.Users;
+using EpicGames.Horde.Jobs;
 
 namespace Horde.Server.Users
 {
diff --git a/Engine/Source/Programs/Horde/Horde.Server/Users/UserCollection.cs b/Engine/Source/Programs/Horde/Horde.Server/Users/UserCollection.cs
--- a/Engine/Source/Programs/Horde/Horde.Server/Users/UserCollection.cs
+++ b/Engine/Source/Programs/Horde/Horde.Server/Users/UserCollection.cs
@@ -14,6 +14,7 @@ using MongoDB.Bson.Serialization.Attributes;
 using MongoDB.Driver;
 using Horde.Server.Jobs.Bisect;
 using EpicGames.Horde.Users;
+using EpicGames.Horde.Jobs;
 
 namespace Horde.Server.Users
 {
diff --git a/Engine/Source/Programs/Horde/Horde.Server/Users/UserCollectionV2.cs b/Engine/Source/Programs/Horde/Horde.Server/Users/UserCollectionV2.cs
--- a/Engine/Source/Programs/Horde/Horde.Server/Users/UserCollectionV2.cs
+++ b/Engine/Source/Programs/Horde/Horde.Server/Users/UserCollectionV2.cs
@@ -18,6 +18,7 @@ using EpicGames.Horde.Jobs.Templates;
 using EpicGames.Horde.Streams;
 using EpicGames.Horde.Users;
 using MongoDB.Driver.Linq;
+using EpicGames.Horde.Jobs;
 
 namespace Horde.Server.Users
 {
diff --git a/Engine/Source/Programs/Horde/Horde.Server/Users/UserController.cs b/Engine/Source/Programs/Horde/Horde.Server/Users/UserController.cs
--- a/Engine/Source/Programs/Horde/Horde.Server/Users/UserController.cs
+++ b/Engine/Source/Programs/Horde/Horde.Server/Users/UserController.cs
@@ -15,6 +15,7 @@ using Microsoft.AspNetCore.Mvc;
 using Microsoft.Extensions.Options;
 using EpicGames.Horde.Users;
 using Horde.Server.Jobs.Bisect;
+using EpicGames.Horde.Jobs;
 
 namespace Horde.Server.Users
 {
diff --git a/Engine/Source/Programs/Horde/Horde.Server/Utilities/HordeControllerBase.cs b/Engine/Source/Programs/Horde/Horde.Server/Utilities/HordeControllerBase.cs
--- a/Engine/Source/Programs/Horde/Horde.Server/Utilities/HordeControllerBase.cs
+++ b/Engine/Source/Programs/Horde/Horde.Server/Utilities/HordeControllerBase.cs
@@ -5,6 +5,7 @@ using EpicGames.Horde.Agents;
 using EpicGames.Horde.Agents.Leases;
 using EpicGames.Horde.Artifacts;
 using EpicGames.Horde.Compute;
+using EpicGames.Horde.Jobs;
 using EpicGames.Horde.Jobs.Templates;
 using EpicGames.Horde.Projects;
 using EpicGames.Horde.Secrets;
diff --git a/Engine/Source/Programs/Horde/Horde.Server/Utilities/OpenTelemetry.cs b/Engine/Source/Programs/Horde/Horde.Server/Utilities/OpenTelemetry.cs
--- a/Engine/Source/Programs/Horde/Horde.Server/Utilities/OpenTelemetry.cs
+++ b/Engine/Source/Programs/Horde/Horde.Server/Utilities/OpenTelemetry.cs
@@ -8,6 +8,7 @@ using System.Linq;
 using System.Net.Http;
 using System.Runtime.CompilerServices;
 using EpicGames.Core;
+using EpicGames.Horde.Jobs;
 using EpicGames.Horde.Jobs.Templates;
 using EpicGames.Horde.Streams;
 using Horde.Server.Jobs;
diff --git a/Engine/Source/Programs/Shared/EpicGames.Horde/Jobs/JobId.cs b/Engine/Source/Programs/Shared/EpicGames.Horde/Jobs/JobId.cs
new file mode 100644
--- /dev/null
+++ b/Engine/Source/Programs/Shared/EpicGames.Horde/Jobs/JobId.cs
@@ -0,0 +1,41 @@
+// Copyright Epic Games, Inc. All Rights Reserved.
+
+using System.ComponentModel;
+using EpicGames.Core;
+using EpicGames.Horde;
+
+namespace EpicGames.Horde.Jobs
+{
+	/// <summary>
+	/// Identifier for a job
+	/// </summary>
+	/// <param name="Id">Id to construct from</param>
+	[JsonSchemaString]
+	[TypeConverter(typeof(BinaryIdTypeConverter<JobId, JobIdConverter>))]
+	[BinaryIdConverter(typeof(JobIdConverter))]
+	public record struct JobId(BinaryId Id)
+	{
+		/// <summary>
+		/// Constant value for an empty job id
+		/// </summary>
+		public static JobId Empty { get; } = default;
+
+		/// <inheritdoc cref="BinaryId.Parse(System.String)"/>
+		public static JobId Parse(string text) => new JobId(BinaryId.Parse(text));
+
+		/// <inheritdoc/>
+		public override string ToString() => Id.ToString();
+	}
+
+	/// <summary>
+	/// Converter to and from <see cref="BinaryId"/> instances.
+	/// </summary>
+	class JobIdConverter : BinaryIdConverter<JobId>
+	{
+		/// <inheritdoc/>
+		public override JobId FromBinaryId(BinaryId id) => new JobId(id);
+
+		/// <inheritdoc/>
+		public override BinaryId ToBinaryId(JobId value) => value.Id;
+	}
+}
diff --git a/Engine/Source/Programs/Shared/EpicGames.Horde/Jobs/JobStepBatchId.cs b/Engine/Source/Programs/Shared/EpicGames.Horde/Jobs/JobStepBatchId.cs
new file mode 100644
--- /dev/null
+++ b/Engine/Source/Programs/Shared/EpicGames.Horde/Jobs/JobStepBatchId.cs
@@ -0,0 +1,46 @@
+// Copyright Epic Games, Inc. All Rights Reserved.
+
+using System.ComponentModel;
+using EpicGames.Core;
+
+namespace EpicGames.Horde.Jobs
+{
+	/// <summary>
+	/// Identifier for a job step batch
+	/// </summary>
+	/// <param name="SubResourceId">Id to construct from</param>
+	[JsonSchemaString]
+	[TypeConverter(typeof(SubResourceIdTypeConverter<JobStepBatchId, JobStepBatchIdConverter>))]
+	[SubResourceIdConverter(typeof(JobStepBatchIdConverter))]
+	public record struct JobStepBatchId(SubResourceId SubResourceId)
+	{
+		/// <summary>
+		/// Constructor
+		/// </summary>
+		public JobStepBatchId(ushort value) : this(new SubResourceId(value))
+		{ }
+
+		/// <summary>
+		/// Creates a new <see cref="JobId"/>
+		/// </summary>
+		public static JobStepBatchId GenerateNewId() => new JobStepBatchId(SubResourceId.GenerateNewId());
+
+		/// <inheritdoc cref="SubResourceId.Parse(System.String)"/>
+		public static JobStepBatchId Parse(string text) => new JobStepBatchId(SubResourceId.Parse(text));
+
+		/// <inheritdoc/>
+		public override string ToString() => SubResourceId.ToString();
+	}
+
+	/// <summary>
+	/// Converter to and from <see cref="SubResourceId"/> instances.
+	/// </summary>
+	class JobStepBatchIdConverter : SubResourceIdConverter<JobStepBatchId>
+	{
+		/// <inheritdoc/>
+		public override JobStepBatchId FromSubResourceId(SubResourceId subResourceId) => new JobStepBatchId(subResourceId);
+
+		/// <inheritdoc/>
+		public override SubResourceId ToSubResourceId(JobStepBatchId value) => value.SubResourceId;
+	}
+}
diff --git a/Engine/Source/Programs/Shared/EpicGames.Horde/Jobs/JobStepId.cs b/Engine/Source/Programs/Shared/EpicGames.Horde/Jobs/JobStepId.cs
new file mode 100644
--- /dev/null
+++ b/Engine/Source/Programs/Shared/EpicGames.Horde/Jobs/JobStepId.cs
@@ -0,0 +1,46 @@
+// Copyright Epic Games, Inc. All Rights Reserved.
+
+using System.ComponentModel;
+using EpicGames.Core;
+
+namespace EpicGames.Horde.Jobs
+{
+	/// <summary>
+	/// Identifier for a jobstep
+	/// </summary>
+	/// <param name="Id">Id to construct from</param>
+	[JsonSchemaString]
+	[TypeConverter(typeof(SubResourceIdTypeConverter<JobStepId, JobStepIdConverter>))]
+	[SubResourceIdConverter(typeof(JobStepIdConverter))]
+	public record struct JobStepId(SubResourceId Id)
+	{
+		/// <summary>
+		/// Constructor
+		/// </summary>
+		public JobStepId(ushort value) : this(new SubResourceId(value))
+		{ }
+
+		/// <summary>
+		/// Creates a new <see cref="JobId"/>
+		/// </summary>
+		public static JobStepId GenerateNewId() => new JobStepId(SubResourceId.GenerateNewId());
+
+		/// <inheritdoc cref="SubResourceId.Parse(System.String)"/>
+		public static JobStepId Parse(string text) => new JobStepId(SubResourceId.Parse(text));
+
+		/// <inheritdoc/>
+		public override string ToString() => Id.ToString();
+	}
+
+	/// <summary>
+	/// Converter to and from <see cref="SubResourceId"/> instances.
+	/// </summary>
+	class JobStepIdConverter : SubResourceIdConverter<JobStepId>
+	{
+		/// <inheritdoc/>
+		public override JobStepId FromSubResourceId(SubResourceId id) => new JobStepId(id);
+
+		/// <inheritdoc/>
+		public override SubResourceId ToSubResourceId(JobStepId value) => value.Id;
+	}
+}
diff --git a/Engine/Source/Programs/Shared/EpicGames.Horde/Jobs/SubResourceId.cs b/Engine/Source/Programs/Shared/EpicGames.Horde/Jobs/SubResourceId.cs
new file mode 100644
--- /dev/null
+++ b/Engine/Source/Programs/Shared/EpicGames.Horde/Jobs/SubResourceId.cs
@@ -0,0 +1,164 @@
+// Copyright Epic Games, Inc. All Rights Reserved.
+
+using System;
+using System.ComponentModel;
+using System.Diagnostics;
+using System.Globalization;
+
+namespace EpicGames.Horde.Jobs
+{
+	/// <summary>
+	/// Identifier for subresources. Assigning unique ids to subresources prevents against race conditions using indices when subresources are added and removed.
+	/// 
+	/// Subresource identifiers are stored as 16-bit integers formatted as a 4-digit hex code, in order to keep URLs short. Calling Next() will generate a new
+	/// identifier with more entropy than just incrementing the value but an identical period before repeating, in order to make URL fragments more distinctive.
+	/// </summary>
+	[TypeConverter(typeof(SubResourceIdTypeConverter))]
+	public struct SubResourceId : IEquatable<SubResourceId>
+	{
+		/// <summary>
+		/// The unique identifier value
+		/// </summary>
+		public ushort Value { get; }
+
+		/// <summary>
+		/// Constructor
+		/// </summary>
+		/// <param name="value">New identifier for this subresource</param>
+		public SubResourceId(ushort value)
+		{
+			Value = value;
+		}
+
+		/// <summary>
+		/// Creates a new random subresource id. We use random numbers for this to increase distinctiveness.
+		/// </summary>
+		/// <returns>New subresource id</returns>
+		public static SubResourceId GenerateNewId()
+		{
+			return new SubResourceId((ushort)Stopwatch.GetTimestamp());
+		}
+
+		/// <summary>
+		/// Updates the current value, and returns a copy of the previous value.
+		/// </summary>
+		/// <returns>New subresource identifier</returns>
+		public SubResourceId Next()
+		{
+			// 771 is a factor of 65535, so we won't repeat when wrapping round 64k.
+			return new SubResourceId((ushort)((int)Value + 771));
+		}
+
+		/// <summary>
+		/// Parse a subresource id from a string
+		/// </summary>
+		/// <param name="text">Text to parse</param>
+		/// <returns>New subresource id</returns>
+		public static SubResourceId Parse(string text)
+		{
+			return new SubResourceId(UInt16.Parse(text, NumberStyles.HexNumber, CultureInfo.InvariantCulture));
+		}
+
+		/// <summary>
+		/// Converts this identifier to a string
+		/// </summary>
+		/// <returns>String representation of this id</returns>
+		public override string ToString()
+		{
+			return Value.ToString("x4", CultureInfo.InvariantCulture);
+		}
+
+		/// <inheritdoc/>
+		public override bool Equals(object? other)
+		{
+			return other is SubResourceId otherId && Equals(otherId);
+		}
+
+		/// <inheritdoc/>
+		public bool Equals(SubResourceId other)
+		{
+			return Value == other.Value;
+		}
+
+		/// <inheritdoc/>
+		public override int GetHashCode()
+		{
+			return (int)Value;
+		}
+
+		/// <summary>
+		/// Equality operator for identifiers
+		/// </summary>
+		/// <param name="left">First identifier to compare</param>
+		/// <param name="right">Second identifier to compare</param>
+		/// <returns>True if the identifiers are equal</returns>
+		public static bool operator ==(SubResourceId left, SubResourceId right)
+		{
+			return left.Value == right.Value;
+		}
+
+		/// <summary>
+		/// Inequality operator for identifiers
+		/// </summary>
+		/// <param name="left">First identifier to compare</param>
+		/// <param name="right">Second identifier to compare</param>
+		/// <returns>True if the identifiers are equal</returns>
+		public static bool operator !=(SubResourceId left, SubResourceId right)
+		{
+			return left.Value != right.Value;
+		}
+	}
+
+	/// <summary>
+	/// Extension methods for manipulating subresource ids
+	/// </summary>
+	public static class SubResourceIdExtensions
+	{
+		/// <summary>
+		/// Parse a string as a subresource identifier
+		/// </summary>
+		/// <param name="text">Text to parse</param>
+		/// <returns>The new subresource identifier</returns>
+		public static SubResourceId ToSubResourceId(this string text)
+		{
+			return new SubResourceId(UInt16.Parse(text, NumberStyles.HexNumber, CultureInfo.InvariantCulture));
+		}
+	}
+
+	/// <summary>
+	/// Type converter from strings to SubResourceId objects
+	/// </summary>
+	sealed class SubResourceIdTypeConverter : TypeConverter
+	{
+		/// <inheritdoc/>
+		public override bool CanConvertFrom(ITypeDescriptorContext? context, Type sourceType)
+		{
+			return sourceType == typeof(string);
+		}
+
+		/// <inheritdoc/>
+		public override object ConvertFrom(ITypeDescriptorContext? context, CultureInfo? culture, object value)
+		{
+			return SubResourceId.Parse((string)value);
+		}
+
+		/// <inheritdoc/>
+		public override bool CanConvertTo(ITypeDescriptorContext? context, Type? destinationType)
+		{
+			return destinationType == typeof(string);
+		}
+
+		/// <inheritdoc/>
+		public override object? ConvertTo(ITypeDescriptorContext? context, CultureInfo? culture, object? value, Type destinationType)
+		{
+			if (destinationType == typeof(string))
+			{
+				return value?.ToString();
+			}
+			else
+			{
+				return null;
+			}
+		}
+	}
+}
diff --git a/Engine/Source/Programs/Shared/EpicGames.Horde/Jobs/SubResourceIdConverter.cs b/Engine/Source/Programs/Shared/EpicGames.Horde/Jobs/SubResourceIdConverter.cs
new file mode 100644
--- /dev/null
+++ b/Engine/Source/Programs/Shared/EpicGames.Horde/Jobs/SubResourceIdConverter.cs
@@ -0,0 +1,124 @@
+// Copyright Epic Games, Inc. All Rights Reserved.
+
+using System;
+using System.ComponentModel;
+using System.Globalization;
+using System.Reflection;
+using System.Text.Json;
+using System.Text.Json.Serialization;
+
+namespace EpicGames.Horde.Jobs
+{
+	/// <summary>
+	/// Base class for converting to and from types containing a <see cref="SubResourceId"/>. Useful pattern for reducing boilerplate with strongly typed records.
+	/// </summary>
+	/// <typeparam name="T"></typeparam>
+	public abstract class SubResourceIdConverter<T> where T : struct
+	{
+		/// <summary>
+		/// Converts a type to a <see cref="SubResourceId"/>
+		/// </summary>
+		public abstract SubResourceId ToSubResourceId(T value);
+
+		/// <summary>
+		/// Constructs a type from a <see cref="SubResourceId"/>
+		/// </summary>
+		public abstract T FromSubResourceId(SubResourceId id);
+	}
+
+	/// <summary>
+	/// Attribute declaring a <see cref="SubResourceIdConverter{T}"/> for a particular type
+	/// </summary>
+	[AttributeUsage(AttributeTargets.Struct)]
+	public sealed class SubResourceIdConverterAttribute : Attribute
+	{
+		/// <summary>
+		/// The converter type
+		/// </summary>
+		public Type ConverterType { get; }
+
+		/// <summary>
+		/// Constructor
+		/// </summary>
+		public SubResourceIdConverterAttribute(Type converterType) => ConverterType = converterType;
+	}
+
+	/// <summary>
+	/// Class which serializes types with a <see cref="SubResourceIdConverter{T}"/> to Json
+	/// </summary>
+	public sealed class SubResourceIdTypeConverter<TValue, TConverter> : TypeConverter where TValue : struct where TConverter : SubResourceIdConverter<TValue>, new()
+	{
+		readonly TConverter _converter = new TConverter();
+
+		/// <inheritdoc/>
+		public override bool CanConvertFrom(ITypeDescriptorContext? context, Type sourceType)
+		{
+			return sourceType == typeof(string) || sourceType == typeof(SubResourceId);
+		}
+
+		/// <inheritdoc/>
+		public override object? ConvertFrom(ITypeDescriptorContext? context, CultureInfo? culture, object value)
+		{
+			if (value is string str)
+			{
+				return _converter.FromSubResourceId(SubResourceId.Parse(str));
+			}
+			if (value is SubResourceId stringId)
+			{
+				return _converter.FromSubResourceId(stringId);
+			}
+			return null;
+		}
+
+		/// <inheritdoc/>
+		public override bool CanConvertTo(ITypeDescriptorContext? context, Type? destinationType) => destinationType == typeof(string) || destinationType == typeof(SubResourceId);
+
+		/// <inheritdoc/>
+		public override object? ConvertTo(ITypeDescriptorContext? context, CultureInfo? culture, object? value, Type destinationType)
+		{
+			if (destinationType == typeof(string))
+			{
+				return _converter.ToSubResourceId((TValue)value!).ToString();
+			}
+			if (destinationType == typeof(SubResourceId))
+			{
+				return _converter.ToSubResourceId((TValue)value!);
+			}
+			return null;
+		}
+	}
+
+	/// <summary>
+	/// Class which serializes types with a <see cref="SubResourceIdConverter{T}"/> to Json
+	/// </summary>
+	public sealed class SubResourceIdJsonConverter<TValue, TConverter> : JsonConverter<TValue> where TValue : struct where TConverter : SubResourceIdConverter<TValue>, new()
+	{
+		readonly TConverter _converter = new TConverter();
+
+		/// <inheritdoc/>
+		public override TValue Read(ref Utf8JsonReader reader, Type typeToConvert, JsonSerializerOptions options) => _converter.FromSubResourceId(SubResourceId.Parse(reader.GetString() ?? String.Empty));
+
+		/// <inheritdoc/>
+		public override void Write(Utf8JsonWriter writer, TValue value, JsonSerializerOptions options) => writer.WriteStringValue(_converter.ToSubResourceId(value).ToString());
+	}
+
+	/// <summary>
+	/// Creates constructors for types with a <see cref="SubResourceIdConverter{T}"/> to Json
+	/// </summary>
+	public sealed class SubResourceIdJsonConverterFactory : JsonConverterFactory
+	{
+		/// <inheritdoc/>
+		public override bool CanConvert(Type typeToConvert) => typeToConvert.GetCustomAttribute<SubResourceIdConverterAttribute>() != null;
+
+		/// <inheritdoc/>
+		public override JsonConverter? CreateConverter(Type type, JsonSerializerOptions options)
+		{
+			SubResourceIdConverterAttribute? attribute = type.GetCustomAttribute<SubResourceIdConverterAttribute>();
+			if (attribute == null)
+			{
+				return null;
+			}
+			return (JsonConverter?)Activator.CreateInstance(typeof(SubResourceIdJsonConverter<,>).MakeGenericType(type, attribute.ConverterType));
+		}
+	}
+}
