Horde: Minor static analysis fixes.

diff --git a/Engine/Source/Programs/Shared/EpicGames.Horde/BinaryId.cs b/Engine/Source/Programs/Shared/EpicGames.Horde/BinaryId.cs
--- a/Engine/Source/Programs/Shared/EpicGames.Horde/BinaryId.cs
+++ b/Engine/Source/Programs/Shared/EpicGames.Horde/BinaryId.cs
@@ -208,6 +208,13 @@ namespace EpicGames.Horde
 		/// <param name="right">Second string id</param>
 		/// <returns>True if the two string ids are not equal</returns>
 		public static bool operator !=(BinaryId left, BinaryId right) => !left.Equals(right);
+
+#pragma warning disable CS1591 // Missing XML comment for publicly visible type or member
+		public static bool operator <(BinaryId left, BinaryId right) => left.CompareTo(right) < 0;
+		public static bool operator <=(BinaryId left, BinaryId right) => left.CompareTo(right) <= 0;
+		public static bool operator >(BinaryId left, BinaryId right) => left.CompareTo(right) > 0;
+		public static bool operator >=(BinaryId left, BinaryId right) => left.CompareTo(right) >= 0;
+#pragma warning restore CS1591 // Missing XML comment for publicly visible type or member
 	}
 
 	/// <summary>
diff --git a/Engine/Source/Programs/Shared/EpicGames.Horde/Jobs/JobId.cs b/Engine/Source/Programs/Shared/EpicGames.Horde/Jobs/JobId.cs
--- a/Engine/Source/Programs/Shared/EpicGames.Horde/Jobs/JobId.cs
+++ b/Engine/Source/Programs/Shared/EpicGames.Horde/Jobs/JobId.cs
@@ -2,7 +2,6 @@
 
 using System.ComponentModel;
 using EpicGames.Core;
-using EpicGames.Horde;
 
 namespace EpicGames.Horde.Jobs
 {
