Horde: Fix BinaryId CompareTo implementation. (#11207)

diff --git a/Engine/Source/Programs/Shared/EpicGames.Horde/BinaryId.cs b/Engine/Source/Programs/Shared/EpicGames.Horde/BinaryId.cs
--- a/Engine/Source/Programs/Shared/EpicGames.Horde/BinaryId.cs
+++ b/Engine/Source/Programs/Shared/EpicGames.Horde/BinaryId.cs
@@ -184,10 +184,10 @@ namespace EpicGames.Horde
 			int result = _a.CompareTo(other._a);
 			if (result == 0)
 			{
-				result = _b.CompareTo(_b);
+				result = _b.CompareTo(other._b);
 				if (result == 0)
 				{
-					result = _c.CompareTo(_c);
+					result = _c.CompareTo(other._c);
 				}
 			}
 			return result;
