Horde: Move PoolId into the EpicGames.Horde assembly.

diff --git a/Engine/Source/Programs/Horde/Horde.Server.Tests/AgentServiceTest.cs b/Engine/Source/Programs/Horde/Horde.Server.Tests/AgentServiceTest.cs
--- a/Engine/Source/Programs/Horde/Horde.Server.Tests/AgentServiceTest.cs
+++ b/Engine/Source/Programs/Horde/Horde.Server.Tests/AgentServiceTest.cs
@@ -8,8 +8,8 @@ using System.Threading;
 using System.Threading.Tasks;
 using EpicGames.Core;
 using EpicGames.Horde.Agents;
+using EpicGames.Horde.Agents.Pools;
 using Horde.Server.Agents;
-using Horde.Server.Agents.Pools;
 using Horde.Server.Auditing;
 using Horde.Server.Jobs;
 using Horde.Server.Server;
diff --git a/Engine/Source/Programs/Horde/Horde.Server.Tests/Agents/AgentCollectionTests.cs b/Engine/Source/Programs/Horde/Horde.Server.Tests/Agents/AgentCollectionTests.cs
--- a/Engine/Source/Programs/Horde/Horde.Server.Tests/Agents/AgentCollectionTests.cs
+++ b/Engine/Source/Programs/Horde/Horde.Server.Tests/Agents/AgentCollectionTests.cs
@@ -4,11 +4,11 @@ using System;
 using System.Collections.Generic;
 using System.Threading.Tasks;
 using Horde.Server.Agents;
-using Horde.Server.Agents.Pools;
 using Horde.Server.Agents.Sessions;
 using HordeCommon;
 using Microsoft.VisualStudio.TestTools.UnitTesting;
 using EpicGames.Horde.Agents.Leases;
+using EpicGames.Horde.Agents.Pools;
 
 namespace Horde.Server.Tests.Agents;
 
diff --git a/Engine/Source/Programs/Horde/Horde.Server.Tests/Compute/ComputeServiceTest.cs b/Engine/Source/Programs/Horde/Horde.Server.Tests/Compute/ComputeServiceTest.cs
--- a/Engine/Source/Programs/Horde/Horde.Server.Tests/Compute/ComputeServiceTest.cs
+++ b/Engine/Source/Programs/Horde/Horde.Server.Tests/Compute/ComputeServiceTest.cs
@@ -7,9 +7,9 @@ using System.Linq;
 using System.Net;
 using System.Threading;
 using System.Threading.Tasks;
+using EpicGames.Horde.Agents.Pools;
 using EpicGames.Horde.Compute;
 using Horde.Server.Agents;
-using Horde.Server.Agents.Pools;
 using Horde.Server.Compute;
 using Horde.Server.Server;
 using Microsoft.Extensions.Logging.Abstractions;
diff --git a/Engine/Source/Programs/Horde/Horde.Server.Tests/Fixture.cs b/Engine/Source/Programs/Horde/Horde.Server.Tests/Fixture.cs
--- a/Engine/Source/Programs/Horde/Horde.Server.Tests/Fixture.cs
+++ b/Engine/Source/Programs/Horde/Horde.Server.Tests/Fixture.cs
@@ -5,12 +5,12 @@ using System.IO;
 using System.Text;
 using System.Threading.Tasks;
 using EpicGames.Core;
+using EpicGames.Horde.Agents.Pools;
 using EpicGames.Horde.Jobs;
 using EpicGames.Horde.Jobs.Templates;
 using EpicGames.Horde.Projects;
 using EpicGames.Horde.Streams;
 using Horde.Server.Agents;
-using Horde.Server.Agents.Pools;
 using Horde.Server.Configuration;
 using Horde.Server.Jobs;
 using Horde.Server.Jobs.Artifacts;
diff --git a/Engine/Source/Programs/Horde/Horde.Server.Tests/Fleet/AwsAutoScalingLifecycleTest.cs b/Engine/Source/Programs/Horde/Horde.Server.Tests/Fleet/AwsAutoScalingLifecycleTest.cs
--- a/Engine/Source/Programs/Horde/Horde.Server.Tests/Fleet/AwsAutoScalingLifecycleTest.cs
+++ b/Engine/Source/Programs/Horde/Horde.Server.Tests/Fleet/AwsAutoScalingLifecycleTest.cs
@@ -11,7 +11,6 @@ using Amazon.AutoScaling;
 using Amazon.AutoScaling.Model;
 using Horde.Server.Agents;
 using Horde.Server.Agents.Fleet;
-using Horde.Server.Agents.Pools;
 using Horde.Server.Utilities;
 using HordeCommon;
 using Microsoft.Extensions.DependencyInjection;
@@ -19,6 +18,7 @@ using Microsoft.Extensions.Logging;
 using Microsoft.VisualStudio.TestTools.UnitTesting;
 using Moq;
 using EpicGames.Horde.Agents.Leases;
+using EpicGames.Horde.Agents.Pools;
 
 namespace Horde.Server.Tests.Fleet;
 
diff --git a/Engine/Source/Programs/Horde/Horde.Server.Tests/Fleet/JobQueueStrategyTest.cs b/Engine/Source/Programs/Horde/Horde.Server.Tests/Fleet/JobQueueStrategyTest.cs
--- a/Engine/Source/Programs/Horde/Horde.Server.Tests/Fleet/JobQueueStrategyTest.cs
+++ b/Engine/Source/Programs/Horde/Horde.Server.Tests/Fleet/JobQueueStrategyTest.cs
@@ -16,6 +16,7 @@ using Horde.Server.Agents;
 using EpicGames.Horde.Jobs.Templates;
 using EpicGames.Horde.Projects;
 using EpicGames.Horde.Streams;
+using EpicGames.Horde.Agents.Pools;
 
 namespace Horde.Server.Tests.Fleet
 {
diff --git a/Engine/Source/Programs/Horde/Horde.Server.Tests/Fleet/LeaseUtilizationAwsMetricStrategyTests.cs b/Engine/Source/Programs/Horde/Horde.Server.Tests/Fleet/LeaseUtilizationAwsMetricStrategyTests.cs
--- a/Engine/Source/Programs/Horde/Horde.Server.Tests/Fleet/LeaseUtilizationAwsMetricStrategyTests.cs
+++ b/Engine/Source/Programs/Horde/Horde.Server.Tests/Fleet/LeaseUtilizationAwsMetricStrategyTests.cs
@@ -18,6 +18,7 @@ using Microsoft.VisualStudio.TestTools.UnitTesting;
 using EpicGames.Horde.Agents.Leases;
 using EpicGames.Horde.Streams;
 using Microsoft.Extensions.Logging.Abstractions;
+using EpicGames.Horde.Agents.Pools;
 
 namespace Horde.Server.Tests.Fleet
 {
diff --git a/Engine/Source/Programs/Horde/Horde.Server.Tests/JobCollectionTests.cs b/Engine/Source/Programs/Horde/Horde.Server.Tests/JobCollectionTests.cs
--- a/Engine/Source/Programs/Horde/Horde.Server.Tests/JobCollectionTests.cs
+++ b/Engine/Source/Programs/Horde/Horde.Server.Tests/JobCollectionTests.cs
@@ -12,13 +12,13 @@ using EpicGames.Horde.Users;
 using Horde.Server.Jobs;
 using Horde.Server.Jobs.Graphs;
 using Horde.Server.Logs;
-using Horde.Server.Agents.Pools;
 using Horde.Server.Jobs.Templates;
 using Horde.Server.Utilities;
 using HordeCommon;
 using Microsoft.VisualStudio.TestTools.UnitTesting;
 using Moq;
 using Horde.Server.Agents.Sessions;
+using EpicGames.Horde.Agents.Pools;
 
 namespace Horde.Server.Tests
 {
diff --git a/Engine/Source/Programs/Horde/Horde.Server.Tests/JobServiceTests.cs b/Engine/Source/Programs/Horde/Horde.Server.Tests/JobServiceTests.cs
--- a/Engine/Source/Programs/Horde/Horde.Server.Tests/JobServiceTests.cs
+++ b/Engine/Source/Programs/Horde/Horde.Server.Tests/JobServiceTests.cs
@@ -6,10 +6,10 @@ using System.Text;
 using System.Threading;
 using System.Threading.Tasks;
 using EpicGames.Core;
+using EpicGames.Horde.Agents.Pools;
 using EpicGames.Horde.Jobs.Templates;
 using EpicGames.Horde.Streams;
 using Horde.Server.Agents;
-using Horde.Server.Agents.Pools;
 using Horde.Server.Jobs;
 using Horde.Server.Jobs.Graphs;
 using Horde.Server.Jobs.Templates;
diff --git a/Engine/Source/Programs/Horde/Horde.Server.Tests/JobTaskSourceTests.cs b/Engine/Source/Programs/Horde/Horde.Server.Tests/JobTaskSourceTests.cs
--- a/Engine/Source/Programs/Horde/Horde.Server.Tests/JobTaskSourceTests.cs
+++ b/Engine/Source/Programs/Horde/Horde.Server.Tests/JobTaskSourceTests.cs
@@ -14,6 +14,7 @@ using Horde.Server.Agents;
 using Horde.Server.Agents.Pools;
 using Horde.Server.Streams;
 using EpicGames.Horde.Users;
+using EpicGames.Horde.Agents.Pools;
 
 namespace Horde.Server.Tests
 {
diff --git a/Engine/Source/Programs/Horde/Horde.Server.Tests/PoolServiceTests.cs b/Engine/Source/Programs/Horde/Horde.Server.Tests/PoolServiceTests.cs
--- a/Engine/Source/Programs/Horde/Horde.Server.Tests/PoolServiceTests.cs
+++ b/Engine/Source/Programs/Horde/Horde.Server.Tests/PoolServiceTests.cs
@@ -3,6 +3,7 @@
 using System;
 using System.Collections.Generic;
 using System.Threading.Tasks;
+using EpicGames.Horde.Agents.Pools;
 using Horde.Server.Agents.Pools;
 using Microsoft.VisualStudio.TestTools.UnitTesting;
 
diff --git a/Engine/Source/Programs/Horde/Horde.Server.Tests/PoolsControllerTest.cs b/Engine/Source/Programs/Horde/Horde.Server.Tests/PoolsControllerTest.cs
--- a/Engine/Source/Programs/Horde/Horde.Server.Tests/PoolsControllerTest.cs
+++ b/Engine/Source/Programs/Horde/Horde.Server.Tests/PoolsControllerTest.cs
@@ -2,6 +2,7 @@
 
 using System.Collections.Generic;
 using System.Threading.Tasks;
+using EpicGames.Horde.Agents.Pools;
 using Horde.Server.Agents.Fleet;
 using Horde.Server.Agents.Pools;
 using Microsoft.AspNetCore.Mvc;
diff --git a/Engine/Source/Programs/Horde/Horde.Server.Tests/RpcServiceTest.cs b/Engine/Source/Programs/Horde/Horde.Server.Tests/RpcServiceTest.cs
--- a/Engine/Source/Programs/Horde/Horde.Server.Tests/RpcServiceTest.cs
+++ b/Engine/Source/Programs/Horde/Horde.Server.Tests/RpcServiceTest.cs
@@ -12,7 +12,6 @@ using System.Threading.Tasks;
 using Google.Protobuf;
 using Grpc.Core;
 using Horde.Server.Utilities;
-using Horde.Server.Agents.Pools;
 using Horde.Server.Logs;
 using Horde.Server.Server;
 using HordeCommon.Rpc;
@@ -27,6 +26,7 @@ using HordeCommon.Rpc.Messages;
 using Horde.Server.Agents.Sessions;
 using EpicGames.Horde.Agents;
 using EpicGames.Horde.Agents.Leases;
+using EpicGames.Horde.Agents.Pools;
 
 namespace Horde.Server.Tests
 {
diff --git a/Engine/Source/Programs/Horde/Horde.Server.Tests/TestSetup.cs b/Engine/Source/Programs/Horde/Horde.Server.Tests/TestSetup.cs
--- a/Engine/Source/Programs/Horde/Horde.Server.Tests/TestSetup.cs
+++ b/Engine/Source/Programs/Horde/Horde.Server.Tests/TestSetup.cs
@@ -65,6 +65,7 @@ using Horde.Server.Dashboard;
 using OpenTelemetry.Trace;
 using Horde.Server.Jobs.Bisect;
 using Horde.Server.Telemetry.Sinks;
+using EpicGames.Horde.Agents.Pools;
 
 namespace Horde.Server.Tests
 {
diff --git a/Engine/Source/Programs/Horde/Horde.Server/Agents/AgentCollection.cs b/Engine/Source/Programs/Horde/Horde.Server/Agents/AgentCollection.cs
--- a/Engine/Source/Programs/Horde/Horde.Server/Agents/AgentCollection.cs
+++ b/Engine/Source/Programs/Horde/Horde.Server/Agents/AgentCollection.cs
@@ -7,9 +7,9 @@ using System.Linq.Expressions;
 using System.Threading.Tasks;
 using EpicGames.Horde.Agents;
 using EpicGames.Horde.Agents.Leases;
+using EpicGames.Horde.Agents.Pools;
 using EpicGames.Redis;
 using Google.Protobuf.WellKnownTypes;
-using Horde.Server.Agents.Pools;
 using Horde.Server.Agents.Sessions;
 using Horde.Server.Auditing;
 using Horde.Server.Server;
diff --git a/Engine/Source/Programs/Horde/Horde.Server/Agents/AgentService.cs b/Engine/Source/Programs/Horde/Horde.Server/Agents/AgentService.cs
--- a/Engine/Source/Programs/Horde/Horde.Server/Agents/AgentService.cs
+++ b/Engine/Source/Programs/Horde/Horde.Server/Agents/AgentService.cs
@@ -11,6 +11,7 @@ using System.Threading.Tasks;
 using EpicGames.Core;
 using EpicGames.Horde.Agents;
 using EpicGames.Horde.Agents.Leases;
+using EpicGames.Horde.Agents.Pools;
 using EpicGames.Redis;
 using EpicGames.Serialization;
 using Google.Protobuf.WellKnownTypes;
diff --git a/Engine/Source/Programs/Horde/Horde.Server/Agents/AgentsController.cs b/Engine/Source/Programs/Horde/Horde.Server/Agents/AgentsController.cs
--- a/Engine/Source/Programs/Horde/Horde.Server/Agents/AgentsController.cs
+++ b/Engine/Source/Programs/Horde/Horde.Server/Agents/AgentsController.cs
@@ -6,7 +6,6 @@ using System.Linq;
 using System.Threading.Tasks;
 using Horde.Server.Acls;
 using Horde.Server.Agents.Leases;
-using Horde.Server.Agents.Pools;
 using Horde.Server.Agents.Sessions;
 using Horde.Server.Auditing;
 using Horde.Server.Server;
@@ -19,6 +18,7 @@ using EpicGames.Core;
 using EpicGames.Horde.Agents;
 using EpicGames.Horde.Agents.Leases;
 using EpicGames.Horde.Common;
+using EpicGames.Horde.Agents.Pools;
 
 namespace Horde.Server.Agents
 {
diff --git a/Engine/Source/Programs/Horde/Horde.Server/Agents/Fleet/FleetService.cs b/Engine/Source/Programs/Horde/Horde.Server/Agents/Fleet/FleetService.cs
--- a/Engine/Source/Programs/Horde/Horde.Server/Agents/Fleet/FleetService.cs
+++ b/Engine/Source/Programs/Horde/Horde.Server/Agents/Fleet/FleetService.cs
@@ -8,6 +8,7 @@ using System.Linq;
 using System.Text.Json;
 using System.Threading;
 using System.Threading.Tasks;
+using EpicGames.Horde.Agents.Pools;
 using Horde.Server.Agents.Leases;
 using Horde.Server.Agents.Pools;
 using Horde.Server.Jobs;
diff --git a/Engine/Source/Programs/Horde/Horde.Server/Agents/Fleet/JobQueueStrategy.cs b/Engine/Source/Programs/Horde/Horde.Server/Agents/Fleet/JobQueueStrategy.cs
--- a/Engine/Source/Programs/Horde/Horde.Server/Agents/Fleet/JobQueueStrategy.cs
+++ b/Engine/Source/Programs/Horde/Horde.Server/Agents/Fleet/JobQueueStrategy.cs
@@ -17,6 +17,7 @@ using Microsoft.Extensions.Caching.Memory;
 using Microsoft.Extensions.Options;
 using OpenTelemetry.Trace;
 using EpicGames.Horde.Streams;
+using EpicGames.Horde.Agents.Pools;
 
 namespace Horde.Server.Agents.Fleet
 {
diff --git a/Engine/Source/Programs/Horde/Horde.Server/Agents/Fleet/LeaseUtilizationStrategy.cs b/Engine/Source/Programs/Horde/Horde.Server/Agents/Fleet/LeaseUtilizationStrategy.cs
--- a/Engine/Source/Programs/Horde/Horde.Server/Agents/Fleet/LeaseUtilizationStrategy.cs
+++ b/Engine/Source/Programs/Horde/Horde.Server/Agents/Fleet/LeaseUtilizationStrategy.cs
@@ -14,6 +14,7 @@ using HordeCommon.Rpc.Tasks;
 using Microsoft.Extensions.Caching.Memory;
 using OpenTelemetry.Trace;
 using EpicGames.Horde.Agents;
+using EpicGames.Horde.Agents.Pools;
 
 namespace Horde.Server.Agents.Fleet
 {
diff --git a/Engine/Source/Programs/Horde/Horde.Server/Agents/Fleet/Providers/AwsAsgFleetManager.cs b/Engine/Source/Programs/Horde/Horde.Server/Agents/Fleet/Providers/AwsAsgFleetManager.cs
--- a/Engine/Source/Programs/Horde/Horde.Server/Agents/Fleet/Providers/AwsAsgFleetManager.cs
+++ b/Engine/Source/Programs/Horde/Horde.Server/Agents/Fleet/Providers/AwsAsgFleetManager.cs
@@ -8,6 +8,7 @@ using System.Threading;
 using System.Threading.Tasks;
 using Amazon.AutoScaling;
 using Amazon.AutoScaling.Model;
+using EpicGames.Horde.Agents.Pools;
 using Horde.Server.Agents.Pools;
 using Microsoft.Extensions.Logging;
 using OpenTelemetry.Trace;
diff --git a/Engine/Source/Programs/Horde/Horde.Server/Agents/IAgent.cs b/Engine/Source/Programs/Horde/Horde.Server/Agents/IAgent.cs
--- a/Engine/Source/Programs/Horde/Horde.Server/Agents/IAgent.cs
+++ b/Engine/Source/Programs/Horde/Horde.Server/Agents/IAgent.cs
@@ -9,6 +9,7 @@ using System.Threading.Tasks;
 using EpicGames.Core;
 using EpicGames.Horde.Agents;
 using EpicGames.Horde.Agents.Leases;
+using EpicGames.Horde.Agents.Pools;
 using EpicGames.Horde.Common;
 using EpicGames.Horde.Compute;
 using EpicGames.Horde.Logs;
diff --git a/Engine/Source/Programs/Horde/Horde.Server/Agents/IAgentCollection.cs b/Engine/Source/Programs/Horde/Horde.Server/Agents/IAgentCollection.cs
--- a/Engine/Source/Programs/Horde/Horde.Server/Agents/IAgentCollection.cs
+++ b/Engine/Source/Programs/Horde/Horde.Server/Agents/IAgentCollection.cs
@@ -3,12 +3,12 @@
 using System;
 using System.Collections.Generic;
 using System.Threading.Tasks;
-using Horde.Server.Agents.Pools;
 using Horde.Server.Agents.Sessions;
 using Horde.Server.Auditing;
 using HordeCommon;
 using EpicGames.Horde.Agents.Leases;
 using EpicGames.Horde.Agents;
+using EpicGames.Horde.Agents.Pools;
 
 namespace Horde.Server.Agents
 {
diff --git a/Engine/Source/Programs/Horde/Horde.Server/Agents/Leases/ILease.cs b/Engine/Source/Programs/Horde/Horde.Server/Agents/Leases/ILease.cs
--- a/Engine/Source/Programs/Horde/Horde.Server/Agents/Leases/ILease.cs
+++ b/Engine/Source/Programs/Horde/Horde.Server/Agents/Leases/ILease.cs
@@ -3,13 +3,13 @@
 using System;
 using Google.Protobuf;
 using Google.Protobuf.WellKnownTypes;
-using Horde.Server.Agents.Pools;
 using Horde.Server.Agents.Sessions;
 using HordeCommon;
 using EpicGames.Horde.Streams;
 using EpicGames.Horde.Agents;
 using EpicGames.Horde.Agents.Leases;
 using EpicGames.Horde.Logs;
+using EpicGames.Horde.Agents.Pools;
 
 namespace Horde.Server.Agents.Leases
 {
diff --git a/Engine/Source/Programs/Horde/Horde.Server/Agents/Leases/ILeaseCollection.cs b/Engine/Source/Programs/Horde/Horde.Server/Agents/Leases/ILeaseCollection.cs
--- a/Engine/Source/Programs/Horde/Horde.Server/Agents/Leases/ILeaseCollection.cs
+++ b/Engine/Source/Programs/Horde/Horde.Server/Agents/Leases/ILeaseCollection.cs
@@ -3,13 +3,13 @@
 using System;
 using System.Collections.Generic;
 using System.Threading.Tasks;
-using Horde.Server.Agents.Pools;
 using Horde.Server.Agents.Sessions;
 using HordeCommon;
 using EpicGames.Horde.Agents;
 using EpicGames.Horde.Agents.Leases;
 using EpicGames.Horde.Streams;
 using EpicGames.Horde.Logs;
+using EpicGames.Horde.Agents.Pools;
 
 namespace Horde.Server.Agents.Leases
 {
diff --git a/Engine/Source/Programs/Horde/Horde.Server/Agents/Leases/LeaseCollection.cs b/Engine/Source/Programs/Horde/Horde.Server/Agents/Leases/LeaseCollection.cs
--- a/Engine/Source/Programs/Horde/Horde.Server/Agents/Leases/LeaseCollection.cs
+++ b/Engine/Source/Programs/Horde/Horde.Server/Agents/Leases/LeaseCollection.cs
@@ -4,7 +4,6 @@ using System;
 using System.Collections.Generic;
 using System.Linq;
 using System.Threading.Tasks;
-using Horde.Server.Agents.Pools;
 using Horde.Server.Agents.Sessions;
 using Horde.Server.Server;
 using Horde.Server.Utilities;
@@ -15,6 +14,7 @@ using EpicGames.Horde.Agents.Leases;
 using EpicGames.Horde.Agents;
 using EpicGames.Horde.Streams;
 using EpicGames.Horde.Logs;
+using EpicGames.Horde.Agents.Pools;
 
 namespace Horde.Server.Agents.Leases
 {
diff --git a/Engine/Source/Programs/Horde/Horde.Server/Agents/Pools/IPoolCollection.cs b/Engine/Source/Programs/Horde/Horde.Server/Agents/Pools/IPoolCollection.cs
--- a/Engine/Source/Programs/Horde/Horde.Server/Agents/Pools/IPoolCollection.cs
+++ b/Engine/Source/Programs/Horde/Horde.Server/Agents/Pools/IPoolCollection.cs
@@ -3,6 +3,7 @@
 using System;
 using System.Collections.Generic;
 using System.Threading.Tasks;
+using EpicGames.Horde.Agents.Pools;
 using EpicGames.Horde.Common;
 using Horde.Server.Agents.Fleet;
 
diff --git a/Engine/Source/Programs/Horde/Horde.Server/Agents/Pools/PoolCollection.cs b/Engine/Source/Programs/Horde/Horde.Server/Agents/Pools/PoolCollection.cs
--- a/Engine/Source/Programs/Horde/Horde.Server/Agents/Pools/PoolCollection.cs
+++ b/Engine/Source/Programs/Horde/Horde.Server/Agents/Pools/PoolCollection.cs
@@ -4,6 +4,7 @@ using System;
 using System.Collections.Generic;
 using System.Linq;
 using System.Threading.Tasks;
+using EpicGames.Horde.Agents.Pools;
 using EpicGames.Horde.Common;
 using Horde.Server.Agents.Fleet;
 using Horde.Server.Server;
diff --git a/Engine/Source/Programs/Horde/Horde.Server/Agents/Pools/PoolConfig.cs b/Engine/Source/Programs/Horde/Horde.Server/Agents/Pools/PoolConfig.cs
--- a/Engine/Source/Programs/Horde/Horde.Server/Agents/Pools/PoolConfig.cs
+++ b/Engine/Source/Programs/Horde/Horde.Server/Agents/Pools/PoolConfig.cs
@@ -3,6 +3,7 @@
 using System;
 using System.Collections.Generic;
 using System.Diagnostics;
+using EpicGames.Horde.Agents.Pools;
 using EpicGames.Horde.Common;
 using Horde.Server.Agents.Fleet;
 
diff --git a/Engine/Source/Programs/Horde/Horde.Server/Agents/Pools/PoolId.cs b/Engine/Source/Programs/Horde/Horde.Server/Agents/Pools/PoolId.cs
deleted file mode 100644
diff --git a/Engine/Source/Programs/Horde/Horde.Server/Agents/Pools/PoolService.cs b/Engine/Source/Programs/Horde/Horde.Server/Agents/Pools/PoolService.cs
--- a/Engine/Source/Programs/Horde/Horde.Server/Agents/Pools/PoolService.cs
+++ b/Engine/Source/Programs/Horde/Horde.Server/Agents/Pools/PoolService.cs
@@ -6,6 +6,7 @@ using System.Linq;
 using System.Threading;
 using System.Threading.Tasks;
 using EpicGames.Horde;
+using EpicGames.Horde.Agents.Pools;
 using Horde.Server.Server;
 using HordeCommon;
 
diff --git a/Engine/Source/Programs/Horde/Horde.Server/Agents/Pools/PoolUpdateService.cs b/Engine/Source/Programs/Horde/Horde.Server/Agents/Pools/PoolUpdateService.cs
--- a/Engine/Source/Programs/Horde/Horde.Server/Agents/Pools/PoolUpdateService.cs
+++ b/Engine/Source/Programs/Horde/Horde.Server/Agents/Pools/PoolUpdateService.cs
@@ -5,6 +5,7 @@ using System.Collections.Generic;
 using System.Linq;
 using System.Threading;
 using System.Threading.Tasks;
+using EpicGames.Horde.Agents.Pools;
 using Horde.Server.Server;
 using Horde.Server.Streams;
 using HordeCommon;
diff --git a/Engine/Source/Programs/Horde/Horde.Server/Agents/Pools/PoolsController.cs b/Engine/Source/Programs/Horde/Horde.Server/Agents/Pools/PoolsController.cs
--- a/Engine/Source/Programs/Horde/Horde.Server/Agents/Pools/PoolsController.cs
+++ b/Engine/Source/Programs/Horde/Horde.Server/Agents/Pools/PoolsController.cs
@@ -4,6 +4,7 @@ using System;
 using System.Collections.Generic;
 using System.Linq;
 using System.Threading.Tasks;
+using EpicGames.Horde.Agents.Pools;
 using Horde.Server.Acls;
 using Horde.Server.Agents.Fleet;
 using Horde.Server.Server;
diff --git a/Engine/Source/Programs/Horde/Horde.Server/Agents/Telemetry/Telemetry.cs b/Engine/Source/Programs/Horde/Horde.Server/Agents/Telemetry/Telemetry.cs
--- a/Engine/Source/Programs/Horde/Horde.Server/Agents/Telemetry/Telemetry.cs
+++ b/Engine/Source/Programs/Horde/Horde.Server/Agents/Telemetry/Telemetry.cs
@@ -2,8 +2,8 @@
 
 using System;
 using System.Collections.Generic;
+using EpicGames.Horde.Agents.Pools;
 using EpicGames.Horde.Streams;
-using Horde.Server.Agents.Pools;
 
 namespace Horde.Server.Agents.Telemetry
 {
diff --git a/Engine/Source/Programs/Horde/Horde.Server/Agents/Telemetry/TelemetryCollection.cs b/Engine/Source/Programs/Horde/Horde.Server/Agents/Telemetry/TelemetryCollection.cs
--- a/Engine/Source/Programs/Horde/Horde.Server/Agents/Telemetry/TelemetryCollection.cs
+++ b/Engine/Source/Programs/Horde/Horde.Server/Agents/Telemetry/TelemetryCollection.cs
@@ -5,8 +5,8 @@ using System.Collections.Generic;
 using System.Linq;
 using System.Threading.Tasks;
 using EpicGames.Core;
+using EpicGames.Horde.Agents.Pools;
 using EpicGames.Horde.Streams;
-using Horde.Server.Agents.Pools;
 using Horde.Server.Server;
 using Horde.Server.Utilities;
 using MongoDB.Bson.Serialization.Attributes;
diff --git a/Engine/Source/Programs/Horde/Horde.Server/Agents/Telemetry/TelemetryService.cs b/Engine/Source/Programs/Horde/Horde.Server/Agents/Telemetry/TelemetryService.cs
--- a/Engine/Source/Programs/Horde/Horde.Server/Agents/Telemetry/TelemetryService.cs
+++ b/Engine/Source/Programs/Horde/Horde.Server/Agents/Telemetry/TelemetryService.cs
@@ -12,6 +12,7 @@ using HordeCommon;
 using Microsoft.Extensions.Hosting;
 using Microsoft.Extensions.Logging;
 using EpicGames.Horde.Agents;
+using EpicGames.Horde.Agents.Pools;
 
 namespace Horde.Server.Agents.Telemetry
 {
diff --git a/Engine/Source/Programs/Horde/Horde.Server/Compute/ComputeController.cs b/Engine/Source/Programs/Horde/Horde.Server/Compute/ComputeController.cs
--- a/Engine/Source/Programs/Horde/Horde.Server/Compute/ComputeController.cs
+++ b/Engine/Source/Programs/Horde/Horde.Server/Compute/ComputeController.cs
@@ -7,10 +7,10 @@ using System.Net;
 using System.Threading;
 using System.Threading.Tasks;
 using EpicGames.Core;
+using EpicGames.Horde.Agents.Pools;
 using EpicGames.Horde.Compute;
 using Horde.Server.Acls;
 using Horde.Server.Agents;
-using Horde.Server.Agents.Pools;
 using Horde.Server.Server;
 using Horde.Server.Utilities;
 using HordeCommon.Rpc.Tasks;
diff --git a/Engine/Source/Programs/Horde/Horde.Server/Jobs/IJob.cs b/Engine/Source/Programs/Horde/Horde.Server/Jobs/IJob.cs
--- a/Engine/Source/Programs/Horde/Horde.Server/Jobs/IJob.cs
+++ b/Engine/Source/Programs/Horde/Horde.Server/Jobs/IJob.cs
@@ -12,7 +12,6 @@ using EpicGames.Horde.Jobs.Templates;
 using EpicGames.Horde.Streams;
 using EpicGames.Horde.Users;
 using Horde.Server.Acls;
-using Horde.Server.Agents.Pools;
 using Horde.Server.Agents.Sessions;
 using Horde.Server.Jobs.Graphs;
 using Horde.Server.Jobs.Timing;
@@ -26,6 +25,7 @@ using MongoDB.Driver;
 using Horde.Server.Jobs.Bisect;
 using EpicGames.Horde.Jobs;
 using EpicGames.Horde.Logs;
+using EpicGames.Horde.Agents.Pools;
 
 namespace Horde.Server.Jobs
 {
diff --git a/Engine/Source/Programs/Horde/Horde.Server/Jobs/IJobCollection.cs b/Engine/Source/Programs/Horde/Horde.Server/Jobs/IJobCollection.cs
--- a/Engine/Source/Programs/Horde/Horde.Server/Jobs/IJobCollection.cs
+++ b/Engine/Source/Programs/Horde/Horde.Server/Jobs/IJobCollection.cs
@@ -9,7 +9,6 @@ using EpicGames.Horde.Agents.Leases;
 using EpicGames.Horde.Jobs.Templates;
 using EpicGames.Horde.Streams;
 using EpicGames.Horde.Users;
-using Horde.Server.Agents.Pools;
 using Horde.Server.Agents.Sessions;
 using Horde.Server.Jobs.Graphs;
 using Horde.Server.Streams;
@@ -21,6 +20,7 @@ using System.Threading;
 using Horde.Server.Acls;
 using EpicGames.Horde.Jobs;
 using EpicGames.Horde.Logs;
+using EpicGames.Horde.Agents.Pools;
 
 namespace Horde.Server.Jobs
 {
diff --git a/Engine/Source/Programs/Horde/Horde.Server/Jobs/IJobStepRef.cs b/Engine/Source/Programs/Horde/Horde.Server/Jobs/IJobStepRef.cs
--- a/Engine/Source/Programs/Horde/Horde.Server/Jobs/IJobStepRef.cs
+++ b/Engine/Source/Programs/Horde/Horde.Server/Jobs/IJobStepRef.cs
@@ -3,7 +3,6 @@
 using System;
 using System.Collections.Generic;
 using System.Diagnostics.CodeAnalysis;
-using Horde.Server.Agents.Pools;
 using HordeCommon;
 using MongoDB.Bson.Serialization;
 using MongoDB.Bson.Serialization.Attributes;
@@ -12,6 +11,7 @@ using EpicGames.Horde.Streams;
 using EpicGames.Horde.Jobs.Templates;
 using EpicGames.Horde.Jobs;
 using EpicGames.Horde.Logs;
+using EpicGames.Horde.Agents.Pools;
 
 namespace Horde.Server.Jobs
 {
diff --git a/Engine/Source/Programs/Horde/Horde.Server/Jobs/IJobStepRefCollection.cs b/Engine/Source/Programs/Horde/Horde.Server/Jobs/IJobStepRefCollection.cs
--- a/Engine/Source/Programs/Horde/Horde.Server/Jobs/IJobStepRefCollection.cs
+++ b/Engine/Source/Programs/Horde/Horde.Server/Jobs/IJobStepRefCollection.cs
@@ -4,7 +4,6 @@ using System;
 using System.Collections.Generic;
 using System.Threading;
 using System.Threading.Tasks;
-using Horde.Server.Agents.Pools;
 using Horde.Server.Jobs.Graphs;
 using HordeCommon;
 using Microsoft.Extensions.Logging;
@@ -13,6 +12,7 @@ using EpicGames.Horde.Jobs.Templates;
 using EpicGames.Horde.Streams;
 using EpicGames.Horde.Jobs;
 using EpicGames.Horde.Logs;
+using EpicGames.Horde.Agents.Pools;
 
 namespace Horde.Server.Jobs
 {
diff --git a/Engine/Source/Programs/Horde/Horde.Server/Jobs/JobCollection.cs b/Engine/Source/Programs/Horde/Horde.Server/Jobs/JobCollection.cs
--- a/Engine/Source/Programs/Horde/Horde.Server/Jobs/JobCollection.cs
+++ b/Engine/Source/Programs/Horde/Horde.Server/Jobs/JobCollection.cs
@@ -14,7 +14,6 @@ using EpicGames.Horde.Agents.Leases;
 using EpicGames.Horde.Streams;
 using EpicGames.Horde.Users;
 using Horde.Server.Acls;
-using Horde.Server.Agents.Pools;
 using Horde.Server.Agents.Sessions;
 using Horde.Server.Jobs.Graphs;
 using Horde.Server.Jobs.Bisect;
@@ -35,6 +34,7 @@ using OpenTelemetry.Trace;
 using EpicGames.Horde.Jobs.Templates;
 using EpicGames.Horde.Jobs;
 using EpicGames.Horde.Logs;
+using EpicGames.Horde.Agents.Pools;
 
 namespace Horde.Server.Jobs
 {
diff --git a/Engine/Source/Programs/Horde/Horde.Server/Jobs/JobStepRefCollection.cs b/Engine/Source/Programs/Horde/Horde.Server/Jobs/JobStepRefCollection.cs
--- a/Engine/Source/Programs/Horde/Horde.Server/Jobs/JobStepRefCollection.cs
+++ b/Engine/Source/Programs/Horde/Horde.Server/Jobs/JobStepRefCollection.cs
@@ -5,7 +5,6 @@ using System.Collections.Generic;
 using System.Linq;
 using System.Threading;
 using System.Threading.Tasks;
-using Horde.Server.Agents.Pools;
 using Horde.Server.Server;
 using Horde.Server.Telemetry;
 using Horde.Server.Utilities;
@@ -17,6 +16,7 @@ using EpicGames.Horde.Streams;
 using EpicGames.Horde.Jobs.Templates;
 using EpicGames.Horde.Jobs;
 using EpicGames.Horde.Logs;
+using EpicGames.Horde.Agents.Pools;
 
 namespace Horde.Server.Jobs
 {
diff --git a/Engine/Source/Programs/Horde/Horde.Server/Jobs/JobTaskSource.cs b/Engine/Source/Programs/Horde/Horde.Server/Jobs/JobTaskSource.cs
--- a/Engine/Source/Programs/Horde/Horde.Server/Jobs/JobTaskSource.cs
+++ b/Engine/Source/Programs/Horde/Horde.Server/Jobs/JobTaskSource.cs
@@ -38,6 +38,7 @@ using EpicGames.Horde.Agents;
 using EpicGames.Horde.Storage;
 using EpicGames.Horde.Jobs;
 using EpicGames.Horde.Logs;
+using EpicGames.Horde.Agents.Pools;
 
 namespace Horde.Server.Jobs
 {
diff --git a/Engine/Source/Programs/Horde/Horde.Server/Streams/StreamConfig.cs b/Engine/Source/Programs/Horde/Horde.Server/Streams/StreamConfig.cs
--- a/Engine/Source/Programs/Horde/Horde.Server/Streams/StreamConfig.cs
+++ b/Engine/Source/Programs/Horde/Horde.Server/Streams/StreamConfig.cs
@@ -12,7 +12,6 @@ using EpicGames.Horde;
 using EpicGames.Horde.Common;
 using EpicGames.Horde.Jobs.Templates;
 using Horde.Server.Acls;
-using Horde.Server.Agents.Pools;
 using Horde.Server.Configuration;
 using Horde.Server.Issues;
 using Horde.Server.Jobs.Graphs;
@@ -24,6 +23,7 @@ using HordeCommon;
 using HordeCommon.Rpc.Tasks;
 using System.Diagnostics;
 using EpicGames.Horde.Streams;
+using EpicGames.Horde.Agents.Pools;
 
 namespace Horde.Server.Streams
 {
diff --git a/Engine/Source/Programs/Shared/EpicGames.Horde/Agents/Pools/PoolId.cs b/Engine/Source/Programs/Shared/EpicGames.Horde/Agents/Pools/PoolId.cs
new file mode 100644
--- /dev/null
+++ b/Engine/Source/Programs/Shared/EpicGames.Horde/Agents/Pools/PoolId.cs
@@ -0,0 +1,44 @@
+// Copyright Epic Games, Inc. All Rights Reserved.
+
+using System.ComponentModel;
+using EpicGames.Core;
+using EpicGames.Serialization;
+
+namespace EpicGames.Horde.Agents.Pools
+{
+	/// <summary>
+	/// Identifier for a pool
+	/// </summary>
+	/// <param name="Id">Id to construct from</param>
+	[JsonSchemaString]
+	[TypeConverter(typeof(StringIdTypeConverter<PoolId, PoolIdConverter>))]
+	[StringIdConverter(typeof(PoolIdConverter))]
+	[CbConverter(typeof(StringIdCbConverter<PoolId, PoolIdConverter>))]
+	public record struct PoolId(StringId Id)
+	{
+		/// <summary>
+		/// Constructor
+		/// </summary>
+		public PoolId(string id) : this(new StringId(id))
+		{
+		}
+
+		/// <inheritdoc cref="StringId.IsEmpty"/>
+		public bool IsEmpty => Id.IsEmpty;
+
+		/// <inheritdoc/>
+		public override string ToString() => Id.ToString();
+	}
+
+	/// <summary>
+	/// Converter to and from <see cref="StringId"/> instances.
+	/// </summary>
+	class PoolIdConverter : StringIdConverter<PoolId>
+	{
+		/// <inheritdoc/>
+		public override PoolId FromStringId(StringId id) => new PoolId(id);
+
+		/// <inheritdoc/>
+		public override StringId ToStringId(PoolId value) => value.Id;
+	}
+}
diff --git a/Engine/Source/Programs/Shared/EpicGames.Horde/Logs/LogId.cs b/Engine/Source/Programs/Shared/EpicGames.Horde/Logs/LogId.cs
--- a/Engine/Source/Programs/Shared/EpicGames.Horde/Logs/LogId.cs
+++ b/Engine/Source/Programs/Shared/EpicGames.Horde/Logs/LogId.cs
@@ -2,7 +2,6 @@
 
 using System.ComponentModel;
 using EpicGames.Core;
-using EpicGames.Horde;
 
 namespace EpicGames.Horde.Logs
 {
