Horde: Move MetricId into EpicGames.Horde assembly.

diff --git a/Engine/Source/Programs/Horde/Horde.Server.Tests/TelemetryTests.cs b/Engine/Source/Programs/Horde/Horde.Server.Tests/TelemetryTests.cs
--- a/Engine/Source/Programs/Horde/Horde.Server.Tests/TelemetryTests.cs
+++ b/Engine/Source/Programs/Horde/Horde.Server.Tests/TelemetryTests.cs
@@ -6,6 +6,7 @@ using System.Linq;
 using System.Text.Json;
 using System.Threading;
 using System.Threading.Tasks;
+using EpicGames.Horde.Telemetry.Metrics;
 using Horde.Server.Server;
 using Horde.Server.Telemetry;
 using Horde.Server.Telemetry.Metrics;
diff --git a/Engine/Source/Programs/Horde/Horde.Server/Telemetry/Metrics/IMetric.cs b/Engine/Source/Programs/Horde/Horde.Server/Telemetry/Metrics/IMetric.cs
--- a/Engine/Source/Programs/Horde/Horde.Server/Telemetry/Metrics/IMetric.cs
+++ b/Engine/Source/Programs/Horde/Horde.Server/Telemetry/Metrics/IMetric.cs
@@ -1,6 +1,7 @@
 // Copyright Epic Games, Inc. All Rights Reserved.
 
 using System;
+using EpicGames.Horde.Telemetry.Metrics;
 
 namespace Horde.Server.Telemetry.Metrics
 {
diff --git a/Engine/Source/Programs/Horde/Horde.Server/Telemetry/Metrics/IMetricCollection.cs b/Engine/Source/Programs/Horde/Horde.Server/Telemetry/Metrics/IMetricCollection.cs
--- a/Engine/Source/Programs/Horde/Horde.Server/Telemetry/Metrics/IMetricCollection.cs
+++ b/Engine/Source/Programs/Horde/Horde.Server/Telemetry/Metrics/IMetricCollection.cs
@@ -5,6 +5,7 @@ using System.Collections.Generic;
 using System.Text.Json.Nodes;
 using System.Threading;
 using System.Threading.Tasks;
+using EpicGames.Horde.Telemetry.Metrics;
 
 namespace Horde.Server.Telemetry.Metrics
 {
diff --git a/Engine/Source/Programs/Horde/Horde.Server/Telemetry/Metrics/MetricCollection.cs b/Engine/Source/Programs/Horde/Horde.Server/Telemetry/Metrics/MetricCollection.cs
--- a/Engine/Source/Programs/Horde/Horde.Server/Telemetry/Metrics/MetricCollection.cs
+++ b/Engine/Source/Programs/Horde/Horde.Server/Telemetry/Metrics/MetricCollection.cs
@@ -8,6 +8,7 @@ using System.Text.Json.Nodes;
 using System.Threading;
 using System.Threading.Tasks;
 using EpicGames.Core;
+using EpicGames.Horde.Telemetry.Metrics;
 using Horde.Server.Server;
 using Horde.Server.Utilities;
 using HordeCommon;
diff --git a/Engine/Source/Programs/Horde/Horde.Server/Telemetry/Metrics/MetricConfig.cs b/Engine/Source/Programs/Horde/Horde.Server/Telemetry/Metrics/MetricConfig.cs
--- a/Engine/Source/Programs/Horde/Horde.Server/Telemetry/Metrics/MetricConfig.cs
+++ b/Engine/Source/Programs/Horde/Horde.Server/Telemetry/Metrics/MetricConfig.cs
@@ -5,6 +5,7 @@ using System.Collections.Generic;
 using System.Linq;
 using System.Text.Json.Serialization;
 using EpicGames.Core;
+using EpicGames.Horde.Telemetry.Metrics;
 using Horde.Server.Utilities;
 using Json.Path;
 
diff --git a/Engine/Source/Programs/Horde/Horde.Server/Telemetry/Metrics/MetricId.cs b/Engine/Source/Programs/Horde/Horde.Server/Telemetry/Metrics/MetricId.cs
deleted file mode 100644
diff --git a/Engine/Source/Programs/Horde/Horde.Server/Telemetry/TelemetryController.cs b/Engine/Source/Programs/Horde/Horde.Server/Telemetry/TelemetryController.cs
--- a/Engine/Source/Programs/Horde/Horde.Server/Telemetry/TelemetryController.cs
+++ b/Engine/Source/Programs/Horde/Horde.Server/Telemetry/TelemetryController.cs
@@ -6,6 +6,7 @@ using System.ComponentModel.DataAnnotations;
 using System.Text.Json.Nodes;
 using System.Threading;
 using System.Threading.Tasks;
+using EpicGames.Horde.Telemetry.Metrics;
 using Horde.Server.Acls;
 using Horde.Server.Server;
 using Horde.Server.Telemetry.Metrics;
diff --git a/Engine/Source/Programs/Shared/EpicGames.Horde/Telemetry/Metrics/MetricId.cs b/Engine/Source/Programs/Shared/EpicGames.Horde/Telemetry/Metrics/MetricId.cs
new file mode 100644
--- /dev/null
+++ b/Engine/Source/Programs/Shared/EpicGames.Horde/Telemetry/Metrics/MetricId.cs
@@ -0,0 +1,44 @@
+// Copyright Epic Games, Inc. All Rights Reserved.
+
+using System.ComponentModel;
+using EpicGames.Core;
+using EpicGames.Serialization;
+
+namespace EpicGames.Horde.Telemetry.Metrics
+{
+	/// <summary>
+	/// Identifier for a particular metric
+	/// </summary>
+	/// <param name="Id">Id to construct from</param>
+	[JsonSchemaString]
+	[TypeConverter(typeof(StringIdTypeConverter<MetricId, MetricIdConverter>))]
+	[StringIdConverter(typeof(MetricIdConverter))]
+	[CbConverter(typeof(StringIdCbConverter<MetricId, MetricIdConverter>))]
+	public record struct MetricId(StringId Id)
+	{
+		/// <summary>
+		/// Constructor
+		/// </summary>
+		public MetricId(string id) : this(new StringId(id))
+		{
+		}
+
+		/// <inheritdoc cref="StringId.IsEmpty"/>
+		public bool IsEmpty => Id.IsEmpty;
+
+		/// <inheritdoc/>
+		public override string ToString() => Id.ToString();
+	}
+
+	/// <summary>
+	/// Converter to and from <see cref="StringId"/> instances.
+	/// </summary>
+	class MetricIdConverter : StringIdConverter<MetricId>
+	{
+		/// <inheritdoc/>
+		public override MetricId FromStringId(StringId id) => new MetricId(id);
+
+		/// <inheritdoc/>
+		public override StringId ToStringId(MetricId value) => value.Id;
+	}
+}
