Horde: Fix static analyzer warnings.

diff --git a/Engine/Source/Programs/Horde/Horde.Agent/Execution/JobExecutor.cs b/Engine/Source/Programs/Horde/Horde.Agent/Execution/JobExecutor.cs
--- a/Engine/Source/Programs/Horde/Horde.Agent/Execution/JobExecutor.cs
+++ b/Engine/Source/Programs/Horde/Horde.Agent/Execution/JobExecutor.cs
@@ -195,7 +195,6 @@ namespace Horde.Agent.Execution
 
 		protected ISession Session { get; }
 		protected HttpStorageClientFactory StorageFactory { get; }
-		private readonly string _token;
 		protected JobOptions JobOptions { get; }
 
 		protected IRpcConnection RpcConnection => Session.RpcConnection;
@@ -212,7 +211,6 @@ namespace Horde.Agent.Execution
 			BatchId = options.BatchId;
 			Batch = options.Batch;
 
-			_token = options.Token;
 			JobOptions = options.JobOptions;
 
 			_envVars[HordeHttpClient.HordeUrlEnvVarName] = options.Session.ServerUrl.ToString();
diff --git a/Engine/Source/Programs/Horde/Horde.Agent/Leases/Handlers/JobHandler.cs b/Engine/Source/Programs/Horde/Horde.Agent/Leases/Handlers/JobHandler.cs
--- a/Engine/Source/Programs/Horde/Horde.Agent/Leases/Handlers/JobHandler.cs
+++ b/Engine/Source/Programs/Horde/Horde.Agent/Leases/Handlers/JobHandler.cs
@@ -4,7 +4,6 @@ using System.Diagnostics;
 using System.Reflection;
 using System.Runtime.InteropServices;
 using EpicGames.Core;
-using EpicGames.Horde.Storage;
 using EpicGames.Horde.Storage.Clients;
 using Google.Protobuf;
 using Grpc.Core;
diff --git a/Engine/Source/Programs/Horde/Horde.Agent/Utility/TempStorage.cs b/Engine/Source/Programs/Horde/Horde.Agent/Utility/TempStorage.cs
--- a/Engine/Source/Programs/Horde/Horde.Agent/Utility/TempStorage.cs
+++ b/Engine/Source/Programs/Horde/Horde.Agent/Utility/TempStorage.cs
@@ -6,7 +6,6 @@ using EpicGames.Horde.Storage;
 using EpicGames.Horde.Storage.Nodes;
 using Horde.Agent.Utility;
 using Horde.Common.Rpc;
-using Microsoft.CodeAnalysis.CSharp.Syntax;
 using Microsoft.Extensions.Logging;
 using System.Diagnostics;
 using System.Diagnostics.CodeAnalysis;
diff --git a/Engine/Source/Programs/Horde/Horde.Server.Tests/ArtifactTests.cs b/Engine/Source/Programs/Horde/Horde.Server.Tests/ArtifactTests.cs
--- a/Engine/Source/Programs/Horde/Horde.Server.Tests/ArtifactTests.cs
+++ b/Engine/Source/Programs/Horde/Horde.Server.Tests/ArtifactTests.cs
@@ -5,12 +5,9 @@ using System.Collections.Generic;
 using System.Linq;
 using System.Threading;
 using System.Threading.Tasks;
-using EpicGames.Horde.Storage;
 using EpicGames.Horde.Artifacts;
 using Horde.Server.Acls;
 using Horde.Server.Artifacts;
-using Horde.Server.Storage;
-using Horde.Server.Utilities;
 using HordeCommon;
 using Microsoft.Extensions.DependencyInjection;
 using Microsoft.VisualStudio.TestTools.UnitTesting;
diff --git a/Engine/Source/Programs/Shared/EpicGames.Horde/Artifacts/ArtifactMessages.cs b/Engine/Source/Programs/Shared/EpicGames.Horde/Artifacts/ArtifactMessages.cs
--- a/Engine/Source/Programs/Shared/EpicGames.Horde/Artifacts/ArtifactMessages.cs
+++ b/Engine/Source/Programs/Shared/EpicGames.Horde/Artifacts/ArtifactMessages.cs
@@ -38,6 +38,11 @@ namespace EpicGames.Horde.Artifacts
 		/// </summary>
 		public ArtifactId Id { get; }
 
+		/// <summary>
+		/// Name of the artifact
+		/// </summary>
+		public ArtifactName Name { get; }
+
 		/// <summary>
 		/// Type of artifact
 		/// </summary>
@@ -61,9 +66,10 @@ namespace EpicGames.Horde.Artifacts
 		/// <summary>
 		/// Constructor
 		/// </summary>
-		public GetArtifactResponse(ArtifactId id, ArtifactType type, StreamId streamId, int change, IReadOnlyList<string> keys)
+		public GetArtifactResponse(ArtifactId id, ArtifactName name, ArtifactType type, StreamId streamId, int change, IReadOnlyList<string> keys)
 		{
 			Id = id;
+			Name = name;
 			Type = type;
 			StreamId = streamId;
 			Change = change;
diff --git a/Engine/Source/Programs/Shared/EpicGames.Horde/Artifacts/ArtifactName.cs b/Engine/Source/Programs/Shared/EpicGames.Horde/Artifacts/ArtifactName.cs
--- a/Engine/Source/Programs/Shared/EpicGames.Horde/Artifacts/ArtifactName.cs
+++ b/Engine/Source/Programs/Shared/EpicGames.Horde/Artifacts/ArtifactName.cs
@@ -1,9 +1,6 @@
 // Copyright Epic Games, Inc. All Rights Reserved.
 
-using System.Collections.Generic;
 using System.ComponentModel;
-using System.Text;
-using System.Threading.Tasks;
 using EpicGames.Core;
 
 namespace EpicGames.Horde.Artifacts
