Horde: Horde: Add an attribute to indicate that a value that should be tagged with a type name in structured log messages, and a log value type for embedding an external hyperlink into the log output.

diff --git a/Engine/Source/Programs/Shared/EpicGames.Core.Tests/LogValueTests.cs b/Engine/Source/Programs/Shared/EpicGames.Core.Tests/LogValueTests.cs
new file mode 100644
--- /dev/null
+++ b/Engine/Source/Programs/Shared/EpicGames.Core.Tests/LogValueTests.cs
@@ -0,0 +1,56 @@
+// Copyright Epic Games, Inc. All Rights Reserved.
+
+using System;
+using System.Buffers;
+using System.Collections.Generic;
+using System.Linq;
+using System.Text;
+using System.Text.Json;
+using System.Threading.Tasks;
+using Microsoft.Extensions.Logging;
+using Microsoft.VisualStudio.TestTools.UnitTesting;
+
+namespace EpicGames.Core.Tests
+{
+	[TestClass]
+	public class LogValueTests
+	{
+		[LogValueType]
+		class TaggedType
+		{
+			public override string ToString() => "hello";
+		}
+
+		static string RenderLogValue(object value)
+		{
+			ArrayBufferWriter<byte> buffer = new ArrayBufferWriter<byte>();
+			using (Utf8JsonWriter writer = new Utf8JsonWriter(buffer))
+			{
+				LogValueFormatter.Format(value, writer);
+			}
+			return Encoding.UTF8.GetString(buffer.WrittenSpan);
+		}
+
+		[TestMethod]
+		public void TaggedTypeTest()
+		{
+			string value = RenderLogValue(new TaggedType());
+			Assert.AreEqual(value, "{\"$type\":\"TaggedType\",\"$text\":\"hello\"}");
+		}
+
+		[TestMethod]
+		public void LinkTest()
+		{
+			string value = RenderLogValue(LogValue.Link(new Uri("https://epicgames.com/"), "Epic Games"));
+			Assert.AreEqual(value, "{\"$type\":\"Link\",\"$text\":\"Epic Games\",\"target\":\"https://epicgames.com/\"}");
+		}
+
+		[TestMethod]
+		public void SourceFile()
+		{
+			FileReference file = new FileReference("foo.txt");
+			string value = RenderLogValue(LogValue.SourceFile(file, "foo"));
+			Assert.AreEqual(value, "{\"$type\":\"SourceFile\",\"$text\":\"foo\",\"file\":\"" + file.FullName.Replace("\\", "\\\\") + "\"}");
+		}
+	}
+}
diff --git a/Engine/Source/Programs/Shared/EpicGames.Core/LogEvent.cs b/Engine/Source/Programs/Shared/EpicGames.Core/LogEvent.cs
--- a/Engine/Source/Programs/Shared/EpicGames.Core/LogEvent.cs
+++ b/Engine/Source/Programs/Shared/EpicGames.Core/LogEvent.cs
@@ -35,6 +35,7 @@ namespace EpicGames.Core
 		public static readonly Utf8String Identifier = new Utf8String("identifier"); // For symbols
 		public static readonly Utf8String RelativePath = new Utf8String("relativePath");
 		public static readonly Utf8String DepotPath = new Utf8String("depotPath");
+		public static readonly Utf8String Target = new Utf8String("target"); // For hyperlinks
 
 		public static readonly Utf8String Exception = new Utf8String("exception");
 		public static readonly Utf8String Trace = new Utf8String("trace");
diff --git a/Engine/Source/Programs/Shared/EpicGames.Core/LogValue.cs b/Engine/Source/Programs/Shared/EpicGames.Core/LogValue.cs
--- a/Engine/Source/Programs/Shared/EpicGames.Core/LogValue.cs
+++ b/Engine/Source/Programs/Shared/EpicGames.Core/LogValue.cs
@@ -26,9 +26,27 @@ namespace EpicGames.Core
 		public static readonly Utf8String ToolName = new Utf8String("ToolName");
 		public static readonly Utf8String ScreenshotTest = new Utf8String("ScreenshotTest");
 		public static readonly Utf8String DepotPath = new Utf8String("DepotPath");
+		public static readonly Utf8String Link = new Utf8String("Link");
 #pragma warning restore CS1591
 	}
 
+	/// <summary>
+	/// Attribute indicating that a type should be tagged in log output
+	/// </summary>
+	[AttributeUsage(AttributeTargets.Class | AttributeTargets.Struct)]
+	public sealed class LogValueTypeAttribute : Attribute
+	{
+		/// <summary>
+		/// Name to use for the type tag
+		/// </summary>
+		public string? Name { get; }
+		
+		/// <summary>
+		/// Constructor
+		/// </summary>
+		public LogValueTypeAttribute(string? name = null) => Name = name;
+	}
+
 	/// <summary>
 	/// Information for a structured value for use in log events
 	/// </summary>
@@ -67,12 +85,30 @@ namespace EpicGames.Core
 		/// </summary>
 		/// <param name="file">Source file to reference</param>
 		/// <param name="text">Display text for the fiel</param>
-		/// <returns>New log value instance</returns>
 		public static LogValue SourceFile(FileReference file, string text)
 		{
 			return new LogValue(LogValueType.SourceFile, text, new Dictionary<Utf8String, object> { [LogEventPropertyName.File] = file.FullName });
 		}
 
+		/// <summary>
+		/// Constructs a value which contains a hyperlink to an external page
+		/// </summary>
+		/// <param name="target">Target URL</param>
+		public static LogValue Link(Uri target)
+		{
+			return Link(target, target.ToString());
+		}
+
+		/// <summary>
+		/// Constructs a value which contains a hyperlink to an external page
+		/// </summary>
+		/// <param name="target">Target URL</param>
+		/// <param name="text">Text to render for the link</param>
+		public static LogValue Link(Uri target, string text)
+		{
+			return new LogValue(LogValueType.Link, text, new Dictionary<Utf8String, object> { [LogEventPropertyName.Target] = target.ToString() });
+		}
+
 		/// <summary>
 		/// Creates a LogValue from an object, overriding the type and display text for it
 		/// </summary>
diff --git a/Engine/Source/Programs/Shared/EpicGames.Core/LogValueFormatter.cs b/Engine/Source/Programs/Shared/EpicGames.Core/LogValueFormatter.cs
--- a/Engine/Source/Programs/Shared/EpicGames.Core/LogValueFormatter.cs
+++ b/Engine/Source/Programs/Shared/EpicGames.Core/LogValueFormatter.cs
@@ -312,20 +312,7 @@ namespace EpicGames.Core
 					return formatter;
 				}
 
-				LogValueFormatterAttribute? formatterAttribute = type.GetCustomAttribute<LogValueFormatterAttribute>();
-				if (formatterAttribute != null)
-				{
-					formatter = (ILogValueFormatter)Activator.CreateInstance(formatterAttribute.Type)!;
-				}
-				else if (TryGetEnumerableType(type, out Type? elementType))
-				{
-					formatter = new EnumerableFormatter(GetFormatter(elementType));
-				}
-				else
-				{
-					formatter = s_stringFormatter;
-				}
-
+				formatter = GetFormatterUncached(type);
 				if (s_formatters.TryAdd(type, formatter))
 				{
 					return formatter;
@@ -333,6 +320,28 @@ namespace EpicGames.Core
 			}
 		}
 
+		static ILogValueFormatter GetFormatterUncached(Type type)
+		{
+			LogValueFormatterAttribute? formatterAttribute = type.GetCustomAttribute<LogValueFormatterAttribute>();
+			if (formatterAttribute != null)
+			{
+				return (ILogValueFormatter)Activator.CreateInstance(formatterAttribute.Type)!;
+			}
+
+			LogValueTypeAttribute? typeAttribute = type.GetCustomAttribute<LogValueTypeAttribute>();
+			if (typeAttribute != null)
+			{
+				return new AnnotateTypeFormatter(typeAttribute.Name ?? type.Name);
+			}
+
+			if (TryGetEnumerableType(type, out Type? elementType))
+			{
+				return new EnumerableFormatter(GetFormatter(elementType));
+			}
+
+			return s_stringFormatter;
+		}
+
 		static bool TryGetEnumerableType(Type type, [NotNullWhen(true)] out Type? elementType)
 		{
 			foreach (Type interfaceType in type.GetInterfaces())
