Horde: Annotate Horde id types with the [LogValueType] attribute, so the dashboard can link them appropriately.

diff --git a/Engine/Source/Programs/Shared/EpicGames.Core.Tests/LogValueTests.cs b/Engine/Source/Programs/Shared/EpicGames.Core.Tests/LogValueTests.cs
--- a/Engine/Source/Programs/Shared/EpicGames.Core.Tests/LogValueTests.cs
+++ b/Engine/Source/Programs/Shared/EpicGames.Core.Tests/LogValueTests.cs
@@ -50,7 +50,7 @@ namespace EpicGames.Core.Tests
 		{
 			FileReference file = new FileReference("foo.txt");
 			string value = RenderLogValue(LogValue.SourceFile(file, "foo"));
-			Assert.AreEqual(value, "{\"$type\":\"SourceFile\",\"$text\":\"foo\",\"file\":\"" + file.FullName.Replace("\\", "\\\\") + "\"}");
+			Assert.AreEqual(value, "{\"$type\":\"SourceFile\",\"$text\":\"foo\",\"file\":\"" + file.FullName.Replace("\\", "\\\\", StringComparison.Ordinal) + "\"}");
 		}
 	}
 }
diff --git a/Engine/Source/Programs/Shared/EpicGames.Horde/Agents/AgentId.cs b/Engine/Source/Programs/Shared/EpicGames.Horde/Agents/AgentId.cs
--- a/Engine/Source/Programs/Shared/EpicGames.Horde/Agents/AgentId.cs
+++ b/Engine/Source/Programs/Shared/EpicGames.Horde/Agents/AgentId.cs
@@ -12,9 +12,9 @@ namespace EpicGames.Horde.Agents
 	/// <summary>
 	/// Normalized hostname of an agent
 	/// </summary>
+	[LogValueType]
 	[TypeConverter(typeof(AgentIdTypeConverter))]
 	[JsonConverter(typeof(AgentIdJsonConverter))]
-	[LogValueFormatter(typeof(AgentIdLogFormatter))]
 	public readonly struct AgentId : IEquatable<AgentId>
 	{
 		/// <summary>
@@ -167,19 +167,4 @@ namespace EpicGames.Horde.Agents
 		/// <inheritdoc/>
 		public override void Write(Utf8JsonWriter writer, AgentId value, JsonSerializerOptions options) => writer.WriteStringValue(value.ToString());
 	}
-
-	/// <summary>
-	/// Formats an AgentId as a typed log value
-	/// </summary>
-	class AgentIdLogFormatter : ILogValueFormatter
-	{
-		/// <inheritdoc/>
-		public void Format(object value, Utf8JsonWriter writer)
-		{
-			writer.WriteStartObject();
-			writer.WriteString(LogEventPropertyName.Type, "AgentId");
-			writer.WriteString(LogEventPropertyName.Text, ((AgentId)value).ToString());
-			writer.WriteEndObject();
-		}
-	}
 }
diff --git a/Engine/Source/Programs/Shared/EpicGames.Horde/Agents/Leases/LeaseId.cs b/Engine/Source/Programs/Shared/EpicGames.Horde/Agents/Leases/LeaseId.cs
--- a/Engine/Source/Programs/Shared/EpicGames.Horde/Agents/Leases/LeaseId.cs
+++ b/Engine/Source/Programs/Shared/EpicGames.Horde/Agents/Leases/LeaseId.cs
@@ -10,10 +10,10 @@ namespace EpicGames.Horde.Agents.Leases
 	/// Identifier for a user
 	/// </summary>
 	/// <param name="Id">Id to construct from</param>
+	[LogValueType]
 	[JsonSchemaString]
 	[TypeConverter(typeof(BinaryIdTypeConverter<LeaseId, LeaseIdConverter>))]
 	[BinaryIdConverter(typeof(LeaseIdConverter))]
-	[LogValueFormatter(typeof(LeaseIdLogFormatter))]
 	public record struct LeaseId(BinaryId Id)
 	{
 		/// <inheritdoc cref="BinaryId.TryParse(System.String, out BinaryId)"/>
@@ -50,19 +50,4 @@ namespace EpicGames.Horde.Agents.Leases
 		/// <inheritdoc/>
 		public override BinaryId ToBinaryId(LeaseId value) => value.Id;
 	}
-
-	/// <summary>
-	/// Formats a LeaseId as a typed log value
-	/// </summary>
-	class LeaseIdLogFormatter : ILogValueFormatter
-	{
-		/// <inheritdoc/>
-		public void Format(object value, Utf8JsonWriter writer)
-		{
-			writer.WriteStartObject();
-			writer.WriteString(LogEventPropertyName.Type, "LeaseId");
-			writer.WriteString(LogEventPropertyName.Text, ((LeaseId)value).ToString());
-			writer.WriteEndObject();
-		}
-	}
 }
diff --git a/Engine/Source/Programs/Shared/EpicGames.Horde/Agents/Pools/PoolId.cs b/Engine/Source/Programs/Shared/EpicGames.Horde/Agents/Pools/PoolId.cs
--- a/Engine/Source/Programs/Shared/EpicGames.Horde/Agents/Pools/PoolId.cs
+++ b/Engine/Source/Programs/Shared/EpicGames.Horde/Agents/Pools/PoolId.cs
@@ -10,6 +10,7 @@ namespace EpicGames.Horde.Agents.Pools
 	/// Identifier for a pool
 	/// </summary>
 	/// <param name="Id">Id to construct from</param>
+	[LogValueType]
 	[JsonSchemaString]
 	[TypeConverter(typeof(StringIdTypeConverter<PoolId, PoolIdConverter>))]
 	[StringIdConverter(typeof(PoolIdConverter))]
diff --git a/Engine/Source/Programs/Shared/EpicGames.Horde/Agents/Sessions/SessionId.cs b/Engine/Source/Programs/Shared/EpicGames.Horde/Agents/Sessions/SessionId.cs
--- a/Engine/Source/Programs/Shared/EpicGames.Horde/Agents/Sessions/SessionId.cs
+++ b/Engine/Source/Programs/Shared/EpicGames.Horde/Agents/Sessions/SessionId.cs
@@ -9,6 +9,7 @@ namespace EpicGames.Horde.Agents.Sessions
 	/// Identifier for a session
 	/// </summary>
 	/// <param name="Id">Id to construct from</param>
+	[LogValueType]
 	[JsonSchemaString]
 	[TypeConverter(typeof(BinaryIdTypeConverter<SessionId, SessionIdConverter>))]
 	[BinaryIdConverter(typeof(SessionIdConverter))]
diff --git a/Engine/Source/Programs/Shared/EpicGames.Horde/Artifacts/ArtifactId.cs b/Engine/Source/Programs/Shared/EpicGames.Horde/Artifacts/ArtifactId.cs
--- a/Engine/Source/Programs/Shared/EpicGames.Horde/Artifacts/ArtifactId.cs
+++ b/Engine/Source/Programs/Shared/EpicGames.Horde/Artifacts/ArtifactId.cs
@@ -1,6 +1,7 @@
 // Copyright Epic Games, Inc. All Rights Reserved.
 
 using System.ComponentModel;
+using EpicGames.Core;
 
 namespace EpicGames.Horde.Artifacts
 {
@@ -8,6 +9,7 @@ namespace EpicGames.Horde.Artifacts
 	/// Unique id for an artifact
 	/// </summary>
 	/// <param name="Id">Identifier for the artifact</param>
+	[LogValueType]
 	[TypeConverter(typeof(BinaryIdTypeConverter<ArtifactId, ArtifactIdConverter>))]
 	[BinaryIdConverter(typeof(ArtifactIdConverter))]
 	public record struct ArtifactId(BinaryId Id)
diff --git a/Engine/Source/Programs/Shared/EpicGames.Horde/Compute/ClusterId.cs b/Engine/Source/Programs/Shared/EpicGames.Horde/Compute/ClusterId.cs
--- a/Engine/Source/Programs/Shared/EpicGames.Horde/Compute/ClusterId.cs
+++ b/Engine/Source/Programs/Shared/EpicGames.Horde/Compute/ClusterId.cs
@@ -13,6 +13,7 @@ namespace EpicGames.Horde.Compute
 	/// <summary>
 	/// Identifier for a compute cluster
 	/// </summary>
+	[LogValueType]
 	[JsonSchemaString]
 	[CbConverter(typeof(ClusterIdCbConverter))]
 	[JsonConverter(typeof(ClusterIdJsonConverter))]
diff --git a/Engine/Source/Programs/Shared/EpicGames.Horde/Devices/DeviceId.cs b/Engine/Source/Programs/Shared/EpicGames.Horde/Devices/DeviceId.cs
--- a/Engine/Source/Programs/Shared/EpicGames.Horde/Devices/DeviceId.cs
+++ b/Engine/Source/Programs/Shared/EpicGames.Horde/Devices/DeviceId.cs
@@ -10,6 +10,7 @@ namespace EpicGames.Horde.Devices
 	/// Identifier for a pool
 	/// </summary>
 	/// <param name="Id">Id to construct from</param>
+	[LogValueType]
 	[JsonSchemaString]
 	[TypeConverter(typeof(StringIdTypeConverter<DeviceId, DeviceIdConverter>))]
 	[StringIdConverter(typeof(DeviceIdConverter))]
diff --git a/Engine/Source/Programs/Shared/EpicGames.Horde/Devices/DevicePlatformId.cs b/Engine/Source/Programs/Shared/EpicGames.Horde/Devices/DevicePlatformId.cs
--- a/Engine/Source/Programs/Shared/EpicGames.Horde/Devices/DevicePlatformId.cs
+++ b/Engine/Source/Programs/Shared/EpicGames.Horde/Devices/DevicePlatformId.cs
@@ -10,6 +10,7 @@ namespace EpicGames.Horde.Devices
 	/// Identifier for a device platform
 	/// </summary>
 	/// <param name="Id">Id to construct from</param>
+	[LogValueType]
 	[JsonSchemaString]
 	[TypeConverter(typeof(StringIdTypeConverter<DevicePlatformId, DevicePlatformIdConverter>))]
 	[StringIdConverter(typeof(DevicePlatformIdConverter))]
diff --git a/Engine/Source/Programs/Shared/EpicGames.Horde/Devices/DevicePoolId.cs b/Engine/Source/Programs/Shared/EpicGames.Horde/Devices/DevicePoolId.cs
--- a/Engine/Source/Programs/Shared/EpicGames.Horde/Devices/DevicePoolId.cs
+++ b/Engine/Source/Programs/Shared/EpicGames.Horde/Devices/DevicePoolId.cs
@@ -10,6 +10,7 @@ namespace EpicGames.Horde.Devices
 	/// Identifier for a device pool
 	/// </summary>
 	/// <param name="Id">Id to construct from</param>
+	[LogValueType]
 	[JsonSchemaString]
 	[TypeConverter(typeof(StringIdTypeConverter<DevicePoolId, DevicePoolIdConverter>))]
 	[StringIdConverter(typeof(DevicePoolIdConverter))]
diff --git a/Engine/Source/Programs/Shared/EpicGames.Horde/Jobs/Bisect/BisectTaskId.cs b/Engine/Source/Programs/Shared/EpicGames.Horde/Jobs/Bisect/BisectTaskId.cs
--- a/Engine/Source/Programs/Shared/EpicGames.Horde/Jobs/Bisect/BisectTaskId.cs
+++ b/Engine/Source/Programs/Shared/EpicGames.Horde/Jobs/Bisect/BisectTaskId.cs
@@ -9,6 +9,7 @@ namespace EpicGames.Horde.Jobs.Bisect
 	/// Identifier for a job
 	/// </summary>
 	/// <param name="Id">Id to construct from</param>
+	[LogValueType]
 	[JsonSchemaString]
 	[TypeConverter(typeof(BinaryIdTypeConverter<BisectTaskId, BisectTaskIdConverter>))]
 	[BinaryIdConverter(typeof(BisectTaskIdConverter))]
diff --git a/Engine/Source/Programs/Shared/EpicGames.Horde/Jobs/JobId.cs b/Engine/Source/Programs/Shared/EpicGames.Horde/Jobs/JobId.cs
--- a/Engine/Source/Programs/Shared/EpicGames.Horde/Jobs/JobId.cs
+++ b/Engine/Source/Programs/Shared/EpicGames.Horde/Jobs/JobId.cs
@@ -9,6 +9,7 @@ namespace EpicGames.Horde.Jobs
 	/// Identifier for a job
 	/// </summary>
 	/// <param name="Id">Id to construct from</param>
+	[LogValueType]
 	[JsonSchemaString]
 	[TypeConverter(typeof(BinaryIdTypeConverter<JobId, JobIdConverter>))]
 	[BinaryIdConverter(typeof(JobIdConverter))]
diff --git a/Engine/Source/Programs/Shared/EpicGames.Horde/Jobs/JobStepBatchId.cs b/Engine/Source/Programs/Shared/EpicGames.Horde/Jobs/JobStepBatchId.cs
--- a/Engine/Source/Programs/Shared/EpicGames.Horde/Jobs/JobStepBatchId.cs
+++ b/Engine/Source/Programs/Shared/EpicGames.Horde/Jobs/JobStepBatchId.cs
@@ -9,6 +9,7 @@ namespace EpicGames.Horde.Jobs
 	/// Identifier for a job step batch
 	/// </summary>
 	/// <param name="SubResourceId">Id to construct from</param>
+	[LogValueType]
 	[JsonSchemaString]
 	[TypeConverter(typeof(SubResourceIdTypeConverter<JobStepBatchId, JobStepBatchIdConverter>))]
 	[SubResourceIdConverter(typeof(JobStepBatchIdConverter))]
diff --git a/Engine/Source/Programs/Shared/EpicGames.Horde/Jobs/JobStepId.cs b/Engine/Source/Programs/Shared/EpicGames.Horde/Jobs/JobStepId.cs
--- a/Engine/Source/Programs/Shared/EpicGames.Horde/Jobs/JobStepId.cs
+++ b/Engine/Source/Programs/Shared/EpicGames.Horde/Jobs/JobStepId.cs
@@ -9,6 +9,7 @@ namespace EpicGames.Horde.Jobs
 	/// Identifier for a jobstep
 	/// </summary>
 	/// <param name="Id">Id to construct from</param>
+	[LogValueType]
 	[JsonSchemaString]
 	[TypeConverter(typeof(SubResourceIdTypeConverter<JobStepId, JobStepIdConverter>))]
 	[SubResourceIdConverter(typeof(JobStepIdConverter))]
diff --git a/Engine/Source/Programs/Shared/EpicGames.Horde/Jobs/Templates/TemplateId.cs b/Engine/Source/Programs/Shared/EpicGames.Horde/Jobs/Templates/TemplateId.cs
--- a/Engine/Source/Programs/Shared/EpicGames.Horde/Jobs/Templates/TemplateId.cs
+++ b/Engine/Source/Programs/Shared/EpicGames.Horde/Jobs/Templates/TemplateId.cs
@@ -10,6 +10,7 @@ namespace EpicGames.Horde.Jobs.Templates
 	/// Identifier for a job template
 	/// </summary>
 	/// <param name="Id">Id to construct from</param>
+	[LogValueType]
 	[JsonSchemaString]
 	[TypeConverter(typeof(StringIdTypeConverter<TemplateId, TemplateIdConverter>))]
 	[StringIdConverter(typeof(TemplateIdConverter))]
diff --git a/Engine/Source/Programs/Shared/EpicGames.Horde/Logs/LogId.cs b/Engine/Source/Programs/Shared/EpicGames.Horde/Logs/LogId.cs
--- a/Engine/Source/Programs/Shared/EpicGames.Horde/Logs/LogId.cs
+++ b/Engine/Source/Programs/Shared/EpicGames.Horde/Logs/LogId.cs
@@ -9,6 +9,7 @@ namespace EpicGames.Horde.Logs
 	/// Identifier for a log
 	/// </summary>
 	/// <param name="Id">Id to construct from</param>
+	[LogValueType]
 	[JsonSchemaString]
 	[TypeConverter(typeof(BinaryIdTypeConverter<LogId, LogIdConverter>))]
 	[BinaryIdConverter(typeof(LogIdConverter))]
diff --git a/Engine/Source/Programs/Shared/EpicGames.Horde/Projects/ProjectId.cs b/Engine/Source/Programs/Shared/EpicGames.Horde/Projects/ProjectId.cs
--- a/Engine/Source/Programs/Shared/EpicGames.Horde/Projects/ProjectId.cs
+++ b/Engine/Source/Programs/Shared/EpicGames.Horde/Projects/ProjectId.cs
@@ -10,6 +10,7 @@ namespace EpicGames.Horde.Projects
 	/// Identifier for a pool
 	/// </summary>
 	/// <param name="Id">Id to construct from</param>
+	[LogValueType]
 	[JsonSchemaString]
 	[TypeConverter(typeof(StringIdTypeConverter<ProjectId, ProjectIdConverter>))]
 	[StringIdConverter(typeof(ProjectIdConverter))]
diff --git a/Engine/Source/Programs/Shared/EpicGames.Horde/Secrets/SecretId.cs b/Engine/Source/Programs/Shared/EpicGames.Horde/Secrets/SecretId.cs
--- a/Engine/Source/Programs/Shared/EpicGames.Horde/Secrets/SecretId.cs
+++ b/Engine/Source/Programs/Shared/EpicGames.Horde/Secrets/SecretId.cs
@@ -10,6 +10,7 @@ namespace EpicGames.Horde.Secrets
 	/// Identifier for a secret
 	/// </summary>
 	/// <param name="Id">Id to construct from</param>
+	[LogValueType]
 	[JsonSchemaString]
 	[TypeConverter(typeof(StringIdTypeConverter<SecretId, SecretIdConverter>))]
 	[StringIdConverter(typeof(SecretIdConverter))]
diff --git a/Engine/Source/Programs/Shared/EpicGames.Horde/Storage/NamespaceId.cs b/Engine/Source/Programs/Shared/EpicGames.Horde/Storage/NamespaceId.cs
--- a/Engine/Source/Programs/Shared/EpicGames.Horde/Storage/NamespaceId.cs
+++ b/Engine/Source/Programs/Shared/EpicGames.Horde/Storage/NamespaceId.cs
@@ -12,6 +12,7 @@ namespace EpicGames.Horde.Storage
 	/// <summary>
 	/// Identifier for a storage namespace
 	/// </summary>
+	[LogValueType]
 	[JsonSchemaString]
 	[JsonConverter(typeof(NamespaceIdJsonConverter))]
 	[TypeConverter(typeof(NamespaceIdTypeConverter))]
diff --git a/Engine/Source/Programs/Shared/EpicGames.Horde/Storage/WorkspaceLayerId.cs b/Engine/Source/Programs/Shared/EpicGames.Horde/Storage/WorkspaceLayerId.cs
--- a/Engine/Source/Programs/Shared/EpicGames.Horde/Storage/WorkspaceLayerId.cs
+++ b/Engine/Source/Programs/Shared/EpicGames.Horde/Storage/WorkspaceLayerId.cs
@@ -1,6 +1,7 @@
 // Copyright Epic Games, Inc. All Rights Reserved.
 
 using System.ComponentModel;
+using EpicGames.Core;
 
 namespace EpicGames.Horde.Storage
 {
@@ -8,6 +9,7 @@ namespace EpicGames.Horde.Storage
 	/// Identifier for a workspace layer
 	/// </summary>
 	/// <param name="Id">Identifier for the layer</param>
+	[LogValueType]
 	[TypeConverter(typeof(StringIdTypeConverter<WorkspaceLayerId, WorkspaceLayerIdConverter>))]
 	[BinaryIdConverter(typeof(WorkspaceLayerIdConverter))]
 	public record struct WorkspaceLayerId(StringId Id)
diff --git a/Engine/Source/Programs/Shared/EpicGames.Horde/Streams/StreamId.cs b/Engine/Source/Programs/Shared/EpicGames.Horde/Streams/StreamId.cs
--- a/Engine/Source/Programs/Shared/EpicGames.Horde/Streams/StreamId.cs
+++ b/Engine/Source/Programs/Shared/EpicGames.Horde/Streams/StreamId.cs
@@ -10,6 +10,7 @@ namespace EpicGames.Horde.Streams
 	/// Identifier for a stream
 	/// </summary>
 	/// <param name="Id">Id to construct from</param>
+	[LogValueType]
 	[JsonSchemaString]
 	[TypeConverter(typeof(StringIdTypeConverter<StreamId, StreamIdConverter>))]
 	[StringIdConverter(typeof(StreamIdConverter))]
diff --git a/Engine/Source/Programs/Shared/EpicGames.Horde/Telemetry/Metrics/MetricId.cs b/Engine/Source/Programs/Shared/EpicGames.Horde/Telemetry/Metrics/MetricId.cs
--- a/Engine/Source/Programs/Shared/EpicGames.Horde/Telemetry/Metrics/MetricId.cs
+++ b/Engine/Source/Programs/Shared/EpicGames.Horde/Telemetry/Metrics/MetricId.cs
@@ -10,6 +10,7 @@ namespace EpicGames.Horde.Telemetry.Metrics
 	/// Identifier for a particular metric
 	/// </summary>
 	/// <param name="Id">Id to construct from</param>
+	[LogValueType]
 	[JsonSchemaString]
 	[TypeConverter(typeof(StringIdTypeConverter<MetricId, MetricIdConverter>))]
 	[StringIdConverter(typeof(MetricIdConverter))]
diff --git a/Engine/Source/Programs/Shared/EpicGames.Horde/Tools/ToolDeploymentId.cs b/Engine/Source/Programs/Shared/EpicGames.Horde/Tools/ToolDeploymentId.cs
--- a/Engine/Source/Programs/Shared/EpicGames.Horde/Tools/ToolDeploymentId.cs
+++ b/Engine/Source/Programs/Shared/EpicGames.Horde/Tools/ToolDeploymentId.cs
@@ -1,6 +1,7 @@
 // Copyright Epic Games, Inc. All Rights Reserved.
 
 using System.ComponentModel;
+using EpicGames.Core;
 
 namespace EpicGames.Horde.Tools
 {
@@ -8,6 +9,7 @@ namespace EpicGames.Horde.Tools
 	/// Identifier for a tool deployment
 	/// </summary>
 	/// <param name="Id">Identifier for the artifact</param>
+	[LogValueType]
 	[TypeConverter(typeof(BinaryIdTypeConverter<ToolDeploymentId, ToolDeploymentIdConverter>))]
 	[BinaryIdConverter(typeof(ToolDeploymentIdConverter))]
 	public record struct ToolDeploymentId(BinaryId Id)
diff --git a/Engine/Source/Programs/Shared/EpicGames.Horde/Tools/ToolId.cs b/Engine/Source/Programs/Shared/EpicGames.Horde/Tools/ToolId.cs
--- a/Engine/Source/Programs/Shared/EpicGames.Horde/Tools/ToolId.cs
+++ b/Engine/Source/Programs/Shared/EpicGames.Horde/Tools/ToolId.cs
@@ -10,6 +10,7 @@ namespace EpicGames.Horde.Tools
 	/// Identifier for a tool
 	/// </summary>
 	/// <param name="Id">Id to construct from</param>
+	[LogValueType]
 	[JsonSchemaString]
 	[TypeConverter(typeof(StringIdTypeConverter<ToolId, ToolIdConverter>))]
 	[StringIdConverter(typeof(ToolIdConverter))]
diff --git a/Engine/Source/Programs/Shared/EpicGames.Horde/Users/UserId.cs b/Engine/Source/Programs/Shared/EpicGames.Horde/Users/UserId.cs
--- a/Engine/Source/Programs/Shared/EpicGames.Horde/Users/UserId.cs
+++ b/Engine/Source/Programs/Shared/EpicGames.Horde/Users/UserId.cs
@@ -9,6 +9,7 @@ namespace EpicGames.Horde.Users
 	/// Identifier for a user
 	/// </summary>
 	/// <param name="Id">Id to construct from</param>
+	[LogValueType]
 	[JsonSchemaString]
 	[TypeConverter(typeof(BinaryIdTypeConverter<UserId, UserIdConverter>))]
 	[BinaryIdConverter(typeof(UserIdConverter))]
