Horde: Fix static analysis warnings.

diff --git a/Engine/Source/Programs/Shared/EpicGames.Core.Tests/LogValueTests.cs b/Engine/Source/Programs/Shared/EpicGames.Core.Tests/LogValueTests.cs
--- a/Engine/Source/Programs/Shared/EpicGames.Core.Tests/LogValueTests.cs
+++ b/Engine/Source/Programs/Shared/EpicGames.Core.Tests/LogValueTests.cs
@@ -2,12 +2,10 @@
 
 using System;
 using System.Buffers;
-using System.Collections.Generic;
 using System.Linq;
 using System.Text;
 using System.Text.Json;
 using System.Threading.Tasks;
-using Microsoft.Extensions.Logging;
 using Microsoft.VisualStudio.TestTools.UnitTesting;
 
 namespace EpicGames.Core.Tests
diff --git a/Engine/Source/Programs/Shared/EpicGames.Horde/Agents/Leases/LeaseId.cs b/Engine/Source/Programs/Shared/EpicGames.Horde/Agents/Leases/LeaseId.cs
--- a/Engine/Source/Programs/Shared/EpicGames.Horde/Agents/Leases/LeaseId.cs
+++ b/Engine/Source/Programs/Shared/EpicGames.Horde/Agents/Leases/LeaseId.cs
@@ -1,7 +1,6 @@
 // Copyright Epic Games, Inc. All Rights Reserved.
 
 using System.ComponentModel;
-using System.Text.Json;
 using EpicGames.Core;
 
 namespace EpicGames.Horde.Agents.Leases
