OidcToken - Allow you to only specify the list of redirect uris instead of both that and the single redirect uri

diff --git a/Engine/Source/Programs/Shared/EpicGames.OIDC/OidcTokenManager.cs b/Engine/Source/Programs/Shared/EpicGames.OIDC/OidcTokenManager.cs
--- a/Engine/Source/Programs/Shared/EpicGames.OIDC/OidcTokenManager.cs
+++ b/Engine/Source/Programs/Shared/EpicGames.OIDC/OidcTokenManager.cs
@@ -204,7 +204,11 @@ namespace EpicGames.OIDC
 			_authorityUri = providerInfo.ServerUri;
 			_clientId = providerInfo.ClientId;
 
-			List<Uri> possibleRedirectUris = new List<Uri> { providerInfo.RedirectUri };
+			List<Uri> possibleRedirectUris = new List<Uri>();
+			if (providerInfo.RedirectUri != null)
+			{
+				possibleRedirectUris.Add(providerInfo.RedirectUri);
+			}
 			if (providerInfo.PossibleRedirectUri != null)
 			{
 				possibleRedirectUris.AddRange(providerInfo.PossibleRedirectUri);
@@ -626,7 +630,7 @@ namespace EpicGames.OIDC
 
 		[Required] public string DisplayName { get; set; } = null!;
 
-		public Uri RedirectUri { get; set; } = null!;
+		public Uri? RedirectUri { get; set; } = null;
 		public List<Uri>? PossibleRedirectUri { get; set; } = null!;
 		[Required] public bool LoadClaimsFromUserProfile { get; set; } = false;
 
