Horde: Add helper methods for writing directory trees to storage writers.

diff --git a/Engine/Source/Programs/AutomationTool/BuildGraph/Tasks/DeployToolTask.cs b/Engine/Source/Programs/AutomationTool/BuildGraph/Tasks/DeployToolTask.cs
--- a/Engine/Source/Programs/AutomationTool/BuildGraph/Tasks/DeployToolTask.cs
+++ b/Engine/Source/Programs/AutomationTool/BuildGraph/Tasks/DeployToolTask.cs
@@ -168,7 +168,7 @@ namespace AutomationTool.Tasks
 				else if (Parameters.Directory != null)
 				{
 					DirectoryInfo directoryInfo = ResolveDirectory(Parameters.Directory).ToDirectoryInfo();
-					await sandbox.CopyFromDirectoryAsync(directoryInfo, new ChunkingOptions(), treeWriter, null);
+					await sandbox.AddFilesAsync(directoryInfo, treeWriter);
 				}
 				else
 				{
diff --git a/Engine/Source/Programs/Horde/Horde.Agent/Execution/JobExecutor.cs b/Engine/Source/Programs/Horde/Horde.Agent/Execution/JobExecutor.cs
--- a/Engine/Source/Programs/Horde/Horde.Agent/Execution/JobExecutor.cs
+++ b/Engine/Source/Programs/Horde/Horde.Agent/Execution/JobExecutor.cs
@@ -522,7 +522,7 @@ namespace Horde.Agent.Execution
 					await using (IStorageWriter treeWriter = CreateStorageWriter(storage, artifact.RefName, logger))
 					{
 						DirectoryNode buildGraphNode = new DirectoryNode();
-						await buildGraphNode.AddFilesAsync(workspaceDir, buildGraphFiles, new ChunkingOptions(), treeWriter, null, cancellationToken);
+						await buildGraphNode.AddFilesAsync(workspaceDir, buildGraphFiles, treeWriter, cancellationToken: cancellationToken);
 						HashedNodeRef<DirectoryNode> outputNodeRef = await treeWriter.WriteHashedNodeAsync(buildGraphNode, cancellationToken);
 
 						DirectoryNode rootNode = new DirectoryNode();
@@ -846,7 +846,7 @@ namespace Horde.Agent.Execution
 					try
 					{
 						DirectoryNode dir = new DirectoryNode();
-						await dir.AddFilesAsync(baseDir, files, new ChunkingOptions(), writer, new CopyStatsLogger(logger), cancellationToken);
+						await dir.AddFilesAsync(baseDir, files, writer, progress: new CopyStatsLogger(logger), cancellationToken: cancellationToken);
 						rootRef = await writer.WriteNodeAsync(dir, cancellationToken);
 					}
 					catch (Exception ex)
diff --git a/Engine/Source/Programs/Horde/Horde.Agent/Utility/TempStorage.cs b/Engine/Source/Programs/Horde/Horde.Agent/Utility/TempStorage.cs
--- a/Engine/Source/Programs/Horde/Horde.Agent/Utility/TempStorage.cs
+++ b/Engine/Source/Programs/Horde/Horde.Agent/Utility/TempStorage.cs
@@ -686,7 +686,7 @@ namespace Horde.Storage.Utility
 
 			// Create the file tree
 			DirectoryNode rootNode = new DirectoryNode();
-			await rootNode.AddFilesAsync(workspaceDir, archiveFiles, new ChunkingOptions(), writer, new CopyStatsLogger(logger), cancellationToken);
+			await rootNode.AddFilesAsync(workspaceDir, archiveFiles, writer, progress: new CopyStatsLogger(logger), cancellationToken: cancellationToken);
 
 			HashedNodeRef<DirectoryNode> rootNodeRef = await writer.WriteHashedNodeAsync(rootNode, cancellationToken);
 			return new DirectoryEntry(blockDirectoryName, rootNode.Length, rootNodeRef);
diff --git a/Engine/Source/Programs/Horde/Horde.Server/Commands/Install/SetupCommand.cs b/Engine/Source/Programs/Horde/Horde.Server/Commands/Install/SetupCommand.cs
--- a/Engine/Source/Programs/Horde/Horde.Server/Commands/Install/SetupCommand.cs
+++ b/Engine/Source/Programs/Horde/Horde.Server/Commands/Install/SetupCommand.cs
@@ -67,7 +67,7 @@ namespace Horde.Server.Commands.Install
 					await using (IStorageWriter writer = client.CreateWriter(refName))
 					{
 						DirectoryNode dirNode = new DirectoryNode();
-						await dirNode.CopyFromDirectoryAsync(looseAgentDir.ToDirectoryInfo(), new ChunkingOptions(), writer, null);
+						await dirNode.AddFilesAsync(looseAgentDir.ToDirectoryInfo(), writer);
 						dirNodeRef = await writer.WriteNodeAsync(dirNode);
 					}
 					await client.WriteRefAsync(refName, dirNodeRef.Handle);
@@ -86,7 +86,7 @@ namespace Horde.Server.Commands.Install
 					await using (IStorageWriter writer = client.CreateWriter(refName))
 					{
 						DirectoryNode dirNode = new DirectoryNode();
-						await dirNode.CopyFromDirectoryAsync(looseAgentInstallerDir.ToDirectoryInfo(), new ChunkingOptions(), writer, null);
+						await dirNode.AddFilesAsync(looseAgentInstallerDir.ToDirectoryInfo(), writer);
 						dirNodeRef = await writer.WriteNodeAsync(dirNode);
 					}
 					await client.WriteRefTargetAsync(refName, dirNodeRef);
diff --git a/Engine/Source/Programs/Horde/Horde/Commands/Bundles/BundleCreate.cs b/Engine/Source/Programs/Horde/Horde/Commands/Bundles/BundleCreate.cs
--- a/Engine/Source/Programs/Horde/Horde/Commands/Bundles/BundleCreate.cs
+++ b/Engine/Source/Programs/Horde/Horde/Commands/Bundles/BundleCreate.cs
@@ -87,8 +87,7 @@ namespace Horde.Commands.Bundles
 				List<FileInfo> fileInfos = files.ConvertAll(x => x.ToFileInfo());
 				CopyStatsLoggerWithTotals copyStatsLogger = new CopyStatsLoggerWithTotals(files.Count, fileInfos.Sum(x => x.Length), logger);
 
-				DirectoryNode node = await DirectoryNode.CreateAsync(baseDir, files.ConvertAll(x => x.ToFileInfo()), options, writer, copyStatsLogger, CancellationToken.None);
-				HashedNodeRef<DirectoryNode> nodeRef = await writer.WriteHashedNodeAsync(node, CancellationToken.None);
+				HashedNodeRef<DirectoryNode> nodeRef = await writer.WriteFilesAsync(baseDir.ToDirectoryInfo(), fileInfos, options, copyStatsLogger, CancellationToken.None);
 
 				await writer.FlushAsync();
 
diff --git a/Engine/Source/Programs/Horde/Horde/Commands/Compute/ComputeRun.cs b/Engine/Source/Programs/Horde/Horde/Commands/Compute/ComputeRun.cs
--- a/Engine/Source/Programs/Horde/Horde/Commands/Compute/ComputeRun.cs
+++ b/Engine/Source/Programs/Horde/Horde/Commands/Compute/ComputeRun.cs
@@ -153,11 +153,10 @@ namespace Horde.Commands.Compute
 		{
 			await using IStorageWriter writer = storage.CreateWriter();
 
-			DirectoryNode sandbox = new DirectoryNode();
-			await sandbox.CopyFromDirectoryAsync(taskFile.Directory.ToDirectoryInfo(), new ChunkingOptions(), writer, null, cancellationToken);
+			HashedNodeRef<DirectoryNode> sandbox = await writer.WriteFilesAsync(taskFile.Directory, cancellationToken: cancellationToken);
+			await writer.FlushAsync(cancellationToken);
 
-			IBlobHandle handle = await writer.FlushAsync(sandbox, cancellationToken);
-			return handle.GetLocator();
+			return sandbox.Handle.GetLocator();
 		}
 	}
 }
diff --git a/Engine/Source/Programs/Horde/Horde/Commands/Tools/ToolUpload.cs b/Engine/Source/Programs/Horde/Horde/Commands/Tools/ToolUpload.cs
--- a/Engine/Source/Programs/Horde/Horde/Commands/Tools/ToolUpload.cs
+++ b/Engine/Source/Programs/Horde/Horde/Commands/Tools/ToolUpload.cs
@@ -40,15 +40,13 @@ namespace Horde.Commands
 		{
 			using IStorageClient storageClient = _httpStorageClientFactory.CreateClientWithPath($"api/v1/tools/{ToolId}");
 
-			IBlobHandle handle;
+			HashedNodeRef<DirectoryNode> target;
 			await using (IStorageWriter writer = storageClient.CreateWriter())
 			{
-				DirectoryNode sandbox = new DirectoryNode();
-				await sandbox.CopyFromDirectoryAsync(InputDir.ToDirectoryInfo(), new ChunkingOptions(), writer, null);
-				handle = await writer.FlushAsync(sandbox);
+				target = await writer.WriteFilesAsync(InputDir);
 			}
 
-			ToolDeploymentId deploymentId = await _hordeHttpClient.CreateToolDeploymentAsync(ToolId, Version, null, null, handle.GetLocator());
+			ToolDeploymentId deploymentId = await _hordeHttpClient.CreateToolDeploymentAsync(ToolId, Version, null, null, target.Handle.GetLocator());
 			logger.LogInformation("Created deployment {DeploymentId}", deploymentId);
 
 			return 0;
diff --git a/Engine/Source/Programs/Horde/Samples/RemoteClient/RemoteClientApp.cs b/Engine/Source/Programs/Horde/Samples/RemoteClient/RemoteClientApp.cs
--- a/Engine/Source/Programs/Horde/Samples/RemoteClient/RemoteClientApp.cs
+++ b/Engine/Source/Programs/Horde/Samples/RemoteClient/RemoteClientApp.cs
@@ -114,12 +114,11 @@ namespace RemoteClient
 				using MemoryStorageClient memoryStorage = new MemoryStorageClient();
 				using BundleStorageClient storage = new BundleStorageClient(memoryStorage, cache, logger);
 
-				await using (IStorageWriter treeWriter = storage.CreateWriter())
+				await using (IStorageWriter writer = storage.CreateWriter())
 				{
-					DirectoryNode sandbox = new DirectoryNode();
-					await sandbox.CopyFromDirectoryAsync(uploadDir.ToDirectoryInfo(), new ChunkingOptions(), treeWriter, null);
-					IBlobHandle handle = await treeWriter.FlushAsync(sandbox);
-					await channel.UploadFilesAsync("", handle.GetLocator(), storage);
+					HashedNodeRef<DirectoryNode> sandbox = await writer.WriteFilesAsync(uploadDir);
+					await writer.FlushAsync();
+					await channel.UploadFilesAsync("", sandbox.Handle.GetLocator(), storage);
 				}
 
 				// Run the task remotely in the background and echo the output to the console
diff --git a/Engine/Source/Programs/Shared/EpicGames.Core.Tests/LogValueTests.cs b/Engine/Source/Programs/Shared/EpicGames.Core.Tests/LogValueTests.cs
--- a/Engine/Source/Programs/Shared/EpicGames.Core.Tests/LogValueTests.cs
+++ b/Engine/Source/Programs/Shared/EpicGames.Core.Tests/LogValueTests.cs
@@ -2,10 +2,8 @@
 
 using System;
 using System.Buffers;
-using System.Linq;
 using System.Text;
 using System.Text.Json;
-using System.Threading.Tasks;
 using Microsoft.VisualStudio.TestTools.UnitTesting;
 
 namespace EpicGames.Core.Tests
diff --git a/Engine/Source/Programs/Shared/EpicGames.Horde/Storage/Nodes/DirectoryNode.cs b/Engine/Source/Programs/Shared/EpicGames.Horde/Storage/Nodes/DirectoryNode.cs
--- a/Engine/Source/Programs/Shared/EpicGames.Horde/Storage/Nodes/DirectoryNode.cs
+++ b/Engine/Source/Programs/Shared/EpicGames.Horde/Storage/Nodes/DirectoryNode.cs
@@ -487,32 +487,22 @@ namespace EpicGames.Horde.Storage.Nodes
 
 		#endregion
 
-		/// <summary>
-		/// Adds files from a flat list of paths
-		/// </summary>
-		/// <param name="baseDir">Base directory to base paths relative to</param>
-		/// <param name="files">Files to add</param>
-		/// <param name="options">Options for chunking file content</param>
-		/// <param name="writer">Writer for new node data</param>
-		/// <param name="progress">Feedback interface for progress updates</param>
-		/// <param name="cancellationToken">Cancellation token for the operation</param>
-		public static async Task<DirectoryNode> CreateAsync(DirectoryReference baseDir, IReadOnlyList<FileInfo> files, ChunkingOptions options, IStorageWriter writer, IProgress<ICopyStats>? progress, CancellationToken cancellationToken)
+		/// <inheritdoc cref="AddFilesAsync(DirectoryReference, IEnumerable{FileInfo}, IStorageWriter, ChunkingOptions?, IProgress{ICopyStats}?, CancellationToken)"/>
+		public async Task AddFilesAsync(DirectoryInfo directoryInfo, IStorageWriter writer, ChunkingOptions? options = null, IProgress<ICopyStats>? progress = null, CancellationToken cancellationToken = default)
 		{
-			DirectoryNode directoryNode = new DirectoryNode();
-			await directoryNode.AddFilesAsync(baseDir, files, options, writer, progress, cancellationToken);
-			return directoryNode;
+			await AddFilesAsync(new DirectoryReference(directoryInfo), directoryInfo.EnumerateFiles("*", SearchOption.AllDirectories).ToList(), writer, options, progress, cancellationToken);
 		}
 
-		/// <inheritdoc cref="AddFilesAsync(DirectoryReference, IEnumerable{FileInfo}, ChunkingOptions, IStorageWriter, IProgress{ICopyStats}?, CancellationToken)"/>
-		public Task AddFilesAsync(DirectoryReference baseDir, IEnumerable<FileReference> files, ChunkingOptions options, IStorageWriter writer, IProgress<ICopyStats>? progress, CancellationToken cancellationToken)
+		/// <inheritdoc cref="AddFilesAsync(DirectoryReference, IEnumerable{FileInfo}, IStorageWriter, ChunkingOptions?, IProgress{ICopyStats}?, CancellationToken)"/>
+		public Task AddFilesAsync(DirectoryReference baseDir, IEnumerable<FileReference> files, IStorageWriter writer, ChunkingOptions? options = null, IProgress<ICopyStats>? progress = null, CancellationToken cancellationToken = default)
 		{
-			return AddFilesAsync(baseDir, files.Select(x => x.ToFileInfo()).ToList(), options, writer, progress, cancellationToken);
+			return AddFilesAsync(baseDir, files.Select(x => x.ToFileInfo()).ToList(), writer, options, progress, cancellationToken);
 		}
 
-		/// <inheritdoc cref="AddFilesAsync(DirectoryReference, IEnumerable{FileInfo}, ChunkingOptions, IStorageWriter, IProgress{ICopyStats}?, CancellationToken)"/>
-		public Task AddFilesAsync(DirectoryInfo baseDir, IEnumerable<FileInfo> files, ChunkingOptions options, IStorageWriter writer, IProgress<ICopyStats>? progress, CancellationToken cancellationToken)
+		/// <inheritdoc cref="AddFilesAsync(DirectoryReference, IEnumerable{FileInfo}, IStorageWriter, ChunkingOptions?, IProgress{ICopyStats}?, CancellationToken)"/>
+		public Task AddFilesAsync(DirectoryInfo baseDir, IEnumerable<FileInfo> files, IStorageWriter writer, ChunkingOptions? options = null, IProgress<ICopyStats>? progress = null, CancellationToken cancellationToken = default)
 		{
-			return AddFilesAsync(new DirectoryReference(baseDir), files.ToList(), options, writer, progress, cancellationToken);
+			return AddFilesAsync(new DirectoryReference(baseDir), files.ToList(), writer, options, progress, cancellationToken);
 		}
 
 		/// <summary>
@@ -524,8 +514,10 @@ namespace EpicGames.Horde.Storage.Nodes
 		/// <param name="writer">Writer for new node data</param>
 		/// <param name="progress">Feedback interface for progress updates</param>
 		/// <param name="cancellationToken">Cancellation token for the operation</param>
-		public async Task AddFilesAsync(DirectoryReference baseDir, IEnumerable<FileInfo> files, ChunkingOptions options, IStorageWriter writer, IProgress<ICopyStats>? progress, CancellationToken cancellationToken)
+		public async Task AddFilesAsync(DirectoryReference baseDir, IEnumerable<FileInfo> files, IStorageWriter writer, ChunkingOptions? options = null, IProgress<ICopyStats>? progress = null, CancellationToken cancellationToken = default)
 		{
+			options ??= new ChunkingOptions();
+
 			CopyStats? copyStats = null;
 			if (progress != null)
 			{
@@ -747,19 +739,6 @@ namespace EpicGames.Horde.Storage.Nodes
 			await UpdateAsync(updates, writer, cancellationToken);
 		}
 
-		/// <summary>
-		/// Adds files from a directory on disk
-		/// </summary>
-		/// <param name="directoryInfo"></param>
-		/// <param name="options">Options for chunking file content</param>
-		/// <param name="writer">Writer for new node data</param>
-		/// <param name="progress">Feedback interface for progress updates</param>
-		/// <param name="cancellationToken">Cancellation token for the operation</param>
-		public async Task CopyFromDirectoryAsync(DirectoryInfo directoryInfo, ChunkingOptions options, IStorageWriter writer, IProgress<ICopyStats>? progress, CancellationToken cancellationToken = default)
-		{
-			await AddFilesAsync(new DirectoryReference(directoryInfo), directoryInfo.EnumerateFiles("*", SearchOption.AllDirectories).ToList(), options, writer, progress, cancellationToken);
-		}
-
 		/// <summary>
 		/// Utility function to allow extracting a packed directory to disk
 		/// </summary>
@@ -1153,4 +1132,40 @@ namespace EpicGames.Horde.Storage.Nodes
 		/// <inheritdoc/>
 		public override void Write(byte[] buffer, int offset, int count) => throw new NotSupportedException();
 	}
+
+	/// <summary>
+	/// Extension methods for writing directory nodes
+	/// </summary>
+	public static class DirectoryNodeExtensions
+	{
+		/// <summary>
+		/// Writes a tree of files to a storage writer
+		/// </summary>
+		public static async Task<HashedNodeRef<DirectoryNode>> WriteFilesAsync(this IStorageWriter writer, DirectoryReference baseDir, ChunkingOptions? options = null, IProgress<ICopyStats>? progress = null, CancellationToken cancellationToken = default)
+		{
+			DirectoryNode outputNode = new DirectoryNode();
+			await outputNode.AddFilesAsync(baseDir, DirectoryReference.EnumerateFiles(baseDir, "*", SearchOption.AllDirectories), writer, options, progress, cancellationToken);
+			return await writer.WriteHashedNodeAsync(outputNode, cancellationToken);
+		}
+
+		/// <summary>
+		/// Writes a tree of files to a storage writer
+		/// </summary>
+		public static async Task<HashedNodeRef<DirectoryNode>> WriteFilesAsync(this IStorageWriter writer, DirectoryInfo baseDir, IReadOnlyList<FileInfo> files, ChunkingOptions? options = null, IProgress<ICopyStats>? progress = null, CancellationToken cancellationToken = default)
+		{
+			DirectoryNode outputNode = new DirectoryNode();
+			await outputNode.AddFilesAsync(baseDir, files, writer, options, progress, cancellationToken);
+			return await writer.WriteHashedNodeAsync(outputNode, cancellationToken);
+		}
+
+		/// <summary>
+		/// Writes a tree of files to a storage writer
+		/// </summary>
+		public static async Task<HashedNodeRef<DirectoryNode>> WriteFilesAsync(this IStorageWriter writer, DirectoryReference baseDir, IReadOnlyList<FileReference> files, ChunkingOptions? options = null, IProgress<ICopyStats>? progress = null, CancellationToken cancellationToken = default)
+		{
+			DirectoryNode outputNode = new DirectoryNode();
+			await outputNode.AddFilesAsync(baseDir, files, writer, options, progress, cancellationToken);
+			return await writer.WriteHashedNodeAsync(outputNode, cancellationToken);
+		}
+	}
 }
diff --git a/Engine/Source/Programs/UnrealBuildTool/Executors/Experimental/UBAHordeSession.cs b/Engine/Source/Programs/UnrealBuildTool/Executors/Experimental/UBAHordeSession.cs
--- a/Engine/Source/Programs/UnrealBuildTool/Executors/Experimental/UBAHordeSession.cs
+++ b/Engine/Source/Programs/UnrealBuildTool/Executors/Experimental/UBAHordeSession.cs
@@ -303,7 +303,7 @@ namespace UnrealBuildTool
 		{
 			await using IStorageWriter writer = _storage.CreateWriter();
 			DirectoryNode sandbox = new();
-			await sandbox.AddFilesAsync(baseDir, files, new ChunkingOptions(), writer, null, cancellationToken);
+			await sandbox.AddFilesAsync(baseDir, files, writer, cancellationToken: cancellationToken);
 			IBlobHandle handle = await writer.FlushAsync(sandbox, cancellationToken);
 			return handle.GetLocator();
 		}
