Extension methods to allow acquiring an IDisposable reference to a semaphore.

diff --git a/Engine/Source/Programs/Horde/Horde.Server/Storage/Backends/AwsStorageBackend.cs b/Engine/Source/Programs/Horde/Horde.Server/Storage/Backends/AwsStorageBackend.cs
--- a/Engine/Source/Programs/Horde/Horde.Server/Storage/Backends/AwsStorageBackend.cs
+++ b/Engine/Source/Programs/Horde/Horde.Server/Storage/Backends/AwsStorageBackend.cs
@@ -293,7 +293,7 @@ namespace Horde.Server.Storage.Backends
 			GetObjectResponse? response = null;
 			try
 			{
-				semaLock = await _semaphore.UseWaitAsync(cancellationToken);
+				semaLock = await _semaphore.WaitDisposableAsync(cancellationToken);
 
 				semaphoreSpan = OpenTelemetryTracers.Horde.StartActiveSpan($"{nameof(AwsStorageBackend)}.{nameof(OpenAsync)}.Semaphore");
 				semaphoreSpan.SetAttribute("path", path);
@@ -418,7 +418,7 @@ namespace Horde.Server.Storage.Backends
 			{
 				try
 				{
-					using IDisposable semaLock = await _semaphore.UseWaitAsync(cancellationToken);
+					using IDisposable semaLock = await _semaphore.WaitDisposableAsync(cancellationToken);
 					await WriteInternalAsync(fullPath, inputStream, cancellationToken);
 					_logger.LogDebug("Written data to {Path}", path);
 					break;
diff --git a/Engine/Source/Programs/Horde/Horde.Server/Utilities/SemaphoreSlimExtensions.cs b/Engine/Source/Programs/Horde/Horde.Server/Utilities/SemaphoreSlimExtensions.cs
deleted file mode 100644
diff --git a/Engine/Source/Programs/Shared/EpicGames.Core/SemaphoreSlimExtensions.cs b/Engine/Source/Programs/Shared/EpicGames.Core/SemaphoreSlimExtensions.cs
new file mode 100644
--- /dev/null
+++ b/Engine/Source/Programs/Shared/EpicGames.Core/SemaphoreSlimExtensions.cs
@@ -0,0 +1,61 @@
+// Copyright Epic Games, Inc. All Rights Reserved.
+
+using System;
+using System.Collections.Generic;
+using System.Linq;
+using System.Text;
+using System.Threading;
+using System.Threading.Tasks;
+
+namespace EpicGames.Core
+{
+	/// <summary>
+	/// Extension method to wrap a semaphore in a using statement
+	/// </summary>
+	public static class SemaphoreSlimExtensions
+	{
+		struct ReleaseWrapper : IDisposable
+		{
+			private readonly SemaphoreSlim _semaphore;
+			private int _disposed;
+
+			public ReleaseWrapper(SemaphoreSlim semaphore)
+			{
+				_semaphore = semaphore;
+				_disposed = 0;
+			}
+
+			public void Dispose()
+			{
+				if (Interlocked.CompareExchange(ref _disposed, 1, 0) == 0)
+				{
+					_semaphore.Release();
+				}
+			}
+		}
+
+		/// <summary>
+		/// Returns a disposable semaphore handle
+		/// </summary>
+		/// <param name="semaphore">the semaphore to wrap</param>
+		/// <param name="cancellationToken">optional cancellation token</param>
+		/// <returns></returns>
+		public static IDisposable WaitDisposable(this SemaphoreSlim semaphore, CancellationToken cancellationToken = default)
+		{
+			semaphore.Wait(cancellationToken);
+			return new ReleaseWrapper(semaphore);
+		}
+
+		/// <summary>
+		/// Returns a disposable semaphore handle
+		/// </summary>
+		/// <param name="semaphore">the semaphore to wrap</param>
+		/// <param name="cancellationToken">optional cancellation token</param>
+		/// <returns></returns>
+		public static async Task<IDisposable> WaitDisposableAsync(this SemaphoreSlim semaphore, CancellationToken cancellationToken = default)
+		{
+			await semaphore.WaitAsync(cancellationToken).ConfigureAwait(false);
+			return new ReleaseWrapper(semaphore);
+		}
+	}
+}
