Horde: Improve callstacks for adding new values to the bundle cache.

diff --git a/Engine/Source/Programs/Shared/EpicGames.Horde/Storage/Bundles/BundleCache.cs b/Engine/Source/Programs/Shared/EpicGames.Horde/Storage/Bundles/BundleCache.cs
--- a/Engine/Source/Programs/Shared/EpicGames.Horde/Storage/Bundles/BundleCache.cs
+++ b/Engine/Source/Programs/Shared/EpicGames.Horde/Storage/Bundles/BundleCache.cs
@@ -60,7 +60,7 @@ namespace EpicGames.Horde.Storage
 				_refCount = 2; // Will be released by RunAndUnlock
 				Key = key;
 				Node = new LinkedListNode<CacheValue>(this);
-				InitTask = Task.Run(() => RunAndUnlockAsync(initTask));
+				InitTask = RunAndUnlockAsync(initTask);
 			}
 
 			public void Dispose()
@@ -70,6 +70,7 @@ namespace EpicGames.Horde.Storage
 
 			async Task<IDisposable> RunAndUnlockAsync(Func<Task<IDisposable>> initTask)
 			{
+				await Task.Yield();
 				try
 				{
 					return await initTask();
