Fix static analysis warning.

diff --git a/Engine/Source/Programs/Shared/EpicGames.Core/SemaphoreSlimExtensions.cs b/Engine/Source/Programs/Shared/EpicGames.Core/SemaphoreSlimExtensions.cs
--- a/Engine/Source/Programs/Shared/EpicGames.Core/SemaphoreSlimExtensions.cs
+++ b/Engine/Source/Programs/Shared/EpicGames.Core/SemaphoreSlimExtensions.cs
@@ -1,9 +1,6 @@
 // Copyright Epic Games, Inc. All Rights Reserved.
 
 using System;
-using System.Collections.Generic;
-using System.Linq;
-using System.Text;
 using System.Threading;
 using System.Threading.Tasks;
 
