Horde: Add support for AppContainers in ManagedProcess

diff --git a/Engine/Source/Programs/Shared/EpicGames.Core.Tests/AppContainerTests.cs b/Engine/Source/Programs/Shared/EpicGames.Core.Tests/AppContainerTests.cs
new file mode 100644
--- /dev/null
+++ b/Engine/Source/Programs/Shared/EpicGames.Core.Tests/AppContainerTests.cs
@@ -0,0 +1,159 @@
+// Copyright Epic Games, Inc. All Rights Reserved.
+
+using System;
+using System.Collections.Generic;
+using System.Diagnostics;
+using System.IO;
+using System.Runtime.InteropServices;
+using System.Security.AccessControl;
+using System.Text;
+using System.Threading;
+using System.Threading.Tasks;
+using Microsoft.VisualStudio.TestTools.UnitTesting;
+
+namespace EpicGames.Core.Tests;
+
+// Suppress call sites not available on all platforms
+// OS check is done in test class initializer below
+#pragma warning disable CA1416
+
+[TestClass]
+public class AppContainerTests
+{
+	private const string FileContent = "foobar";
+	private const string RegeditExePath = @"c:\Windows\regedit.exe";
+	private const string DisplayName = "Unreal Engine Horde Test";
+	private const string Description = "Created by Unreal Engine tests";
+
+	private static int s_idCounter = 1000;
+	private static readonly HashSet<string> s_allocatedNames = new();
+	private readonly DirectoryInfo _tempDir;
+	private readonly FileInfo _tempFile;
+
+	public AppContainerTests()
+	{
+		string uniqueTempDir = Path.Join(Path.GetTempPath(), Path.GetRandomFileName());
+		string uniqueTempFile = Path.Join(uniqueTempDir, Path.GetRandomFileName());
+		_tempDir = Directory.CreateDirectory(uniqueTempDir);
+		_tempFile = new FileInfo(uniqueTempFile);
+		File.WriteAllText(uniqueTempFile, FileContent);
+	}
+	
+	[ClassInitialize]
+	public static void ClassInit(TestContext _)
+	{
+		if (!RuntimeInformation.IsOSPlatform(OSPlatform.Windows))
+		{
+			Assert.Inconclusive("AppContainer tests only run on Windows");
+		}
+	}
+	
+	[TestCleanup]
+	public void Cleanup()
+	{
+		foreach (string containerName in s_allocatedNames)
+		{
+			AppContainer.Delete(containerName);
+		}
+		
+		_tempDir.Delete(true);
+	}
+	
+	[TestMethod]
+	public void CreateOrUseExisting()
+	{
+		string containerName = GetName(20);
+		AppContainer.Delete(containerName);
+		using AppContainer ac = AppContainer.Create(containerName, DisplayName, Description);
+		using AppContainer sameAc = AppContainer.Create(containerName, DisplayName, Description);
+
+		Assert.ThrowsException<AppContainerException>(() =>
+		{
+			using AppContainer errored = AppContainer.Create(containerName, DisplayName, Description, true);
+		});
+		Assert.AreEqual(ac.ContainerName, sameAc.ContainerName);
+	}
+	
+	[TestMethod]
+	public async Task FileSystem_CanRead_AllowedDir_Async()
+	{
+		using AppContainer appContainer = Create();
+		Assert.IsTrue(await IsFileReadableAsync(null, _tempFile.FullName));
+		Assert.IsFalse(await IsFileReadableAsync(appContainer, _tempFile.FullName));
+		appContainer.AddDirectoryAccess(_tempDir, FileSystemRights.FullControl);
+		Assert.IsTrue(await IsFileReadableAsync(appContainer, _tempFile.FullName));
+	}
+	
+	[TestMethod]
+	public async Task FileSystem_CannotRead_DefaultTempDir_Async()
+	{
+		using AppContainer appContainer = Create();
+		Assert.IsTrue(await IsFileReadableAsync(null));
+		Assert.IsFalse(await IsFileReadableAsync(appContainer));
+	}
+	
+	[TestMethod]
+	public async Task FileSystem_CanRead_WindowsDir_Async()
+	{
+		// Windows dir can be read by restricted apps by default
+		using AppContainer appContainer = Create();
+		Assert.IsTrue(await IsFileReadableAsync(null, RegeditExePath, "run in DOS mode"));
+		Assert.IsTrue(await IsFileReadableAsync(appContainer, RegeditExePath, "run in DOS mode"));
+	}
+
+	private static async Task<bool> IsFileReadableAsync(AppContainer? appContainer, string? filename = null, string? matchString = FileContent)
+	{
+		if (filename == null)
+		{
+			filename = Path.GetTempFileName();
+			await File.WriteAllTextAsync(filename, FileContent);
+		}
+
+		// Use built-in Windows tool findstr.exe as an external process to launch and read from the file system with
+		// findstr.exe finds a string in the given file. If exit code is zero, file could be read and string was found.
+		string exe = @"c:\Windows\System32\findstr.exe";
+		
+		if (appContainer != null)
+		{
+			string destFileName = Path.Join(appContainer.GetFolderPath().FullName, "target.exe");
+			File.Copy(exe, destFileName, true);
+			exe = destFileName;
+		}
+		
+		(int exitCode, string _) = await RunProcessAsync(exe, $"\"{matchString}\" {filename}", appContainer);
+
+		return exitCode switch
+		{
+			0 => true,
+			1 => false,
+			2 => false,
+			_ => throw new Exception("Unknown exit code: " + exitCode)
+		};
+	}
+
+	private static string GetName(int? id = null)
+	{
+		id ??= Interlocked.Increment(ref s_idCounter);
+		string containerName = "EpicGames-UE-Test-" + id;
+		s_allocatedNames.Add(containerName);
+		return containerName;
+	}
+	
+	private static AppContainer Create()
+	{
+		string containerName = GetName();
+		AppContainer.Delete(containerName);
+		return AppContainer.Create(containerName, DisplayName, Description);
+	}
+
+	private static async Task<(int exitCode, string output)> RunProcessAsync(string fileName, string commandLine, AppContainer? appContainer)
+	{
+		using ManagedProcess process = new (null, fileName, commandLine, null, null, null, ProcessPriorityClass.Normal, appContainer);
+		using MemoryStream ms = new(500);
+		await process.CopyToAsync((buffer, offset, length) => { ms.Write(buffer, offset, length); }, 4096, CancellationToken.None);
+		await process.WaitForExitAsync(CancellationToken.None);
+		return (process.ExitCode, Encoding.UTF8.GetString(ms.ToArray()));
+	}
+}
+
+#pragma warning restore CA1416
\ No newline at end of file
diff --git a/Engine/Source/Programs/Shared/EpicGames.Core/AppContainer.cs b/Engine/Source/Programs/Shared/EpicGames.Core/AppContainer.cs
new file mode 100644
--- /dev/null
+++ b/Engine/Source/Programs/Shared/EpicGames.Core/AppContainer.cs
@@ -0,0 +1,319 @@
+// Copyright Epic Games, Inc. All Rights Reserved.
+
+using System;
+using System.ComponentModel;
+using System.Diagnostics.CodeAnalysis;
+using System.IO;
+using System.Runtime.InteropServices;
+using System.Runtime.Versioning;
+using System.Security.AccessControl;
+using System.Security.Principal;
+
+namespace EpicGames.Core;
+
+/// <summary>
+/// Exception for AppContainer
+/// </summary>
+public class AppContainerException : Exception
+{
+	/// <inheritdoc/>
+	public AppContainerException(string? message) : base(message)
+	{
+	}
+	
+	/// <inheritdoc/>
+	public AppContainerException(string? message, Exception? innerException) : base(message, innerException)
+	{
+	}
+}
+
+/// <summary>
+/// Manages a setup and initialization of an AppContainer on Windows
+/// </summary>
+[SupportedOSPlatform("windows")]
+public sealed class AppContainer : IDisposable
+{
+	#region P/Invoke
+	private enum HRESULT : uint
+	{
+		S_FALSE = 0x0001,
+		S_OK = 0x0000,
+		E_INVALIDARG = 0x80070057,
+	}
+	
+	private enum PROC_THREAD_ATTRIBUTES
+	{
+		PROC_THREAD_ATTRIBUTE_SECURITY_CAPABILITIES = 0x00020009,
+	}
+
+	private enum Win32Error : uint
+	{
+		ERROR_ALREADY_EXISTS = 0x000000B7
+	}
+	
+	[StructLayout(LayoutKind.Sequential)]
+	[SuppressMessage("Style", "IDE1006:Naming Styles", Justification = "Win32 API naming convention")]
+	private struct SID_AND_ATTRIBUTES
+	{
+		public IntPtr Sid;
+		public int Attributes;
+	}
+	
+	[StructLayout(LayoutKind.Sequential)]
+	[SuppressMessage("Style", "IDE1006:Naming Styles", Justification = "Win32 API naming convention")]
+	private struct SECURITY_CAPABILITIES
+	{
+		public IntPtr AppContainerSid;
+		public IntPtr Capabilities;
+		public uint CapabilityCount;
+		public uint Reserved;
+	}
+	
+	[DllImport("advapi32.dll")]
+	private static extern IntPtr FreeSid(IntPtr pSid);
+	
+	[DllImport("advapi32.dll", CharSet = CharSet.Auto, SetLastError = true)]
+	private static extern bool ConvertSidToStringSid(IntPtr pSid, out IntPtr ptrSid);
+	
+	[DllImport("kernel32.dll", SetLastError = true)]
+	private static extern bool DeleteProcThreadAttributeList(IntPtr lpAttributeList);
+	
+	[DllImport("kernel32.dll", SetLastError = true)]
+	static extern IntPtr LocalFree(IntPtr hMem);
+	
+	[DllImport("kernel32.dll", SetLastError = true)]
+	[return: MarshalAs(UnmanagedType.Bool)]
+	private static extern bool UpdateProcThreadAttribute(
+		IntPtr lpAttributeList,
+		uint dwFlags,
+		IntPtr attribute,
+		IntPtr lpValue,
+		IntPtr cbSize,
+		IntPtr lpPreviousValue,
+		IntPtr lpReturnSize);
+	
+	[DllImport("kernel32.dll", SetLastError = true)]
+	[return: MarshalAs(UnmanagedType.Bool)]
+	private static extern bool InitializeProcThreadAttributeList(
+		IntPtr lpAttributeList,
+		int dwAttributeCount,
+		int dwFlags,
+		ref IntPtr lpSize);
+	
+	[DllImport("userenv.dll", SetLastError = false, ExactSpelling = true)]
+	private static extern HRESULT CreateAppContainerProfile(
+		[MarshalAs(UnmanagedType.LPWStr)] string pszAppContainerName,
+		[MarshalAs(UnmanagedType.LPWStr)] string pszDisplayName,
+		[MarshalAs(UnmanagedType.LPWStr)] string pszDescription,
+		[In] SID_AND_ATTRIBUTES[] pCapabilities,
+		uint dwCapabilityCount,
+		out IntPtr ppSidAppContainerSid);
+	
+	[DllImport("userenv.dll", SetLastError = false, ExactSpelling = true)]
+	private static extern HRESULT GetAppContainerFolderPath(
+		[MarshalAs(UnmanagedType.LPWStr)] string pszAppContainerSid,
+		[MarshalAs(UnmanagedType.LPWStr)] out string ppszPath);
+
+	
+	[DllImport("userenv.dll", SetLastError = false, ExactSpelling = true)]
+	private static extern HRESULT DeleteAppContainerProfile([MarshalAs(UnmanagedType.LPWStr)] string pszAppContainerName);
+
+	[DllImport("userenv.dll", SetLastError = false, ExactSpelling = true)]
+	private static extern HRESULT DeriveAppContainerSidFromAppContainerName([MarshalAs(UnmanagedType.LPWStr)] string pszAppContainerName, out IntPtr ppsidAppContainerSid);
+	
+	private static string ManagedConvertSidToStringSid(IntPtr sid)
+	{
+		IntPtr sidStringPtr = IntPtr.Zero;
+		try
+		{
+			if (!ConvertSidToStringSid(sid, out sidStringPtr))
+			{
+				throw new AppContainerException("Failed converting SID to string SID", new Win32Exception());
+			}
+			
+			string? sidString = Marshal.PtrToStringAuto(sidStringPtr);
+			if (sidString == null)
+			{
+				throw new AppContainerException("Failed converting string pointer");
+			}
+			return sidString;
+		}
+		finally
+		{
+			if (sidStringPtr != IntPtr.Zero)
+			{
+				LocalFree(sidStringPtr);
+			}
+		}
+	}
+	
+	private static HRESULT HResultFromWin32(uint errorCode)
+	{
+		const int FacilityWin32 = 7;
+		if (errorCode <= 0)
+		{
+			return (HRESULT)errorCode;
+		}
+		return (HRESULT)((errorCode & 0x0000FFFF) | (FacilityWin32 << 16) | 0x80000000);
+	}
+	
+	#endregion
+	
+	/// <summary>
+	/// Name of the container
+	/// </summary>
+	public string ContainerName { get; }
+
+	private readonly IntPtr _sid;
+	private IntPtr _attribList;
+	private bool _isDisposed;
+	
+	private AppContainer(string containerName, IntPtr sid)
+	{
+		ContainerName = containerName;
+		_sid = sid;
+	}
+	
+	/// <inheritdoc/>
+	public void Dispose()
+	{
+		if (!_isDisposed)
+		{
+			if (_sid != IntPtr.Zero)
+			{
+				FreeSid(_sid);
+			}
+			
+			if (_attribList != IntPtr.Zero)
+			{
+				DeleteProcThreadAttributeList(_attribList);	
+			}
+			
+			GC.SuppressFinalize(this);
+			_isDisposed = true;
+		}
+	}
+	
+	/// <summary>
+	/// Create an AppContainer
+	/// </summary>
+	/// <param name="containerName">Unique container name (see CreateAppContainerProfile call for details)</param>
+	/// <param name="displayName">Display name</param>
+	/// <param name="description">Description</param>
+	/// <param name="errorIfExists">Throw an exception if container already exists, do not attempt to get it</param>
+	/// <returns>An initialized AppContainer</returns>
+	/// <exception cref="AppContainerException"></exception>
+	public static AppContainer Create(string containerName, string displayName, string description, bool errorIfExists = false)
+	{
+		HRESULT res = CreateAppContainerProfile(containerName, displayName, description, null!, 0, out IntPtr sid);
+		if (!errorIfExists && res == HResultFromWin32((uint)Win32Error.ERROR_ALREADY_EXISTS))
+		{
+			res = DeriveAppContainerSidFromAppContainerName(containerName, out sid);
+		}
+	
+		if (res != HRESULT.S_OK)
+		{
+			FreeSid(sid);
+			throw new AppContainerException($"Unable to create AppContainer profile '{containerName}. Error: {res}");
+		}
+		
+		return new AppContainer(containerName, sid);
+	}
+	
+	/// <summary>
+	/// Delete an AppContainer
+	/// </summary>
+	/// <param name="containerName">Name</param>
+	/// <exception cref="AppContainerException">If deletion fails</exception>
+	public static void Delete(string containerName)
+	{
+		HRESULT res = DeleteAppContainerProfile(containerName);
+		if (res != HRESULT.S_OK)
+		{
+			throw new AppContainerException($"Unable to delete AppContainer profile '{containerName}. Error: {res}");
+		}
+	}
+	
+	/// <summary>
+	/// Get home directory path for the AppContainer
+	/// </summary>
+	/// <returns>Path to directory</returns>
+	/// <exception cref="Exception"></exception>
+	public DirectoryInfo GetFolderPath()
+	{
+		string sidString = ManagedConvertSidToStringSid(_sid);
+		HRESULT res = GetAppContainerFolderPath(sidString, out string folderPath);
+		if (res != HRESULT.S_OK)
+		{
+			throw new Exception("Unable to get AppContainer folder path. Error: " + res, new Win32Exception());
+		}
+
+		return new DirectoryInfo(folderPath);
+	}
+	
+	/// <summary>
+	/// Grant permissions to a directory for current AppContainer
+	/// </summary>
+	/// <param name="dirInfo">Directory to modify</param>
+	/// <param name="rights">Permissions to grant</param>
+	public void AddDirectoryAccess(DirectoryInfo dirInfo, FileSystemRights rights)
+	{
+		AccessControlType controlType = AccessControlType.Allow;
+		DirectorySecurity dirSecurity = dirInfo.GetAccessControl();
+		SecurityIdentifier sid = new (_sid);
+		const InheritanceFlags InheritanceFlags = InheritanceFlags.ContainerInherit | InheritanceFlags.ObjectInherit;
+
+		dirSecurity.AddAccessRule(new FileSystemAccessRule(sid, rights, InheritanceFlags, PropagationFlags.None, controlType));
+		dirInfo.SetAccessControl(dirSecurity);
+	}
+	
+	/// <summary>
+	/// Get attribute list for use in STARTUPINFOEX struct given to CreateProcess
+	/// </summary>
+	/// <returns>An attribute list pointer</returns>
+	/// <exception cref="Exception"></exception>
+	internal IntPtr GetAttributeList()
+	{
+		if (_attribList != IntPtr.Zero)
+		{
+			return _attribList;
+		}
+		
+		SECURITY_CAPABILITIES securityCapabilities = new() { AppContainerSid = _sid, Capabilities = IntPtr.Zero, CapabilityCount = 0, Reserved = 0 };
+		IntPtr lpSize = IntPtr.Zero;
+
+		bool isSuccess = InitializeProcThreadAttributeList(IntPtr.Zero, 1, 0, ref lpSize);
+		if (isSuccess || lpSize == IntPtr.Zero)
+		{
+			throw new Exception("Unable to initialize thread attribute list size", new Win32Exception());
+		}
+		
+		IntPtr attribList = Marshal.AllocHGlobal(lpSize);
+		isSuccess = InitializeProcThreadAttributeList(attribList, 1, 0, ref lpSize);
+		if (!isSuccess)
+		{
+			throw new Exception("Unable to initialize thread attribute list", new Win32Exception());
+		}
+		
+		IntPtr scPtr = Marshal.AllocHGlobal(Marshal.SizeOf(securityCapabilities));
+		Marshal.StructureToPtr(securityCapabilities, scPtr, false);
+
+		isSuccess = UpdateProcThreadAttribute(
+			attribList,
+			0,
+			(IntPtr)PROC_THREAD_ATTRIBUTES.PROC_THREAD_ATTRIBUTE_SECURITY_CAPABILITIES,
+			scPtr,
+			(IntPtr)Marshal.SizeOf(securityCapabilities),
+			IntPtr.Zero,
+			IntPtr.Zero);
+			
+		if (!isSuccess)
+		{
+			throw new Exception("Unable to update thread attribute list", new Win32Exception());
+		}
+
+		_attribList = attribList;
+		return attribList;
+	}
+}
+
diff --git a/Engine/Source/Programs/Shared/EpicGames.Core/ManagedProcess.cs b/Engine/Source/Programs/Shared/EpicGames.Core/ManagedProcess.cs
--- a/Engine/Source/Programs/Shared/EpicGames.Core/ManagedProcess.cs
+++ b/Engine/Source/Programs/Shared/EpicGames.Core/ManagedProcess.cs
@@ -275,6 +275,13 @@ namespace EpicGames.Core
 			public SafeHandle? hStdOutput;
 			public SafeHandle? hStdError;
 		}
+		
+		[StructLayout(LayoutKind.Sequential)]
+		class STARTUPINFOEX
+		{
+			public STARTUPINFO? StartupInfo;
+			public IntPtr lpAttributeList;
+		}
 
 		[DllImport("kernel32.dll", SetLastError = true)]
 		static extern int CloseHandle(IntPtr hObject);
@@ -300,11 +307,12 @@ namespace EpicGames.Core
 			REALTIME_PRIORITY_CLASS = 0x00000100,
 			BELOW_NORMAL_PRIORITY_CLASS = 0x00004000,
 			ABOVE_NORMAL_PRIORITY_CLASS = 0x00008000,
+			EXTENDED_STARTUPINFO_PRESENT = 0x00080000
 		}
 
-		[DllImport("kernel32.dll", SetLastError = true)]
 #pragma warning disable CA1838 // Avoid 'StringBuilder' parameters for P/Invokes
-		static extern int CreateProcess(/*[MarshalAs(UnmanagedType.LPTStr)]*/ string? lpApplicationName, StringBuilder lpCommandLine, IntPtr lpProcessAttributes, IntPtr lpThreadAttributes, bool bInheritHandles, ProcessCreationFlags dwCreationFlags, IntPtr lpEnvironment, /*[MarshalAs(UnmanagedType.LPTStr)]*/ string? lpCurrentDirectory, STARTUPINFO lpStartupInfo, PROCESS_INFORMATION lpProcessInformation);
+		[DllImport("kernel32.dll", SetLastError = true)]
+		static extern int CreateProcess(/*[MarshalAs(UnmanagedType.LPTStr)]*/ string? lpApplicationName, StringBuilder lpCommandLine, IntPtr lpProcessAttributes, IntPtr lpThreadAttributes, bool bInheritHandles, ProcessCreationFlags dwCreationFlags, IntPtr lpEnvironment, /*[MarshalAs(UnmanagedType.LPTStr)]*/ string? lpCurrentDirectory, STARTUPINFOEX lpStartupInfo, PROCESS_INFORMATION lpProcessInformation);
 #pragma warning restore CA1838 // Avoid 'StringBuilder' parameters for P/Invokes
 
 		[DllImport("kernel32.dll", SetLastError = true)]
@@ -557,9 +565,10 @@ namespace EpicGames.Core
 		/// <param name="environment">Environment variables for the new process. May be null, in which case the current process' environment is inherited</param>
 		/// <param name="input">Text to be passed via stdin to the new process. May be null.</param>
 		/// <param name="priority">Priority for the child process</param>
+		/// <param name="appContainer">AppContainer to use for process (only available on Windows)</param>
 		/// <param name="flags">Flags for the process</param>
-		public ManagedProcess(ManagedProcessGroup? group, string fileName, string commandLine, string? workingDirectory, IReadOnlyDictionary<string, string>? environment, byte[]? input, ProcessPriorityClass priority, ManagedProcessFlags flags = ManagedProcessFlags.MergeOutputPipes)
-			: this(group, fileName, commandLine, workingDirectory, environment, priority, flags)
+		public ManagedProcess(ManagedProcessGroup? group, string fileName, string commandLine, string? workingDirectory, IReadOnlyDictionary<string, string>? environment, byte[]? input, ProcessPriorityClass priority, AppContainer? appContainer = null, ManagedProcessFlags flags = ManagedProcessFlags.MergeOutputPipes)
+			: this(group, fileName, commandLine, workingDirectory, environment, priority, appContainer, flags)
 		{
 			if (input != null)
 			{
@@ -577,8 +586,9 @@ namespace EpicGames.Core
 		/// <param name="workingDirectory">Working directory for the new process. May be null to use the current working directory.</param>
 		/// <param name="environment">Environment variables for the new process. May be null, in which case the current process' environment is inherited</param>
 		/// <param name="priority">Priority for the child process</param>
+		/// <param name="appContainer">AppContainer to use for process (only available on Windows)</param>
 		/// <param name="flags">Flags for the process</param>
-		public ManagedProcess(ManagedProcessGroup? group, string fileName, string commandLine, string? workingDirectory, IReadOnlyDictionary<string, string>? environment, ProcessPriorityClass priority, ManagedProcessFlags flags = ManagedProcessFlags.MergeOutputPipes)
+		public ManagedProcess(ManagedProcessGroup? group, string fileName, string commandLine, string? workingDirectory, IReadOnlyDictionary<string, string>? environment, ProcessPriorityClass priority, AppContainer? appContainer = null, ManagedProcessFlags flags = ManagedProcessFlags.MergeOutputPipes)
 		{
 			// Create the child process
 			// NOTE: Child process must be created in a separate method to avoid stomping exception callstacks (https://stackoverflow.com/a/2494150)
@@ -586,7 +596,7 @@ namespace EpicGames.Core
 			{
 				if (ManagedProcessGroup.SupportsJobObjects)
 				{
-					CreateManagedProcessWin32(group, fileName, commandLine, workingDirectory, environment, priority, flags);
+					CreateManagedProcessWin32(group, fileName, commandLine, workingDirectory, environment, priority, appContainer, flags);
 				}
 				else
 				{
@@ -616,8 +626,9 @@ namespace EpicGames.Core
 		/// <param name="workingDirectory">Working directory for the new process. May be null to use the current working directory.</param>
 		/// <param name="environment">Environment variables for the new process. May be null, in which case the current process' environment is inherited</param>
 		/// <param name="priority">Priority for the child process</param>
+		/// <param name="appContainer">AppContainer to use for the new process</param>
 		/// <param name="managedFlags">Flags controlling how the new process is created</param>
-		private void CreateManagedProcessWin32(ManagedProcessGroup? group, string fileName, string commandLine, string? workingDirectory, IReadOnlyDictionary<string, string>? environment, ProcessPriorityClass priority, ManagedProcessFlags managedFlags)
+		private void CreateManagedProcessWin32(ManagedProcessGroup? group, string fileName, string commandLine, string? workingDirectory, IReadOnlyDictionary<string, string>? environment, ProcessPriorityClass priority, AppContainer? appContainer, ManagedProcessFlags? managedFlags)
 		{
 			IntPtr environmentBlock = IntPtr.Zero;
 			try
@@ -709,17 +720,25 @@ namespace EpicGames.Core
 
 							// Create the new process as suspended, so we can modify it before it starts executing (and potentially preempting us)
 							STARTUPINFO startupInfo = new STARTUPINFO();
-							startupInfo.cb = Marshal.SizeOf(startupInfo);
+							startupInfo.cb = Marshal.SizeOf(typeof(STARTUPINFOEX));
 							startupInfo.hStdInput = stdInRead;
 							startupInfo.hStdOutput = stdOutWrite;
 							startupInfo.hStdError = stdErrWrite;
 							startupInfo.dwFlags = STARTF_USESTDHANDLES;
+							
+							flags |= ProcessCreationFlags.EXTENDED_STARTUPINFO_PRESENT;
+							STARTUPINFOEX startupInfoEx = new() { StartupInfo = startupInfo, lpAttributeList = IntPtr.Zero };
+
+							if (appContainer != null && OperatingSystem.IsWindows())
+							{
+								startupInfoEx.lpAttributeList = appContainer.GetAttributeList();
+							}
 
 							// Under heavy load (ie. spawning large number of processes, typically Clang) we see CreateProcess very occasionally failing with ERROR_ACCESS_DENIED.
 							int[] retryDelay = { 100, 200, 1000, 5000 };
 							for(int attemptIdx = 0; ; attemptIdx++)
 							{
-								if (CreateProcess(null, new StringBuilder("\"" + fileName + "\" " + commandLine), IntPtr.Zero, IntPtr.Zero, true, flags, environmentBlock, workingDirectory, startupInfo, processInfo) != 0)
+								if (CreateProcess(null, new StringBuilder("\"" + fileName + "\" " + commandLine), IntPtr.Zero, IntPtr.Zero, true, flags, environmentBlock, workingDirectory, startupInfoEx, processInfo) != 0)
 								{
 									break;
 								}
