Horde: Ensure process exit is awaited for compute tasks before getting the exit code

Without this, code path occasionally raises an exception with "Process must exit before requested information can be determined".

diff --git a/Engine/Source/Programs/Shared/EpicGames.Horde/Compute/AgentMessageHandler.cs b/Engine/Source/Programs/Shared/EpicGames.Horde/Compute/AgentMessageHandler.cs
--- a/Engine/Source/Programs/Shared/EpicGames.Horde/Compute/AgentMessageHandler.cs
+++ b/Engine/Source/Programs/Shared/EpicGames.Horde/Compute/AgentMessageHandler.cs
@@ -403,22 +403,20 @@ namespace EpicGames.Horde.Compute
 					_logger.LogWarning("Working dir {Path} does not exist", resolvedWorkingDir);	
 				}
 
-				using (ManagedProcessGroup group = new ManagedProcessGroup())
+				using ManagedProcessGroup group = new ManagedProcessGroup();
+				using ManagedProcess process = new ManagedProcess(group, resolvedExecutable, resolvedCommandLine, resolvedWorkingDir, resolvedEnvVars, null, ProcessPriorityClass.Normal);
+				byte[] buffer = new byte[1024];
+				
+				for (; ; )
 				{
-					using (ManagedProcess process = new ManagedProcess(group, resolvedExecutable, resolvedCommandLine, resolvedWorkingDir, resolvedEnvVars, null, ProcessPriorityClass.Normal))
+					int length = await process.ReadAsync(buffer, 0, buffer.Length, cancellationToken);
+					if (length == 0)
 					{
-						byte[] buffer = new byte[1024];
-						for (; ; )
-						{
-							int length = await process.ReadAsync(buffer, 0, buffer.Length, cancellationToken);
-							if (length == 0)
-							{
-								await channel.SendExecuteResultAsync(process.ExitCode, cancellationToken);
-								return;
-							}
-							await channel.SendExecuteOutputAsync(buffer.AsMemory(0, length), cancellationToken);
-						}
+						await process.WaitForExitAsync(cancellationToken);
+						await channel.SendExecuteResultAsync(process.ExitCode, cancellationToken);
+						return;
 					}
+					await channel.SendExecuteOutputAsync(buffer.AsMemory(0, length), cancellationToken);
 				}
 			}
 		}
