Horde: Write the output from any remotely executed processes to the log.

diff --git a/Engine/Source/Programs/Shared/EpicGames.Core/ByteArrayBuilder.cs b/Engine/Source/Programs/Shared/EpicGames.Core/ByteArrayBuilder.cs
--- a/Engine/Source/Programs/Shared/EpicGames.Core/ByteArrayBuilder.cs
+++ b/Engine/Source/Programs/Shared/EpicGames.Core/ByteArrayBuilder.cs
@@ -107,6 +107,22 @@ namespace EpicGames.Core
 			}
 		}
 
+		/// <summary>
+		/// Gets a contiguous, non-owned buffer containing the current data. Will merge together current chunks if necessary.
+		/// </summary>
+		/// <returns></returns>
+		public ReadOnlyMemory<byte> AsMemory()
+		{
+			if (_chunks.Count == 1)
+			{
+				return _chunks[0].WrittenMemory;
+			}
+			else
+			{
+				return ToByteArray();
+			}
+		}
+
 		/// <summary>
 		/// Gets a sequence representing the bytes that have been written so far
 		/// </summary>
diff --git a/Engine/Source/Programs/Shared/EpicGames.Horde/Compute/AgentMessageHandler.cs b/Engine/Source/Programs/Shared/EpicGames.Horde/Compute/AgentMessageHandler.cs
--- a/Engine/Source/Programs/Shared/EpicGames.Horde/Compute/AgentMessageHandler.cs
+++ b/Engine/Source/Programs/Shared/EpicGames.Horde/Compute/AgentMessageHandler.cs
@@ -15,6 +15,7 @@ using System.IO;
 using System.Linq;
 using EpicGames.Horde.Storage.Clients;
 using System.Runtime.ExceptionServices;
+using System.Text;
 
 namespace EpicGames.Horde.Compute
 {
@@ -314,6 +315,52 @@ namespace EpicGames.Horde.Compute
 			}
 		}
 
+		// Helper class to take raw UTF8 output and merge it into log lines
+		class ProcessOutputWriter
+		{
+			readonly string _prefix;
+			readonly ByteArrayBuilder _lineBuffer = new ByteArrayBuilder();
+			readonly ILogger _logger;
+
+			public ProcessOutputWriter(string prefix, ILogger logger)
+			{
+				_prefix = prefix;
+				_logger = logger;
+			}
+
+			public void WriteBytes(ReadOnlySpan<byte> span)
+			{
+				for (; ; )
+				{
+					int newlineIdx = span.IndexOf((byte)'\n');
+					if (newlineIdx == -1)
+					{
+						_lineBuffer.WriteFixedLengthBytes(span);
+						break;
+					}
+
+					ReadOnlySpan<byte> line = span.Slice(0, newlineIdx);
+					if (line.Length > 0 && line[line.Length - 1] == (byte)'\r')
+					{
+						line = line.Slice(0, line.Length - 1);
+					}
+
+					if (_lineBuffer.Length > 0)
+					{
+						_lineBuffer.WriteFixedLengthBytes(line);
+						_logger.LogInformation("{Prefix}: {Line}", _prefix, Encoding.UTF8.GetString(_lineBuffer.AsMemory().Span));
+						_lineBuffer.Clear();
+					}
+					else
+					{
+						_logger.LogInformation("{Prefix}: {Line}", _prefix, Encoding.UTF8.GetString(line));
+					}
+
+					span = span.Slice(newlineIdx + 1);
+				}
+			}
+		}
+
 		async Task ExecuteProcessInternalAsync(AgentMessageChannel channel, string executable, IReadOnlyList<string> arguments, string? workingDir, IReadOnlyDictionary<string, string?>? envVars, ExecuteProcessFlags flags, CancellationToken cancellationToken)
 		{
 			string resolvedExecutable = FileReference.Combine(_sandboxDir, executable).FullName;
@@ -406,7 +453,8 @@ namespace EpicGames.Horde.Compute
 				using ManagedProcessGroup group = new ManagedProcessGroup();
 				using ManagedProcess process = new ManagedProcess(group, resolvedExecutable, resolvedCommandLine, resolvedWorkingDir, resolvedEnvVars, null, ProcessPriorityClass.Normal);
 				byte[] buffer = new byte[1024];
-				
+
+				ProcessOutputWriter outputWriter = new ProcessOutputWriter($"{Path.GetFileNameWithoutExtension(resolvedExecutable)}> ", _logger);
 				for (; ; )
 				{
 					int length = await process.ReadAsync(buffer, 0, buffer.Length, cancellationToken);
@@ -416,7 +464,11 @@ namespace EpicGames.Horde.Compute
 						await channel.SendExecuteResultAsync(process.ExitCode, cancellationToken);
 						return;
 					}
-					await channel.SendExecuteOutputAsync(buffer.AsMemory(0, length), cancellationToken);
+
+					ReadOnlyMemory<byte> output = buffer.AsMemory(0, length);
+					await channel.SendExecuteOutputAsync(output, cancellationToken);
+
+					outputWriter.WriteBytes(output.Span);
 				}
 			}
 		}
