Horde: Add a better error when mapping an empty memory mapped file.

diff --git a/Engine/Source/Programs/Shared/EpicGames.Horde/Storage/Backends/FileStorageBackend.cs b/Engine/Source/Programs/Shared/EpicGames.Horde/Storage/Backends/FileStorageBackend.cs
--- a/Engine/Source/Programs/Shared/EpicGames.Horde/Storage/Backends/FileStorageBackend.cs
+++ b/Engine/Source/Programs/Shared/EpicGames.Horde/Storage/Backends/FileStorageBackend.cs
@@ -204,6 +204,10 @@ namespace EpicGames.Horde.Storage.Backends
 			if (!_pathToMappedFile.TryGetValue(path, out mappedFile))
 			{
 				FileInfo fileInfo = GetBlobFile(path).ToFileInfo();
+				if (fileInfo.Length == 0)
+				{
+					throw new Exception($"Unable to map empty memory mapped file: {fileInfo.FullName}");
+				}
 
 				long maxSize = MaxMappedSize - fileInfo.Length;
 				if (_mappedSize > maxSize || _mappedFiles.Count + 1 > MaxMappedCount)
