Horde: Disabling filesystem cache for bundle data; seems to be causing errors.

diff --git a/Engine/Source/Programs/Horde/Horde.Server/Storage/StorageBackendProvider.cs b/Engine/Source/Programs/Horde/Horde.Server/Storage/StorageBackendProvider.cs
--- a/Engine/Source/Programs/Horde/Horde.Server/Storage/StorageBackendProvider.cs
+++ b/Engine/Source/Programs/Horde/Horde.Server/Storage/StorageBackendProvider.cs
@@ -194,7 +194,7 @@ namespace Horde.Server.Storage
 					{
 #pragma warning disable CA2000 // False positive? (Will be disposed with cache backend wrapper)
 						IStorageBackend backend = new AwsStorageBackend(_serviceProvider.GetRequiredService<IConfiguration>(), config, _serviceProvider.GetRequiredService<ILogger<AwsStorageBackend>>());
-						return _storageBackendCache.CreateWrapper(config.Id.ToString(), backend);
+						return backend;// _storageBackendCache.CreateWrapper(config.Id.ToString(), backend);
 #pragma warning restore CA2000
 					}
 				case StorageBackendType.Memory:
diff --git a/Engine/Source/Programs/Shared/EpicGames.Horde/Storage/Clients/HttpStorageClient.cs b/Engine/Source/Programs/Shared/EpicGames.Horde/Storage/Clients/HttpStorageClient.cs
--- a/Engine/Source/Programs/Shared/EpicGames.Horde/Storage/Clients/HttpStorageClient.cs
+++ b/Engine/Source/Programs/Shared/EpicGames.Horde/Storage/Clients/HttpStorageClient.cs
@@ -320,10 +320,10 @@ namespace EpicGames.Horde.Storage.Clients
 			}
 
 			IStorageBackend backend = new HttpStorageBackend(basePath, CreateClient, _backendLogger);
-			if (_backendCache != null && withBackendCache)
-			{
-				backend = _backendCache.CreateWrapper(basePath, backend);
-			}
+//			if (_backendCache != null && withBackendCache)
+//			{
+//				backend = _backendCache.CreateWrapper(basePath, backend);
+//			}
 
 			HttpStorageClient client = new HttpStorageClient(basePath, CreateClient, backend, _clientLogger);
 			return new BundleStorageClient(client, _bundleCache, _clientLogger);
