Horde: Improve auth error messaging in compute client

diff --git a/Engine/Source/Programs/Shared/EpicGames.Horde/Compute/Clients/ServerComputeClient.cs b/Engine/Source/Programs/Shared/EpicGames.Horde/Compute/Clients/ServerComputeClient.cs
--- a/Engine/Source/Programs/Shared/EpicGames.Horde/Compute/Clients/ServerComputeClient.cs
+++ b/Engine/Source/Programs/Shared/EpicGames.Horde/Compute/Clients/ServerComputeClient.cs
@@ -253,6 +253,20 @@ namespace EpicGames.Horde.Compute.Clients
 					yield break;
 				}
 
+				if (httpResponse.StatusCode == HttpStatusCode.Unauthorized)
+				{
+					throw new ComputeClientException($"Bad authentication credentials. Check or refresh token. (HTTP status {httpResponse.StatusCode})");
+				}
+				
+				if (httpResponse.StatusCode == HttpStatusCode.Forbidden)
+				{
+					LogEvent? logEvent = await httpResponse.Content.ReadFromJsonAsync<LogEvent>(HordeHttpClient.JsonSerializerOptions, cancellationToken);
+					if (logEvent != null)
+					{
+						throw new ComputeClientException($"{logEvent.Message} (HTTP status {httpResponse.StatusCode})");
+					}
+				}
+
 				httpResponse.EnsureSuccessStatusCode();
 				response = await httpResponse.Content.ReadFromJsonAsync<AssignComputeResponse>(HordeHttpClient.JsonSerializerOptions, cancellationToken);
 				if (response == null)
@@ -363,7 +377,7 @@ namespace EpicGames.Horde.Compute.Clients
 	/// <summary>
 	/// Exception indicating that no matching compute agents were found
 	/// </summary>
-	public sealed class NoComputeAgentsFoundException : Exception
+	public sealed class NoComputeAgentsFoundException : ComputeClientException
 	{
 		/// <summary>
 		/// The compute cluster requested
diff --git a/Engine/Source/Programs/Shared/EpicGames.Horde/Compute/IComputeClient.cs b/Engine/Source/Programs/Shared/EpicGames.Horde/Compute/IComputeClient.cs
--- a/Engine/Source/Programs/Shared/EpicGames.Horde/Compute/IComputeClient.cs
+++ b/Engine/Source/Programs/Shared/EpicGames.Horde/Compute/IComputeClient.cs
@@ -62,4 +62,15 @@ namespace EpicGames.Horde.Compute
 			return computeClient.TryAssignWorkerAsync(clusterId, requirements, requestId, null, logger, cancellationToken);
 		}
 	}
+	
+	/// <summary>
+	/// Exception from ComputeClient
+	/// </summary>
+	public class ComputeClientException : Exception
+	{
+		/// <inheritdoc/>
+		public ComputeClientException(string? message) : base(message)
+		{
+		}
+	}
 }
