Horde: Add a FileMetadataDb class, which stores the contents of a directory tree in a SQLite database.

diff --git a/Engine/Source/Programs/Shared/EpicGames.Horde.Tests/FileMetadataDbTests.cs b/Engine/Source/Programs/Shared/EpicGames.Horde.Tests/FileMetadataDbTests.cs
new file mode 100644
--- /dev/null
+++ b/Engine/Source/Programs/Shared/EpicGames.Horde.Tests/FileMetadataDbTests.cs
@@ -0,0 +1,209 @@
+﻿// Copyright Epic Games, Inc. All Rights Reserved.
+
+using System;
+using System.Collections.Generic;
+using System.Linq;
+using System.Threading.Tasks;
+using EpicGames.Core;
+using EpicGames.Horde.Storage;
+using Microsoft.VisualStudio.TestTools.UnitTesting;
+
+namespace EpicGames.Horde.Tests
+{
+	[TestClass]
+	public class FileMetadataDbTests
+	{
+		[TestMethod]
+		public async Task DirectoryTestsAsync()
+		{
+			using FileMetadataDb db = await FileMetadataDb.CreateInMemoryAsync();
+
+			int a = await db.AddDirectoryAsync(FileMetadataDb.RootDirectoryId, "A");
+			int b = await db.AddDirectoryAsync(FileMetadataDb.RootDirectoryId, "B");
+			int c = await db.AddDirectoryAsync(b, "C");
+			int d = await db.AddDirectoryAsync(b, "D");
+
+			{
+				List<DirectoryRow> rows = (await db.GetDirectoriesAsync(FileMetadataDb.RootDirectoryId)).OrderBy(x => x.Id).ToList();
+				Assert.AreEqual(2, rows.Count);
+				Assert.AreEqual(a, rows[0].Id);
+				Assert.AreEqual("A", rows[0].Name);
+				Assert.AreEqual(b, rows[1].Id);
+				Assert.AreEqual("B", rows[1].Name);
+			}
+
+			{
+				List<DirectoryRow> rows = (await db.GetDirectoriesAsync(a)).OrderBy(x => x.Id).ToList();
+				Assert.AreEqual(0, rows.Count);
+			}
+
+			{
+				List<DirectoryRow> rows = (await db.GetDirectoriesAsync(b)).OrderBy(x => x.Id).ToList();
+				Assert.AreEqual(2, rows.Count);
+				Assert.AreEqual(c, rows[0].Id);
+				Assert.AreEqual("C", rows[0].Name);
+				Assert.AreEqual(d, rows[1].Id);
+				Assert.AreEqual("D", rows[1].Name);
+			}
+
+			Assert.AreEqual("", await db.GetDirectoryPathAsync(FileMetadataDb.RootDirectoryId));
+			Assert.AreEqual("A/", await db.GetDirectoryPathAsync(a));
+			Assert.AreEqual("B/", await db.GetDirectoryPathAsync(b));
+			Assert.AreEqual("B/C/", await db.GetDirectoryPathAsync(c));
+			Assert.AreEqual("B/D/", await db.GetDirectoryPathAsync(d));
+		}
+
+		[TestMethod]
+		public async Task FileTestsAsync()
+		{
+			using FileMetadataDb db = await FileMetadataDb.CreateInMemoryAsync();
+
+			int a = await db.AddDirectoryAsync(FileMetadataDb.RootDirectoryId, "A");
+
+			DateTime time = new DateTime(2024, 1, 1, 12, 30, 45, DateTimeKind.Utc);
+			int c = await db.AddFileAsync(FileMetadataDb.RootDirectoryId, "C.txt", time, 1234);
+			int d = await db.AddFileAsync(FileMetadataDb.RootDirectoryId, "D.txt", time, 2345);
+			int e = await db.AddFileAsync(a, "E.txt", time, 3456);
+			int f = await db.AddFileAsync(a, "F.txt", time, 4567);
+
+			{
+				List<FileRow> files = (await db.FindFilesInDirectoryAsync(FileMetadataDb.RootDirectoryId)).OrderBy(x => x.Id).ToList();
+				Assert.AreEqual(2, files.Count);
+
+				Assert.AreEqual(c, files[0].Id);
+				Assert.AreEqual(time, files[0].Time);
+				Assert.AreEqual(1234, files[0].Length);
+
+				Assert.AreEqual(d, files[1].Id);
+				Assert.AreEqual(time, files[1].Time);
+				Assert.AreEqual(2345, files[1].Length);
+			}
+
+			{
+				List<FileRow> files = (await db.FindFilesInDirectoryAsync(a)).OrderBy(x => x.Id).ToList();
+				Assert.AreEqual(2, files.Count);
+
+				Assert.AreEqual(e, files[0].Id);
+				Assert.AreEqual(time, files[0].Time);
+				Assert.AreEqual(3456, files[0].Length);
+
+				Assert.AreEqual(f, files[1].Id);
+				Assert.AreEqual(time, files[1].Time);
+				Assert.AreEqual(4567, files[1].Length);
+			}
+
+			Assert.AreEqual("C.txt", await db.GetFilePathAsync(c));
+			Assert.AreEqual("D.txt", await db.GetFilePathAsync(d));
+			Assert.AreEqual("A/E.txt", await db.GetFilePathAsync(e));
+			Assert.AreEqual("A/F.txt", await db.GetFilePathAsync(f));
+		}
+
+		[TestMethod]
+		public async Task BatchFileTestsAsync()
+		{
+			using FileMetadataDb db = await FileMetadataDb.CreateInMemoryAsync();
+			int root = await db.AddDirectoryAsync(FileMetadataDb.RootDirectoryId, "");
+
+			List<FileRow> files = new List<FileRow>();
+			files.Add(new FileRow(root, "A", DateTime.MinValue, 1));
+			files.Add(new FileRow(root, "B", DateTime.MinValue, 1));
+			files.Add(new FileRow(root, "C", DateTime.MinValue, 1));
+
+			await db.AddFilesAsync(files);
+			Assert.AreEqual(3, (await db.FindFilesInDirectoryAsync(root)).Count());
+		}
+
+		[TestMethod]
+		public async Task BatchDirectoryTestsAsync()
+		{
+			using FileMetadataDb db = await FileMetadataDb.CreateInMemoryAsync();
+			await db.AddDirectoriesAsync(new[] { new DirectoryRow(FileMetadataDb.RootDirectoryId, "A"), new DirectoryRow(FileMetadataDb.RootDirectoryId, "B"), new DirectoryRow(FileMetadataDb.RootDirectoryId, "C") });
+			Assert.AreEqual(3, (await db.GetDirectoriesAsync(FileMetadataDb.RootDirectoryId)).Count());
+		}
+
+		[TestMethod]
+		public async Task ChunkTestsAsync()
+		{
+			using FileMetadataDb db = await FileMetadataDb.CreateInMemoryAsync();
+
+			int d1 = await db.AddDirectoryAsync(FileMetadataDb.RootDirectoryId, "D1");
+			int d2 = await db.AddDirectoryAsync(d1, "D2");
+
+			DateTime time = new DateTime(2024, 1, 1, 12, 30, 45, DateTimeKind.Utc);
+			int f1 = await db.AddFileAsync(d1, "F1.txt", time, 1234);
+			int f2 = await db.AddFileAsync(d1, "F2.txt", time, 2345);
+			int f3 = await db.AddFileAsync(d2, "F3.txt", time, 3456);
+			int f4 = await db.AddFileAsync(d2, "F4.txt", time, 4567);
+
+			IoHash h1 = IoHash.Compute(new byte[] { 1, 2, 3 });
+			IoHash h2 = IoHash.Compute(new byte[] { 1, 2, 3, 4 });
+
+			int c1 = await db.AddChunkAsync(f1, 0, 10, h1);
+			int c2 = await db.AddChunkAsync(f2, 10, 10, h1);
+			int c3 = await db.AddChunkAsync(f3, 20, 10, h1);
+			int c4 = await db.AddChunkAsync(f4, 0, 10, h2);
+
+			{
+				List<ChunkRow> chunks = (await db.FindChunksAsync(h1, 10)).OrderBy(x => x.Id).ToList();
+				Assert.AreEqual(3, chunks.Count);
+				Assert.AreEqual(new ChunkRow(c1, f1, 0, 10, h1), chunks[0]);
+				Assert.AreEqual(new ChunkRow(c2, f2, 10, 10, h1), chunks[1]);
+				Assert.AreEqual(new ChunkRow(c3, f3, 20, 10, h1), chunks[2]);
+			}
+
+			{
+				List<ChunkRow> chunks = (await db.FindChunksAsync(h2, 10)).OrderBy(x => x.Id).ToList();
+				Assert.AreEqual(1, chunks.Count);
+				Assert.AreEqual(new ChunkRow(c4, f4, 0, 10, h2), chunks[0]);
+			}
+		}
+
+		[TestMethod]
+		public async Task RemoveTreeTestAsync()
+		{
+			using FileMetadataDb db = await FileMetadataDb.CreateInMemoryAsync();
+
+			int d1 = await db.AddDirectoryAsync(FileMetadataDb.RootDirectoryId, "D1");
+			int d2 = await db.AddDirectoryAsync(d1, "D2");
+			int d3 = await db.AddDirectoryAsync(d2, "D3");
+			int d4 = await db.AddDirectoryAsync(d2, "D4");
+			int d5 = await db.AddDirectoryAsync(d4, "D5");
+
+			DateTime time = new DateTime(2024, 1, 1, 12, 30, 45, DateTimeKind.Utc);
+			int f1 = await db.AddFileAsync(d1, "F1.txt", time, 1234);
+			int f2 = await db.AddFileAsync(d2, "F2.txt", time, 2345);
+			int f3 = await db.AddFileAsync(d3, "F3.txt", time, 3456);
+			int f4 = await db.AddFileAsync(d4, "F4.txt", time, 4567);
+			int f5 = await db.AddFileAsync(d5, "F5.txt", time, 5678);
+
+			IoHash h1 = IoHash.Compute(new byte[] { 1, 2, 3 });
+
+			int c1 = await db.AddChunkAsync(f1, 0, 10, h1);
+			int c2 = await db.AddChunkAsync(f2, 10, 10, h1);
+			int c3 = await db.AddChunkAsync(f3, 20, 10, h1);
+			int c4 = await db.AddChunkAsync(f4, 0, 10, h1);
+			int c5 = await db.AddChunkAsync(f5, 10, 10, h1);
+
+			{
+				List<ChunkRow> chunks = (await db.FindChunksAsync(h1, 10)).OrderBy(x => x.Id).ToList();
+				Assert.AreEqual(5, chunks.Count);
+				Assert.AreEqual(new ChunkRow(c1, f1, 0, 10, h1), chunks[0]);
+				Assert.AreEqual(new ChunkRow(c2, f2, 10, 10, h1), chunks[1]);
+				Assert.AreEqual(new ChunkRow(c3, f3, 20, 10, h1), chunks[2]);
+				Assert.AreEqual(new ChunkRow(c4, f4, 0, 10, h1), chunks[3]);
+				Assert.AreEqual(new ChunkRow(c5, f5, 10, 10, h1), chunks[4]);
+			}
+
+			await db.RemoveDirectoryAsync(d4);
+
+			{
+				List<ChunkRow> chunks = (await db.FindChunksAsync(h1, 10)).OrderBy(x => x.Id).ToList();
+				Assert.AreEqual(3, chunks.Count);
+
+				Assert.AreEqual(new ChunkRow(c1, f1, 0, 10, h1), chunks[0]);
+				Assert.AreEqual(new ChunkRow(c2, f2, 10, 10, h1), chunks[1]);
+				Assert.AreEqual(new ChunkRow(c3, f3, 20, 10, h1), chunks[2]);
+			}
+		}
+	}
+}
diff --git a/Engine/Source/Programs/Shared/EpicGames.Horde/EpicGames.Horde.csproj b/Engine/Source/Programs/Shared/EpicGames.Horde/EpicGames.Horde.csproj
--- a/Engine/Source/Programs/Shared/EpicGames.Horde/EpicGames.Horde.csproj
+++ b/Engine/Source/Programs/Shared/EpicGames.Horde/EpicGames.Horde.csproj
@@ -38,8 +38,10 @@
     </Compile>
   </ItemGroup>
   <ItemGroup>
+    <PackageReference Include="Dapper" Version="2.1.24" />
     <PackageReference Include="K4os.Compression.LZ4" Version="1.2.16" />
     <PackageReference Include="Microsoft.CodeAnalysis.CSharp.Workspaces" Version="4.2.0" />
+    <PackageReference Include="Microsoft.Data.Sqlite" Version="8.0.0" />
     <PackageReference Include="Microsoft.Extensions.Caching.Abstractions" Version="6.0.0" />
     <PackageReference Include="Microsoft.Extensions.Caching.Memory" Version="6.0.1" />
     <PackageReference Include="Microsoft.Extensions.DependencyInjection" Version="6.0.0" />
diff --git a/Engine/Source/Programs/Shared/EpicGames.Horde/Storage/FileMetadataDb.cs b/Engine/Source/Programs/Shared/EpicGames.Horde/Storage/FileMetadataDb.cs
new file mode 100644
--- /dev/null
+++ b/Engine/Source/Programs/Shared/EpicGames.Horde/Storage/FileMetadataDb.cs
@@ -0,0 +1,457 @@
+// Copyright Epic Games, Inc. All Rights Reserved.
+
+using System;
+using System.Collections.Generic;
+using System.Data;
+using System.Linq;
+using System.Text;
+using System.Threading;
+using System.Threading.Tasks;
+using Dapper;
+using EpicGames.Core;
+using Microsoft.Data.Sqlite;
+
+namespace EpicGames.Horde.Storage
+{
+	/// <summary>
+	/// Database of file metadata, storing timestamps for trees of files and hashes for ranges within them. Implemented as a SQLite database.
+	/// </summary>
+	public sealed class FileMetadataDb : IDisposable
+	{
+		/// <summary>
+		/// Id value for the root directory
+		/// </summary>
+		public const int RootDirectoryId = 0;
+
+		class IoHashTypeMapper : SqlMapper.TypeHandler<IoHash>
+		{
+			public override IoHash Parse(object value)
+			{
+				return IoHash.Parse((string)value);
+			}
+
+			public override void SetValue(IDbDataParameter parameter, IoHash value)
+			{
+				parameter.DbType = DbType.AnsiStringFixedLength;
+				parameter.Value = value.ToString();
+			}
+		}
+
+		static readonly IoHashTypeMapper s_hashTypeMapper = new IoHashTypeMapper();
+
+		readonly SqliteConnection _connection;
+
+		FileMetadataDb(SqliteConnection connection)
+		{
+			_connection = connection;
+		}
+
+		/// <inheritdoc/>
+		public void Dispose()
+		{
+			_connection.Dispose();
+		}
+
+		/// <summary>
+		/// Creates a new in-memory database for file metadata
+		/// </summary>
+		public static async Task<FileMetadataDb> CreateInMemoryAsync(CancellationToken cancellationToken = default)
+		{
+			return await CreateAsync("Data Source=:memory:", cancellationToken);
+		}
+
+		/// <summary>
+		/// Creates a new database for metadata backed by a file on disk
+		/// </summary>
+		public static async Task<FileMetadataDb> CreateFromFileAsync(FileReference file, CancellationToken cancellationToken = default)
+		{
+			return await CreateAsync($"Data Source={file}", cancellationToken);
+		}
+
+		static async Task<FileMetadataDb> CreateAsync(string connectionString, CancellationToken cancellationToken)
+		{
+			SqlMapper.AddTypeHandler<IoHash>(s_hashTypeMapper);
+
+			SqliteConnection? connection = new SqliteConnection(connectionString);
+			try
+			{
+				await connection.OpenAsync(cancellationToken);
+
+				int version = await connection.ExecuteScalarAsync<int>("PRAGMA user_version;");
+				if(version == 0)
+				{
+					// Configure the directories table
+					await connection.ExecuteAsync(
+						"CREATE TABLE Directories (" +
+							" id INTEGER NOT NULL PRIMARY KEY AUTOINCREMENT," +
+							" parentDirectoryId INTEGER," +
+							" name TEXT NOT NULL," +
+							" FOREIGN KEY (parentDirectoryId) REFERENCES Directories(id));", cancellationToken);
+
+					await connection.ExecuteAsync(
+						"CREATE INDEX DirectoryParents ON Directories(parentDirectoryId);", cancellationToken);
+
+					await connection.ExecuteAsync(
+						"INSERT INTO Directories VALUES (0, NULL, \"\");", cancellationToken);
+
+					// Configure the files table
+					await connection.ExecuteAsync(
+						"CREATE TABLE Files (" +
+							" id INTEGER NOT NULL PRIMARY KEY AUTOINCREMENT," +
+							" directoryId INTEGER NOT NULL," +
+							" name TEXT NOT NULL," +
+							" time INTEGER NOT NULL," +
+							" length INTEGER NOT NULL," +
+							" FOREIGN KEY(directoryId) REFERENCES Directories(id));", cancellationToken);
+
+					await connection.ExecuteAsync(
+						"CREATE INDEX FileParents ON Files(directoryId);", cancellationToken);
+
+					// Configure the chunks table
+					await connection.ExecuteAsync(
+						"CREATE TABLE Chunks (" +
+							" id INTEGER NOT NULL PRIMARY KEY AUTOINCREMENT," +
+							" fileId INTEGER NOT NULL," +
+							" offset INTEGER NOT NULL," +
+							" length INTEGER NOT NULL," +
+							" hash VARCHAR NOT NULL," +
+							" FOREIGN KEY(fileId) REFERENCES Files(id));", cancellationToken);
+
+					await connection.ExecuteAsync(
+						"CREATE INDEX ChunkParents ON Chunks(fileId);", cancellationToken);
+
+					await connection.ExecuteAsync(
+						"CREATE INDEX ChunkHashes ON Chunks(hash);", cancellationToken);
+
+					// Set the DB version number to 1
+					await connection.ExecuteAsync(
+						"PRAGMA user_version = 1;", cancellationToken);
+
+					version = 1;
+				}
+				return new FileMetadataDb(connection);
+			}
+			catch
+			{
+				await connection.DisposeAsync();
+				throw;
+			}
+		}
+
+		#region Files
+
+		/// <summary>
+		/// Adds a single file to the database
+		/// </summary>
+		public async Task<int> AddFileAsync(int directoryId, string name, DateTime time, long length)
+		{
+			FileRow row = new FileRow(directoryId, name, time, length);
+			return await _connection.ExecuteScalarAsync<int>("INSERT INTO Files (DirectoryId, Name, Time, Length) VALUES (@DirectoryId, @Name, @Time, @Length); SELECT last_insert_rowid();", row);
+		}
+
+		/// <summary>
+		/// Adds a set of files to the database
+		/// </summary>
+		public async Task AddFilesAsync(IEnumerable<FileRow> files)
+		{
+			await _connection.ExecuteAsync("INSERT INTO Files (DirectoryId, Name, Time, Length) VALUES (@DirectoryId, @Name, @Time, @Length);", files);
+		}
+
+		/// <summary>
+		/// Finds all the files in a particular directory
+		/// </summary>
+		public Task<FileRow> GetFileAsync(int fileId)
+		{
+			return _connection.QuerySingleAsync<FileRow>("SELECT * FROM Files WHERE Id = @FileId", new { FileId = fileId });
+		}
+
+		/// <summary>
+		/// Finds all the files in a particular directory
+		/// </summary>
+		public Task<IEnumerable<FileRow>> FindFilesInDirectoryAsync(int directoryId)
+		{
+			return _connection.QueryAsync<FileRow>("SELECT * FROM Files WHERE DirectoryId = @DirectoryId", new { DirectoryId = directoryId });
+		}
+
+		/// <summary>
+		/// Finds all the files in a particular directory
+		/// </summary>
+		public Task<IEnumerable<FileRow>> FindFilesInDirectoriesAsync(IEnumerable<int> directoryIds)
+		{
+			return _connection.QueryAsync<FileRow>("SELECT * FROM Files WHERE DirectoryId IN (@DirectoryIds)", new { DirectoryIds = directoryIds });
+		}
+
+		/// <summary>
+		/// Gets the full name of a file
+		/// </summary>
+		public async Task<string> GetFilePathAsync(int id)
+		{
+			StringBuilder builder = new StringBuilder();
+			await GetFilePathAsync(id, builder);
+			return builder.ToString();
+		}
+
+		/// <summary>
+		/// Gets the full name of a directory
+		/// </summary>
+		public async Task GetFilePathAsync(int id, StringBuilder builder)
+		{
+			FileRow file = await GetFileAsync(id);
+			await GetDirectoryPathAsync(file.DirectoryId, builder);
+			builder.Append(file.Name);
+		}
+
+		/// <summary>
+		/// Removes a file and all its chunk metadata
+		/// </summary>
+		public async Task RemoveFileAsync(int fileId)
+		{
+			await _connection.ExecuteAsync("DELETE FROM Chunks WHERE FileId = @FileId; DELETE FROM Files WHERE FileId = @FileId;", fileId);
+		}
+
+		#endregion
+
+		#region Chunks
+
+		/// <summary>
+		/// Adds a record for a new file chunk
+		/// </summary>
+		public async Task<int> AddChunkAsync(int fileId, long offset, long length, IoHash hash)
+		{
+			ChunkRow row = new ChunkRow(fileId, offset, length, hash);
+			return await _connection.ExecuteScalarAsync<int>("INSERT INTO Chunks (FileId, Offset, Length, Hash) VALUES (@FileId, @Offset, @Length, @Hash); SELECT last_insert_rowid();", row);
+		}
+
+		/// <summary>
+		/// Adds multiple file chunk records
+		/// </summary>
+		public async Task AddChunksAsync(IEnumerable<ChunkRow> chunks)
+		{
+			await _connection.ExecuteScalarAsync<int>("INSERT INTO Chunks (FileId, Offset, Length, Hash) VALUES (@FileId, @Offset, @Length, @Hash);", chunks);
+		}
+
+		/// <summary>
+		/// Gets a chunk row
+		/// </summary>
+		public async Task<ChunkRow> GetChunkAsync(int chunkId)
+		{
+			return await _connection.QuerySingleAsync<ChunkRow>("SELECT * FROM Chunks WHERE id = @ChunkId;", new { ChunkId = chunkId });
+		}
+
+		/// <summary>
+		/// Find all chunks with a particular hash and length
+		/// </summary>
+		public async Task<IEnumerable<ChunkRow>> FindChunksAsync(IoHash hash, long length)
+		{
+			return await _connection.QueryAsync<ChunkRow>("SELECT * FROM Chunks WHERE hash = @Hash AND length = @Length;", new { Hash = hash, Length = length });
+		}
+
+		/// <summary>
+		/// Finds all the chunks within a particular file
+		/// </summary>
+		public async Task<IEnumerable<ChunkRow>> FindChunksForFileAsync(int fileId)
+		{
+			return await _connection.QueryAsync<ChunkRow>("SELECT * FROM Chunks WHERE fileId = @FileId;", new { FileId = fileId });
+		}
+
+		/// <summary>
+		/// Remove all chunks for a particular file
+		/// </summary>
+		public async Task RemoveChunksForFileAsync(int fileId)
+		{
+			await _connection.ExecuteAsync("DELETE FROM Chunks WHERE FileId = @FileId; DELETE FROM Files WHERE FileId = @FileId;", fileId);
+		}
+
+		/// <summary>
+		/// Remove all chunks for a set of files
+		/// </summary>
+		public async Task RemoveChunksForFilesAsync(IEnumerable<int> fileIds)
+		{
+			await _connection.ExecuteAsync("DELETE FROM Chunks WHERE FileId IN (@FileIds);", new { FileIds = fileIds });
+		}
+
+		#endregion
+
+		#region Directories
+
+		/// <summary>
+		/// Adds a new directory to the collection
+		/// </summary>
+		public async Task<int> AddDirectoryAsync(int parentDirectoryId, string name)
+		{
+			DirectoryRow row = new DirectoryRow { ParentDirectoryId = parentDirectoryId, Name = name };
+			return await _connection.ExecuteScalarAsync<int>("INSERT INTO Directories (ParentDirectoryId, Name) VALUES (@ParentDirectoryId, @Name); SELECT last_insert_rowid();", row);
+		}
+
+		/// <summary>
+		/// Adds multiple directories to the collection
+		/// </summary>
+		public async Task AddDirectoriesAsync(IEnumerable<DirectoryRow> dirs)
+		{
+			await _connection.ExecuteAsync("INSERT INTO Directories (ParentDirectoryId, Name) VALUES (@ParentDirectoryId, @Name);", dirs);
+		}
+
+		/// <summary>
+		/// Gets the definition for a particular directory
+		/// </summary>
+		public async Task<DirectoryRow> GetDirectoryAsync(int id)
+		{
+			return await _connection.QuerySingleAsync<DirectoryRow>("SELECT * FROM Directories WHERE Id = @Id;", new { Id = id });
+		}
+
+		/// <summary>
+		/// Gets the full name of a directory
+		/// </summary>
+		public async Task<string> GetDirectoryPathAsync(int id)
+		{
+			StringBuilder builder = new StringBuilder();
+			await GetDirectoryPathAsync(id, builder);
+			return builder.ToString();
+		}
+
+		/// <summary>
+		/// Gets the full name of a directory
+		/// </summary>
+		public async Task GetDirectoryPathAsync(int id, StringBuilder builder)
+		{
+			if (id != RootDirectoryId)
+			{
+				DirectoryRow row = await GetDirectoryAsync(id);
+				await GetDirectoryPathAsync(row.ParentDirectoryId, builder);
+				builder.Append(row.Name);
+				builder.Append('/');
+			}
+		}
+
+		/// <summary>
+		/// Finds all directories within a given parent directory
+		/// </summary>
+		public async Task<IEnumerable<DirectoryRow>> GetDirectoriesAsync(int parentDirectoryId)
+		{
+			return await _connection.QueryAsync<DirectoryRow>("SELECT * FROM Directories WHERE ParentDirectoryId = @ParentDirectoryId;", new { ParentDirectoryId = parentDirectoryId });
+		}
+
+		/// <summary>
+		/// Removes a directory and all its subdirectories
+		/// </summary>
+		public async Task RemoveDirectoryAsync(int directoryId)
+		{
+			await RemoveDirectoryContentsAsync(directoryId);
+			await _connection.ExecuteAsync("DELETE FROM Directories WHERE Id = @Id;", new { Id = directoryId });
+		}
+
+		/// <summary>
+		/// Removes a directory and all its subdirectories
+		/// </summary>
+		public async Task RemoveDirectoriesAsync(IEnumerable<int> directoryIds)
+		{
+			await RemoveDirectoryContentsAsync(directoryIds);
+			await _connection.ExecuteAsync("DELETE FROM Directories WHERE Id IN (@DirectoryIds);", new { DirectoryIds = directoryIds });
+		}
+
+		/// <summary>
+		/// Removes the contents of a directory, without removing the directory itself
+		/// </summary>
+		public async Task RemoveDirectoryContentsAsync(int directoryId)
+		{
+			await RemoveDirectoryContentsAsync(new[] { directoryId });
+		}
+
+		/// <summary>
+		/// Removes all subdirectories and files starting at the given roots
+		/// </summary>
+		public async Task RemoveDirectoryContentsAsync(IEnumerable<int> directoryIds)
+		{
+			object directoryIdList = new { DirectoryIds = directoryIds };
+
+			IEnumerable<int> subDirectoryIds = await _connection.QueryAsync<int>("SELECT Id FROM Directories WHERE ParentDirectoryId IN (@DirectoryIds);", directoryIdList);
+			if (subDirectoryIds.Any())
+			{
+				await RemoveDirectoryContentsAsync(subDirectoryIds);
+				await _connection.ExecuteAsync("DELETE FROM Directories WHERE ParentDirectoryId IN (@DirectoryIds);", directoryIdList);
+			}
+
+			IEnumerable<int> fileIds = await _connection.QueryAsync<int>("SELECT (Id) FROM Files WHERE DirectoryId IN (@DirectoryIds);", directoryIdList);
+			if (fileIds.Any())
+			{
+				await RemoveChunksForFilesAsync(fileIds);
+				await _connection.ExecuteAsync("DELETE FROM Files WHERE DirectoryId IN (@DirectoryIds);", directoryIdList);
+			}
+		}
+
+		#endregion
+	}
+
+	/// <summary>
+	/// Metadata for a file
+	/// </summary>
+	/// <param name="Id">Unique id for this file</param>
+	/// <param name="DirectoryId">Identifier for the directory containing this file</param>
+	/// <param name="Name">Name of the file</param>
+	/// <param name="Time">Last modified timestamp for the file</param>
+	/// <param name="Length">Length of the file</param>
+	public record class FileRow(int Id, int DirectoryId, string Name, DateTime Time, long Length)
+	{
+		/// <summary>
+		/// Default constructor
+		/// </summary>
+		public FileRow()
+			: this(-1, -1, String.Empty, DateTime.MinValue, 0)
+		{ }
+
+		/// <summary>
+		/// Constructor for new file rows
+		/// </summary>
+		public FileRow(int directoryId, string name, DateTime time, long length)
+			: this(-1, directoryId, name, time, length)
+		{ }
+	}
+
+	/// <summary>
+	/// Metadata for a directory
+	/// </summary>
+	/// <param name="Id">Unique id for this directory</param>
+	/// <param name="ParentDirectoryId">Parent directory identifier</param>
+	/// <param name="Name">Name of the directory</param>
+	public record class DirectoryRow(int Id, int ParentDirectoryId, string Name)
+	{
+		/// <summary>
+		/// Default constructor
+		/// </summary>
+		public DirectoryRow()
+			: this(-1, -1, String.Empty)
+		{ }
+
+		/// <summary>
+		/// Constructor for new directory rows
+		/// </summary>
+		public DirectoryRow(int parentDirectoryId, string name)
+			: this(-1, parentDirectoryId, name)
+		{ }
+	}
+
+	/// <summary>
+	/// Metadata for a file chunk
+	/// </summary>
+	/// <param name="Id">Unique id for the row</param>
+	/// <param name="FileId">Id of the file that this chunk belongs to</param>
+	/// <param name="Offset">Starting offset within the file of this chunk</param>
+	/// <param name="Length">Length of the chunk</param>
+	/// <param name="Hash">Hash of the chunk data</param>
+	public record class ChunkRow(int Id, int FileId, long Offset, long Length, IoHash Hash)
+	{
+		/// <summary>
+		/// Default constructor
+		/// </summary>
+		public ChunkRow()
+			: this(-1, -1, 0, 0, IoHash.Zero)
+		{ }
+
+		/// <summary>
+		/// Constructor for new chunk rows
+		/// </summary>
+		public ChunkRow(int fileId, long offset, long length, IoHash hash)
+			: this(-1, fileId, offset, length, hash)
+		{ }
+	}
+}
