Added missing Property Flag parsing code added in cl 30258132

diff --git a/Engine/Source/Programs/Shared/EpicGames.Core/UnrealEngineTypes.cs b/Engine/Source/Programs/Shared/EpicGames.Core/UnrealEngineTypes.cs
--- a/Engine/Source/Programs/Shared/EpicGames.Core/UnrealEngineTypes.cs
+++ b/Engine/Source/Programs/Shared/EpicGames.Core/UnrealEngineTypes.cs
@@ -1210,10 +1210,15 @@ namespace EpicGames.Core
 		TObjectPtrWrapper = TObjectPtr | UObjectWrapper,
 
 		/// <summary>
-		/// ****Experimental*** Property will use different logic to serialize knowing what changes are done against its default use the overridable information provided by the overidable manager on the object
+		/// ****Experimental*** Property will use different logic to serialize knowing what changes are done against its default use the overridable information provided by the overridable manager on the object
 		/// </summary>
 		ExperimentalOverridableLogic = 0x0200000000000000,
 
+		/// <summary>
+		/// ****Experimental*** Property should never inherit from the parent when using overridable serialization
+		/// </summary>
+		ExperimentalAlwaysOverriden = 0x0400000000000000,
+
 		/// <summary>
 		/// All Native Access Specifier flags
 		/// </summary>
diff --git a/Engine/Source/Programs/Shared/EpicGames.UHT/Specifiers/UhtPropertyMemberSpecifiers.cs b/Engine/Source/Programs/Shared/EpicGames.UHT/Specifiers/UhtPropertyMemberSpecifiers.cs
--- a/Engine/Source/Programs/Shared/EpicGames.UHT/Specifiers/UhtPropertyMemberSpecifiers.cs
+++ b/Engine/Source/Programs/Shared/EpicGames.UHT/Specifiers/UhtPropertyMemberSpecifiers.cs
@@ -389,6 +389,12 @@ namespace EpicGames.UHT.Parsers
 			context.PropertySettings.PropertyFlags |= EPropertyFlags.ExperimentalOverridableLogic;
 		}
 
+		[UhtSpecifier(Extends = UhtTableNames.PropertyMember, ValueType = UhtSpecifierValueType.Legacy)]
+		private static void ExperimentalAlwaysOverridenSpecifier(UhtSpecifierContext specifierContext)
+		{
+			UhtPropertySpecifierContext context = (UhtPropertySpecifierContext)specifierContext;
+			context.PropertySettings.PropertyFlags |= EPropertyFlags.ExperimentalAlwaysOverriden;
+		}
 
 		[UhtSpecifier(Extends = UhtTableNames.PropertyMember, ValueType = UhtSpecifierValueType.OptionalString)]
 		private static void GetterSpecifier(UhtSpecifierContext specifierContext, StringView? value)
