Horde: Disable the backend cache by default.

diff --git a/Engine/Source/Programs/Shared/EpicGames.Horde/Storage/Clients/HttpStorageClient.cs b/Engine/Source/Programs/Shared/EpicGames.Horde/Storage/Clients/HttpStorageClient.cs
--- a/Engine/Source/Programs/Shared/EpicGames.Horde/Storage/Clients/HttpStorageClient.cs
+++ b/Engine/Source/Programs/Shared/EpicGames.Horde/Storage/Clients/HttpStorageClient.cs
@@ -307,7 +307,7 @@ namespace EpicGames.Horde.Storage.Clients
 		/// <param name="basePath">Base path for all requests</param>
 		/// <param name="accessToken">Custom access token to use for requests</param>
 		/// <param name="withBackendCache"></param>
-		public IStorageClient CreateClientWithPath(string basePath, string? accessToken = null, bool withBackendCache = true)
+		public IStorageClient CreateClientWithPath(string basePath, string? accessToken = null, bool withBackendCache = false)
 		{
 			HttpClient CreateClient()
 			{
@@ -335,7 +335,7 @@ namespace EpicGames.Horde.Storage.Clients
 		/// <param name="namespaceId">Namespace to create a client for</param>
 		/// <param name="accessToken">Custom access token to use for requests</param>
 		/// <param name="withBackendCache">Whether to enable the backend cache, which caches full bundles to disk</param>
-		public IStorageClient CreateClient(NamespaceId namespaceId, string? accessToken = null, bool withBackendCache = true) => CreateClientWithPath($"api/v1/storage/{namespaceId}", accessToken, withBackendCache);
+		public IStorageClient CreateClient(NamespaceId namespaceId, string? accessToken = null, bool withBackendCache = false) => CreateClientWithPath($"api/v1/storage/{namespaceId}", accessToken, withBackendCache);
 
 		/// <inheritdoc/>
 		IStorageClient? IStorageClientFactory.TryCreateClient(NamespaceId namespaceId) => CreateClient(namespaceId);
