Horde: Update to new bundle version again.

diff --git a/Engine/Source/Programs/Shared/EpicGames.Horde/Storage/Bundles/BundleOptions.cs b/Engine/Source/Programs/Shared/EpicGames.Horde/Storage/Bundles/BundleOptions.cs
--- a/Engine/Source/Programs/Shared/EpicGames.Horde/Storage/Bundles/BundleOptions.cs
+++ b/Engine/Source/Programs/Shared/EpicGames.Horde/Storage/Bundles/BundleOptions.cs
@@ -15,7 +15,7 @@ namespace EpicGames.Horde.Storage.Bundles
 		/// <summary>
 		/// Maximum version number of bundles to write
 		/// </summary>
-		public BundleVersion MaxVersion { get; set; } = BundleVersion.LatestV1;
+		public BundleVersion MaxVersion { get; set; } = BundleVersion.LatestV2;
 
 		/// <summary>
 		/// Maximum payload size fo a blob
@@ -36,10 +36,5 @@ namespace EpicGames.Horde.Storage.Bundles
 		/// Maximum amount of data to store in memory. This includes any background writes as well as bundles being built.
 		/// </summary>
 		public long MaxWriteQueueLength { get; set; } = 256 * 1024 * 1024;
-
-		/// <summary>
-		/// Number of nodes to cache
-		/// </summary>
-		public int NodeCacheSize { get; set; } = 1024;
 	}
 }
