EpicGames.Core: FileSystemReference cleanup

#rnx

diff --git a/Engine/Source/Programs/Shared/EpicGames.Core/DirectoryReference.cs b/Engine/Source/Programs/Shared/EpicGames.Core/DirectoryReference.cs
--- a/Engine/Source/Programs/Shared/EpicGames.Core/DirectoryReference.cs
+++ b/Engine/Source/Programs/Shared/EpicGames.Core/DirectoryReference.cs
@@ -84,10 +84,7 @@ namespace EpicGames.Core
 		/// Gets the top level directory name
 		/// </summary>
 		/// <returns>The name of the directory</returns>
-		public string GetDirectoryName()
-		{
-			return Path.GetFileName(FullName);
-		}
+		public string GetDirectoryName() => Path.GetFileName(FullName);
 
 		/// <summary>
 		/// Gets the directory containing this object
@@ -140,17 +137,14 @@ namespace EpicGames.Core
 		public static DirectoryReference? GetSpecialFolder(Environment.SpecialFolder folder)
 		{
 			string folderPath = Environment.GetFolderPath(folder);
-			return String.IsNullOrEmpty(folderPath)? null : new DirectoryReference(folderPath);
+			return String.IsNullOrEmpty(folderPath) ? null : new DirectoryReference(folderPath);
 		}
 
 		/// <summary>
 		/// Determines whether this path represents a root directory in the filesystem
 		/// </summary>
 		/// <returns>True if this path is a root directory, false otherwise</returns>
-		public bool IsRootDirectory()
-		{
-			return FullName[^1] == Path.DirectorySeparatorChar;
-		}
+		public bool IsRootDirectory() => FullName[^1] == Path.DirectorySeparatorChar;
 
 		/// <summary>
 		/// Combine several fragments with a base directory, to form a new directory name
@@ -158,11 +152,7 @@ namespace EpicGames.Core
 		/// <param name="baseDirectory">The base directory</param>
 		/// <param name="fragments">Fragments to combine with the base directory</param>
 		/// <returns>The new directory name</returns>
-		public static DirectoryReference Combine(DirectoryReference baseDirectory, params string[] fragments)
-		{
-			string fullName = FileSystemReference.CombineStrings(baseDirectory, fragments);
-			return new DirectoryReference(fullName, Sanitize.None);
-		}
+		public static DirectoryReference Combine(DirectoryReference baseDirectory, params string[] fragments) => new DirectoryReference(CombineStrings(baseDirectory, fragments), Sanitize.None);
 
 		/// <summary>
 		/// Compares two filesystem object names for equality. Uses the canonical name representation, not the display name representation.
@@ -192,64 +182,49 @@ namespace EpicGames.Core
 		/// <param name="a">First object to compare.</param>
 		/// <param name="b">Second object to compare.</param>
 		/// <returns>False if the names represent the same object, true otherwise</returns>
-		public static bool operator !=(DirectoryReference? a, DirectoryReference? b)
-		{
-			return !(a == b);
-		}
+		public static bool operator !=(DirectoryReference? a, DirectoryReference? b) => !(a == b);
 
 		/// <summary>
 		/// Compares against another object for equality.
 		/// </summary>
 		/// <param name="obj">other instance to compare.</param>
 		/// <returns>True if the names represent the same object, false otherwise</returns>
-		public override bool Equals(object? obj)
-		{
-			return (obj is DirectoryReference dir) && dir == this;
-		}
+		public override bool Equals(object? obj) => (obj is DirectoryReference dir) && dir == this;
 
 		/// <summary>
 		/// Compares against another object for equality.
 		/// </summary>
 		/// <param name="obj">other instance to compare.</param>
 		/// <returns>True if the names represent the same object, false otherwise</returns>
-		public bool Equals(DirectoryReference? obj)
-		{
-			return obj == this;
-		}
+		public bool Equals(DirectoryReference? obj) => obj == this;
 
 		/// <summary>
 		/// Returns a hash code for this object
 		/// </summary>
 		/// <returns></returns>
-		public override int GetHashCode()
-		{
-			return Comparer.GetHashCode(FullName);
-		}
+		public override int GetHashCode() => Comparer.GetHashCode(FullName);
 
 		/// <inheritdoc/>
 		public int CompareTo(DirectoryReference? other) => Comparer.Compare(FullName, other?.FullName);
 
 		/// <inheritdoc/>
-		public static bool operator <(DirectoryReference left, DirectoryReference right) => ReferenceEquals(left, null) ? !ReferenceEquals(right, null) : left.CompareTo(right) < 0;
+		public static bool operator <(DirectoryReference left, DirectoryReference right) => left is null ? right is not null : left.CompareTo(right) < 0;
 
 		/// <inheritdoc/>
-		public static bool operator <=(DirectoryReference left, DirectoryReference right) => ReferenceEquals(left, null) || left.CompareTo(right) <= 0;
+		public static bool operator <=(DirectoryReference left, DirectoryReference right) => left is null || left.CompareTo(right) <= 0;
 
 		/// <inheritdoc/>
-		public static bool operator >(DirectoryReference left, DirectoryReference right) => !ReferenceEquals(left, null) && left.CompareTo(right) > 0;
+		public static bool operator >(DirectoryReference left, DirectoryReference right) => left is not null && left.CompareTo(right) > 0;
 
 		/// <inheritdoc/>
-		public static bool operator >=(DirectoryReference left, DirectoryReference right) => ReferenceEquals(left, null) ? ReferenceEquals(right, null) : left.CompareTo(right) >= 0;
+		public static bool operator >=(DirectoryReference left, DirectoryReference right) => left is null ? right is null : left.CompareTo(right) >= 0;
 
 		/// <summary>
 		/// Helper function to create a remote directory reference. Unlike normal DirectoryReference objects, these aren't converted to a full path in the local filesystem.
 		/// </summary>
 		/// <param name="absolutePath">The absolute path in the remote file system</param>
 		/// <returns>New directory reference</returns>
-		public static DirectoryReference MakeRemote(string absolutePath)
-		{
-			return new DirectoryReference(absolutePath, Sanitize.None);
-		}
+		public static DirectoryReference MakeRemote(string absolutePath) => new DirectoryReference(absolutePath, Sanitize.None);
 
 		/// <summary>
 		/// Gets the parent directory for a file, or returns null if it's null.
@@ -257,53 +232,27 @@ namespace EpicGames.Core
 		/// <param name="file">The file to create a directory reference for</param>
 		/// <returns>The directory containing the file  </returns>
 		[return: NotNullIfNotNull("file")]
-		public static DirectoryReference? FromFile(FileReference? file)
-		{
-			if(file == null)
-			{
-				return null;
-			}
-			else
-			{
-				return file.Directory;
-			}
-		}
+		public static DirectoryReference? FromFile(FileReference? file) => file?.Directory;
 
 		/// <summary>
 		/// Create a DirectoryReference from a string. If the string is null, returns a null DirectoryReference.
 		/// </summary>
 		/// <param name="directoryName">Path for the new object</param>
 		/// <returns>Returns a FileReference representing the given string, or null.</returns>
-		public static DirectoryReference? FromString(string? directoryName)
-		{
-			if(String.IsNullOrEmpty(directoryName))
-			{
-				return null;
-			}
-			else
-			{
-				return new DirectoryReference(directoryName);
-			}
-		}
+		public static DirectoryReference? FromString(string? directoryName) => String.IsNullOrEmpty(directoryName) ? null : new DirectoryReference(directoryName);
 
 		/// <summary>
 		/// Finds the correct case to match the location of this file on disk. Uses the given case for parts of the path that do not exist.
 		/// </summary>
 		/// <param name="location">The path to find the correct case for</param>
 		/// <returns>Location of the file with the correct case</returns>
-		public static DirectoryReference FindCorrectCase(DirectoryReference location)
-		{
-			return new DirectoryReference(DirectoryUtils.FindCorrectCase(location.ToDirectoryInfo()));
-		}
+		public static DirectoryReference FindCorrectCase(DirectoryReference location) => new DirectoryReference(DirectoryUtils.FindCorrectCase(location.ToDirectoryInfo()));
 
 		/// <summary>
 		/// Constructs a DirectoryInfo object from this reference
 		/// </summary>
 		/// <returns>New DirectoryInfo object</returns>
-		public DirectoryInfo ToDirectoryInfo()
-		{
-			return new DirectoryInfo(FullName);
-		}
+		public DirectoryInfo ToDirectoryInfo() => new DirectoryInfo(FullName);
 
 		#region System.IO.Directory Wrapper Methods
 
@@ -311,61 +260,40 @@ namespace EpicGames.Core
 		/// Finds the current directory
 		/// </summary>
 		/// <returns>The current directory</returns>
-		public static DirectoryReference GetCurrentDirectory()
-		{
-			return new DirectoryReference(Directory.GetCurrentDirectory());
-		}
+		public static DirectoryReference GetCurrentDirectory() => new DirectoryReference(Directory.GetCurrentDirectory());
 
 		/// <summary>
 		/// Creates a directory
 		/// </summary>
 		/// <param name="location">Location of the directory</param>
-		public static void CreateDirectory(DirectoryReference location)
-		{
-			Directory.CreateDirectory(location.FullName);
-		}
+		public static void CreateDirectory(DirectoryReference location) => Directory.CreateDirectory(location.FullName);
 
 		/// <summary>
 		/// Deletes a directory
 		/// </summary>
 		/// <param name="location">Location of the directory</param>
-		public static void Delete(DirectoryReference location)
-		{
-			Directory.Delete(location.FullName);
-		}
+		public static void Delete(DirectoryReference location) => Directory.Delete(location.FullName);
 
 		/// <summary>
 		/// Deletes a directory
 		/// </summary>
 		/// <param name="location">Location of the directory</param>
 		/// <param name="bRecursive">Whether to remove directories recursively</param>
-		public static void Delete(DirectoryReference location, bool bRecursive)
-		{
-			Directory.Delete(location.FullName, bRecursive);
-		}
+		public static void Delete(DirectoryReference location, bool bRecursive) => Directory.Delete(location.FullName, bRecursive);
 
 		/// <summary>
 		/// Checks whether the directory exists
 		/// </summary>
 		/// <param name="location">Location of the directory</param>
 		/// <returns>True if this directory exists</returns>
-		public static bool Exists(DirectoryReference location)
-		{
-			return Directory.Exists(location.FullName);
-		}
+		public static bool Exists(DirectoryReference location) => Directory.Exists(location.FullName);
 
 		/// <summary>
 		/// Enumerate files from a given directory
 		/// </summary>
 		/// <param name="baseDir">Base directory to search in</param>
 		/// <returns>Sequence of file references</returns>
-		public static IEnumerable<FileReference> EnumerateFiles(DirectoryReference baseDir)
-		{
-			foreach (string fileName in Directory.EnumerateFiles(baseDir.FullName))
-			{
-				yield return new FileReference(fileName, FileReference.Sanitize.None);
-			}
-		}
+		public static IEnumerable<FileReference> EnumerateFiles(DirectoryReference baseDir) => Directory.EnumerateFiles(baseDir.FullName).Select(fileName => new FileReference(fileName, FileReference.Sanitize.None));
 
 		/// <summary>
 		/// Enumerate files from a given directory
@@ -373,13 +301,7 @@ namespace EpicGames.Core
 		/// <param name="baseDir">Base directory to search in</param>
 		/// <param name="pattern">Pattern for matching files</param>
 		/// <returns>Sequence of file references</returns>
-		public static IEnumerable<FileReference> EnumerateFiles(DirectoryReference baseDir, string pattern)
-		{
-			foreach (string fileName in Directory.EnumerateFiles(baseDir.FullName, pattern))
-			{
-				yield return new FileReference(fileName, FileReference.Sanitize.None);
-			}
-		}
+		public static IEnumerable<FileReference> EnumerateFiles(DirectoryReference baseDir, string pattern) => Directory.EnumerateFiles(baseDir.FullName, pattern).Select(fileName => new FileReference(fileName, FileReference.Sanitize.None));
 
 		/// <summary>
 		/// Enumerate files from a given directory
@@ -388,26 +310,14 @@ namespace EpicGames.Core
 		/// <param name="pattern">Pattern for matching files</param>
 		/// <param name="option">Options for the search</param>
 		/// <returns>Sequence of file references</returns>
-		public static IEnumerable<FileReference> EnumerateFiles(DirectoryReference baseDir, string pattern, SearchOption option)
-		{
-			foreach (string fileName in Directory.EnumerateFiles(baseDir.FullName, pattern, option))
-			{
-				yield return new FileReference(fileName, FileReference.Sanitize.None);
-			}
-		}
+		public static IEnumerable<FileReference> EnumerateFiles(DirectoryReference baseDir, string pattern, SearchOption option) => Directory.EnumerateFiles(baseDir.FullName, pattern, option).Select(fileName => new FileReference(fileName, FileReference.Sanitize.None));
 
 		/// <summary>
 		/// Enumerate subdirectories in a given directory
 		/// </summary>
 		/// <param name="baseDir">Base directory to search in</param>
 		/// <returns>Sequence of directory references</returns>
-		public static IEnumerable<DirectoryReference> EnumerateDirectories(DirectoryReference baseDir)
-		{
-			foreach (string directoryName in Directory.EnumerateDirectories(baseDir.FullName))
-			{
-				yield return new DirectoryReference(directoryName, Sanitize.None);
-			}
-		}
+		public static IEnumerable<DirectoryReference> EnumerateDirectories(DirectoryReference baseDir) => Directory.EnumerateDirectories(baseDir.FullName).Select(directoryName => new DirectoryReference(directoryName, Sanitize.None));
 
 		/// <summary>
 		/// Enumerate subdirectories in a given directory
@@ -415,13 +325,7 @@ namespace EpicGames.Core
 		/// <param name="baseDir">Base directory to search in</param>
 		/// <param name="pattern">Pattern for matching directories</param>
 		/// <returns>Sequence of directory references</returns>
-		public static IEnumerable<DirectoryReference> EnumerateDirectories(DirectoryReference baseDir, string pattern)
-		{
-			foreach (string directoryName in Directory.EnumerateDirectories(baseDir.FullName, pattern))
-			{
-				yield return new DirectoryReference(directoryName, Sanitize.None);
-			}
-		}
+		public static IEnumerable<DirectoryReference> EnumerateDirectories(DirectoryReference baseDir, string pattern) => Directory.EnumerateDirectories(baseDir.FullName, pattern).Select(directoryName => new DirectoryReference(directoryName, Sanitize.None));
 
 		/// <summary>
 		/// Enumerate subdirectories in a given directory
@@ -430,13 +334,7 @@ namespace EpicGames.Core
 		/// <param name="pattern">Pattern for matching files</param>
 		/// <param name="option">Options for the search</param>
 		/// <returns>Sequence of directory references</returns>
-		public static IEnumerable<DirectoryReference> EnumerateDirectories(DirectoryReference baseDir, string pattern, SearchOption option)
-		{
-			foreach (string directoryName in Directory.EnumerateDirectories(baseDir.FullName, pattern, option))
-			{
-				yield return new DirectoryReference(directoryName, Sanitize.None);
-			}
-		}
+		public static IEnumerable<DirectoryReference> EnumerateDirectories(DirectoryReference baseDir, string pattern, SearchOption option) => Directory.EnumerateDirectories(baseDir.FullName, pattern, option).Select(directoryName => new DirectoryReference(directoryName, Sanitize.None));
 
 		/// <summary>
 		/// Sets the current directory
@@ -466,10 +364,7 @@ namespace EpicGames.Core
 		/// </summary>
 		/// <param name="writer">Binary writer to write to</param>
 		/// <param name="directory">The directory reference to write</param>
-		public static void Write(this BinaryWriter writer, DirectoryReference directory)
-		{
-			writer.Write((directory == null) ? String.Empty : directory.FullName);
-		}
+		public static void Write(this BinaryWriter writer, DirectoryReference directory) => writer.Write((directory == null) ? String.Empty : directory.FullName);
 
 		/// <summary>
 		/// Manually deserialize a directory reference from a binary stream.
@@ -487,20 +382,14 @@ namespace EpicGames.Core
 		/// </summary>
 		/// <param name="reader">Binary reader to read from</param>
 		/// <returns>New DirectoryReference object</returns>
-		public static DirectoryReference ReadDirectoryReferenceNotNull(this BinaryReader reader)
-		{
-			return BinaryArchiveReader.NotNull<DirectoryReference>(ReadDirectoryReference(reader));
-		}
+		public static DirectoryReference ReadDirectoryReferenceNotNull(this BinaryReader reader) => BinaryArchiveReader.NotNull(ReadDirectoryReference(reader));
 
 		/// <summary>
 		/// Writes a directory reference  to a binary archive
 		/// </summary>
 		/// <param name="writer">The writer to output data to</param>
 		/// <param name="directory">The item to write</param>
-		public static void WriteDirectoryReference(this BinaryArchiveWriter writer, DirectoryReference? directory)
-		{
-			writer.WriteString(directory?.FullName);
-		}
+		public static void WriteDirectoryReference(this BinaryArchiveWriter writer, DirectoryReference? directory) => writer.WriteString(directory?.FullName);
 
 		/// <summary>
 		/// Reads a directory reference from a binary archive
@@ -510,14 +399,7 @@ namespace EpicGames.Core
 		public static DirectoryReference? ReadDirectoryReference(this BinaryArchiveReader reader)
 		{
 			string? fullName = reader.ReadString();
-			if (fullName == null)
-			{
-				return null;
-			}
-			else
-			{
-				return new DirectoryReference(fullName, DirectoryReference.Sanitize.None);
-			}
+			return fullName == null ? null : new DirectoryReference(fullName, DirectoryReference.Sanitize.None);
 		}
 
 		/// <summary>
@@ -525,9 +407,6 @@ namespace EpicGames.Core
 		/// </summary>
 		/// <param name="reader">Reader to serialize data from</param>
 		/// <returns>New directory reference instance</returns>
-		public static DirectoryReference ReadDirectoryReferenceNotNull(this BinaryArchiveReader reader)
-		{
-			return BinaryArchiveReader.NotNull(ReadDirectoryReference(reader));
-		}
+		public static DirectoryReference ReadDirectoryReferenceNotNull(this BinaryArchiveReader reader) => BinaryArchiveReader.NotNull(ReadDirectoryReference(reader));
 	}
 }
diff --git a/Engine/Source/Programs/Shared/EpicGames.Core/FileReference.cs b/Engine/Source/Programs/Shared/EpicGames.Core/FileReference.cs
--- a/Engine/Source/Programs/Shared/EpicGames.Core/FileReference.cs
+++ b/Engine/Source/Programs/Shared/EpicGames.Core/FileReference.cs
@@ -35,7 +35,7 @@ namespace EpicGames.Core
 		public FileReference(string inPath)
 			: base(Path.GetFullPath(inPath))
 		{
-			if(FullName[^1] == '\\' || FullName[^1] == '/')
+			if (FullName[^1] == '\\' || FullName[^1] == '/')
 			{
 				throw new ArgumentException("File names may not be terminated by a path separator character");
 			}
@@ -67,35 +67,19 @@ namespace EpicGames.Core
 		/// <param name="fileName">FileName for the string</param>
 		/// <returns>Returns a FileReference representing the given string, or null.</returns>
 		[return: NotNullIfNotNull("fileName")]
-		public static FileReference? FromString(string? fileName)
-		{
-			if(String.IsNullOrEmpty(fileName))
-			{
-				return null;
-			}
-			else
-			{
-				return new FileReference(fileName);
-			}
-		}
+		public static FileReference? FromString(string? fileName) => String.IsNullOrEmpty(fileName) ? null : new FileReference(fileName);
 
 		/// <summary>
 		/// Gets the file name without path information
 		/// </summary>
 		/// <returns>A string containing the file name</returns>
-		public string GetFileName()
-		{
-			return Path.GetFileName(FullName);
-		}
+		public string GetFileName() => Path.GetFileName(FullName);
 
 		/// <summary>
 		/// Gets the file name without path information or an extension
 		/// </summary>
 		/// <returns>A string containing the file name without an extension</returns>
-		public string GetFileNameWithoutExtension()
-		{
-			return Path.GetFileNameWithoutExtension(FullName);
-		}
+		public string GetFileNameWithoutExtension() => Path.GetFileNameWithoutExtension(FullName);
 
 		/// <summary>
 		/// Gets the file name without path or any extensions
@@ -120,21 +104,14 @@ namespace EpicGames.Core
 		/// Gets the extension for this filename
 		/// </summary>
 		/// <returns>A string containing the extension of this filename</returns>
-		public string GetExtension()
-		{
-			return Path.GetExtension(FullName);
-		}
+		public string GetExtension() => Path.GetExtension(FullName);
 
 		/// <summary>
 		/// Change the file's extension to something else
 		/// </summary>
 		/// <param name="extension">The new extension</param>
 		/// <returns>A FileReference with the same path and name, but with the new extension</returns>
-		public FileReference ChangeExtension(string? extension)
-		{
-			string newFullName = Path.ChangeExtension(FullName, extension);
-			return new FileReference(newFullName, Sanitize.None);
-		}
+		public FileReference ChangeExtension(string? extension) => new FileReference(Path.ChangeExtension(FullName, extension), Sanitize.None);
 
 		/// <summary>
 		/// Gets the directory containing this file
@@ -168,11 +145,7 @@ namespace EpicGames.Core
 		/// <param name="baseDirectory">The base directory</param>
 		/// <param name="fragments">Fragments to combine with the base directory</param>
 		/// <returns>The new file name</returns>
-		public static FileReference Combine(DirectoryReference baseDirectory, params string[] fragments)
-		{
-			string fullName = FileSystemReference.CombineStrings(baseDirectory, fragments);
-			return new FileReference(fullName, Sanitize.None);
-		}
+		public static FileReference Combine(DirectoryReference baseDirectory, params string[] fragments) => new FileReference(CombineStrings(baseDirectory, fragments), Sanitize.None);
 
 		/// <summary>
 		/// Append a string to the end of a filename
@@ -180,10 +153,7 @@ namespace EpicGames.Core
 		/// <param name="a">The base file reference</param>
 		/// <param name="b">Suffix to be appended</param>
 		/// <returns>The new file reference</returns>
-		public static FileReference operator +(FileReference a, string b)
-		{
-			return new FileReference(a.FullName + b, Sanitize.None);
-		}
+		public static FileReference operator +(FileReference a, string b) => new FileReference(a.FullName + b, Sanitize.None);
 
 		/// <summary>
 		/// Compares two filesystem object names for equality. Uses the canonical name representation, not the display name representation.
@@ -213,10 +183,7 @@ namespace EpicGames.Core
 		/// <param name="a">First object to compare.</param>
 		/// <param name="b">Second object to compare.</param>
 		/// <returns>False if the names represent the same object, true otherwise</returns>
-		public static bool operator !=(FileReference? a, FileReference? b)
-		{
-			return !(a == b);
-		}
+		public static bool operator !=(FileReference? a, FileReference? b) => !(a == b);
 
 		/// <summary>
 		/// Compares against another object for equality.
@@ -230,19 +197,13 @@ namespace EpicGames.Core
 		/// </summary>
 		/// <param name="obj">other instance to compare.</param>
 		/// <returns>True if the names represent the same object, false otherwise</returns>
-		public bool Equals(FileReference? obj)
-		{
-			return obj == this;
-		}
+		public bool Equals(FileReference? obj) => obj == this;
 
 		/// <summary>
 		/// Returns a hash code for this object
 		/// </summary>
 		/// <returns></returns>
-		public override int GetHashCode()
-		{
-			return Comparer.GetHashCode(FullName);
-		}
+		public override int GetHashCode() => Comparer.GetHashCode(FullName);
 
 		/// <inheritdoc/>
 		public int CompareTo(FileReference? other) => Comparer.Compare(FullName, other?.FullName);
@@ -265,10 +226,7 @@ namespace EpicGames.Core
 		/// </summary>
 		/// <param name="absolutePath">The absolute path in the remote file system</param>
 		/// <returns>New file reference</returns>
-		public static FileReference MakeRemote(string absolutePath)
-		{
-			return new FileReference(absolutePath, Sanitize.None);
-		}
+		public static FileReference MakeRemote(string absolutePath) => new FileReference(absolutePath, Sanitize.None);
 
 		/// <summary>
 		/// Makes a file location writeable; 
@@ -276,10 +234,10 @@ namespace EpicGames.Core
 		/// <param name="location">Location of the file</param>
 		public static void MakeWriteable(FileReference location)
 		{
-			if(Exists(location))
+			if (Exists(location))
 			{
 				FileAttributes attributes = GetAttributes(location);
-				if((attributes & FileAttributes.ReadOnly) != 0)
+				if ((attributes & FileAttributes.ReadOnly) != 0)
 				{
 					SetAttributes(location, attributes & ~FileAttributes.ReadOnly);
 				}
@@ -291,19 +249,13 @@ namespace EpicGames.Core
 		/// </summary>
 		/// <param name="location">The path to find the correct case for</param>
 		/// <returns>Location of the file with the correct case</returns>
-		public static FileReference FindCorrectCase(FileReference location)
-		{
-			return new FileReference(FileUtils.FindCorrectCase(location.ToFileInfo()));
-		}
+		public static FileReference FindCorrectCase(FileReference location) => new FileReference(FileUtils.FindCorrectCase(location.ToFileInfo()));
 
 		/// <summary>
 		/// Constructs a FileInfo object from this reference
 		/// </summary>
 		/// <returns>New FileInfo object</returns>
-		public FileInfo ToFileInfo()
-		{
-			return new FileInfo(FullName);
-		}
+		public FileInfo ToFileInfo() => new FileInfo(FullName);
 
 		#region System.IO.File methods
 
@@ -312,10 +264,7 @@ namespace EpicGames.Core
 		/// </summary>
 		/// <param name="sourceLocation">Location of the source file</param>
 		/// <param name="targetLocation">Location of the target file</param>
-		public static void Copy(FileReference sourceLocation, FileReference targetLocation)
-		{
-			File.Copy(sourceLocation.FullName, targetLocation.FullName);
-		}
+		public static void Copy(FileReference sourceLocation, FileReference targetLocation) => File.Copy(sourceLocation.FullName, targetLocation.FullName);
 
 		/// <summary>
 		/// Copies a file from one location to another
@@ -323,67 +272,46 @@ namespace EpicGames.Core
 		/// <param name="sourceLocation">Location of the source file</param>
 		/// <param name="targetLocation">Location of the target file</param>
 		/// <param name="bOverwrite">Whether to overwrite the file in the target location</param>
-		public static void Copy(FileReference sourceLocation, FileReference targetLocation, bool bOverwrite)
-		{
-			File.Copy(sourceLocation.FullName, targetLocation.FullName, bOverwrite);
-		}
+		public static void Copy(FileReference sourceLocation, FileReference targetLocation, bool bOverwrite) => File.Copy(sourceLocation.FullName, targetLocation.FullName, bOverwrite);
 
 		/// <summary>
 		/// Deletes this file
 		/// </summary>
-		public static void Delete(FileReference location)
-		{
-			File.Delete(location.FullName);
-		}
+		public static void Delete(FileReference location) => File.Delete(location.FullName);
 
 		/// <summary>
 		/// Determines whether the given filename exists
 		/// </summary>
 		/// <returns>True if it exists, false otherwise</returns>
-		public static bool Exists(FileReference location)
-		{
-			return File.Exists(location.FullName);
-		}
+		public static bool Exists(FileReference location) => File.Exists(location.FullName);
 
 		/// <summary>
 		/// Gets the attributes for a file
 		/// </summary>
 		/// <param name="location">Location of the file</param>
 		/// <returns>Attributes for the file</returns>
-		public static FileAttributes GetAttributes(FileReference location)
-		{
-			return File.GetAttributes(location.FullName);
-		}
+		public static FileAttributes GetAttributes(FileReference location) => File.GetAttributes(location.FullName);
 
 		/// <summary>
 		/// Gets the time that the file was last written to
 		/// </summary>
 		/// <param name="location">Location of the file</param>
 		/// <returns>Last write time, in local time</returns>
-		public static DateTime GetLastWriteTime(FileReference location)
-		{
-			return File.GetLastWriteTime(location.FullName);
-		}
+		public static DateTime GetLastWriteTime(FileReference location) => File.GetLastWriteTime(location.FullName);
 
 		/// <summary>
 		/// Gets the time that the file was last written to
 		/// </summary>
 		/// <param name="location">Location of the file</param>
 		/// <returns>Last write time, in UTC time</returns>
-		public static DateTime GetLastWriteTimeUtc(FileReference location)
-		{
-			return File.GetLastWriteTimeUtc(location.FullName);
-		}
+		public static DateTime GetLastWriteTimeUtc(FileReference location) => File.GetLastWriteTimeUtc(location.FullName);
 
 		/// <summary>
 		/// Moves a file from one location to another
 		/// </summary>
 		/// <param name="sourceLocation">Location of the source file</param>
 		/// <param name="targetLocation">Location of the target file</param>
-		public static void Move(FileReference sourceLocation, FileReference targetLocation)
-		{
-			File.Move(sourceLocation.FullName, targetLocation.FullName);
-		}
+		public static void Move(FileReference sourceLocation, FileReference targetLocation) => File.Move(sourceLocation.FullName, targetLocation.FullName);
 
 		/// <summary>
 		/// Moves a file from one location to another
@@ -391,10 +319,7 @@ namespace EpicGames.Core
 		/// <param name="sourceLocation">Location of the source file</param>
 		/// <param name="targetLocation">Location of the target file</param>
 		/// <param name="overwrite">Whether to overwrite the file in the target location</param>
-		public static void Move(FileReference sourceLocation, FileReference targetLocation, bool overwrite)
-		{
-			File.Move(sourceLocation.FullName, targetLocation.FullName, overwrite);
-		}
+		public static void Move(FileReference sourceLocation, FileReference targetLocation, bool overwrite) => File.Move(sourceLocation.FullName, targetLocation.FullName, overwrite);
 
 		/// <summary>
 		/// Opens a FileStream on the specified path with read/write access
@@ -402,10 +327,7 @@ namespace EpicGames.Core
 		/// <param name="location">Location of the file</param>
 		/// <param name="mode">Mode to use when opening the file</param>
 		/// <returns>New filestream for the given file</returns>
-		public static FileStream Open(FileReference location, FileMode mode)
-		{
-			return File.Open(location.FullName, mode);
-		}
+		public static FileStream Open(FileReference location, FileMode mode) => File.Open(location.FullName, mode);
 
 		/// <summary>
 		/// Opens a FileStream on the specified path
@@ -414,10 +336,7 @@ namespace EpicGames.Core
 		/// <param name="mode">Mode to use when opening the file</param>
 		/// <param name="access">Sharing mode for the new file</param>
 		/// <returns>New filestream for the given file</returns>
-		public static FileStream Open(FileReference location, FileMode mode, FileAccess access)
-		{
-			return File.Open(location.FullName, mode, access);
-		}
+		public static FileStream Open(FileReference location, FileMode mode, FileAccess access) => File.Open(location.FullName, mode, access);
 
 		/// <summary>
 		/// Opens a FileStream on the specified path
@@ -427,20 +346,14 @@ namespace EpicGames.Core
 		/// <param name="access">Access mode for the new file</param>
 		/// <param name="share">Sharing mode for the open file</param>
 		/// <returns>New filestream for the given file</returns>
-		public static FileStream Open(FileReference location, FileMode mode, FileAccess access, FileShare share)
-		{
-			return File.Open(location.FullName, mode, access, share);
-		}
+		public static FileStream Open(FileReference location, FileMode mode, FileAccess access, FileShare share) => File.Open(location.FullName, mode, access, share);
 
 		/// <summary>
 		/// Reads the contents of a file
 		/// </summary>
 		/// <param name="location">Location of the file</param>
 		/// <returns>Byte array containing the contents of the file</returns>
-		public static byte[] ReadAllBytes(FileReference location)
-		{
-			return File.ReadAllBytes(location.FullName);
-		}
+		public static byte[] ReadAllBytes(FileReference location) => File.ReadAllBytes(location.FullName);
 
 		/// <summary>
 		/// Reads the contents of a file
@@ -448,10 +361,7 @@ namespace EpicGames.Core
 		/// <param name="location">Location of the file</param>
 		/// <param name="cancellationToken">Cancellation token for the operation</param>
 		/// <returns>Byte array containing the contents of the file</returns>
-		public static Task<byte[]> ReadAllBytesAsync(FileReference location, CancellationToken cancellationToken = default)
-		{
-			return File.ReadAllBytesAsync(location.FullName, cancellationToken);
-		}
+		public static Task<byte[]> ReadAllBytesAsync(FileReference location, CancellationToken cancellationToken = default) => File.ReadAllBytesAsync(location.FullName, cancellationToken);
 
 		/// <summary>
 		/// Reads the contents of a file
@@ -488,10 +398,7 @@ namespace EpicGames.Core
 		/// <param name="location">Location of the file</param>
 		/// <param name="encoding">Encoding of the file</param>
 		/// <returns>Contents of the file as a single string</returns>
-		public static string ReadAllText(FileReference location, Encoding encoding)
-		{
-			return File.ReadAllText(location.FullName, encoding);
-		}
+		public static string ReadAllText(FileReference location, Encoding encoding) => File.ReadAllText(location.FullName, encoding);
 
 		/// <summary>
 		/// Reads the contents of a file
@@ -499,10 +406,7 @@ namespace EpicGames.Core
 		/// <param name="location">Location of the file</param>
 		/// <param name="cancellationToken">Cancellation token for the operation</param>
 		/// <returns>Contents of the file as a single string</returns>
-		public static Task<string> ReadAllTextAsync(FileReference location, CancellationToken cancellationToken = default)
-		{
-			return File.ReadAllTextAsync(location.FullName, cancellationToken);
-		}
+		public static Task<string> ReadAllTextAsync(FileReference location, CancellationToken cancellationToken = default) => File.ReadAllTextAsync(location.FullName, cancellationToken);
 
 		/// <summary>
 		/// Reads the contents of a file
@@ -511,20 +415,14 @@ namespace EpicGames.Core
 		/// <param name="encoding">Encoding of the file</param>
 		/// <param name="cancellationToken">Cancellation token for the operation</param>
 		/// <returns>Contents of the file as a single string</returns>
-		public static Task<string> ReadAllTextAsync(FileReference location, Encoding encoding, CancellationToken cancellationToken = default)
-		{
-			return File.ReadAllTextAsync(location.FullName, encoding, cancellationToken);
-		}
+		public static Task<string> ReadAllTextAsync(FileReference location, Encoding encoding, CancellationToken cancellationToken = default) => File.ReadAllTextAsync(location.FullName, encoding, cancellationToken);
 
 		/// <summary>
 		/// Reads the contents of a file
 		/// </summary>
 		/// <param name="location">Location of the file</param>
 		/// <returns>String array containing the contents of the file</returns>
-		public static string[] ReadAllLines(FileReference location)
-		{
-			return File.ReadAllLines(location.FullName);
-		}
+		public static string[] ReadAllLines(FileReference location) => File.ReadAllLines(location.FullName);
 
 		/// <summary>
 		/// Reads the contents of a file
@@ -532,10 +430,7 @@ namespace EpicGames.Core
 		/// <param name="location">Location of the file</param>
 		/// <param name="encoding">The encoding to use when parsing the file</param>
 		/// <returns>String array containing the contents of the file</returns>
-		public static string[] ReadAllLines(FileReference location, Encoding encoding)
-		{
-			return File.ReadAllLines(location.FullName, encoding);
-		}
+		public static string[] ReadAllLines(FileReference location, Encoding encoding) => File.ReadAllLines(location.FullName, encoding);
 
 		/// <summary>
 		/// Reads the contents of a file
@@ -543,10 +438,7 @@ namespace EpicGames.Core
 		/// <param name="location">Location of the file</param>
 		/// <param name="cancellationToken">Cancellation token for the operation</param>
 		/// <returns>String array containing the contents of the file</returns>
-		public static Task<string[]> ReadAllLinesAsync(FileReference location, CancellationToken cancellationToken = default)
-		{
-			return File.ReadAllLinesAsync(location.FullName, cancellationToken);
-		}
+		public static Task<string[]> ReadAllLinesAsync(FileReference location, CancellationToken cancellationToken = default) => File.ReadAllLinesAsync(location.FullName, cancellationToken);
 
 		/// <summary>
 		/// Reads the contents of a file
@@ -555,70 +447,49 @@ namespace EpicGames.Core
 		/// <param name="encoding">The encoding to use when parsing the file</param>
 		/// <param name="cancellationToken">Cancellation token for the operation</param>
 		/// <returns>String array containing the contents of the file</returns>
-		public static Task<string[]> ReadAllLinesAsync(FileReference location, Encoding encoding, CancellationToken cancellationToken = default)
-		{
-			return File.ReadAllLinesAsync(location.FullName, encoding, cancellationToken);
-		}
+		public static Task<string[]> ReadAllLinesAsync(FileReference location, Encoding encoding, CancellationToken cancellationToken = default) => File.ReadAllLinesAsync(location.FullName, encoding, cancellationToken);
 
 		/// <summary>
 		/// Sets the attributes for a file
 		/// </summary>
 		/// <param name="location">Location of the file</param>
 		/// <param name="attributes">New attributes for the file</param>
-		public static void SetAttributes(FileReference location, FileAttributes attributes)
-		{
-			File.SetAttributes(location.FullName, attributes);
-		}
+		public static void SetAttributes(FileReference location, FileAttributes attributes) => File.SetAttributes(location.FullName, attributes);
 
 		/// <summary>
 		/// Sets the time that the file was last written to
 		/// </summary>
 		/// <param name="location">Location of the file</param>
 		/// <param name="lastWriteTime">Last write time, in local time</param>
-		public static void SetLastWriteTime(FileReference location, DateTime lastWriteTime)
-		{
-			File.SetLastWriteTime(location.FullName, lastWriteTime);
-		}
+		public static void SetLastWriteTime(FileReference location, DateTime lastWriteTime) => File.SetLastWriteTime(location.FullName, lastWriteTime);
 
 		/// <summary>
 		/// Sets the time that the file was last written to
 		/// </summary>
 		/// <param name="location">Location of the file</param>
 		/// <param name="lastWriteTimeUtc">Last write time, in UTC time</param>
-		public static void SetLastWriteTimeUtc(FileReference location, DateTime lastWriteTimeUtc)
-		{
-			File.SetLastWriteTimeUtc(location.FullName, lastWriteTimeUtc);
-		}
+		public static void SetLastWriteTimeUtc(FileReference location, DateTime lastWriteTimeUtc) => File.SetLastWriteTimeUtc(location.FullName, lastWriteTimeUtc);
 
 		/// <summary>
 		/// Sets the time that the file was last accessed.
 		/// </summary>
 		/// <param name="location">Location of the file.</param>
 		/// <param name="lastWriteTime">Last access time, in local time.</param>
-		public static void SetLastAccessTime(FileReference location, DateTime lastWriteTime)
-		{
-			File.SetLastWriteTime(location.FullName, lastWriteTime);
-		}
+		public static void SetLastAccessTime(FileReference location, DateTime lastWriteTime) => File.SetLastWriteTime(location.FullName, lastWriteTime);
 
 		/// <summary>
 		/// Sets the time that the file was last accessed.
 		/// </summary>
 		/// <param name="location">Location of the file.</param>
 		/// <param name="lastWriteTimeUtc">Last access time, in UTC time.</param>
-		public static void SetLastAccessTimeUtc(FileReference location, DateTime lastWriteTimeUtc)
-		{
-			File.SetLastWriteTimeUtc(location.FullName, lastWriteTimeUtc);
-		}
+		public static void SetLastAccessTimeUtc(FileReference location, DateTime lastWriteTimeUtc) => File.SetLastWriteTimeUtc(location.FullName, lastWriteTimeUtc);
 
 		/// <summary>
 		/// Writes the contents of a file
 		/// </summary>
 		/// <param name="location">Location of the file</param>
 		/// <param name="contents">Contents of the file</param>
-		public static void WriteAllBytes(FileReference location, byte[] contents)
-		{
-			File.WriteAllBytes(location.FullName, contents);
-		}
+		public static void WriteAllBytes(FileReference location, byte[] contents) => File.WriteAllBytes(location.FullName, contents);
 
 		/// <summary>
 		/// Writes the contents of a file
@@ -626,10 +497,7 @@ namespace EpicGames.Core
 		/// <param name="location">Location of the file</param>
 		/// <param name="contents">Contents of the file</param>
 		/// <param name="cancellationToken">Cancellation token for the operation</param>
-		public static Task WriteAllBytesAsync(FileReference location, byte[] contents, CancellationToken cancellationToken = default)
-		{
-			return File.WriteAllBytesAsync(location.FullName, contents, cancellationToken);
-		}
+		public static Task WriteAllBytesAsync(FileReference location, byte[] contents, CancellationToken cancellationToken = default) => File.WriteAllBytesAsync(location.FullName, contents, cancellationToken);
 
 		/// <summary>
 		/// Writes the data to the given file, if it's different from what's there already.
@@ -639,10 +507,10 @@ namespace EpicGames.Core
 		/// <param name="contents">Contents of the file</param>
 		public static bool WriteAllBytesIfDifferent(FileReference location, byte[] contents)
 		{
-			if(FileReference.Exists(location))
+			if (FileReference.Exists(location))
 			{
 				byte[] currentContents = FileReference.ReadAllBytes(location);
-				if(contents.AsSpan().SequenceEqual(currentContents))
+				if (contents.AsSpan().SequenceEqual(currentContents))
 				{
 					return false;
 				}
@@ -659,7 +527,7 @@ namespace EpicGames.Core
 		/// <param name="contents">Contents of the file</param>
 		public static bool WriteAllTextIfDifferent(FileReference location, string contents)
 		{
-			if(FileReference.Exists(location))
+			if (FileReference.Exists(location))
 			{
 				string currentContents = FileReference.ReadAllText(location);
 				if (String.Equals(contents, currentContents, StringComparison.Ordinal))
@@ -676,10 +544,7 @@ namespace EpicGames.Core
 		/// </summary>
 		/// <param name="location">Location of the file</param>
 		/// <param name="contents">Contents of the file</param>
-		public static void WriteAllLines(FileReference location, IEnumerable<string> contents)
-		{
-			File.WriteAllLines(location.FullName, contents);
-		}
+		public static void WriteAllLines(FileReference location, IEnumerable<string> contents) => File.WriteAllLines(location.FullName, contents);
 
 		/// <summary>
 		/// Writes the contents of a file
@@ -687,20 +552,14 @@ namespace EpicGames.Core
 		/// <param name="location">Location of the file</param>
 		/// <param name="contents">Contents of the file</param>
 		/// <param name="encoding">The encoding to use when parsing the file</param>
-		public static void WriteAllLines(FileReference location, IEnumerable<string> contents, Encoding encoding)
-		{
-			File.WriteAllLines(location.FullName, contents, encoding);
-		}
+		public static void WriteAllLines(FileReference location, IEnumerable<string> contents, Encoding encoding) => File.WriteAllLines(location.FullName, contents, encoding);
 
 		/// <summary>
 		/// Writes the contents of a file
 		/// </summary>
 		/// <param name="location">Location of the file</param>
 		/// <param name="contents">Contents of the file</param>
-		public static void WriteAllLines(FileReference location, string[] contents)
-		{
-			File.WriteAllLines(location.FullName, contents);
-		}
+		public static void WriteAllLines(FileReference location, string[] contents) => File.WriteAllLines(location.FullName, contents);
 
 		/// <summary>
 		/// Writes the contents of a file
@@ -708,10 +567,7 @@ namespace EpicGames.Core
 		/// <param name="location">Location of the file</param>
 		/// <param name="contents">Contents of the file</param>
 		/// <param name="encoding">The encoding to use when parsing the file</param>
-		public static void WriteAllLines(FileReference location, string[] contents, Encoding encoding)
-		{
-			File.WriteAllLines(location.FullName, contents, encoding);
-		}
+		public static void WriteAllLines(FileReference location, string[] contents, Encoding encoding) => File.WriteAllLines(location.FullName, contents, encoding);
 
 		/// <summary>
 		/// Writes the contents of a file
@@ -719,10 +575,7 @@ namespace EpicGames.Core
 		/// <param name="location">Location of the file</param>
 		/// <param name="contents">Contents of the file</param>
 		/// <param name="cancellationToken">Cancellation token for the operation</param>
-		public static Task WriteAllLinesAsync(FileReference location, IEnumerable<string> contents, CancellationToken cancellationToken = default)
-		{
-			return File.WriteAllLinesAsync(location.FullName, contents, cancellationToken);
-		}
+		public static Task WriteAllLinesAsync(FileReference location, IEnumerable<string> contents, CancellationToken cancellationToken = default) => File.WriteAllLinesAsync(location.FullName, contents, cancellationToken);
 
 		/// <summary>
 		/// Writes the contents of a file
@@ -731,10 +584,7 @@ namespace EpicGames.Core
 		/// <param name="contents">Contents of the file</param>
 		/// <param name="encoding">The encoding to use when parsing the file</param>
 		/// <param name="cancellationToken">Cancellation token for the operation</param>
-		public static Task WriteAllLinesAsync(FileReference location, IEnumerable<string> contents, Encoding encoding, CancellationToken cancellationToken = default)
-		{
-			return File.WriteAllLinesAsync(location.FullName, contents, encoding, cancellationToken);
-		}
+		public static Task WriteAllLinesAsync(FileReference location, IEnumerable<string> contents, Encoding encoding, CancellationToken cancellationToken = default) => File.WriteAllLinesAsync(location.FullName, contents, encoding, cancellationToken);
 
 		/// <summary>
 		/// Writes the contents of a file
@@ -742,10 +592,7 @@ namespace EpicGames.Core
 		/// <param name="location">Location of the file</param>
 		/// <param name="contents">Contents of the file</param>
 		/// <param name="cancellationToken">Cancellation token for the operation</param>
-		public static Task WriteAllLinesAsync(FileReference location, string[] contents, CancellationToken cancellationToken = default)
-		{
-			return File.WriteAllLinesAsync(location.FullName, contents, cancellationToken);
-		}
+		public static Task WriteAllLinesAsync(FileReference location, string[] contents, CancellationToken cancellationToken = default) => File.WriteAllLinesAsync(location.FullName, contents, cancellationToken);
 
 		/// <summary>
 		/// Writes the contents of a file
@@ -754,20 +601,14 @@ namespace EpicGames.Core
 		/// <param name="contents">Contents of the file</param>
 		/// <param name="encoding">The encoding to use when parsing the file</param>
 		/// <param name="cancellationToken">Cancellation token for the operation</param>
-		public static Task WriteAllLinesAsync(FileReference location, string[] contents, Encoding encoding, CancellationToken cancellationToken = default)
-		{
-			return File.WriteAllLinesAsync(location.FullName, contents, encoding, cancellationToken);
-		}
+		public static Task WriteAllLinesAsync(FileReference location, string[] contents, Encoding encoding, CancellationToken cancellationToken = default) => File.WriteAllLinesAsync(location.FullName, contents, encoding, cancellationToken);
 
 		/// <summary>
 		/// Writes the contents of a file
 		/// </summary>
 		/// <param name="location">Location of the file</param>
 		/// <param name="contents">Contents of the file</param>
-		public static void WriteAllText(FileReference location, string contents)
-		{
-			File.WriteAllText(location.FullName, contents);
-		}
+		public static void WriteAllText(FileReference location, string contents) => File.WriteAllText(location.FullName, contents);
 
 		/// <summary>
 		/// Writes the contents of a file
@@ -775,20 +616,14 @@ namespace EpicGames.Core
 		/// <param name="location">Location of the file</param>
 		/// <param name="contents">Contents of the file</param>
 		/// <param name="encoding">The encoding to use when parsing the file</param>
-		public static void WriteAllText(FileReference location, string contents, Encoding encoding)
-		{
-			File.WriteAllText(location.FullName, contents, encoding);
-		}
+		public static void WriteAllText(FileReference location, string contents, Encoding encoding) => File.WriteAllText(location.FullName, contents, encoding);
 
 		/// <summary>
 		/// Writes the contents of a file
 		/// </summary>
 		/// <param name="location">Location of the file</param>
 		/// <param name="contents">Contents of the file</param>
-		public static Task WriteAllTextAsync(FileReference location, string contents)
-		{
-			return File.WriteAllTextAsync(location.FullName, contents);
-		}
+		public static Task WriteAllTextAsync(FileReference location, string contents) => File.WriteAllTextAsync(location.FullName, contents);
 
 		/// <summary>
 		/// Writes the contents of a file
@@ -796,30 +631,21 @@ namespace EpicGames.Core
 		/// <param name="location">Location of the file</param>
 		/// <param name="contents">Contents of the file</param>
 		/// <param name="encoding">The encoding to use when parsing the file</param>
-		public static Task WriteAllTextAsync(FileReference location, string contents, Encoding encoding)
-		{
-			return File.WriteAllTextAsync(location.FullName, contents, encoding);
-		}
+		public static Task WriteAllTextAsync(FileReference location, string contents, Encoding encoding) => File.WriteAllTextAsync(location.FullName, contents, encoding);
 
 		/// <summary>
 		/// Appends the contents to a file
 		/// </summary>
 		/// <param name="location">Location of the file</param>
 		/// <param name="contents">Contents to append to the file</param>
-		public static void AppendAllLines(FileReference location, IEnumerable<string> contents)
-		{
-			File.AppendAllLines(location.FullName, contents);
-		}
+		public static void AppendAllLines(FileReference location, IEnumerable<string> contents) => File.AppendAllLines(location.FullName, contents);
 
 		/// <summary>
 		/// Appends the contents to a file
 		/// </summary>
 		/// <param name="location">Location of the file</param>
 		/// <param name="contents">Contents to append to the file</param>
-		public static Task AppendAllLinesAsync(FileReference location, IEnumerable<string> contents)
-		{
-			return File.AppendAllLinesAsync(location.FullName, contents);
-		}
+		public static Task AppendAllLinesAsync(FileReference location, IEnumerable<string> contents) => File.AppendAllLinesAsync(location.FullName, contents);
 
 		/// <summary>
 		/// Appends the contents to a file
@@ -827,10 +653,7 @@ namespace EpicGames.Core
 		/// <param name="location">Location of the file</param>
 		/// <param name="contents">Contents to append to the file</param>
 		/// <param name="encoding">The encoding to use when parsing the file</param>
-		public static void AppendAllLines(FileReference location, IEnumerable<string> contents, Encoding encoding)
-		{
-			File.AppendAllLines(location.FullName, contents, encoding);
-		}
+		public static void AppendAllLines(FileReference location, IEnumerable<string> contents, Encoding encoding) => File.AppendAllLines(location.FullName, contents, encoding);
 
 		/// <summary>
 		/// Appends the contents to a file
@@ -838,30 +661,21 @@ namespace EpicGames.Core
 		/// <param name="location">Location of the file</param>
 		/// <param name="contents">Contents to append to the file</param>
 		/// <param name="encoding">The encoding to use when parsing the file</param>
-		public static Task AppendAllLinesAsync(FileReference location, IEnumerable<string> contents, Encoding encoding)
-		{
-			return File.AppendAllLinesAsync(location.FullName, contents, encoding);
-		}
+		public static Task AppendAllLinesAsync(FileReference location, IEnumerable<string> contents, Encoding encoding) => File.AppendAllLinesAsync(location.FullName, contents, encoding);
 
 		/// <summary>
 		/// Appends the contents to a file
 		/// </summary>
 		/// <param name="location">Location of the file</param>
 		/// <param name="contents">Contents to append to the file</param>
-		public static void AppendAllLines(FileReference location, string[] contents)
-		{
-			File.AppendAllLines(location.FullName, contents);
-		}
+		public static void AppendAllLines(FileReference location, string[] contents) => File.AppendAllLines(location.FullName, contents);
 
 		/// <summary>
 		/// Appends the contents to a file
 		/// </summary>
 		/// <param name="location">Location of the file</param>
 		/// <param name="contents">Contents to append to the file</param>
-		public static Task AppendAllLinesAsync(FileReference location, string[] contents)
-		{
-			return File.AppendAllLinesAsync(location.FullName, contents);
-		}
+		public static Task AppendAllLinesAsync(FileReference location, string[] contents) => File.AppendAllLinesAsync(location.FullName, contents);
 
 		/// <summary>
 		/// Appends the contents to a file
@@ -869,10 +683,7 @@ namespace EpicGames.Core
 		/// <param name="location">Location of the file</param>
 		/// <param name="contents">Contents to append to the file</param>
 		/// <param name="encoding">The encoding to use when parsing the file</param>
-		public static void AppendAllLines(FileReference location, string[] contents, Encoding encoding)
-		{
-			File.AppendAllLines(location.FullName, contents, encoding);
-		}
+		public static void AppendAllLines(FileReference location, string[] contents, Encoding encoding) => File.AppendAllLines(location.FullName, contents, encoding);
 
 		/// <summary>
 		/// Appends the contents to a file
@@ -880,30 +691,21 @@ namespace EpicGames.Core
 		/// <param name="location">Location of the file</param>
 		/// <param name="contents">Contents to append to the file</param>
 		/// <param name="encoding">The encoding to use when parsing the file</param>
-		public static Task AppendAllLinesAsync(FileReference location, string[] contents, Encoding encoding)
-		{
-			return File.AppendAllLinesAsync(location.FullName, contents, encoding);
-		}
+		public static Task AppendAllLinesAsync(FileReference location, string[] contents, Encoding encoding) => File.AppendAllLinesAsync(location.FullName, contents, encoding);
 
 		/// <summary>
 		/// Appends the contents to a file
 		/// </summary>
 		/// <param name="location">Location of the file</param>
 		/// <param name="contents">Contents to append to the file</param>
-		public static void AppendAllText(FileReference location, string contents)
-		{
-			File.AppendAllText(location.FullName, contents);
-		}
+		public static void AppendAllText(FileReference location, string contents) => File.AppendAllText(location.FullName, contents);
 
 		/// <summary>
 		/// Appends the contents to a file
 		/// </summary>
 		/// <param name="location">Location of the file</param>
 		/// <param name="contents">Contents to append to the file</param>
-		public static Task AppendAllTextAsync(FileReference location, string contents)
-		{
-			return File.AppendAllTextAsync(location.FullName, contents);
-		}
+		public static Task AppendAllTextAsync(FileReference location, string contents) => File.AppendAllTextAsync(location.FullName, contents);
 
 		/// <summary>
 		/// Appends the contents to a file
@@ -911,10 +713,7 @@ namespace EpicGames.Core
 		/// <param name="location">Location of the file</param>
 		/// <param name="contents">Contents to append to the file</param>
 		/// <param name="encoding">The encoding to use when parsing the file</param>
-		public static void AppendAllText(FileReference location, string contents, Encoding encoding)
-		{
-			File.AppendAllText(location.FullName, contents, encoding);
-		}
+		public static void AppendAllText(FileReference location, string contents, Encoding encoding) => File.AppendAllText(location.FullName, contents, encoding);
 
 		/// <summary>
 		/// Appends the contents to a file
@@ -922,10 +721,7 @@ namespace EpicGames.Core
 		/// <param name="location">Location of the file</param>
 		/// <param name="contents">Contents to append to the file</param>
 		/// <param name="encoding">The encoding to use when parsing the file</param>
-		public static Task AppendAllTextAsync(FileReference location, string contents, Encoding encoding)
-		{
-			return File.AppendAllTextAsync(location.FullName, contents, encoding);
-		}
+		public static Task AppendAllTextAsync(FileReference location, string contents, Encoding encoding) => File.AppendAllTextAsync(location.FullName, contents, encoding);
 
 		#endregion
 	}
@@ -940,10 +736,7 @@ namespace EpicGames.Core
 		/// </summary>
 		/// <param name="writer">Binary writer to write to</param>
 		/// <param name="file">The file reference to write</param>
-		public static void Write(this BinaryWriter writer, FileReference file)
-		{
-			writer.Write((file == null) ? String.Empty : file.FullName);
-		}
+		public static void Write(this BinaryWriter writer, FileReference file) => writer.Write((file == null) ? String.Empty : file.FullName);
 
 		/// <summary>
 		/// Serializes a file reference, using a lookup table to avoid serializing the same name more than once.
@@ -975,10 +768,7 @@ namespace EpicGames.Core
 		/// </summary>
 		/// <param name="reader">Binary reader to read from</param>
 		/// <returns>New FileReference object</returns>
-		public static FileReference ReadFileReference(this BinaryReader reader)
-		{
-			return BinaryArchiveReader.NotNull(ReadFileReferenceOrNull(reader));
-		}
+		public static FileReference ReadFileReference(this BinaryReader reader) => BinaryArchiveReader.NotNull(ReadFileReferenceOrNull(reader));
 
 		/// <summary>
 		/// Manually deserialize a file reference from a binary stream.
@@ -997,10 +787,7 @@ namespace EpicGames.Core
 		/// <param name="reader">The source to read from</param>
 		/// <param name="uniqueFiles">List of previously read file references. The index into this array is used in place of subsequent ocurrences of the file.</param>
 		/// <returns>The file reference that was read</returns>
-		public static FileReference ReadFileReference(this BinaryReader reader, List<FileReference> uniqueFiles)
-		{
-			return BinaryArchiveReader.NotNull(ReadFileReferenceOrNull(reader, uniqueFiles));
-		}
+		public static FileReference ReadFileReference(this BinaryReader reader, List<FileReference> uniqueFiles) => BinaryArchiveReader.NotNull(ReadFileReferenceOrNull(reader, uniqueFiles));
 
 		/// <summary>
 		/// Deserializes a file reference, using a lookup table to avoid writing the same name more than once.
@@ -1034,7 +821,7 @@ namespace EpicGames.Core
 		/// <param name="file">The file reference to write</param>
 		public static void WriteFileReference(this BinaryArchiveWriter writer, FileReference? file)
 		{
-			if(file == null)
+			if (file == null)
 			{
 				writer.WriteString(null);
 			}
@@ -1049,10 +836,7 @@ namespace EpicGames.Core
 		/// </summary>
 		/// <param name="reader">Reader to serialize data from</param>
 		/// <returns>New file reference instance</returns>
-		public static FileReference ReadFileReference(this BinaryArchiveReader reader)
-		{
-			return BinaryArchiveReader.NotNull(ReadFileReferenceOrNull(reader));
-		}
+		public static FileReference ReadFileReference(this BinaryArchiveReader reader) => BinaryArchiveReader.NotNull(ReadFileReferenceOrNull(reader));
 
 		/// <summary>
 		/// Reads a FileReference from a binary archive
@@ -1062,14 +846,7 @@ namespace EpicGames.Core
 		public static FileReference? ReadFileReferenceOrNull(this BinaryArchiveReader reader)
 		{
 			string? fullName = reader.ReadString();
-			if(fullName == null)
-			{
-				return null;
-			}
-			else
-			{
-				return new FileReference(fullName, FileReference.Sanitize.None);
-			}
+			return fullName == null ? null : new FileReference(fullName, FileReference.Sanitize.None);
 		}
 	}
 }
diff --git a/Engine/Source/Programs/Shared/EpicGames.Core/FileSystemReference.cs b/Engine/Source/Programs/Shared/EpicGames.Core/FileSystemReference.cs
--- a/Engine/Source/Programs/Shared/EpicGames.Core/FileSystemReference.cs
+++ b/Engine/Source/Programs/Shared/EpicGames.Core/FileSystemReference.cs
@@ -33,10 +33,7 @@ namespace EpicGames.Core
 		/// <summary>
 		/// Direct constructor for a path
 		/// </summary>
-		protected FileSystemReference(string fullName)
-		{
-			FullName = fullName;
-		}
+		protected FileSystemReference(string fullName) => FullName = fullName;
 
 		static readonly ThreadLocal<StringBuilder> s_combineStringsStringBuilder = new ThreadLocal<StringBuilder>(() => new StringBuilder(260));
 
@@ -122,27 +119,16 @@ namespace EpicGames.Core
 		/// </summary>
 		/// <param name="extension">The extension to check</param>
 		/// <returns>True if this name has the given extension, false otherwise</returns>
-		public bool HasExtension(string extension)
-		{
-			if (extension.Length > 0 && extension[0] != '.')
-			{
-				return FullName.Length >= extension.Length + 1 && FullName[FullName.Length - extension.Length - 1] == '.' && FullName.EndsWith(extension, Comparison);
-			}
-			else
-			{
-				return FullName.EndsWith(extension, Comparison);
-			}
-		}
+		public bool HasExtension(string extension) => extension.Length > 0 && extension[0] != '.'
+				? FullName.Length >= extension.Length + 1 && FullName[FullName.Length - extension.Length - 1] == '.' && FullName.EndsWith(extension, Comparison)
+				: FullName.EndsWith(extension, Comparison);
 
 		/// <summary>
 		/// Determines if the given object is at or under the given directory
 		/// </summary>
 		/// <param name="other">Directory to check against</param>
 		/// <returns>True if this path is under the given directory</returns>
-		public bool IsUnderDirectory(DirectoryReference other)
-		{
-			return FullName.StartsWith(other.FullName, Comparison) && (FullName.Length == other.FullName.Length || FullName[other.FullName.Length] == Path.DirectorySeparatorChar || other.IsRootDirectory());
-		}
+		public bool IsUnderDirectory(DirectoryReference other) => FullName.StartsWith(other.FullName, Comparison) && (FullName.Length == other.FullName.Length || FullName[other.FullName.Length] == Path.DirectorySeparatorChar || other.IsRootDirectory());
 
 		/// <summary>
 		/// Searches the path fragments for the given name. Only complete fragments are considered a match.
@@ -150,10 +136,7 @@ namespace EpicGames.Core
 		/// <param name="name">Name to check for</param>
 		/// <param name="offset">Offset within the string to start the search</param>
 		/// <returns>True if the given name is found within the path</returns>
-		public bool ContainsName(string name, int offset)
-		{
-			return ContainsName(name, offset, FullName.Length - offset);
-		}
+		public bool ContainsName(string name, int offset) => ContainsName(name, offset, FullName.Length - offset);
 
 		/// <summary>
 		/// Searches the path fragments for the given name. Only complete fragments are considered a match.
@@ -165,25 +148,25 @@ namespace EpicGames.Core
 		public bool ContainsName(string name, int offset, int length)
 		{
 			// Check the substring to search is at least long enough to contain a match
-			if(length < name.Length)
+			if (length < name.Length)
 			{
 				return false;
 			}
 
 			// Find each occurence of the name within the remaining string, then test whether it's surrounded by directory separators
 			int matchIdx = offset;
-			for(;;)
+			for (; ; )
 			{
 				// Find the next occurrence
 				matchIdx = FullName.IndexOf(name, matchIdx, offset + length - matchIdx, Comparison);
-				if(matchIdx == -1)
+				if (matchIdx == -1)
 				{
 					return false;
 				}
 
 				// Check if the substring is a directory
 				int matchEndIdx = matchIdx + name.Length;
-				if(FullName[matchIdx - 1] == Path.DirectorySeparatorChar && (matchEndIdx == FullName.Length || FullName[matchEndIdx] == Path.DirectorySeparatorChar))
+				if (FullName[matchIdx - 1] == Path.DirectorySeparatorChar && (matchEndIdx == FullName.Length || FullName[matchEndIdx] == Path.DirectorySeparatorChar))
 				{
 					return true;
 				}
@@ -199,18 +182,7 @@ namespace EpicGames.Core
 		/// <param name="name">Name of a subfolder to also check for</param>
 		/// <param name="baseDir">Base directory to check against</param>
 		/// <returns>True if the path is under the given directory</returns>
-		public bool ContainsName(string name, DirectoryReference baseDir)
-		{
-			// Check that this is under the base directory
-			if(!IsUnderDirectory(baseDir))
-			{
-				return false;
-			}
-			else
-			{
-				return ContainsName(name, baseDir.FullName.Length);
-			}
-		}
+		public bool ContainsName(string name, DirectoryReference baseDir) => IsUnderDirectory(baseDir) && ContainsName(name, baseDir.FullName.Length);
 
 		/// <summary>
 		/// Determines if the given object is under the given directory, within a subfolder of the given name. Useful for masking out directories by name.
@@ -218,18 +190,7 @@ namespace EpicGames.Core
 		/// <param name="names">Names of subfolders to also check for</param>
 		/// <param name="baseDir">Base directory to check against</param>
 		/// <returns>True if the path is under the given directory</returns>
-		public bool ContainsAnyNames(IEnumerable<string> names, DirectoryReference baseDir)
-		{
-			// Check that this is under the base directory
-			if(!IsUnderDirectory(baseDir))
-			{
-				return false;
-			}
-			else
-			{
-				return names.Any(x => ContainsName(x, baseDir.FullName.Length));
-			}
-		}
+		public bool ContainsAnyNames(IEnumerable<string> names, DirectoryReference baseDir) => IsUnderDirectory(baseDir) && names.Any(x => ContainsName(x, baseDir.FullName.Length));
 
 		static readonly ThreadLocal<StringBuilder> s_makeRelativeToStringBuilder = new ThreadLocal<StringBuilder>(() => new StringBuilder(260));
 
@@ -284,7 +245,7 @@ namespace EpicGames.Core
 				else
 				{
 					// Check the two paths match, and bail if they don't. Increase the common directory length if we've reached a separator.
-					if(String.Compare(FullName, idx, directory.FullName, idx, 1, Comparison) != 0)
+					if (String.Compare(FullName, idx, directory.FullName, idx, 1, Comparison) != 0)
 					{
 						break;
 					}
@@ -333,18 +294,12 @@ namespace EpicGames.Core
 		/// Normalize the path to using forward slashes
 		/// </summary>
 		/// <returns></returns>
-		public string ToNormalizedPath()
-		{
-			return FullName.Replace("\\", "/", StringComparison.Ordinal);
-		}
+		public string ToNormalizedPath() => FullName.Replace("\\", "/", StringComparison.Ordinal);
 
 		/// <summary>
 		/// Returns a string representation of this filesystem object
 		/// </summary>
 		/// <returns>Full path to the object</returns>
-		public override string ToString()
-		{
-			return FullName;
-		}
+		public override string ToString() => FullName;
 	}
 }
