Horde: Prepare for new DirectoryNode version by discarding additional data from FileEntry objects.

diff --git a/Engine/Source/Programs/Shared/EpicGames.Horde/Storage/Nodes/FileEntry.cs b/Engine/Source/Programs/Shared/EpicGames.Horde/Storage/Nodes/FileEntry.cs
--- a/Engine/Source/Programs/Shared/EpicGames.Horde/Storage/Nodes/FileEntry.cs
+++ b/Engine/Source/Programs/Shared/EpicGames.Horde/Storage/Nodes/FileEntry.cs
@@ -94,7 +94,7 @@ namespace EpicGames.Horde.Storage.Nodes
 		/// </summary>
 		/// <param name="reader"></param>
 		public FileEntry(IBlobReader reader)
-			: base(reader)
+			: base(ReadNodeRef(reader))
 		{
 			Name = reader.ReadString();
 			Flags = (FileEntryFlags)reader.ReadUnsignedVarInt();
@@ -108,6 +108,18 @@ namespace EpicGames.Horde.Storage.Nodes
 			}
 		}
 
+		static HashedNodeRef<ChunkedDataNode> ReadNodeRef(IBlobReader reader)
+		{
+			if (reader.Version >= 2)
+			{
+				return new ChunkedDataNodeRef(reader);
+			}
+			else
+			{
+				return reader.ReadHashedNodeRef<ChunkedDataNode>();
+			}
+		}
+
 		/// <summary>
 		/// Serialize this entry
 		/// </summary>
