Horde: Include the type of file entry nodes in file entries. This allows scanning for the handles needed to extract from storage without having to expand leaf nodes to determine their type.

* Also change FileEntry and DirectoryEntry to be standalone types rather than deriving from HashedNodeRef<T>.

diff --git a/Engine/Source/Programs/Horde/Horde.Agent/Utility/TempStorage.cs b/Engine/Source/Programs/Horde/Horde.Agent/Utility/TempStorage.cs
--- a/Engine/Source/Programs/Horde/Horde.Agent/Utility/TempStorage.cs
+++ b/Engine/Source/Programs/Horde/Horde.Agent/Utility/TempStorage.cs
@@ -765,7 +765,7 @@ namespace Horde.Storage.Utility
 				Stopwatch timer = Stopwatch.StartNew();
 
 				// Add all the files and flush the ref
-				DirectoryNode rootDirNode = await rootDirEntry.ExpandAsync(cancellationToken);
+				DirectoryNode rootDirNode = await rootDirEntry.Target.ExpandAsync(cancellationToken);
 				await rootDirNode.CopyToDirectoryAsync(rootDir.ToDirectoryInfo(), new CopyStatsLogger(logger), logger, cancellationToken);
 
 				StorageStats deltaStats = StorageStats.GetDelta(initialStats, storageClient.GetStats());
diff --git a/Engine/Source/Programs/Horde/Horde.Server/Artifacts/ArtifactsController.cs b/Engine/Source/Programs/Horde/Horde.Server/Artifacts/ArtifactsController.cs
--- a/Engine/Source/Programs/Horde/Horde.Server/Artifacts/ArtifactsController.cs
+++ b/Engine/Source/Programs/Horde/Horde.Server/Artifacts/ArtifactsController.cs
@@ -292,7 +292,7 @@ namespace Horde.Server.Artifacts
 					{
 						return NotFound();
 					}
-					directoryNode = await nextDirectoryEntry.ExpandAsync(cancellationToken);
+					directoryNode = await nextDirectoryEntry.Target.ExpandAsync(cancellationToken);
 				}
 			}
 
@@ -308,9 +308,9 @@ namespace Horde.Server.Artifacts
 				response.Directories = new List<GetArtifactDirectoryEntryResponse>();
 				foreach (DirectoryEntry subDirectoryEntry in directoryNode.Directories)
 				{
-					DirectoryNode subDirectoryNode = await subDirectoryEntry.ExpandAsync(cancellationToken);
+					DirectoryNode subDirectoryNode = await subDirectoryEntry.Target.ExpandAsync(cancellationToken);
 
-					GetArtifactDirectoryEntryResponse subDirectoryEntryResponse = new GetArtifactDirectoryEntryResponse(subDirectoryEntry.Name.ToString(), subDirectoryEntry.Length, subDirectoryEntry.Hash);
+					GetArtifactDirectoryEntryResponse subDirectoryEntryResponse = new GetArtifactDirectoryEntryResponse(subDirectoryEntry.Name.ToString(), subDirectoryEntry.Length, subDirectoryEntry.Target.Hash);
 					if (depth == 0)
 					{
 						if (subDirectoryNode.Directories.Count + subDirectoryNode.Files.Count < 16)
@@ -338,7 +338,7 @@ namespace Horde.Server.Artifacts
 
 			if (directoryNode.Files.Count > 0)
 			{
-				response.Files = directoryNode.Files.Select(x => new GetArtifactFileEntryResponse(x.Name.ToString(), x.Length, x.Hash)).ToList();
+				response.Files = directoryNode.Files.Select(x => new GetArtifactFileEntryResponse(x.Name.ToString(), x.Length, x.StreamHash)).ToList();
 			}
 		}
 
diff --git a/Engine/Source/Programs/Horde/Horde.Server/Perforce/PerforceReplicator.cs b/Engine/Source/Programs/Horde/Horde.Server/Perforce/PerforceReplicator.cs
--- a/Engine/Source/Programs/Horde/Horde.Server/Perforce/PerforceReplicator.cs
+++ b/Engine/Source/Programs/Horde/Horde.Server/Perforce/PerforceReplicator.cs
@@ -407,7 +407,7 @@ namespace Horde.Server.Perforce
 				}
 				else
 				{
-					syncNode = new SyncNode(change, parentChange, parent.Contents);
+					syncNode = new SyncNode(change, parentChange, parent.Contents.Target);
 				}
 			}
 
@@ -683,7 +683,7 @@ namespace Horde.Server.Perforce
 			await store.WriteRefTargetAsync(refName, commitNodeRef, options.RefOptions, cancellationToken: cancellationToken);
 
 			// Log the snapshot info
-			_logger.LogInformation("Snapshot for {StreamId} CL {Change} is ref {RefName} (commit: {CommitHandle}, root: {RootHandle})", streamConfig.Id, change, refName, commitNodeRef.Handle.GetLocator(), rootRef.Handle.GetLocator());
+			_logger.LogInformation("Snapshot for {StreamId} CL {Change} is ref {RefName} (commit: {CommitHandle}, root: {RootHandle})", streamConfig.Id, change, refName, commitNodeRef.Handle.GetLocator(), rootRef.Target.Handle.GetLocator());
 		}
 
 		static int GetFileOffset(string path)
diff --git a/Engine/Source/Programs/Horde/Horde.Server/Storage/StorageController.cs b/Engine/Source/Programs/Horde/Horde.Server/Storage/StorageController.cs
--- a/Engine/Source/Programs/Horde/Horde.Server/Storage/StorageController.cs
+++ b/Engine/Source/Programs/Horde/Horde.Server/Storage/StorageController.cs
@@ -405,13 +405,13 @@ namespace Horde.Server.Storage
 						List<object> directories = new List<object>();
 						foreach ((string name, DirectoryEntry entry) in directoryNode.NameToDirectory)
 						{
-							directories.Add(new { name = name.ToString(), length = entry.Length, hash = entry.Hash, link = GetNodeLink(linkBase, entry.Handle) });
+							directories.Add(new { name = name.ToString(), length = entry.Length, hash = entry.Target.Hash, link = GetNodeLink(linkBase, entry.Target.Handle) });
 						}
 
 						List<object> files = new List<object>();
 						foreach ((string name, FileEntry entry) in directoryNode.NameToFile)
 						{
-							files.Add(new { name = name.ToString(), length = entry.Length, flags = entry.Flags, hash = entry.Hash, link = GetNodeLink(linkBase, entry.Handle) });
+							files.Add(new { name = name.ToString(), length = entry.Length, flags = entry.Flags, hash = entry.StreamHash, link = GetNodeLink(linkBase, entry.Target.Handle) });
 						}
 
 						content = new { directoryNode.Length, directories, files };
diff --git a/Engine/Source/Programs/Horde/Horde/Commands/Vcs/VcsCheckout.cs b/Engine/Source/Programs/Horde/Horde/Commands/Vcs/VcsCheckout.cs
--- a/Engine/Source/Programs/Horde/Horde/Commands/Vcs/VcsCheckout.cs
+++ b/Engine/Source/Programs/Horde/Horde/Commands/Vcs/VcsCheckout.cs
@@ -73,7 +73,7 @@ namespace Horde.Commands.Vcs
 
 			DirectoryState newState = new DirectoryState();
 
-			DirectoryNode directoryNode = await directoryRef.ExpandAsync();
+			DirectoryNode directoryNode = await directoryRef.Target.ExpandAsync();
 			foreach ((string name, DirectoryEntry? subDirEntry, DirectoryState? subDirState) in EnumerableExtensions.Zip(directoryNode.NameToDirectory, directoryState?.Directories))
 			{
 				DirectoryReference subDirPath = DirectoryReference.Combine(dirPath, name.ToString());
@@ -100,7 +100,7 @@ namespace Horde.Commands.Vcs
 						FileReference.Delete(filePath);
 					}
 				}
-				else if (fileState == null || fileState.Hash != fileEntry.Hash)
+				else if (fileState == null || fileState.Hash != fileEntry.StreamHash)
 				{
 					newState.Files[name] = await CheckoutFileAsync(fileEntry, filePath.ToFileInfo(), logger);
 				}
@@ -110,16 +110,16 @@ namespace Horde.Commands.Vcs
 				}
 			}
 
-			newState.Hash = directoryRef.Hash;
+			newState.Hash = directoryRef.Target.Hash;
 			return newState;
 		}
 
 		static async Task<FileState> CheckoutFileAsync(FileEntry fileRef, FileInfo fileInfo, ILogger logger)
 		{
-			logger.LogInformation("Updating {File} to {Hash}", fileInfo, fileRef.Hash);
-			await ChunkedDataNode.CopyToFileAsync(fileRef.Handle, fileInfo, CancellationToken.None);
+			logger.LogInformation("Updating {File} to {Hash}", fileInfo, fileRef.StreamHash);
+			await ChunkedDataNode.CopyToFileAsync(fileRef.Target.Handle, fileInfo, CancellationToken.None);
 			fileInfo.Refresh();
-			return new FileState(fileInfo, fileRef.Hash);
+			return new FileState(fileInfo, fileRef.StreamHash);
 		}
 	}
 }
diff --git a/Engine/Source/Programs/Shared/EpicGames.Horde.Tests/BundleStoreTests.cs b/Engine/Source/Programs/Shared/EpicGames.Horde.Tests/BundleStoreTests.cs
--- a/Engine/Source/Programs/Shared/EpicGames.Horde.Tests/BundleStoreTests.cs
+++ b/Engine/Source/Programs/Shared/EpicGames.Horde.Tests/BundleStoreTests.cs
@@ -264,11 +264,11 @@ namespace EpicGames.Horde.Tests
 			Assert.AreEqual(1, root.Directories.Count);
 			Assert.AreEqual("hello", root.Directories.First().Name);
 
-			DirectoryNode hello = await root.Directories.First().ExpandAsync(CancellationToken.None);
+			DirectoryNode hello = await root.Directories.First().Target.ExpandAsync(CancellationToken.None);
 			Assert.AreEqual(1, hello.Directories.Count);
 			Assert.AreEqual("world", hello.Directories.First().Name);
 
-			DirectoryNode world = await hello.Directories.First().ExpandAsync(CancellationToken.None);
+			DirectoryNode world = await hello.Directories.First().Target.ExpandAsync(CancellationToken.None);
 			Assert.AreEqual(0, world.Directories.Count);
 		}
 
@@ -310,12 +310,12 @@ namespace EpicGames.Horde.Tests
 			Assert.AreEqual(1, root.Directories.Count);
 			Assert.AreEqual("hello", root.Directories.First().Name);
 
-			DirectoryNode hello = await root.Directories.First().ExpandAsync();
+			DirectoryNode hello = await root.Directories.First().Target.ExpandAsync();
 			Assert.AreEqual(0, hello.Directories.Count);
 			Assert.AreEqual(1, hello.Files.Count);
 			Assert.AreEqual("world", hello.Files.First().Name);
 
-			ChunkedDataNode world = await hello.Files.First().ExpandAsync();
+			ChunkedDataNode world = await hello.Files.First().Target.ExpandAsync();
 
 			byte[] worldData = await GetFileDataAsync(world);
 			Assert.IsTrue(worldData.SequenceEqual(Encoding.UTF8.GetBytes("world")));
@@ -406,10 +406,10 @@ namespace EpicGames.Horde.Tests
 				await CompareTreesAsync(root, newRoot);
 				await CheckLargeFileTreeAsync(root, data);
 
-				HashedNodeRef<ChunkedDataNode> file = root.GetFileEntry("test");
+				FileEntry file = root.GetFileEntry("test");
 
 				Dictionary<BlobLocator, long> locatorToSize = new Dictionary<BlobLocator, long>();
-				await GetUniqueBlobsAsync(file.Handle, locatorToSize);
+				await GetUniqueBlobsAsync(file.Target.Handle, locatorToSize);
 
 				long uniqueSize = locatorToSize.Sum(x => x.Value);
 				Assert.IsTrue(uniqueSize < data.Length / 3); // random fraction meaning "lots of dedupe happened"
@@ -435,8 +435,8 @@ namespace EpicGames.Horde.Tests
 
 			foreach ((FileEntry oldFileEntry, FileEntry newFileEntry) in oldNode.Files.Zip(newNode.Files))
 			{
-				ChunkedDataNode oldFile = await oldFileEntry.ExpandAsync();
-				ChunkedDataNode newFile = await newFileEntry.ExpandAsync();
+				ChunkedDataNode oldFile = await oldFileEntry.Target.ExpandAsync();
+				ChunkedDataNode newFile = await newFileEntry.Target.ExpandAsync();
 				await CompareTreesAsync(oldFile, newFile);
 			}
 		}
@@ -473,7 +473,7 @@ namespace EpicGames.Horde.Tests
 			Assert.AreEqual(0, root.Directories.Count);
 			Assert.AreEqual(1, root.Files.Count);
 
-			ChunkedDataNode world = await root.Files.First().ExpandAsync(CancellationToken.None);
+			ChunkedDataNode world = await root.Files.First().Target.ExpandAsync(CancellationToken.None);
 
 			int length = await CheckFileDataAsync(world, data);
 			Assert.AreEqual(data.Length, length);
diff --git a/Engine/Source/Programs/Shared/EpicGames.Horde/Storage/Nodes/ChunkedDataNode.cs b/Engine/Source/Programs/Shared/EpicGames.Horde/Storage/Nodes/ChunkedDataNode.cs
--- a/Engine/Source/Programs/Shared/EpicGames.Horde/Storage/Nodes/ChunkedDataNode.cs
+++ b/Engine/Source/Programs/Shared/EpicGames.Horde/Storage/Nodes/ChunkedDataNode.cs
@@ -448,7 +448,7 @@ namespace EpicGames.Horde.Storage.Nodes
 		/// <returns>Handle to the root node of the tree</returns>
 		public static async Task<ChunkedData> CreateTreeAsync(LeafChunkedData leafChunkedData, InteriorChunkedDataNodeOptions options, IStorageWriter writer, CancellationToken cancellationToken)
 		{
-			HashedNodeRef<ChunkedDataNode> rootRef = await CreateTreeAsync(leafChunkedData.LeafHandles, options, writer, cancellationToken);
+			ChunkedDataNodeRef rootRef = await CreateTreeAsync(leafChunkedData.LeafHandles, options, writer, cancellationToken);
 			return new ChunkedData(leafChunkedData.Hash, rootRef);
 		}
 
@@ -460,7 +460,7 @@ namespace EpicGames.Horde.Storage.Nodes
 		/// <param name="writer">Output writer for new interior nodes</param>
 		/// <param name="cancellationToken">Cancellation token for the operation</param>
 		/// <returns>Handle to the root node of the tree</returns>
-		public static async Task<HashedNodeRef<ChunkedDataNode>> CreateTreeAsync(List<ChunkedDataNodeRef> nodeRefs, InteriorChunkedDataNodeOptions options, IStorageWriter writer, CancellationToken cancellationToken)
+		public static async Task<ChunkedDataNodeRef> CreateTreeAsync(List<ChunkedDataNodeRef> nodeRefs, InteriorChunkedDataNodeOptions options, IStorageWriter writer, CancellationToken cancellationToken)
 		{
 			List<ChunkedDataNodeRef> handleBuffer = new List<ChunkedDataNodeRef>();
 
diff --git a/Engine/Source/Programs/Shared/EpicGames.Horde/Storage/Nodes/ChunkedDataWriter.cs b/Engine/Source/Programs/Shared/EpicGames.Horde/Storage/Nodes/ChunkedDataWriter.cs
--- a/Engine/Source/Programs/Shared/EpicGames.Horde/Storage/Nodes/ChunkedDataWriter.cs
+++ b/Engine/Source/Programs/Shared/EpicGames.Horde/Storage/Nodes/ChunkedDataWriter.cs
@@ -56,9 +56,9 @@ namespace EpicGames.Horde.Storage.Nodes
 	/// <summary>
 	/// Describes a chunked data stream
 	/// </summary>
-	/// <param name="Hash">Hash of the stream</param>
+	/// <param name="StreamHash">Hash of the stream as a contiguous buffer</param>
 	/// <param name="Root">Handle to the root chunk containing the data</param>
-	public record class ChunkedData(IoHash Hash, HashedNodeRef<ChunkedDataNode> Root);
+	public record class ChunkedData(IoHash StreamHash, ChunkedDataNodeRef Root);
 
 	/// <summary>
 	/// Utility class for generating FileNode data directly into <see cref="IStorageWriter"/> instances, without constructing node representations first.
@@ -322,7 +322,7 @@ namespace EpicGames.Horde.Storage.Nodes
 		public async Task<ChunkedData> CompleteAsync(CancellationToken cancellationToken)
 		{
 			await FlushLeafNodeAsync(cancellationToken);
-			HashedNodeRef<ChunkedDataNode> rootHandle = await InteriorChunkedDataNode.CreateTreeAsync(_leafHandles, _options.InteriorOptions, _writer, cancellationToken);
+			ChunkedDataNodeRef rootHandle = await InteriorChunkedDataNode.CreateTreeAsync(_leafHandles, _options.InteriorOptions, _writer, cancellationToken);
 			return new ChunkedData(IoHash.FromBlake3(_hasher), rootHandle);
 		}
 
diff --git a/Engine/Source/Programs/Shared/EpicGames.Horde/Storage/Nodes/CommitNode.cs b/Engine/Source/Programs/Shared/EpicGames.Horde/Storage/Nodes/CommitNode.cs
--- a/Engine/Source/Programs/Shared/EpicGames.Horde/Storage/Nodes/CommitNode.cs
+++ b/Engine/Source/Programs/Shared/EpicGames.Horde/Storage/Nodes/CommitNode.cs
@@ -110,7 +110,7 @@ namespace EpicGames.Horde.Storage.Nodes
 			writer.WriteOptionalString(CommitterId);
 			writer.WriteString(Message);
 			writer.WriteDateTime(Time);
-			writer.WriteHashedNodeRef(Contents);
+			writer.WriteDirectoryNodeRef(Contents);
 			writer.WriteDictionary(Metadata, key => writer.WriteGuidUnrealOrder(key), value => writer.WriteHashedNodeRef(value));
 		}
 	}
diff --git a/Engine/Source/Programs/Shared/EpicGames.Horde/Storage/Nodes/DirectoryEntry.cs b/Engine/Source/Programs/Shared/EpicGames.Horde/Storage/Nodes/DirectoryEntry.cs
--- a/Engine/Source/Programs/Shared/EpicGames.Horde/Storage/Nodes/DirectoryEntry.cs
+++ b/Engine/Source/Programs/Shared/EpicGames.Horde/Storage/Nodes/DirectoryEntry.cs
@@ -47,4 +47,18 @@ namespace EpicGames.Horde.Storage.Nodes
 		/// <inheritdoc/>
 		public override string ToString() => Name;
 	}
+
+	/// <summary>
+	/// Extension methods for <see cref="DirectoryEntry"/>
+	/// </summary>
+	public static class DirectoryEntryExtensions
+	{
+		/// <summary>
+		/// Serialize a directory entry to storage
+		/// </summary>
+		public static void WriteDirectoryEntry(this IBlobWriter writer, DirectoryEntry entry)
+		{
+			entry.Serialize(writer);
+		}
+	}
 }
diff --git a/Engine/Source/Programs/Shared/EpicGames.Horde/Storage/Nodes/DirectoryNode.cs b/Engine/Source/Programs/Shared/EpicGames.Horde/Storage/Nodes/DirectoryNode.cs
--- a/Engine/Source/Programs/Shared/EpicGames.Horde/Storage/Nodes/DirectoryNode.cs
+++ b/Engine/Source/Programs/Shared/EpicGames.Horde/Storage/Nodes/DirectoryNode.cs
@@ -140,7 +140,7 @@ namespace EpicGames.Horde.Storage.Nodes
 	/// <summary>
 	/// A directory node
 	/// </summary>
-	[BlobType("{0714EC11-4D07-291A-8AE7-7F86799980D6}", 1)]
+	[BlobType("{0714EC11-4D07-291A-8AE7-7F86799980D6}", 2)]
 	public class DirectoryNode : Node
 	{
 		readonly SortedDictionary<string, FileEntry> _nameToFileEntry = new SortedDictionary<string, FileEntry>(StringComparer.Ordinal);
@@ -216,13 +216,13 @@ namespace EpicGames.Horde.Storage.Nodes
 			writer.WriteUnsignedVarInt(Files.Count);
 			foreach (FileEntry fileEntry in _nameToFileEntry.Values)
 			{
-				writer.WriteHashedNodeRef(fileEntry);
+				writer.WriteFileEntry(fileEntry);
 			}
 
 			writer.WriteUnsignedVarInt(Directories.Count);
 			foreach (DirectoryEntry directoryEntry in _nameToDirectoryEntry.Values)
 			{
-				writer.WriteHashedNodeRef(directoryEntry);
+				writer.WriteDirectoryEntry(directoryEntry);
 			}
 		}
 
@@ -381,7 +381,7 @@ namespace EpicGames.Horde.Storage.Nodes
 					return null;
 				}
 
-				directoryNode = await directoryEntry.ExpandAsync(cancellationToken);
+				directoryNode = await directoryEntry.Target.ExpandAsync(cancellationToken);
 			}
 			return directoryNode;
 		}
@@ -470,7 +470,7 @@ namespace EpicGames.Horde.Storage.Nodes
 		{
 			if (TryGetDirectoryEntry(name, out DirectoryEntry? entry))
 			{
-				return await entry.ExpandAsync(cancellationToken);
+				return await entry.Target.ExpandAsync(cancellationToken);
 			}
 			else
 			{
@@ -813,7 +813,7 @@ namespace EpicGames.Horde.Storage.Nodes
 				if (windowOffset < directoryEntry.Length && windowOffset + windowLength >= 0)
 				{
 					DirectoryInfo subDirectoryInfo = directoryInfo.CreateSubdirectory(directoryEntry.Name.ToString());
-					DirectoryNode subDirectoryNode = await directoryEntry.ExpandAsync(cancellationToken);
+					DirectoryNode subDirectoryNode = await directoryEntry.Target.ExpandAsync(cancellationToken);
 					await subDirectoryNode.CopyToDirectoryInternalAsync(subDirectoryInfo, windowOffset, windowLength, copyStats, logger, cancellationToken);
 				}
 				windowOffset -= directoryEntry.Length;
@@ -1083,7 +1083,7 @@ namespace EpicGames.Horde.Storage.Nodes
 				string directoryPath = $"{prefix}{directoryEntry.Name}/";
 				if (filter == null || filter.PossiblyMatches(directoryPath))
 				{
-					DirectoryNode node = await directoryEntry.ExpandAsync(cancellationToken);
+					DirectoryNode node = await directoryEntry.Target.ExpandAsync(cancellationToken);
 					await CopyFilesAsync(node, directoryPath, filter, archive, logger, cancellationToken);
 				}
 				numCopiedDirs++;
diff --git a/Engine/Source/Programs/Shared/EpicGames.Horde/Storage/Nodes/DirectoryNodeRef.cs b/Engine/Source/Programs/Shared/EpicGames.Horde/Storage/Nodes/DirectoryNodeRef.cs
--- a/Engine/Source/Programs/Shared/EpicGames.Horde/Storage/Nodes/DirectoryNodeRef.cs
+++ b/Engine/Source/Programs/Shared/EpicGames.Horde/Storage/Nodes/DirectoryNodeRef.cs
@@ -7,28 +7,33 @@ namespace EpicGames.Horde.Storage.Nodes
 	/// <summary>
 	/// Reference to a directory node, including the target hash and length
 	/// </summary>
-	public class DirectoryNodeRef : HashedNodeRef<DirectoryNode>
+	public class DirectoryNodeRef
 	{
 		/// <summary>
 		/// Length of this directory tree
 		/// </summary>
 		public long Length { get; }
 
+		/// <summary>
+		/// Target node
+		/// </summary>
+		public HashedNodeRef<DirectoryNode> Target { get; }
+
 		/// <summary>
 		/// Constructor
 		/// </summary>
-		public DirectoryNodeRef(long length, HashedNodeRef<DirectoryNode> nodeRef)
-			: base(nodeRef)
+		public DirectoryNodeRef(long length, HashedNodeRef<DirectoryNode> target)
 		{
 			Length = length;
+			Target = target;
 		}
 
 		/// <summary>
 		/// Constructor
 		/// </summary>
 		public DirectoryNodeRef(IBlobReader reader)
-			: base(reader)
 		{
+			Target = new HashedNodeRef<DirectoryNode>(reader);
 			Length = (long)reader.ReadUnsignedVarInt();
 		}
 
@@ -36,11 +41,24 @@ namespace EpicGames.Horde.Storage.Nodes
 		/// Serialize this directory entry to disk
 		/// </summary>
 		/// <param name="writer"></param>
-		public override void Serialize(IBlobWriter writer)
+		public virtual void Serialize(IBlobWriter writer)
 		{
-			base.Serialize(writer);
-
+			writer.WriteHashedNodeRef(Target);
 			writer.WriteUnsignedVarInt((ulong)Length);
 		}
 	}
+
+	/// <summary>
+	/// Extension methods for <see cref="DirectoryNodeRef"/>
+	/// </summary>
+	public static class DirectoryNodeRefExtensions
+	{
+		/// <summary>
+		/// Serialize a directory node ref to a blob writer
+		/// </summary>
+		public static void WriteDirectoryNodeRef(this IBlobWriter writer, DirectoryNodeRef nodeRef)
+		{
+			nodeRef.Serialize(writer);
+		}
+	}
 }
diff --git a/Engine/Source/Programs/Shared/EpicGames.Horde/Storage/Nodes/FileEntry.cs b/Engine/Source/Programs/Shared/EpicGames.Horde/Storage/Nodes/FileEntry.cs
--- a/Engine/Source/Programs/Shared/EpicGames.Horde/Storage/Nodes/FileEntry.cs
+++ b/Engine/Source/Programs/Shared/EpicGames.Horde/Storage/Nodes/FileEntry.cs
@@ -50,7 +50,7 @@ namespace EpicGames.Horde.Storage.Nodes
 	/// <summary>
 	/// Entry for a file within a directory node
 	/// </summary>
-	public sealed class FileEntry : HashedNodeRef<ChunkedDataNode>
+	public sealed class FileEntry
 	{
 		/// <summary>
 		/// Name of this file
@@ -68,9 +68,14 @@ namespace EpicGames.Horde.Storage.Nodes
 		public long Length { get; }
 
 		/// <summary>
-		/// Hash of the target node
+		/// Hash of the file as a contiguous stream. This differs from individual node hashes which hash the Merkle tree of chunks forming it.
 		/// </summary>
-		public new IoHash Hash { get; }
+		public IoHash StreamHash { get; }
+
+		/// <summary>
+		/// Reference to the chunked data for the file
+		/// </summary>
+		public ChunkedDataNodeRef Target { get; }
 
 		/// <summary>
 		/// Custom user data for this file entry
@@ -80,13 +85,13 @@ namespace EpicGames.Horde.Storage.Nodes
 		/// <summary>
 		/// Constructor
 		/// </summary>
-		public FileEntry(string name, FileEntryFlags flags, long length, ChunkedData data)
-			: base(data.Root)
+		public FileEntry(string name, FileEntryFlags flags, long length, ChunkedData chunkedData)
 		{
 			Name = name;
 			Flags = flags;
 			Length = length;
-			Hash = data.Hash;
+			StreamHash = chunkedData.StreamHash;
+			Target = chunkedData.Root;
 		}
 
 		/// <summary>
@@ -94,12 +99,20 @@ namespace EpicGames.Horde.Storage.Nodes
 		/// </summary>
 		/// <param name="reader"></param>
 		public FileEntry(IBlobReader reader)
-			: base(ReadNodeRef(reader))
 		{
+			if (reader.Version >= 2)
+			{
+				Target = new ChunkedDataNodeRef(reader);
+			}
+			else
+			{
+				Target = new ChunkedDataNodeRef(ChunkedDataNodeType.Unknown, reader);
+			}
+
 			Name = reader.ReadString();
 			Flags = (FileEntryFlags)reader.ReadUnsignedVarInt();
 			Length = (long)reader.ReadUnsignedVarInt();
-			Hash = reader.ReadIoHash();
+			StreamHash = reader.ReadIoHash();
 
 			if ((Flags & FileEntryFlags.HasCustomData) != 0)
 			{
@@ -108,32 +121,20 @@ namespace EpicGames.Horde.Storage.Nodes
 			}
 		}
 
-		static HashedNodeRef<ChunkedDataNode> ReadNodeRef(IBlobReader reader)
-		{
-			if (reader.Version >= 2)
-			{
-				return new ChunkedDataNodeRef(reader);
-			}
-			else
-			{
-				return reader.ReadHashedNodeRef<ChunkedDataNode>();
-			}
-		}
-
 		/// <summary>
 		/// Serialize this entry
 		/// </summary>
 		/// <param name="writer"></param>
-		public override void Serialize(IBlobWriter writer)
+		public void Serialize(IBlobWriter writer)
 		{
-			base.Serialize(writer);
+			writer.WriteNodeRef(Target);
 
 			FileEntryFlags flags = (CustomData.Length > 0) ? (Flags | FileEntryFlags.HasCustomData) : (Flags & ~FileEntryFlags.HasCustomData);
 
 			writer.WriteString(Name);
 			writer.WriteUnsignedVarInt((ulong)flags);
 			writer.WriteUnsignedVarInt((ulong)Length);
-			writer.WriteIoHash(Hash);
+			writer.WriteIoHash(StreamHash);
 
 			if ((flags & FileEntryFlags.HasCustomData) != 0)
 			{
@@ -154,7 +155,7 @@ namespace EpicGames.Horde.Storage.Nodes
 		/// <param name="cancellationToken">Cancellation token for the operation</param>
 		public async Task CopyToStreamAsync(Stream outputStream, CancellationToken cancellationToken)
 		{
-			await ChunkedDataNode.CopyToStreamAsync(Handle, outputStream, cancellationToken);
+			await ChunkedDataNode.CopyToStreamAsync(Target.Handle, outputStream, cancellationToken);
 		}
 
 		/// <summary>
@@ -167,7 +168,7 @@ namespace EpicGames.Horde.Storage.Nodes
 		{
 			try
 			{
-				await ChunkedDataNode.CopyToFileAsync(Handle, file, cancellationToken);
+				await ChunkedDataNode.CopyToFileAsync(Target.Handle, file, cancellationToken);
 
 				if ((Flags & FileEntryFlags.Executable) != 0)
 				{
@@ -191,6 +192,20 @@ namespace EpicGames.Horde.Storage.Nodes
 		public override string ToString() => Name.ToString();
 	}
 
+	/// <summary>
+	/// Extension methods for <see cref="FileEntry"/>
+	/// </summary>
+	public static class FileEntryExtensions
+	{
+		/// <summary>
+		/// Serialize a file entry to a writer
+		/// </summary>
+		public static void WriteFileEntry(this IBlobWriter writer, FileEntry entry)
+		{
+			entry.Serialize(writer);
+		}
+	}
+
 	/// <summary>
 	/// Stream which returns the content of a file
 	/// </summary>
diff --git a/Engine/Source/Programs/Shared/EpicGames.Horde/Storage/Workspace.cs b/Engine/Source/Programs/Shared/EpicGames.Horde/Storage/Workspace.cs
--- a/Engine/Source/Programs/Shared/EpicGames.Horde/Storage/Workspace.cs
+++ b/Engine/Source/Programs/Shared/EpicGames.Horde/Storage/Workspace.cs
@@ -623,7 +623,7 @@ namespace EpicGames.Horde.Storage
 				FileState fileState = dirState.Files[fileIdx];
 				if ((fileState.LayerFlags & flag) != 0)
 				{
-					if (dirNode == null || !dirNode.TryGetFileEntry(fileState.Name, out FileEntry? entry) || entry.Hash != fileState.Hash)
+					if (dirNode == null || !dirNode.TryGetFileEntry(fileState.Name, out FileEntry? entry) || entry.StreamHash != fileState.Hash)
 					{
 						fileState.LayerFlags &= ~flag;
 
@@ -650,7 +650,7 @@ namespace EpicGames.Horde.Storage
 					DirectoryReference subDirPath = DirectoryReference.Combine(dirRef, subDirEntry.Name.ToString());
 					DirectoryState subDirState = dirState.FindOrAddDirectory(subDirEntry.Name);
 
-					DirectoryNode subDirNode = await subDirEntry.ExpandAsync(cancellationToken);
+					DirectoryNode subDirNode = await subDirEntry.Target.ExpandAsync(cancellationToken);
 					await SyncDirectoryAsync(subDirPath, subDirState, subDirNode, flag, cacheDirState, cancellationToken);
 
 					dirState.LayerFlags |= flag;
@@ -676,7 +676,7 @@ namespace EpicGames.Horde.Storage
 			FileState? fileState;
 			if (dirState.TryGetFile(fileEntry.Name, out fileState))
 			{
-				if (fileState.Hash == fileEntry.Hash)
+				if (fileState.Hash == fileEntry.StreamHash)
 				{
 					fileState.LayerFlags |= flag;
 				}
@@ -689,22 +689,22 @@ namespace EpicGames.Horde.Storage
 			{
 				FileInfo fileInfo = FileReference.Combine(dirRef, fileEntry.Name.ToString()).ToFileInfo();
 
-				fileState = TryMoveCachedDataAsync(fileEntry.Hash, dirState, fileEntry.Name);
+				fileState = TryMoveCachedDataAsync(fileEntry.StreamHash, dirState, fileEntry.Name);
 				if (fileState == null)
 				{
 					fileState = dirState.FindOrAddFile(fileEntry.Name);
 
-					_logger.LogInformation("Updating {File} to {Hash}", fileInfo, fileEntry.Hash);
+					_logger.LogInformation("Updating {File} to {Hash}", fileInfo, fileEntry.StreamHash);
 					using (FileStream stream = fileInfo.Open(FileMode.Create, FileAccess.Write, FileShare.Read))
 					{
-						await ExtractDataAsync(fileEntry, stream, cancellationToken);
+						await ExtractDataAsync(fileEntry.Target, stream, cancellationToken);
 						if (stream.Length != fileEntry.Length)
 						{
 							throw new EndOfStreamException($"Incorrect length for extracted file {fileInfo.FullName}");
 						}
 					}
 
-					fileState.Hash = fileEntry.Hash;
+					fileState.Hash = fileEntry.StreamHash;
 
 					fileInfo.Refresh();
 					fileState.Update(fileInfo);
