Horde: Change FileEntry, DirectoryEntry, and ChunkedDataNodeRef to be standalone types rather than deriving from HashedNodeRef. Inheritance makes serialization convoluted here, since data is serialized and versioned according to the parent node type rather than the NodeRef type.

diff --git a/Engine/Source/Programs/Horde/Horde.Agent/Utility/TempStorage.cs b/Engine/Source/Programs/Horde/Horde.Agent/Utility/TempStorage.cs
--- a/Engine/Source/Programs/Horde/Horde.Agent/Utility/TempStorage.cs
+++ b/Engine/Source/Programs/Horde/Horde.Agent/Utility/TempStorage.cs
@@ -765,7 +765,7 @@ namespace Horde.Storage.Utility
 				Stopwatch timer = Stopwatch.StartNew();
 
 				// Add all the files and flush the ref
-				DirectoryNode rootDirNode = await rootDirEntry.Target.ExpandAsync(cancellationToken);
+				DirectoryNode rootDirNode = await rootDirEntry.ExpandAsync(cancellationToken);
 				await rootDirNode.CopyToDirectoryAsync(rootDir.ToDirectoryInfo(), new CopyStatsLogger(logger), logger, cancellationToken);
 
 				StorageStats deltaStats = StorageStats.GetDelta(initialStats, storageClient.GetStats());
diff --git a/Engine/Source/Programs/Horde/Horde.Server/Artifacts/ArtifactsController.cs b/Engine/Source/Programs/Horde/Horde.Server/Artifacts/ArtifactsController.cs
--- a/Engine/Source/Programs/Horde/Horde.Server/Artifacts/ArtifactsController.cs
+++ b/Engine/Source/Programs/Horde/Horde.Server/Artifacts/ArtifactsController.cs
@@ -292,7 +292,7 @@ namespace Horde.Server.Artifacts
 					{
 						return NotFound();
 					}
-					directoryNode = await nextDirectoryEntry.Target.ExpandAsync(cancellationToken);
+					directoryNode = await nextDirectoryEntry.ExpandAsync(cancellationToken);
 				}
 			}
 
@@ -308,9 +308,9 @@ namespace Horde.Server.Artifacts
 				response.Directories = new List<GetArtifactDirectoryEntryResponse>();
 				foreach (DirectoryEntry subDirectoryEntry in directoryNode.Directories)
 				{
-					DirectoryNode subDirectoryNode = await subDirectoryEntry.Target.ExpandAsync(cancellationToken);
+					DirectoryNode subDirectoryNode = await subDirectoryEntry.ExpandAsync(cancellationToken);
 
-					GetArtifactDirectoryEntryResponse subDirectoryEntryResponse = new GetArtifactDirectoryEntryResponse(subDirectoryEntry.Name.ToString(), subDirectoryEntry.Length, subDirectoryEntry.Target.Hash);
+					GetArtifactDirectoryEntryResponse subDirectoryEntryResponse = new GetArtifactDirectoryEntryResponse(subDirectoryEntry.Name.ToString(), subDirectoryEntry.Length, subDirectoryEntry.Hash);
 					if (depth == 0)
 					{
 						if (subDirectoryNode.Directories.Count + subDirectoryNode.Files.Count < 16)
diff --git a/Engine/Source/Programs/Horde/Horde.Server/Perforce/PerforceReplicator.cs b/Engine/Source/Programs/Horde/Horde.Server/Perforce/PerforceReplicator.cs
--- a/Engine/Source/Programs/Horde/Horde.Server/Perforce/PerforceReplicator.cs
+++ b/Engine/Source/Programs/Horde/Horde.Server/Perforce/PerforceReplicator.cs
@@ -407,7 +407,7 @@ namespace Horde.Server.Perforce
 				}
 				else
 				{
-					syncNode = new SyncNode(change, parentChange, parent.Contents.Target);
+					syncNode = new SyncNode(change, parentChange, new HashedNodeRef<DirectoryNode>(parent.Contents.Hash, parent.Contents.Handle));
 				}
 			}
 
@@ -683,7 +683,7 @@ namespace Horde.Server.Perforce
 			await store.WriteRefTargetAsync(refName, commitNodeRef, options.RefOptions, cancellationToken: cancellationToken);
 
 			// Log the snapshot info
-			_logger.LogInformation("Snapshot for {StreamId} CL {Change} is ref {RefName} (commit: {CommitHandle}, root: {RootHandle})", streamConfig.Id, change, refName, commitNodeRef.Handle.GetLocator(), rootRef.Target.Handle.GetLocator());
+			_logger.LogInformation("Snapshot for {StreamId} CL {Change} is ref {RefName} (commit: {CommitHandle}, root: {RootHandle})", streamConfig.Id, change, refName, commitNodeRef.Handle.GetLocator(), rootRef.Handle.GetLocator());
 		}
 
 		static int GetFileOffset(string path)
diff --git a/Engine/Source/Programs/Horde/Horde.Server/Storage/StorageController.cs b/Engine/Source/Programs/Horde/Horde.Server/Storage/StorageController.cs
--- a/Engine/Source/Programs/Horde/Horde.Server/Storage/StorageController.cs
+++ b/Engine/Source/Programs/Horde/Horde.Server/Storage/StorageController.cs
@@ -405,7 +405,7 @@ namespace Horde.Server.Storage
 						List<object> directories = new List<object>();
 						foreach ((string name, DirectoryEntry entry) in directoryNode.NameToDirectory)
 						{
-							directories.Add(new { name = name.ToString(), length = entry.Length, hash = entry.Target.Hash, link = GetNodeLink(linkBase, entry.Target.Handle) });
+							directories.Add(new { name = name.ToString(), length = entry.Length, hash = entry.Hash, link = GetNodeLink(linkBase, entry.Handle) });
 						}
 
 						List<object> files = new List<object>();
diff --git a/Engine/Source/Programs/Horde/Horde/Commands/Vcs/VcsCheckout.cs b/Engine/Source/Programs/Horde/Horde/Commands/Vcs/VcsCheckout.cs
--- a/Engine/Source/Programs/Horde/Horde/Commands/Vcs/VcsCheckout.cs
+++ b/Engine/Source/Programs/Horde/Horde/Commands/Vcs/VcsCheckout.cs
@@ -73,7 +73,7 @@ namespace Horde.Commands.Vcs
 
 			DirectoryState newState = new DirectoryState();
 
-			DirectoryNode directoryNode = await directoryRef.Target.ExpandAsync();
+			DirectoryNode directoryNode = await directoryRef.ExpandAsync();
 			foreach ((string name, DirectoryEntry? subDirEntry, DirectoryState? subDirState) in EnumerableExtensions.Zip(directoryNode.NameToDirectory, directoryState?.Directories))
 			{
 				DirectoryReference subDirPath = DirectoryReference.Combine(dirPath, name.ToString());
@@ -110,7 +110,7 @@ namespace Horde.Commands.Vcs
 				}
 			}
 
-			newState.Hash = directoryRef.Target.Hash;
+			newState.Hash = directoryRef.Hash;
 			return newState;
 		}
 
diff --git a/Engine/Source/Programs/Shared/EpicGames.Horde.Tests/BundleStoreTests.cs b/Engine/Source/Programs/Shared/EpicGames.Horde.Tests/BundleStoreTests.cs
--- a/Engine/Source/Programs/Shared/EpicGames.Horde.Tests/BundleStoreTests.cs
+++ b/Engine/Source/Programs/Shared/EpicGames.Horde.Tests/BundleStoreTests.cs
@@ -264,11 +264,11 @@ namespace EpicGames.Horde.Tests
 			Assert.AreEqual(1, root.Directories.Count);
 			Assert.AreEqual("hello", root.Directories.First().Name);
 
-			DirectoryNode hello = await root.Directories.First().Target.ExpandAsync(CancellationToken.None);
+			DirectoryNode hello = await root.Directories.First().ExpandAsync(CancellationToken.None);
 			Assert.AreEqual(1, hello.Directories.Count);
 			Assert.AreEqual("world", hello.Directories.First().Name);
 
-			DirectoryNode world = await hello.Directories.First().Target.ExpandAsync(CancellationToken.None);
+			DirectoryNode world = await hello.Directories.First().ExpandAsync(CancellationToken.None);
 			Assert.AreEqual(0, world.Directories.Count);
 		}
 
@@ -310,7 +310,7 @@ namespace EpicGames.Horde.Tests
 			Assert.AreEqual(1, root.Directories.Count);
 			Assert.AreEqual("hello", root.Directories.First().Name);
 
-			DirectoryNode hello = await root.Directories.First().Target.ExpandAsync();
+			DirectoryNode hello = await root.Directories.First().ExpandAsync();
 			Assert.AreEqual(0, hello.Directories.Count);
 			Assert.AreEqual(1, hello.Files.Count);
 			Assert.AreEqual("world", hello.Files.First().Name);
@@ -333,7 +333,7 @@ namespace EpicGames.Horde.Tests
 			new Random(0).NextBytes(chunk);
 
 			// Generate a tree
-			HashedNodeRef<ChunkedDataNode> nodeRef;
+			ChunkedDataNodeRef nodeRef;
 			{
 				await using IStorageWriter writer = store.CreateWriter(options: new BundleOptions { MaxBlobSize = 1024 });
 
@@ -449,7 +449,7 @@ namespace EpicGames.Horde.Tests
 				Assert.AreEqual(oldInteriorNode.Children.Count, newInteriorNode.Children.Count);
 
 				int index = 0;
-				foreach ((HashedNodeRef<ChunkedDataNode> oldFileRef, HashedNodeRef<ChunkedDataNode> newFileRef) in oldInteriorNode.Children.Zip(newInteriorNode.Children))
+				foreach ((ChunkedDataNodeRef oldFileRef, ChunkedDataNodeRef newFileRef) in oldInteriorNode.Children.Zip(newInteriorNode.Children))
 				{
 					ChunkedDataNode oldFile = await oldFileRef.ExpandAsync();
 					ChunkedDataNode newFile = await newFileRef.ExpandAsync();
@@ -489,7 +489,7 @@ namespace EpicGames.Horde.Tests
 			}
 			else if (fileNode is InteriorChunkedDataNode interiorFileNode)
 			{
-				foreach (HashedNodeRef<ChunkedDataNode> childRef in interiorFileNode.Children)
+				foreach (ChunkedDataNodeRef childRef in interiorFileNode.Children)
 				{
 					ChunkedDataNode child = await childRef.ExpandAsync();
 					offset += await CheckFileDataAsync(child, data.Slice(offset));
diff --git a/Engine/Source/Programs/Shared/EpicGames.Horde.Tests/BundleTests.cs b/Engine/Source/Programs/Shared/EpicGames.Horde.Tests/BundleTests.cs
--- a/Engine/Source/Programs/Shared/EpicGames.Horde.Tests/BundleTests.cs
+++ b/Engine/Source/Programs/Shared/EpicGames.Horde.Tests/BundleTests.cs
@@ -80,7 +80,7 @@ namespace EpicGames.Horde.Tests
 			using ChunkedDataWriter fileNodeWriter = new ChunkedDataWriter(writer, options);
 
 			ChunkedDataNode node;
-			HashedNodeRef<ChunkedDataNode> nodeRef;
+			ChunkedDataNodeRef nodeRef;
 			byte[] data = CreateBuffer(1024);
 
 			nodeRef = (await fileNodeWriter.CreateAsync(data.AsMemory(0, 7), CancellationToken.None)).Root;
@@ -127,7 +127,7 @@ namespace EpicGames.Horde.Tests
 			return output;
 		}
 
-		private static async Task TestBufferlessReadsAsync(HashedNodeRef<ChunkedDataNode> nodeRef, ReadOnlyMemory<byte> expected)
+		private static async Task TestBufferlessReadsAsync(ChunkedDataNodeRef nodeRef, ReadOnlyMemory<byte> expected)
 		{
 			using MemoryStream memoryStream = new MemoryStream();
 			await ChunkedDataNode.CopyToStreamAsync(nodeRef.Handle, memoryStream, default);
diff --git a/Engine/Source/Programs/Shared/EpicGames.Horde.Tests/ChunkingTests.cs b/Engine/Source/Programs/Shared/EpicGames.Horde.Tests/ChunkingTests.cs
--- a/Engine/Source/Programs/Shared/EpicGames.Horde.Tests/ChunkingTests.cs
+++ b/Engine/Source/Programs/Shared/EpicGames.Horde.Tests/ChunkingTests.cs
@@ -106,7 +106,7 @@ namespace EpicGames.Horde.Tests
 				data[idx] = (byte)idx;
 			}
 
-			HashedNodeRef<ChunkedDataNode> handle;
+			ChunkedDataNodeRef handle;
 
 			const int NumIterations = 100;
 			{
diff --git a/Engine/Source/Programs/Shared/EpicGames.Horde/Storage/Nodes/ChunkedDataNode.cs b/Engine/Source/Programs/Shared/EpicGames.Horde/Storage/Nodes/ChunkedDataNode.cs
--- a/Engine/Source/Programs/Shared/EpicGames.Horde/Storage/Nodes/ChunkedDataNode.cs
+++ b/Engine/Source/Programs/Shared/EpicGames.Horde/Storage/Nodes/ChunkedDataNode.cs
@@ -110,47 +110,16 @@ namespace EpicGames.Horde.Storage.Nodes
 	/// <summary>
 	/// Reference to a chunked data node
 	/// </summary>
-	public class ChunkedDataNodeRef : HashedNodeRef<ChunkedDataNode>
+	/// <param name="Hash">Hash of the target node</param>
+	/// <param name="Type">Type of the referenced node</param>
+	/// <param name="Handle">Handle to the target node</param>
+	public record class ChunkedDataNodeRef(IoHash Hash, ChunkedDataNodeType Type, IBlobHandle Handle)
 	{
 		/// <summary>
-		/// Type of the referenced node
+		/// Read the node which is the target of this ref
 		/// </summary>
-		public ChunkedDataNodeType Type { get; }
-
-		/// <summary>
-		/// Constructor
-		/// </summary>
-		public ChunkedDataNodeRef(IBlobReader reader)
-			: base(reader)
-		{
-			Type = (ChunkedDataNodeType)reader.ReadUnsignedVarInt();
-		}
-
-		/// <summary>
-		/// Constructor
-		/// </summary>
-		public ChunkedDataNodeRef(ChunkedDataNodeType type, IBlobReader reader)
-			: base(reader)
-		{
-			Type = type;
-		}
-
-		/// <summary>
-		/// Constructor
-		/// </summary>
-		public ChunkedDataNodeRef(ChunkedDataNodeType type, HashedNodeRef<ChunkedDataNode> target)
-			: base(target)
-		{
-			Type = type;
-		}
-
-		/// <inheritdoc/>
-		public override void Serialize(IBlobWriter writer)
-		{
-			base.Serialize(writer);
-
-			writer.WriteUnsignedVarInt((ulong)Type);
-		}
+		public ValueTask<ChunkedDataNode> ExpandAsync(CancellationToken cancellationToken = default)
+			=> Handle.ReadNodeAsync<ChunkedDataNode>(cancellationToken);
 	}
 
 	/// <summary>
@@ -301,7 +270,7 @@ namespace EpicGames.Horde.Storage.Nodes
 				}
 
 				HashedNodeRef<ChunkedDataNode> nodeRef = await writer.WriteHashedNodeRefAsync<ChunkedDataNode>(GetNodeType<LeafChunkedDataNode>(), nextLength, Array.Empty<IBlobHandle>(), cancellationToken);
-				leafNodeRefs.Add(new ChunkedDataNodeRef(ChunkedDataNodeType.Leaf, nodeRef));
+				leafNodeRefs.Add(new ChunkedDataNodeRef(nodeRef.Hash, ChunkedDataNodeType.Leaf, nodeRef.Handle));
 
 				readBuffer.Memory.Slice(nextLength, size - nextLength).CopyTo(readBuffer.Memory);
 				size -= nextLength;
@@ -417,14 +386,21 @@ namespace EpicGames.Horde.Storage.Nodes
 			List<ChunkedDataNodeRef> children = new List<ChunkedDataNodeRef>();
 			while (reader.GetMemory().Length > 0)
 			{
+				IoHash hash = reader.ReadIoHash();
+
+				ChunkedDataNodeType type;
 				if (reader.Version >= 2)
 				{
-					children.Add(new ChunkedDataNodeRef(reader));
+					type = (ChunkedDataNodeType)reader.ReadUnsignedVarInt();
 				}
 				else
 				{
-					children.Add(new ChunkedDataNodeRef(ChunkedDataNodeType.Unknown, reader.ReadHashedNodeRef<ChunkedDataNode>()));
+					type = ChunkedDataNodeType.Unknown;
 				}
+
+				IBlobHandle handle = reader.ReadBlobReference();
+
+				children.Add(new ChunkedDataNodeRef(hash, type, handle));
 			}
 			Children = children;
 		}
@@ -434,7 +410,9 @@ namespace EpicGames.Horde.Storage.Nodes
 		{
 			foreach (ChunkedDataNodeRef child in Children)
 			{
-				writer.WriteHashedNodeRef(child);
+				writer.WriteIoHash(child.Hash);
+				writer.WriteUnsignedVarInt((int)child.Type);
+				writer.WriteBlobReference(child.Handle);
 			}
 		}
 
@@ -474,7 +452,7 @@ namespace EpicGames.Horde.Storage.Nodes
 				foreach (InteriorChunkedDataNode interiorNode in interiorNodes)
 				{
 					HashedNodeRef<ChunkedDataNode> newNodeRef = await writer.WriteHashedNodeAsync<ChunkedDataNode>(interiorNode, cancellationToken);
-					handleBuffer.Add(new ChunkedDataNodeRef(ChunkedDataNodeType.Interior, newNodeRef));
+					handleBuffer.Add(new ChunkedDataNodeRef(newNodeRef.Hash, ChunkedDataNodeType.Interior, newNodeRef.Handle));
 				}
 
 				nodeRefs = handleBuffer;
@@ -520,7 +498,7 @@ namespace EpicGames.Horde.Storage.Nodes
 		/// <inheritdoc/>
 		public override async Task CopyToStreamAsync(Stream outputStream, CancellationToken cancellationToken)
 		{
-			foreach (HashedNodeRef<ChunkedDataNode> childNodeRef in Children)
+			foreach (ChunkedDataNodeRef childNodeRef in Children)
 			{
 				ChunkedDataNode childNode = await childNodeRef.ExpandAsync(cancellationToken);
 				await childNode.CopyToStreamAsync(outputStream, cancellationToken);
diff --git a/Engine/Source/Programs/Shared/EpicGames.Horde/Storage/Nodes/ChunkedDataWriter.cs b/Engine/Source/Programs/Shared/EpicGames.Horde/Storage/Nodes/ChunkedDataWriter.cs
--- a/Engine/Source/Programs/Shared/EpicGames.Horde/Storage/Nodes/ChunkedDataWriter.cs
+++ b/Engine/Source/Programs/Shared/EpicGames.Horde/Storage/Nodes/ChunkedDataWriter.cs
@@ -346,7 +346,7 @@ namespace EpicGames.Horde.Storage.Nodes
 		async ValueTask FlushLeafNodeAsync(CancellationToken cancellationToken)
 		{
 			HashedNodeRef<ChunkedDataNode> nodeRef  = await _writer.WriteHashedNodeRefAsync<ChunkedDataNode>(s_leafNodeType, _leafLength, Array.Empty<IBlobHandle>(), cancellationToken);
-			_leafHandles.Add(new ChunkedDataNodeRef(ChunkedDataNodeType.Leaf, nodeRef));
+			_leafHandles.Add(new ChunkedDataNodeRef(nodeRef.Hash, ChunkedDataNodeType.Leaf, nodeRef.Handle));
 			ResetLeafState();
 		}
 	}
diff --git a/Engine/Source/Programs/Shared/EpicGames.Horde/Storage/Nodes/CommitNode.cs b/Engine/Source/Programs/Shared/EpicGames.Horde/Storage/Nodes/CommitNode.cs
--- a/Engine/Source/Programs/Shared/EpicGames.Horde/Storage/Nodes/CommitNode.cs
+++ b/Engine/Source/Programs/Shared/EpicGames.Horde/Storage/Nodes/CommitNode.cs
@@ -95,7 +95,11 @@ namespace EpicGames.Horde.Storage.Nodes
 			CommitterId = reader.ReadOptionalString();
 			Message = reader.ReadString();
 			Time = reader.ReadDateTime();
-			Contents = new DirectoryNodeRef(reader);
+
+			IoHash hash = reader.ReadIoHash();
+			long length = (long)reader.ReadUnsignedVarInt();
+			Contents = new DirectoryNodeRef(hash, length, reader.ReadBlobReference());
+
 			Metadata = reader.ReadDictionary(() => reader.ReadGuidUnrealOrder(), () => reader.ReadHashedNodeRef());
 		}
 
@@ -110,7 +114,11 @@ namespace EpicGames.Horde.Storage.Nodes
 			writer.WriteOptionalString(CommitterId);
 			writer.WriteString(Message);
 			writer.WriteDateTime(Time);
-			writer.WriteDirectoryNodeRef(Contents);
+
+			writer.WriteIoHash(Contents.Hash);
+			writer.WriteUnsignedVarInt((ulong)Contents.Length);
+			writer.WriteBlobReference(Contents.Handle);
+
 			writer.WriteDictionary(Metadata, key => writer.WriteGuidUnrealOrder(key), value => writer.WriteHashedNodeRef(value));
 		}
 	}
diff --git a/Engine/Source/Programs/Shared/EpicGames.Horde/Storage/Nodes/DirectoryEntry.cs b/Engine/Source/Programs/Shared/EpicGames.Horde/Storage/Nodes/DirectoryEntry.cs
--- a/Engine/Source/Programs/Shared/EpicGames.Horde/Storage/Nodes/DirectoryEntry.cs
+++ b/Engine/Source/Programs/Shared/EpicGames.Horde/Storage/Nodes/DirectoryEntry.cs
@@ -1,5 +1,8 @@
 // Copyright Epic Games, Inc. All Rights Reserved.
 
+using System.Diagnostics;
+using System.Threading;
+using System.Threading.Tasks;
 using EpicGames.Core;
 
 namespace EpicGames.Horde.Storage.Nodes
@@ -7,58 +10,15 @@ namespace EpicGames.Horde.Storage.Nodes
 	/// <summary>
 	/// Entry for a directory within a directory node
 	/// </summary>
-	public class DirectoryEntry : DirectoryNodeRef
+	[DebuggerDisplay("{Name}")]
+	public record class DirectoryEntry(string Name, IoHash Hash, long Length, IBlobHandle Handle) : DirectoryNodeRef(Hash, Length, Handle)
 	{
-		/// <summary>
-		/// Name of this directory
-		/// </summary>
-		public string Name { get; }
-
 		/// <summary>
 		/// Constructor
 		/// </summary>
-		public DirectoryEntry(string name, long length, HashedNodeRef<DirectoryNode> nodeRef)
-			: base(length, nodeRef)
-		{
-			Name = name;
-		}
-
-		/// <summary>
-		/// Deserializing constructor
-		/// </summary>
-		/// <param name="reader"></param>
-		public DirectoryEntry(IBlobReader reader)
-			: base(reader)
-		{
-			Name = reader.ReadString();
-		}
-
-		/// <summary>
-		/// Serialize this directory entry to disk
-		/// </summary>
-		/// <param name="writer"></param>
-		public override void Serialize(IBlobWriter writer)
-		{
-			base.Serialize(writer);
-
-			writer.WriteString(Name);
-		}
-
-		/// <inheritdoc/>
-		public override string ToString() => Name;
-	}
-
-	/// <summary>
-	/// Extension methods for <see cref="DirectoryEntry"/>
-	/// </summary>
-	public static class DirectoryEntryExtensions
-	{
-		/// <summary>
-		/// Serialize a directory entry to storage
-		/// </summary>
-		public static void WriteDirectoryEntry(this IBlobWriter writer, DirectoryEntry entry)
+		public DirectoryEntry(string name, long length, HashedNodeRef<DirectoryNode> target)
+			: this(name, target.Hash, length, target.Handle)
 		{
-			entry.Serialize(writer);
 		}
 	}
 }
diff --git a/Engine/Source/Programs/Shared/EpicGames.Horde/Storage/Nodes/DirectoryNode.cs b/Engine/Source/Programs/Shared/EpicGames.Horde/Storage/Nodes/DirectoryNode.cs
--- a/Engine/Source/Programs/Shared/EpicGames.Horde/Storage/Nodes/DirectoryNode.cs
+++ b/Engine/Source/Programs/Shared/EpicGames.Horde/Storage/Nodes/DirectoryNode.cs
@@ -196,15 +196,43 @@ namespace EpicGames.Horde.Storage.Nodes
 			int fileCount = (int)reader.ReadUnsignedVarInt();
 			for (int idx = 0; idx < fileCount; idx++)
 			{
-				FileEntry entry = new FileEntry(reader);
-				_nameToFileEntry[entry.Name] = entry;
+				IoHash hash = reader.ReadIoHash();
+
+				ChunkedDataNodeType type;
+				if (reader.Version >= 2)
+				{
+					type = (ChunkedDataNodeType)reader.ReadUnsignedVarInt();
+				}
+				else
+				{
+					type = ChunkedDataNodeType.Unknown;
+				}
+
+				ChunkedDataNodeRef target = new ChunkedDataNodeRef(hash, type, reader.ReadBlobReference());
+
+				string name = reader.ReadString();
+				FileEntryFlags flags = (FileEntryFlags)reader.ReadUnsignedVarInt();
+				long length = (long)reader.ReadUnsignedVarInt();
+				IoHash streamHash = reader.ReadIoHash();
+
+				ReadOnlyMemory<byte> customData = default;
+				if ((flags & FileEntryFlags.HasCustomData) != 0)
+				{
+					customData = reader.ReadVariableLengthBytes();
+					flags &= ~FileEntryFlags.HasCustomData;
+				}
+
+				_nameToFileEntry[name] = new FileEntry(name, flags, length, streamHash, target, customData);
 			}
 
 			int directoryCount = (int)reader.ReadUnsignedVarInt();
 			for (int idx = 0; idx < directoryCount; idx++)
 			{
-				DirectoryEntry entry = new DirectoryEntry(reader);
-				_nameToDirectoryEntry[entry.Name] = entry;
+				IoHash hash = reader.ReadIoHash();
+				long length = (long)reader.ReadUnsignedVarInt();
+				string name = reader.ReadString();
+
+				_nameToDirectoryEntry[name] = new DirectoryEntry(name, hash, length, reader.ReadBlobReference());
 			}
 		}
 
@@ -216,13 +244,30 @@ namespace EpicGames.Horde.Storage.Nodes
 			writer.WriteUnsignedVarInt(Files.Count);
 			foreach (FileEntry fileEntry in _nameToFileEntry.Values)
 			{
-				writer.WriteFileEntry(fileEntry);
+				writer.WriteIoHash(fileEntry.Target.Hash);
+				writer.WriteUnsignedVarInt((int)fileEntry.Target.Type);
+				writer.WriteBlobReference(fileEntry.Target.Handle);
+
+				FileEntryFlags flags = (fileEntry.CustomData.Length > 0) ? (fileEntry.Flags | FileEntryFlags.HasCustomData) : (fileEntry.Flags & ~FileEntryFlags.HasCustomData);
+
+				writer.WriteString(fileEntry.Name);
+				writer.WriteUnsignedVarInt((ulong)flags);
+				writer.WriteUnsignedVarInt((ulong)fileEntry.Length);
+				writer.WriteIoHash(fileEntry.StreamHash);
+
+				if ((flags & FileEntryFlags.HasCustomData) != 0)
+				{
+					writer.WriteVariableLengthBytes(fileEntry.CustomData.Span);
+				}
 			}
 
 			writer.WriteUnsignedVarInt(Directories.Count);
 			foreach (DirectoryEntry directoryEntry in _nameToDirectoryEntry.Values)
 			{
-				writer.WriteDirectoryEntry(directoryEntry);
+				writer.WriteIoHash(directoryEntry.Hash);
+				writer.WriteUnsignedVarInt((ulong)directoryEntry.Length);
+				writer.WriteString(directoryEntry.Name);
+				writer.WriteBlobReference(directoryEntry.Handle);
 			}
 		}
 
@@ -381,7 +426,7 @@ namespace EpicGames.Horde.Storage.Nodes
 					return null;
 				}
 
-				directoryNode = await directoryEntry.Target.ExpandAsync(cancellationToken);
+				directoryNode = await directoryEntry.ExpandAsync(cancellationToken);
 			}
 			return directoryNode;
 		}
@@ -470,7 +515,7 @@ namespace EpicGames.Horde.Storage.Nodes
 		{
 			if (TryGetDirectoryEntry(name, out DirectoryEntry? entry))
 			{
-				return await entry.Target.ExpandAsync(cancellationToken);
+				return await entry.ExpandAsync(cancellationToken);
 			}
 			else
 			{
@@ -679,7 +724,7 @@ namespace EpicGames.Horde.Storage.Nodes
 					childNode ??= new DirectoryNode();
 					await childNode.UpdateAsync(directory, writer, cancellationToken);
 					HashedNodeRef<DirectoryNode> nodeRef = await writer.WriteHashedNodeAsync(childNode, cancellationToken);
-					_nameToDirectoryEntry[name] = new DirectoryEntry(name, childNode.Length, nodeRef);
+					_nameToDirectoryEntry[name] = new DirectoryEntry(name, nodeRef.Hash, childNode.Length, nodeRef.Handle);
 				}
 			}
 			foreach ((string name, FileEntry? file) in update.Files)
@@ -813,7 +858,7 @@ namespace EpicGames.Horde.Storage.Nodes
 				if (windowOffset < directoryEntry.Length && windowOffset + windowLength >= 0)
 				{
 					DirectoryInfo subDirectoryInfo = directoryInfo.CreateSubdirectory(directoryEntry.Name.ToString());
-					DirectoryNode subDirectoryNode = await directoryEntry.Target.ExpandAsync(cancellationToken);
+					DirectoryNode subDirectoryNode = await directoryEntry.ExpandAsync(cancellationToken);
 					await subDirectoryNode.CopyToDirectoryInternalAsync(subDirectoryInfo, windowOffset, windowLength, copyStats, logger, cancellationToken);
 				}
 				windowOffset -= directoryEntry.Length;
@@ -910,7 +955,7 @@ namespace EpicGames.Horde.Storage.Nodes
 				}
 			}
 
-			FileEntry entry = new FileEntry(name, flags, length, chunkedData);
+			FileEntry entry = new FileEntry(name, flags, length, chunkedData.StreamHash, chunkedData.Root, default);
 			AddFile(path, entry);
 			return entry;
 		}
@@ -1083,7 +1128,7 @@ namespace EpicGames.Horde.Storage.Nodes
 				string directoryPath = $"{prefix}{directoryEntry.Name}/";
 				if (filter == null || filter.PossiblyMatches(directoryPath))
 				{
-					DirectoryNode node = await directoryEntry.Target.ExpandAsync(cancellationToken);
+					DirectoryNode node = await directoryEntry.ExpandAsync(cancellationToken);
 					await CopyFilesAsync(node, directoryPath, filter, archive, logger, cancellationToken);
 				}
 				numCopiedDirs++;
diff --git a/Engine/Source/Programs/Shared/EpicGames.Horde/Storage/Nodes/DirectoryNodeRef.cs b/Engine/Source/Programs/Shared/EpicGames.Horde/Storage/Nodes/DirectoryNodeRef.cs
--- a/Engine/Source/Programs/Shared/EpicGames.Horde/Storage/Nodes/DirectoryNodeRef.cs
+++ b/Engine/Source/Programs/Shared/EpicGames.Horde/Storage/Nodes/DirectoryNodeRef.cs
@@ -1,5 +1,7 @@
 // Copyright Epic Games, Inc. All Rights Reserved.
 
+using System.Threading;
+using System.Threading.Tasks;
 using EpicGames.Core;
 
 namespace EpicGames.Horde.Storage.Nodes
@@ -7,58 +9,23 @@ namespace EpicGames.Horde.Storage.Nodes
 	/// <summary>
 	/// Reference to a directory node, including the target hash and length
 	/// </summary>
-	public class DirectoryNodeRef
+	/// <param name="Hash">Hash of the target node</param>
+	/// <param name="Length">Sum total of all the file lengths in this directory tree</param>
+	/// <param name="Handle">Handle to the target node</param>
+	public record class DirectoryNodeRef(IoHash Hash, long Length, IBlobHandle Handle)
 	{
-		/// <summary>
-		/// Length of this directory tree
-		/// </summary>
-		public long Length { get; }
-
-		/// <summary>
-		/// Target node
-		/// </summary>
-		public HashedNodeRef<DirectoryNode> Target { get; }
-
 		/// <summary>
 		/// Constructor
 		/// </summary>
 		public DirectoryNodeRef(long length, HashedNodeRef<DirectoryNode> target)
+			: this(target.Hash, length, target.Handle)
 		{
-			Length = length;
-			Target = target;
-		}
-
-		/// <summary>
-		/// Constructor
-		/// </summary>
-		public DirectoryNodeRef(IBlobReader reader)
-		{
-			Target = new HashedNodeRef<DirectoryNode>(reader);
-			Length = (long)reader.ReadUnsignedVarInt();
 		}
 
 		/// <summary>
-		/// Serialize this directory entry to disk
+		/// Get the target directory node
 		/// </summary>
-		/// <param name="writer"></param>
-		public virtual void Serialize(IBlobWriter writer)
-		{
-			writer.WriteHashedNodeRef(Target);
-			writer.WriteUnsignedVarInt((ulong)Length);
-		}
-	}
-
-	/// <summary>
-	/// Extension methods for <see cref="DirectoryNodeRef"/>
-	/// </summary>
-	public static class DirectoryNodeRefExtensions
-	{
-		/// <summary>
-		/// Serialize a directory node ref to a blob writer
-		/// </summary>
-		public static void WriteDirectoryNodeRef(this IBlobWriter writer, DirectoryNodeRef nodeRef)
-		{
-			nodeRef.Serialize(writer);
-		}
+		public ValueTask<DirectoryNode> ExpandAsync(CancellationToken cancellationToken = default)
+			=> Handle.ReadNodeAsync<DirectoryNode>(cancellationToken);
 	}
 }
diff --git a/Engine/Source/Programs/Shared/EpicGames.Horde/Storage/Nodes/FileEntry.cs b/Engine/Source/Programs/Shared/EpicGames.Horde/Storage/Nodes/FileEntry.cs
--- a/Engine/Source/Programs/Shared/EpicGames.Horde/Storage/Nodes/FileEntry.cs
+++ b/Engine/Source/Programs/Shared/EpicGames.Horde/Storage/Nodes/FileEntry.cs
@@ -80,66 +80,27 @@ namespace EpicGames.Horde.Storage.Nodes
 		/// <summary>
 		/// Custom user data for this file entry
 		/// </summary>
-		public ReadOnlyMemory<byte> CustomData { get; set; } = ReadOnlyMemory<byte>.Empty;
+		public ReadOnlyMemory<byte> CustomData { get; set; }
 
 		/// <summary>
 		/// Constructor
 		/// </summary>
-		public FileEntry(string name, FileEntryFlags flags, long length, ChunkedData chunkedData)
+		public FileEntry(string name, FileEntryFlags flags, long length, ChunkedData chunkedData, ReadOnlyMemory<byte> customData = default)
+			: this(name, flags, length, chunkedData.StreamHash, chunkedData.Root, customData)
 		{
-			Name = name;
-			Flags = flags;
-			Length = length;
-			StreamHash = chunkedData.StreamHash;
-			Target = chunkedData.Root;
 		}
 
 		/// <summary>
-		/// Deserialize from a buffer
-		/// </summary>
-		/// <param name="reader"></param>
-		public FileEntry(IBlobReader reader)
-		{
-			if (reader.Version >= 2)
-			{
-				Target = new ChunkedDataNodeRef(reader);
-			}
-			else
-			{
-				Target = new ChunkedDataNodeRef(ChunkedDataNodeType.Unknown, reader);
-			}
-
-			Name = reader.ReadString();
-			Flags = (FileEntryFlags)reader.ReadUnsignedVarInt();
-			Length = (long)reader.ReadUnsignedVarInt();
-			StreamHash = reader.ReadIoHash();
-
-			if ((Flags & FileEntryFlags.HasCustomData) != 0)
-			{
-				CustomData = reader.ReadVariableLengthBytes();
-				Flags &= ~FileEntryFlags.HasCustomData;
-			}
-		}
-
-		/// <summary>
-		/// Serialize this entry
+		/// Constructor
 		/// </summary>
-		/// <param name="writer"></param>
-		public void Serialize(IBlobWriter writer)
+		public FileEntry(string name, FileEntryFlags flags, long length, IoHash streamHash, ChunkedDataNodeRef target, ReadOnlyMemory<byte> customData)
 		{
-			writer.WriteNodeRef(Target);
-
-			FileEntryFlags flags = (CustomData.Length > 0) ? (Flags | FileEntryFlags.HasCustomData) : (Flags & ~FileEntryFlags.HasCustomData);
-
-			writer.WriteString(Name);
-			writer.WriteUnsignedVarInt((ulong)flags);
-			writer.WriteUnsignedVarInt((ulong)Length);
-			writer.WriteIoHash(StreamHash);
-
-			if ((flags & FileEntryFlags.HasCustomData) != 0)
-			{
-				writer.WriteVariableLengthBytes(CustomData.Span);
-			}
+			Name = name;
+			Flags = flags;
+			Length = length;
+			StreamHash = streamHash;
+			Target = target;
+			CustomData = customData;
 		}
 
 		/// <summary>
@@ -192,20 +153,6 @@ namespace EpicGames.Horde.Storage.Nodes
 		public override string ToString() => Name.ToString();
 	}
 
-	/// <summary>
-	/// Extension methods for <see cref="FileEntry"/>
-	/// </summary>
-	public static class FileEntryExtensions
-	{
-		/// <summary>
-		/// Serialize a file entry to a writer
-		/// </summary>
-		public static void WriteFileEntry(this IBlobWriter writer, FileEntry entry)
-		{
-			entry.Serialize(writer);
-		}
-	}
-
 	/// <summary>
 	/// Stream which returns the content of a file
 	/// </summary>
diff --git a/Engine/Source/Programs/Shared/EpicGames.Horde/Storage/Workspace.cs b/Engine/Source/Programs/Shared/EpicGames.Horde/Storage/Workspace.cs
--- a/Engine/Source/Programs/Shared/EpicGames.Horde/Storage/Workspace.cs
+++ b/Engine/Source/Programs/Shared/EpicGames.Horde/Storage/Workspace.cs
@@ -650,7 +650,7 @@ namespace EpicGames.Horde.Storage
 					DirectoryReference subDirPath = DirectoryReference.Combine(dirRef, subDirEntry.Name.ToString());
 					DirectoryState subDirState = dirState.FindOrAddDirectory(subDirEntry.Name);
 
-					DirectoryNode subDirNode = await subDirEntry.Target.ExpandAsync(cancellationToken);
+					DirectoryNode subDirNode = await subDirEntry.ExpandAsync(cancellationToken);
 					await SyncDirectoryAsync(subDirPath, subDirState, subDirNode, flag, cacheDirState, cancellationToken);
 
 					dirState.LayerFlags |= flag;
@@ -721,7 +721,7 @@ namespace EpicGames.Horde.Storage
 			}
 		}
 
-		async Task ExtractDataAsync(HashedNodeRef<ChunkedDataNode> nodeRef, Stream outputStream, CancellationToken cancellationToken)
+		async Task ExtractDataAsync(ChunkedDataNodeRef nodeRef, Stream outputStream, CancellationToken cancellationToken)
 		{
 			BlobType blobType = await nodeRef.Handle.ReadTypeAsync(cancellationToken);
 			if (blobType.Guid == LeafChunkedDataNode.BlobType.Guid)
@@ -730,8 +730,8 @@ namespace EpicGames.Horde.Storage
 			}
 			else if (blobType.Guid == InteriorChunkedDataNode.BlobType.Guid)
 			{
-				InteriorChunkedDataNode interiorNode = (InteriorChunkedDataNode)await nodeRef.ExpandAsync(cancellationToken);
-				foreach (HashedNodeRef<ChunkedDataNode> childNodeRef in interiorNode.Children)
+				InteriorChunkedDataNode interiorNode = await nodeRef.Handle.ReadNodeAsync<InteriorChunkedDataNode>(cancellationToken);
+				foreach (ChunkedDataNodeRef childNodeRef in interiorNode.Children)
 				{
 					await ExtractDataAsync(childNodeRef, outputStream, cancellationToken);
 				}
@@ -742,7 +742,7 @@ namespace EpicGames.Horde.Storage
 			}
 		}
 
-		async Task ExtractLeafDataAsync(HashedNodeRef<ChunkedDataNode> nodeRef, Stream outputStream, CancellationToken cancellationToken)
+		async Task ExtractLeafDataAsync(ChunkedDataNodeRef nodeRef, Stream outputStream, CancellationToken cancellationToken)
 		{
 			// Try to copy cached data for this node
 			HashInfo? hashInfo;
diff --git a/Engine/Source/Programs/UnrealBuildTool/Artifacts/HordeStorageArtifactCache.cs b/Engine/Source/Programs/UnrealBuildTool/Artifacts/HordeStorageArtifactCache.cs
--- a/Engine/Source/Programs/UnrealBuildTool/Artifacts/HordeStorageArtifactCache.cs
+++ b/Engine/Source/Programs/UnrealBuildTool/Artifacts/HordeStorageArtifactCache.cs
@@ -88,7 +88,7 @@ namespace UnrealBuildTool.Artifacts
 				string outputName = artifact.GetFullPath(ArtifactAction.DirectoryMapping);
 				using FileStream stream = new(outputName, FileMode.Open, FileAccess.Read, FileShare.Read);
 				ChunkedData chunkedData = await fileWriter.CreateAsync(stream, leafOptions.TargetSize, cancellationToken);
-				OutputRefs[index++] = new HashedNodeRef<ChunkedDataNode>(chunkedData.Root);
+				OutputRefs[index++] = new HashedNodeRef<ChunkedDataNode>(chunkedData.Root.Hash, chunkedData.Root.Handle);
 			}
 		}
 	}
