Horde: Remove unnecessary using directive.

diff --git a/Engine/Source/Programs/Shared/EpicGames.Horde/Storage/Nodes/DirectoryEntry.cs b/Engine/Source/Programs/Shared/EpicGames.Horde/Storage/Nodes/DirectoryEntry.cs
--- a/Engine/Source/Programs/Shared/EpicGames.Horde/Storage/Nodes/DirectoryEntry.cs
+++ b/Engine/Source/Programs/Shared/EpicGames.Horde/Storage/Nodes/DirectoryEntry.cs
@@ -1,7 +1,6 @@
 // Copyright Epic Games, Inc. All Rights Reserved.
 
 using System.Diagnostics;
-using System.Threading;
 using System.Threading.Tasks;
 using EpicGames.Core;
 
