Horde: Reorder arguments to ChunkedDataNodeRef.

diff --git a/Engine/Source/Programs/Shared/EpicGames.Horde/Storage/Nodes/ChunkedDataNode.cs b/Engine/Source/Programs/Shared/EpicGames.Horde/Storage/Nodes/ChunkedDataNode.cs
--- a/Engine/Source/Programs/Shared/EpicGames.Horde/Storage/Nodes/ChunkedDataNode.cs
+++ b/Engine/Source/Programs/Shared/EpicGames.Horde/Storage/Nodes/ChunkedDataNode.cs
@@ -110,10 +110,10 @@ namespace EpicGames.Horde.Storage.Nodes
 	/// <summary>
 	/// Reference to a chunked data node
 	/// </summary>
-	/// <param name="Hash">Hash of the target node</param>
 	/// <param name="Type">Type of the referenced node</param>
+	/// <param name="Hash">Hash of the target node</param>
 	/// <param name="Handle">Handle to the target node</param>
-	public record class ChunkedDataNodeRef(IoHash Hash, ChunkedDataNodeType Type, IBlobHandle Handle)
+	public record class ChunkedDataNodeRef(ChunkedDataNodeType Type, IoHash Hash, IBlobHandle Handle)
 	{
 		/// <summary>
 		/// Read the node which is the target of this ref
@@ -270,7 +270,7 @@ namespace EpicGames.Horde.Storage.Nodes
 				}
 
 				HashedNodeRef<ChunkedDataNode> nodeRef = await writer.WriteHashedNodeRefAsync<ChunkedDataNode>(GetNodeType<LeafChunkedDataNode>(), nextLength, Array.Empty<IBlobHandle>(), cancellationToken);
-				leafNodeRefs.Add(new ChunkedDataNodeRef(nodeRef.Hash, ChunkedDataNodeType.Leaf, nodeRef.Handle));
+				leafNodeRefs.Add(new ChunkedDataNodeRef(ChunkedDataNodeType.Leaf, nodeRef.Hash, nodeRef.Handle));
 
 				readBuffer.Memory.Slice(nextLength, size - nextLength).CopyTo(readBuffer.Memory);
 				size -= nextLength;
@@ -400,7 +400,7 @@ namespace EpicGames.Horde.Storage.Nodes
 
 				IBlobHandle handle = reader.ReadBlobReference();
 
-				children.Add(new ChunkedDataNodeRef(hash, type, handle));
+				children.Add(new ChunkedDataNodeRef(type, hash, handle));
 			}
 			Children = children;
 		}
@@ -452,7 +452,7 @@ namespace EpicGames.Horde.Storage.Nodes
 				foreach (InteriorChunkedDataNode interiorNode in interiorNodes)
 				{
 					HashedNodeRef<ChunkedDataNode> newNodeRef = await writer.WriteHashedNodeAsync<ChunkedDataNode>(interiorNode, cancellationToken);
-					handleBuffer.Add(new ChunkedDataNodeRef(newNodeRef.Hash, ChunkedDataNodeType.Interior, newNodeRef.Handle));
+					handleBuffer.Add(new ChunkedDataNodeRef(ChunkedDataNodeType.Interior, newNodeRef.Hash, newNodeRef.Handle));
 				}
 
 				nodeRefs = handleBuffer;
diff --git a/Engine/Source/Programs/Shared/EpicGames.Horde/Storage/Nodes/ChunkedDataWriter.cs b/Engine/Source/Programs/Shared/EpicGames.Horde/Storage/Nodes/ChunkedDataWriter.cs
--- a/Engine/Source/Programs/Shared/EpicGames.Horde/Storage/Nodes/ChunkedDataWriter.cs
+++ b/Engine/Source/Programs/Shared/EpicGames.Horde/Storage/Nodes/ChunkedDataWriter.cs
@@ -346,7 +346,7 @@ namespace EpicGames.Horde.Storage.Nodes
 		async ValueTask FlushLeafNodeAsync(CancellationToken cancellationToken)
 		{
 			HashedNodeRef<ChunkedDataNode> nodeRef  = await _writer.WriteHashedNodeRefAsync<ChunkedDataNode>(s_leafNodeType, _leafLength, Array.Empty<IBlobHandle>(), cancellationToken);
-			_leafHandles.Add(new ChunkedDataNodeRef(nodeRef.Hash, ChunkedDataNodeType.Leaf, nodeRef.Handle));
+			_leafHandles.Add(new ChunkedDataNodeRef(ChunkedDataNodeType.Leaf, nodeRef.Hash, nodeRef.Handle));
 			ResetLeafState();
 		}
 	}
diff --git a/Engine/Source/Programs/Shared/EpicGames.Horde/Storage/Nodes/DirectoryNode.cs b/Engine/Source/Programs/Shared/EpicGames.Horde/Storage/Nodes/DirectoryNode.cs
--- a/Engine/Source/Programs/Shared/EpicGames.Horde/Storage/Nodes/DirectoryNode.cs
+++ b/Engine/Source/Programs/Shared/EpicGames.Horde/Storage/Nodes/DirectoryNode.cs
@@ -196,24 +196,19 @@ namespace EpicGames.Horde.Storage.Nodes
 			int fileCount = (int)reader.ReadUnsignedVarInt();
 			for (int idx = 0; idx < fileCount; idx++)
 			{
-				IoHash hash = reader.ReadIoHash();
+				IoHash targetHash = reader.ReadIoHash();
 
-				ChunkedDataNodeType type;
+				ChunkedDataNodeType targetType = ChunkedDataNodeType.Unknown;
 				if (reader.Version >= 2)
 				{
-					type = (ChunkedDataNodeType)reader.ReadUnsignedVarInt();
+					targetType = (ChunkedDataNodeType)reader.ReadUnsignedVarInt();
 				}
-				else
-				{
-					type = ChunkedDataNodeType.Unknown;
-				}
-
-				ChunkedDataNodeRef target = new ChunkedDataNodeRef(hash, type, reader.ReadBlobReference());
 
 				string name = reader.ReadString();
 				FileEntryFlags flags = (FileEntryFlags)reader.ReadUnsignedVarInt();
 				long length = (long)reader.ReadUnsignedVarInt();
 				IoHash streamHash = reader.ReadIoHash();
+				ChunkedDataNodeRef target = new ChunkedDataNodeRef(targetType, targetHash, reader.ReadBlobReference());
 
 				ReadOnlyMemory<byte> customData = default;
 				if ((flags & FileEntryFlags.HasCustomData) != 0)
