Horde: Add a length property to ChunkedDataNodeRef objects to allow seeking within chunked data streams. Currently only supports reading; need to deploy agent widely before updating writer.

diff --git a/Engine/Source/Programs/Shared/EpicGames.Horde/Storage/Nodes/ChunkedDataNode.cs b/Engine/Source/Programs/Shared/EpicGames.Horde/Storage/Nodes/ChunkedDataNode.cs
--- a/Engine/Source/Programs/Shared/EpicGames.Horde/Storage/Nodes/ChunkedDataNode.cs
+++ b/Engine/Source/Programs/Shared/EpicGames.Horde/Storage/Nodes/ChunkedDataNode.cs
@@ -6,6 +6,7 @@ using System.Buffers.Binary;
 using System.Collections.Generic;
 using System.Diagnostics;
 using System.IO;
+using System.Linq;
 using System.Threading;
 using System.Threading.Tasks;
 using EpicGames.Core;
@@ -112,8 +113,9 @@ namespace EpicGames.Horde.Storage.Nodes
 	/// </summary>
 	/// <param name="Type">Type of the referenced node</param>
 	/// <param name="Hash">Hash of the target node</param>
+	/// <param name="Length">Length of the data stream within this node</param>
 	/// <param name="Handle">Handle to the target node</param>
-	public record class ChunkedDataNodeRef(ChunkedDataNodeType Type, IoHash Hash, IBlobHandle Handle)
+	public record class ChunkedDataNodeRef(ChunkedDataNodeType Type, IoHash Hash, long Length, IBlobHandle Handle)
 	{
 		/// <summary>
 		/// Read the node which is the target of this ref
@@ -270,7 +272,7 @@ namespace EpicGames.Horde.Storage.Nodes
 				}
 
 				HashedNodeRef<ChunkedDataNode> nodeRef = await writer.WriteHashedNodeRefAsync<ChunkedDataNode>(GetNodeType<LeafChunkedDataNode>(), nextLength, Array.Empty<IBlobHandle>(), cancellationToken);
-				leafNodeRefs.Add(new ChunkedDataNodeRef(ChunkedDataNodeType.Leaf, nodeRef.Hash, nodeRef.Handle));
+				leafNodeRefs.Add(new ChunkedDataNodeRef(ChunkedDataNodeType.Leaf, nodeRef.Hash, nextLength, nodeRef.Handle));
 
 				readBuffer.Memory.Slice(nextLength, size - nextLength).CopyTo(readBuffer.Memory);
 				size -= nextLength;
@@ -355,7 +357,7 @@ namespace EpicGames.Horde.Storage.Nodes
 	/// <summary>
 	/// An interior file node
 	/// </summary>
-	[BlobType("{F4DEDDBC-4C7A-70CB-11F0-4783B9CDCCAF}", 2)]
+	[BlobType("{F4DEDDBC-4C7A-70CB-11F0-4783B9CDCCAF}", 2)] // Pending V3
 	public class InteriorChunkedDataNode : ChunkedDataNode
 	{
 		/// <summary>
@@ -388,19 +390,21 @@ namespace EpicGames.Horde.Storage.Nodes
 			{
 				IoHash hash = reader.ReadIoHash();
 
-				ChunkedDataNodeType type;
+				ChunkedDataNodeType type = ChunkedDataNodeType.Unknown;
 				if (reader.Version >= 2)
 				{
 					type = (ChunkedDataNodeType)reader.ReadUnsignedVarInt();
 				}
-				else
+
+				long length = 0;
+				if (reader.Version >= 3)
 				{
-					type = ChunkedDataNodeType.Unknown;
+					length = (long)reader.ReadUnsignedVarInt();
 				}
 
 				IBlobHandle handle = reader.ReadBlobReference();
 
-				children.Add(new ChunkedDataNodeRef(type, hash, handle));
+				children.Add(new ChunkedDataNodeRef(type, hash, length, handle));
 			}
 			Children = children;
 		}
@@ -412,6 +416,7 @@ namespace EpicGames.Horde.Storage.Nodes
 			{
 				writer.WriteIoHash(child.Hash);
 				writer.WriteUnsignedVarInt((int)child.Type);
+				// Pending V3: writer.WriteUnsignedVarInt((ulong)child.Length);
 				writer.WriteBlobReference(child.Handle);
 			}
 		}
@@ -442,17 +447,17 @@ namespace EpicGames.Horde.Storage.Nodes
 		{
 			List<ChunkedDataNodeRef> handleBuffer = new List<ChunkedDataNodeRef>();
 
-			List<InteriorChunkedDataNode> interiorNodes = new List<InteriorChunkedDataNode>();
+			List<(InteriorChunkedDataNode, long)> interiorNodes = new List<(InteriorChunkedDataNode, long)>();
 			while (nodeRefs.Count > 1)
 			{
 				interiorNodes.Clear();
 				CreateTreeLayer(nodeRefs, options, interiorNodes);
 
 				handleBuffer.Clear();
-				foreach (InteriorChunkedDataNode interiorNode in interiorNodes)
+				foreach ((InteriorChunkedDataNode interiorNode, long interiorLength) in interiorNodes)
 				{
 					HashedNodeRef<ChunkedDataNode> newNodeRef = await writer.WriteHashedNodeAsync<ChunkedDataNode>(interiorNode, cancellationToken);
-					handleBuffer.Add(new ChunkedDataNodeRef(ChunkedDataNodeType.Interior, newNodeRef.Hash, newNodeRef.Handle));
+					handleBuffer.Add(new ChunkedDataNodeRef(ChunkedDataNodeType.Interior, newNodeRef.Hash, interiorLength, newNodeRef.Handle));
 				}
 
 				nodeRefs = handleBuffer;
@@ -460,11 +465,11 @@ namespace EpicGames.Horde.Storage.Nodes
 
 			return nodeRefs[0];
 		}
-		
+
 		/// <summary>
 		/// Split a list of leaf handles into a layer of interior nodes
 		/// </summary>
-		static void CreateTreeLayer(List<ChunkedDataNodeRef> nodeRefs, InteriorChunkedDataNodeOptions options, List<InteriorChunkedDataNode> interiorNodes)
+		static void CreateTreeLayer(List<ChunkedDataNodeRef> nodeRefs, InteriorChunkedDataNodeOptions options, List<(InteriorChunkedDataNode, long)> interiorNodes)
 		{
 			Span<byte> buffer = stackalloc byte[IoHash.NumBytes];
 
@@ -491,7 +496,8 @@ namespace EpicGames.Horde.Storage.Nodes
 					children[childIndex - minIndex] = nodeRefs[childIndex];
 				}
 
-				interiorNodes.Add(new InteriorChunkedDataNode(children));
+				long interiorLength = children.Sum(x => x.Length);
+				interiorNodes.Add((new InteriorChunkedDataNode(children), interiorLength));
 			}
 		}
 
diff --git a/Engine/Source/Programs/Shared/EpicGames.Horde/Storage/Nodes/ChunkedDataWriter.cs b/Engine/Source/Programs/Shared/EpicGames.Horde/Storage/Nodes/ChunkedDataWriter.cs
--- a/Engine/Source/Programs/Shared/EpicGames.Horde/Storage/Nodes/ChunkedDataWriter.cs
+++ b/Engine/Source/Programs/Shared/EpicGames.Horde/Storage/Nodes/ChunkedDataWriter.cs
@@ -346,7 +346,7 @@ namespace EpicGames.Horde.Storage.Nodes
 		async ValueTask FlushLeafNodeAsync(CancellationToken cancellationToken)
 		{
 			HashedNodeRef<ChunkedDataNode> nodeRef  = await _writer.WriteHashedNodeRefAsync<ChunkedDataNode>(s_leafNodeType, _leafLength, Array.Empty<IBlobHandle>(), cancellationToken);
-			_leafHandles.Add(new ChunkedDataNodeRef(ChunkedDataNodeType.Leaf, nodeRef.Hash, nodeRef.Handle));
+			_leafHandles.Add(new ChunkedDataNodeRef(ChunkedDataNodeType.Leaf, nodeRef.Hash, _leafLength, nodeRef.Handle));
 			ResetLeafState();
 		}
 	}
diff --git a/Engine/Source/Programs/Shared/EpicGames.Horde/Storage/Nodes/DirectoryNode.cs b/Engine/Source/Programs/Shared/EpicGames.Horde/Storage/Nodes/DirectoryNode.cs
--- a/Engine/Source/Programs/Shared/EpicGames.Horde/Storage/Nodes/DirectoryNode.cs
+++ b/Engine/Source/Programs/Shared/EpicGames.Horde/Storage/Nodes/DirectoryNode.cs
@@ -208,7 +208,7 @@ namespace EpicGames.Horde.Storage.Nodes
 				FileEntryFlags flags = (FileEntryFlags)reader.ReadUnsignedVarInt();
 				long length = (long)reader.ReadUnsignedVarInt();
 				IoHash streamHash = reader.ReadIoHash();
-				ChunkedDataNodeRef target = new ChunkedDataNodeRef(targetType, targetHash, reader.ReadBlobReference());
+				ChunkedDataNodeRef target = new ChunkedDataNodeRef(targetType, targetHash, length, reader.ReadBlobReference());
 
 				ReadOnlyMemory<byte> customData = default;
 				if ((flags & FileEntryFlags.HasCustomData) != 0)
