Horde: Include the length of chunked data trees in references from the parent node.

diff --git a/Engine/Source/Programs/Shared/EpicGames.Horde/Storage/Nodes/ChunkedDataNode.cs b/Engine/Source/Programs/Shared/EpicGames.Horde/Storage/Nodes/ChunkedDataNode.cs
--- a/Engine/Source/Programs/Shared/EpicGames.Horde/Storage/Nodes/ChunkedDataNode.cs
+++ b/Engine/Source/Programs/Shared/EpicGames.Horde/Storage/Nodes/ChunkedDataNode.cs
@@ -357,7 +357,7 @@ namespace EpicGames.Horde.Storage.Nodes
 	/// <summary>
 	/// An interior file node
 	/// </summary>
-	[BlobType("{F4DEDDBC-4C7A-70CB-11F0-4783B9CDCCAF}", 2)] // Pending V3
+	[BlobType("{F4DEDDBC-4C7A-70CB-11F0-4783B9CDCCAF}", 3)]
 	public class InteriorChunkedDataNode : ChunkedDataNode
 	{
 		/// <summary>
@@ -416,7 +416,7 @@ namespace EpicGames.Horde.Storage.Nodes
 			{
 				writer.WriteIoHash(child.Hash);
 				writer.WriteUnsignedVarInt((int)child.Type);
-				// Pending V3: writer.WriteUnsignedVarInt((ulong)child.Length);
+				writer.WriteUnsignedVarInt((ulong)child.Length);
 				writer.WriteBlobReference(child.Handle);
 			}
 		}
