Horde: Fix test failure due to chunked node schema change.

diff --git a/Engine/Source/Programs/Shared/EpicGames.Horde/Storage/Nodes/ChunkedDataNode.cs b/Engine/Source/Programs/Shared/EpicGames.Horde/Storage/Nodes/ChunkedDataNode.cs
--- a/Engine/Source/Programs/Shared/EpicGames.Horde/Storage/Nodes/ChunkedDataNode.cs
+++ b/Engine/Source/Programs/Shared/EpicGames.Horde/Storage/Nodes/ChunkedDataNode.cs
@@ -526,7 +526,11 @@ namespace EpicGames.Horde.Storage.Nodes
 				IBlobHandle handle = nodeReader.ReadBlobReference();
 				if (nodeReader.Version >= 2)
 				{
-					_ = nodeReader.ReadUnsignedVarInt();
+					_ = nodeReader.ReadUnsignedVarInt(); // Type
+				}
+				if (nodeReader.Version >= 3)
+				{
+					_ = nodeReader.ReadUnsignedVarInt(); // Length
 				}
 				await ChunkedDataNode.CopyToStreamAsync(handle, outputStream, cancellationToken);
 			}
