Use CbFieldName rather than Utf8String as a parameter to CbSerializer.WriteNamed().

diff --git a/Engine/Source/Programs/Horde/Horde.Server/Ddc/BlobId.cs b/Engine/Source/Programs/Horde/Horde.Server/Ddc/BlobId.cs
--- a/Engine/Source/Programs/Horde/Horde.Server/Ddc/BlobId.cs
+++ b/Engine/Source/Programs/Horde/Horde.Server/Ddc/BlobId.cs
@@ -159,6 +159,6 @@ namespace Horde.Server.Ddc
 		public override void Write(CbWriter writer, BlobId value) => writer.WriteBinaryAttachmentValue(value.Hash);
 
 		/// <inheritdoc/>
-		public override void WriteNamed(CbWriter writer, Utf8String name, BlobId value) => writer.WriteBinaryAttachment(name, value.Hash);
+		public override void WriteNamed(CbWriter writer, CbFieldName name, BlobId value) => writer.WriteBinaryAttachment(name, value.Hash);
 	}
 }
diff --git a/Engine/Source/Programs/Horde/Horde.Server/Ddc/BucketId.cs b/Engine/Source/Programs/Horde/Horde.Server/Ddc/BucketId.cs
--- a/Engine/Source/Programs/Horde/Horde.Server/Ddc/BucketId.cs
+++ b/Engine/Source/Programs/Horde/Horde.Server/Ddc/BucketId.cs
@@ -84,6 +84,6 @@ namespace Horde.Server.Ddc
 		public override void Write(CbWriter writer, BucketId value) => writer.WriteStringValue(value.ToString());
 
 		/// <inheritdoc/>
-		public override void WriteNamed(CbWriter writer, Utf8String name, BucketId value) => writer.WriteString(name, value.ToString());
+		public override void WriteNamed(CbWriter writer, CbFieldName name, BucketId value) => writer.WriteString(name, value.ToString());
 	}
 }
diff --git a/Engine/Source/Programs/Horde/Horde.Server/Ddc/ContentId.cs b/Engine/Source/Programs/Horde/Horde.Server/Ddc/ContentId.cs
--- a/Engine/Source/Programs/Horde/Horde.Server/Ddc/ContentId.cs
+++ b/Engine/Source/Programs/Horde/Horde.Server/Ddc/ContentId.cs
@@ -159,6 +159,6 @@ namespace Horde.Server.Ddc
 		public override void Write(CbWriter writer, ContentId value) => writer.WriteBinaryAttachmentValue(value.Hash);
 
 		/// <inheritdoc/>
-		public override void WriteNamed(CbWriter writer, Utf8String name, ContentId value) => writer.WriteBinaryAttachment(name, value.Hash);
+		public override void WriteNamed(CbWriter writer, CbFieldName name, ContentId value) => writer.WriteBinaryAttachment(name, value.Hash);
 	}
 }
diff --git a/Engine/Source/Programs/Horde/Horde.Server/Ddc/RefId.cs b/Engine/Source/Programs/Horde/Horde.Server/Ddc/RefId.cs
--- a/Engine/Source/Programs/Horde/Horde.Server/Ddc/RefId.cs
+++ b/Engine/Source/Programs/Horde/Horde.Server/Ddc/RefId.cs
@@ -93,6 +93,6 @@ namespace Horde.Server.Ddc
 		public override void Write(CbWriter writer, RefId value) => writer.WriteStringValue(value.ToString());
 
 		/// <inheritdoc/>
-		public override void WriteNamed(CbWriter writer, Utf8String name, RefId value) => writer.WriteString(name, value.ToString());
+		public override void WriteNamed(CbWriter writer, CbFieldName name, RefId value) => writer.WriteString(name, value.ToString());
 	}
 }
diff --git a/Engine/Source/Programs/Shared/EpicGames.AspNet/CbConvertersAspnet.cs b/Engine/Source/Programs/Shared/EpicGames.AspNet/CbConvertersAspnet.cs
--- a/Engine/Source/Programs/Shared/EpicGames.AspNet/CbConvertersAspnet.cs
+++ b/Engine/Source/Programs/Shared/EpicGames.AspNet/CbConvertersAspnet.cs
@@ -52,7 +52,7 @@ namespace EpicGames.AspNet
 		}
 
 		/// <inheritdoc/>
-		public override void WriteNamed(CbWriter writer, Utf8String name, ProblemDetails problemDetails)
+		public override void WriteNamed(CbWriter writer, CbFieldName name, ProblemDetails problemDetails)
 		{
 			writer.WriteField(name, ToCbObject(problemDetails).AsField());
 		}
diff --git a/Engine/Source/Programs/Shared/EpicGames.Horde/BinaryIdConverter.cs b/Engine/Source/Programs/Shared/EpicGames.Horde/BinaryIdConverter.cs
--- a/Engine/Source/Programs/Shared/EpicGames.Horde/BinaryIdConverter.cs
+++ b/Engine/Source/Programs/Shared/EpicGames.Horde/BinaryIdConverter.cs
@@ -65,7 +65,7 @@ namespace EpicGames.Horde
 		}
 
 		/// <inheritdoc/>
-		public override void WriteNamed(CbWriter writer, Utf8String name, TValue value)
+		public override void WriteNamed(CbWriter writer, CbFieldName name, TValue value)
 		{
 			writer.WriteString(name, _converter.ToBinaryId(value).ToString());
 		}
diff --git a/Engine/Source/Programs/Shared/EpicGames.Horde/Common/Condition.cs b/Engine/Source/Programs/Shared/EpicGames.Horde/Common/Condition.cs
--- a/Engine/Source/Programs/Shared/EpicGames.Horde/Common/Condition.cs
+++ b/Engine/Source/Programs/Shared/EpicGames.Horde/Common/Condition.cs
@@ -658,7 +658,7 @@ namespace EpicGames.Horde.Common
 		}
 
 		/// <inheritdoc/>
-		public override void WriteNamed(CbWriter writer, Utf8String name, Condition value)
+		public override void WriteNamed(CbWriter writer, CbFieldName name, Condition value)
 		{
 			if (value != null)
 			{
diff --git a/Engine/Source/Programs/Shared/EpicGames.Horde/Compute/ClusterId.cs b/Engine/Source/Programs/Shared/EpicGames.Horde/Compute/ClusterId.cs
--- a/Engine/Source/Programs/Shared/EpicGames.Horde/Compute/ClusterId.cs
+++ b/Engine/Source/Programs/Shared/EpicGames.Horde/Compute/ClusterId.cs
@@ -65,7 +65,7 @@ namespace EpicGames.Horde.Compute
 		public override void Write(CbWriter writer, ClusterId value) => writer.WriteStringValue(value.ToString());
 
 		/// <inheritdoc/>
-		public override void WriteNamed(CbWriter writer, Utf8String name, ClusterId value) => writer.WriteString(name, value.ToString());
+		public override void WriteNamed(CbWriter writer, CbFieldName name, ClusterId value) => writer.WriteString(name, value.ToString());
 	}
 
 	/// <summary>
diff --git a/Engine/Source/Programs/Shared/EpicGames.Horde/Storage/RefName.cs b/Engine/Source/Programs/Shared/EpicGames.Horde/Storage/RefName.cs
--- a/Engine/Source/Programs/Shared/EpicGames.Horde/Storage/RefName.cs
+++ b/Engine/Source/Programs/Shared/EpicGames.Horde/Storage/RefName.cs
@@ -193,6 +193,6 @@ namespace EpicGames.Horde.Storage
 		public override void Write(CbWriter writer, RefName value) => writer.WriteUtf8StringValue(value.Text);
 
 		/// <inheritdoc/>
-		public override void WriteNamed(CbWriter writer, Utf8String name, RefName value) => writer.WriteUtf8String(name, value.Text);
+		public override void WriteNamed(CbWriter writer, CbFieldName name, RefName value) => writer.WriteUtf8String(name, value.Text);
 	}
 }
diff --git a/Engine/Source/Programs/Shared/EpicGames.Horde/StringIdConverter.cs b/Engine/Source/Programs/Shared/EpicGames.Horde/StringIdConverter.cs
--- a/Engine/Source/Programs/Shared/EpicGames.Horde/StringIdConverter.cs
+++ b/Engine/Source/Programs/Shared/EpicGames.Horde/StringIdConverter.cs
@@ -65,7 +65,7 @@ namespace EpicGames.Horde
 		}
 
 		/// <inheritdoc/>
-		public override void WriteNamed(CbWriter writer, Utf8String name, TValue value)
+		public override void WriteNamed(CbWriter writer, CbFieldName name, TValue value)
 		{
 			writer.WriteString(name, _converter.ToStringId(value).ToString());
 		}
diff --git a/Engine/Source/Programs/Shared/EpicGames.Serialization/CbConverter.cs b/Engine/Source/Programs/Shared/EpicGames.Serialization/CbConverter.cs
--- a/Engine/Source/Programs/Shared/EpicGames.Serialization/CbConverter.cs
+++ b/Engine/Source/Programs/Shared/EpicGames.Serialization/CbConverter.cs
@@ -55,7 +55,7 @@ namespace EpicGames.Serialization
 		/// <param name="writer"></param>
 		/// <param name="name"></param>
 		/// <param name="value"></param>
-		void WriteNamedObject(CbWriter writer, Utf8String name, object? value);
+		void WriteNamedObject(CbWriter writer, CbFieldName name, object? value);
 	}
 
 	/// <summary>
@@ -83,7 +83,7 @@ namespace EpicGames.Serialization
 		/// <param name="writer"></param>
 		/// <param name="name"></param>
 		/// <param name="value"></param>
-		void WriteNamed(CbWriter writer, Utf8String name, T value);
+		void WriteNamed(CbWriter writer, CbFieldName name, T value);
 	}
 
 	/// <summary>
@@ -102,7 +102,7 @@ namespace EpicGames.Serialization
 		public MethodInfo WriteMethod { get; }
 
 		/// <summary>
-		/// Method with the signature CbWriter, Utf8String, T -> void
+		/// Method with the signature CbWriter, CbFieldName, T -> void
 		/// </summary>
 		public MethodInfo WriteNamedMethod { get; }
 	}
@@ -118,7 +118,7 @@ namespace EpicGames.Serialization
 
 			static T Read(CbField field) => s_staticConverter.Read(field);
 			static void Write(CbWriter writer, T value) => s_staticConverter.Write(writer, value);
-			static void WriteNamed(CbWriter writer, Utf8String name, T value) => s_staticConverter.WriteNamed(writer, name, value);
+			static void WriteNamed(CbWriter writer, CbFieldName name, T value) => s_staticConverter.WriteNamed(writer, name, value);
 
 			public MethodInfo ReadMethod { get; } = GetMethodInfo(() => Read(null!));
 			public MethodInfo WriteMethod { get; } = GetMethodInfo(() => Write(null!, default!));
@@ -192,7 +192,7 @@ namespace EpicGames.Serialization
 		public void WriteObject(CbWriter writer, object? value) => Write(writer, (T)value!);
 
 		/// <inheritdoc/>
-		public void WriteNamedObject(CbWriter writer, Utf8String name, object? value) => WriteNamed(writer, name, (T)value!);
+		public void WriteNamedObject(CbWriter writer, CbFieldName name, object? value) => WriteNamed(writer, name, (T)value!);
 
 		/// <inheritdoc/>
 		public abstract T Read(CbField field);
@@ -201,7 +201,7 @@ namespace EpicGames.Serialization
 		public abstract void Write(CbWriter writer, T value);
 
 		/// <inheritdoc/>
-		public abstract void WriteNamed(CbWriter writer, Utf8String name, T value);
+		public abstract void WriteNamed(CbWriter writer, CbFieldName name, T value);
 	}
 
 	/// <summary>
@@ -306,7 +306,7 @@ namespace EpicGames.Serialization
 				public override void Write(CbWriter writer, T value) => _inner.WriteObject(writer, value);
 
 				/// <inheritdoc/>
-				public override void WriteNamed(CbWriter writer, Utf8String name, T value) => _inner.WriteNamedObject(writer, name, value);
+				public override void WriteNamed(CbWriter writer, CbFieldName name, T value) => _inner.WriteNamedObject(writer, name, value);
 			}
 		}
 
diff --git a/Engine/Source/Programs/Shared/EpicGames.Serialization/CbSerializer.cs b/Engine/Source/Programs/Shared/EpicGames.Serialization/CbSerializer.cs
--- a/Engine/Source/Programs/Shared/EpicGames.Serialization/CbSerializer.cs
+++ b/Engine/Source/Programs/Shared/EpicGames.Serialization/CbSerializer.cs
@@ -166,7 +166,7 @@ namespace EpicGames.Serialization
 		/// <param name="writer"></param>
 		/// <param name="name"></param>
 		/// <param name="value"></param>
-		public static void Serialize<T>(CbWriter writer, Utf8String name, T value)
+		public static void Serialize<T>(CbWriter writer, CbFieldName name, T value)
 		{
 			CbConverter.GetConverter<T>().WriteNamed(writer, name, value);
 		}
diff --git a/Engine/Source/Programs/Shared/EpicGames.Serialization/Converters/CbArrayConverter.cs b/Engine/Source/Programs/Shared/EpicGames.Serialization/Converters/CbArrayConverter.cs
--- a/Engine/Source/Programs/Shared/EpicGames.Serialization/Converters/CbArrayConverter.cs
+++ b/Engine/Source/Programs/Shared/EpicGames.Serialization/Converters/CbArrayConverter.cs
@@ -47,7 +47,7 @@ namespace EpicGames.Serialization.Converters
 		}
 
 		/// <inheritdoc/>
-		public override void WriteNamed(CbWriter writer, Utf8String name, T[] array)
+		public override void WriteNamed(CbWriter writer, CbFieldName name, T[] array)
 		{
 			if (array == null)
 			{
diff --git a/Engine/Source/Programs/Shared/EpicGames.Serialization/Converters/CbClassConverter.cs b/Engine/Source/Programs/Shared/EpicGames.Serialization/Converters/CbClassConverter.cs
--- a/Engine/Source/Programs/Shared/EpicGames.Serialization/Converters/CbClassConverter.cs
+++ b/Engine/Source/Programs/Shared/EpicGames.Serialization/Converters/CbClassConverter.cs
@@ -61,7 +61,7 @@ namespace EpicGames.Serialization.Converters
 			WriteConcreteContentsMethod = new DynamicMethod($"WriteConcreteContents_{classType.Name}", null, new Type[] { typeof(CbWriter), classType });
 
 			WriteMethod = new DynamicMethod($"Write_{classType.Name}", null, new Type[] { typeof(CbWriter), classType });
-			WriteNamedMethod = new DynamicMethod($"WriteNamed_{classType.Name}", null, new Type[] { typeof(CbWriter), typeof(Utf8String), classType });
+			WriteNamedMethod = new DynamicMethod($"WriteNamed_{classType.Name}", null, new Type[] { typeof(CbWriter), typeof(CbFieldName), classType });
 
 			if (IsPolymorphic)
 			{
@@ -262,7 +262,7 @@ namespace EpicGames.Serialization.Converters
 
 				generator.Emit(OpCodes.Ldsfld, namesField);
 				generator.Emit(OpCodes.Ldc_I4, idx);
-				generator.Emit(OpCodes.Ldelem, typeof(Utf8String));
+				generator.Emit(OpCodes.Ldelem, typeof(CbFieldName));
 
 				generator.Emit(OpCodes.Ldloc, local);
 				generator.EmitCall(OpCodes.Call, writeMethod, null);
@@ -499,7 +499,7 @@ namespace EpicGames.Serialization.Converters
 
 		readonly Func<CbField, T> _readFunc;
 		readonly Action<CbWriter, T> _writeFunc;
-		readonly Action<CbWriter, Utf8String, T> _writeNamedFunc;
+		readonly Action<CbWriter, CbFieldName, T> _writeNamedFunc;
 
 		public CbClassConverter()
 		{
@@ -507,7 +507,7 @@ namespace EpicGames.Serialization.Converters
 
 			_readFunc = CreateDelegate<Func<CbField, T>>(_methods.ReadMethod);
 			_writeFunc = CreateDelegate<Action<CbWriter, T>>(_methods.WriteMethod);
-			_writeNamedFunc = CreateDelegate<Action<CbWriter, Utf8String, T>>(_methods.WriteNamedMethod);
+			_writeNamedFunc = CreateDelegate<Action<CbWriter, CbFieldName, T>>(_methods.WriteNamedMethod);
 		}
 
 		public MethodInfo ReadMethod => _methods.ReadMethod;
@@ -522,7 +522,7 @@ namespace EpicGames.Serialization.Converters
 
 		public override void Write(CbWriter writer, T value) => _writeFunc(writer, value);
 
-		public override void WriteNamed(CbWriter writer, Utf8String name, T value) => _writeNamedFunc(writer, name, value);
+		public override void WriteNamed(CbWriter writer, CbFieldName name, T value) => _writeNamedFunc(writer, name, value);
 	}
 
 	class CbClassConverterFactory : CbConverterFactory
diff --git a/Engine/Source/Programs/Shared/EpicGames.Serialization/Converters/CbDictionaryConverter.cs b/Engine/Source/Programs/Shared/EpicGames.Serialization/Converters/CbDictionaryConverter.cs
--- a/Engine/Source/Programs/Shared/EpicGames.Serialization/Converters/CbDictionaryConverter.cs
+++ b/Engine/Source/Programs/Shared/EpicGames.Serialization/Converters/CbDictionaryConverter.cs
@@ -63,7 +63,7 @@ namespace EpicGames.Serialization.Converters
 		}
 
 		/// <inheritdoc/>
-		public override void WriteNamed(CbWriter writer, Utf8String name, Dictionary<TKey, TValue> value)
+		public override void WriteNamed(CbWriter writer, CbFieldName name, Dictionary<TKey, TValue> value)
 		{
 			if (value == null)
 			{
diff --git a/Engine/Source/Programs/Shared/EpicGames.Serialization/Converters/CbEnumConverter.cs b/Engine/Source/Programs/Shared/EpicGames.Serialization/Converters/CbEnumConverter.cs
--- a/Engine/Source/Programs/Shared/EpicGames.Serialization/Converters/CbEnumConverter.cs
+++ b/Engine/Source/Programs/Shared/EpicGames.Serialization/Converters/CbEnumConverter.cs
@@ -12,7 +12,7 @@ namespace EpicGames.Serialization.Converters
 	{
 		readonly Func<CbField, T> _readFunc;
 		readonly Action<CbWriter, T> _writeFunc;
-		readonly Action<CbWriter, Utf8String, T> _writeNamedFunc;
+		readonly Action<CbWriter, CbFieldName, T> _writeNamedFunc;
 
 		public CbEnumConverter()
 		{
@@ -26,9 +26,9 @@ namespace EpicGames.Serialization.Converters
 			CreateEnumWriter(WriteMethod.GetILGenerator());
 			_writeFunc = (Action<CbWriter, T>)WriteMethod.CreateDelegate(typeof(Action<CbWriter, T>));
 
-			WriteNamedMethod = new DynamicMethod($"WriteNamed_{type.Name}", null, new Type[] { typeof(CbWriter), typeof(Utf8String), type });
+			WriteNamedMethod = new DynamicMethod($"WriteNamed_{type.Name}", null, new Type[] { typeof(CbWriter), typeof(CbFieldName), type });
 			CreateNamedEnumWriter(WriteNamedMethod.GetILGenerator());
-			_writeNamedFunc = (Action<CbWriter, Utf8String, T>)WriteNamedMethod.CreateDelegate(typeof(Action<CbWriter, Utf8String, T>));
+			_writeNamedFunc = (Action<CbWriter, CbFieldName, T>)WriteNamedMethod.CreateDelegate(typeof(Action<CbWriter, CbFieldName, T>));
 		}
 
 		public DynamicMethod ReadMethod { get; }
@@ -44,7 +44,7 @@ namespace EpicGames.Serialization.Converters
 
 		public override void Write(CbWriter writer, T value) => _writeFunc(writer, value);
 
-		public override void WriteNamed(CbWriter writer, Utf8String name, T value) => _writeNamedFunc(writer, name, value);
+		public override void WriteNamed(CbWriter writer, CbFieldName name, T value) => _writeNamedFunc(writer, name, value);
 
 		static void CreateEnumReader(ILGenerator generator)
 		{
diff --git a/Engine/Source/Programs/Shared/EpicGames.Serialization/Converters/CbFieldConverter.cs b/Engine/Source/Programs/Shared/EpicGames.Serialization/Converters/CbFieldConverter.cs
--- a/Engine/Source/Programs/Shared/EpicGames.Serialization/Converters/CbFieldConverter.cs
+++ b/Engine/Source/Programs/Shared/EpicGames.Serialization/Converters/CbFieldConverter.cs
@@ -22,7 +22,7 @@ namespace EpicGames.Serialization.Converters
 		}
 
 		/// <inheritdoc/>
-		public override void WriteNamed(CbWriter writer, Utf8String name, CbField value)
+		public override void WriteNamed(CbWriter writer, CbFieldName name, CbField value)
 		{
 			writer.WriteField(name, value);
 		}
@@ -46,7 +46,7 @@ namespace EpicGames.Serialization.Converters
 		}
 
 		/// <inheritdoc/>
-		public override void WriteNamed(CbWriter writer, Utf8String name, CbObject obj)
+		public override void WriteNamed(CbWriter writer, CbFieldName name, CbObject obj)
 		{
 			writer.WriteField(name, obj.AsField());
 		}
diff --git a/Engine/Source/Programs/Shared/EpicGames.Serialization/Converters/CbListConverter.cs b/Engine/Source/Programs/Shared/EpicGames.Serialization/Converters/CbListConverter.cs
--- a/Engine/Source/Programs/Shared/EpicGames.Serialization/Converters/CbListConverter.cs
+++ b/Engine/Source/Programs/Shared/EpicGames.Serialization/Converters/CbListConverter.cs
@@ -49,7 +49,7 @@ namespace EpicGames.Serialization.Converters
 		}
 
 		/// <inheritdoc/>
-		public override void WriteNamed(CbWriter writer, Utf8String name, List<T> list)
+		public override void WriteNamed(CbWriter writer, CbFieldName name, List<T> list)
 		{
 			if (list == null)
 			{
@@ -95,7 +95,7 @@ namespace EpicGames.Serialization.Converters
 		}
 
 		/// <inheritdoc/>
-		public override void WriteNamed(CbWriter writer, Utf8String name, List<Utf8String> list)
+		public override void WriteNamed(CbWriter writer, CbFieldName name, List<Utf8String> list)
 		{
 			if (list.Count > 0)
 			{
diff --git a/Engine/Source/Programs/Shared/EpicGames.Serialization/Converters/CbNullableConverter.cs b/Engine/Source/Programs/Shared/EpicGames.Serialization/Converters/CbNullableConverter.cs
--- a/Engine/Source/Programs/Shared/EpicGames.Serialization/Converters/CbNullableConverter.cs
+++ b/Engine/Source/Programs/Shared/EpicGames.Serialization/Converters/CbNullableConverter.cs
@@ -27,7 +27,7 @@ namespace EpicGames.Serialization.Converters
 		}
 
 		/// <inheritdoc/>
-		public override void WriteNamed(CbWriter writer, Utf8String name, T? nullable)
+		public override void WriteNamed(CbWriter writer, CbFieldName name, T? nullable)
 		{
 			if (nullable.HasValue)
 			{
diff --git a/Engine/Source/Programs/Shared/EpicGames.Serialization/Converters/CbPrimitiveConverter.cs b/Engine/Source/Programs/Shared/EpicGames.Serialization/Converters/CbPrimitiveConverter.cs
--- a/Engine/Source/Programs/Shared/EpicGames.Serialization/Converters/CbPrimitiveConverter.cs
+++ b/Engine/Source/Programs/Shared/EpicGames.Serialization/Converters/CbPrimitiveConverter.cs
@@ -16,9 +16,9 @@ namespace EpicGames.Serialization.Converters
 		public Action<CbWriter, T> WriteFunc { get; }
 
 		public MethodInfo WriteNamedMethod { get; }
-		public Action<CbWriter, Utf8String, T> WriteNamedFunc { get; }
+		public Action<CbWriter, CbFieldName, T> WriteNamedFunc { get; }
 
-		public CbPrimitiveConverter(Expression<Func<CbField, T>> read, Expression<Action<CbWriter, T>> write, Expression<Action<CbWriter, Utf8String, T>> writeNamed)
+		public CbPrimitiveConverter(Expression<Func<CbField, T>> read, Expression<Action<CbWriter, T>> write, Expression<Action<CbWriter, CbFieldName, T>> writeNamed)
 		{
 			ReadMethod = ((MethodCallExpression)read.Body).Method;
 			ReadFunc = read.Compile();
@@ -34,6 +34,6 @@ namespace EpicGames.Serialization.Converters
 
 		public override void Write(CbWriter writer, T value) => WriteFunc(writer, value);
 
-		public override void WriteNamed(CbWriter writer, Utf8String name, T value) => WriteNamedFunc(writer, name, value);
+		public override void WriteNamed(CbWriter writer, CbFieldName name, T value) => WriteNamedFunc(writer, name, value);
 	}
 }
diff --git a/Engine/Source/Programs/Shared/EpicGames.Serialization/Converters/CbStringConverter.cs b/Engine/Source/Programs/Shared/EpicGames.Serialization/Converters/CbStringConverter.cs
--- a/Engine/Source/Programs/Shared/EpicGames.Serialization/Converters/CbStringConverter.cs
+++ b/Engine/Source/Programs/Shared/EpicGames.Serialization/Converters/CbStringConverter.cs
@@ -20,7 +20,7 @@ namespace EpicGames.Serialization.Converters
 			return (T)_typeConverter.ConvertFromInvariantString(field.AsString())!;
 		}
 
-		public override void WriteNamed(CbWriter writer, Utf8String name, T value)
+		public override void WriteNamed(CbWriter writer, CbFieldName name, T value)
 		{
 			writer.WriteString(name, _typeConverter.ConvertToInvariantString(value));
 		}
