Change CbConverter to an abstract base class rather than interface.

diff --git a/Engine/Source/Programs/Horde/Horde.Server/Ddc/BlobId.cs b/Engine/Source/Programs/Horde/Horde.Server/Ddc/BlobId.cs
--- a/Engine/Source/Programs/Horde/Horde.Server/Ddc/BlobId.cs
+++ b/Engine/Source/Programs/Horde/Horde.Server/Ddc/BlobId.cs
@@ -150,7 +150,7 @@ namespace Horde.Server.Ddc
 	/// <summary>
 	/// Serializes <see cref="BlobId"/> instances to compact binary
 	/// </summary>
-	public class BlobIdCbConverter : CbConverterBase<BlobId>
+	public class BlobIdCbConverter : CbConverter<BlobId>
 	{
 		/// <inheritdoc/>
 		public override BlobId Read(CbField field) => new BlobId(field.AsHash());
diff --git a/Engine/Source/Programs/Horde/Horde.Server/Ddc/BucketId.cs b/Engine/Source/Programs/Horde/Horde.Server/Ddc/BucketId.cs
--- a/Engine/Source/Programs/Horde/Horde.Server/Ddc/BucketId.cs
+++ b/Engine/Source/Programs/Horde/Horde.Server/Ddc/BucketId.cs
@@ -76,7 +76,7 @@ namespace Horde.Server.Ddc
 		public override object ConvertFrom(ITypeDescriptorContext? context, CultureInfo? culture, object? value) => new BucketId((string)value!);
 	}
 
-	sealed class BucketIdCbConverter : CbConverterBase<BucketId>
+	sealed class BucketIdCbConverter : CbConverter<BucketId>
 	{
 		public override BucketId Read(CbField field) => new BucketId(field.AsString());
 
diff --git a/Engine/Source/Programs/Horde/Horde.Server/Ddc/ContentId.cs b/Engine/Source/Programs/Horde/Horde.Server/Ddc/ContentId.cs
--- a/Engine/Source/Programs/Horde/Horde.Server/Ddc/ContentId.cs
+++ b/Engine/Source/Programs/Horde/Horde.Server/Ddc/ContentId.cs
@@ -150,7 +150,7 @@ namespace Horde.Server.Ddc
 	/// <summary>
 	/// Serializes <see cref="ContentId"/> instances to compact binary
 	/// </summary>
-	public class ContentIdCbConverter : CbConverterBase<ContentId>
+	public class ContentIdCbConverter : CbConverter<ContentId>
 	{
 		/// <inheritdoc/>
 		public override ContentId Read(CbField field) => new ContentId(field.AsHash());
diff --git a/Engine/Source/Programs/Horde/Horde.Server/Ddc/RefId.cs b/Engine/Source/Programs/Horde/Horde.Server/Ddc/RefId.cs
--- a/Engine/Source/Programs/Horde/Horde.Server/Ddc/RefId.cs
+++ b/Engine/Source/Programs/Horde/Horde.Server/Ddc/RefId.cs
@@ -85,7 +85,7 @@ namespace Horde.Server.Ddc
 	/// <summary>
 	/// Type converter to compact binary objects
 	/// </summary>
-	sealed class RefIdCbConverter : CbConverterBase<RefId>
+	sealed class RefIdCbConverter : CbConverter<RefId>
 	{
 		public override RefId Read(CbField field) => RefId.Parse(field.AsString());
 
diff --git a/Engine/Source/Programs/Shared/EpicGames.AspNet/CbConvertersAspnet.cs b/Engine/Source/Programs/Shared/EpicGames.AspNet/CbConvertersAspnet.cs
--- a/Engine/Source/Programs/Shared/EpicGames.AspNet/CbConvertersAspnet.cs
+++ b/Engine/Source/Programs/Shared/EpicGames.AspNet/CbConvertersAspnet.cs
@@ -24,7 +24,7 @@ namespace EpicGames.AspNet
 	/// <summary>
 	/// Converter for asp.net problem details type
 	/// </summary>
-	class CbProblemDetailsConverter : CbConverterBase<ProblemDetails>
+	class CbProblemDetailsConverter : CbConverter<ProblemDetails>
 	{
 		/// <inheritdoc/>
 		public override ProblemDetails Read(CbField field)
diff --git a/Engine/Source/Programs/Shared/EpicGames.Horde/BinaryIdConverter.cs b/Engine/Source/Programs/Shared/EpicGames.Horde/BinaryIdConverter.cs
--- a/Engine/Source/Programs/Shared/EpicGames.Horde/BinaryIdConverter.cs
+++ b/Engine/Source/Programs/Shared/EpicGames.Horde/BinaryIdConverter.cs
@@ -48,7 +48,7 @@ namespace EpicGames.Horde
 	/// <summary>
 	/// Converter to compact binary objects
 	/// </summary>
-	public sealed class BinaryIdCbConverter<TValue, TConverter> : CbConverterBase<TValue> where TValue : struct where TConverter : BinaryIdConverter<TValue>, new()
+	public sealed class BinaryIdCbConverter<TValue, TConverter> : CbConverter<TValue> where TValue : struct where TConverter : BinaryIdConverter<TValue>, new()
 	{
 		readonly TConverter _converter = new TConverter();
 
diff --git a/Engine/Source/Programs/Shared/EpicGames.Horde/Common/Condition.cs b/Engine/Source/Programs/Shared/EpicGames.Horde/Common/Condition.cs
--- a/Engine/Source/Programs/Shared/EpicGames.Horde/Common/Condition.cs
+++ b/Engine/Source/Programs/Shared/EpicGames.Horde/Common/Condition.cs
@@ -629,7 +629,7 @@ namespace EpicGames.Horde.Common
 	/// <summary>
 	/// Converter from conditions to compact binary objects
 	/// </summary>
-	public class ConditionCbConverter : CbConverterBase<Condition>
+	public class ConditionCbConverter : CbConverter<Condition>
 	{
 		/// <inheritdoc/>
 		public override Condition Read(CbField field)
diff --git a/Engine/Source/Programs/Shared/EpicGames.Horde/Compute/ClusterId.cs b/Engine/Source/Programs/Shared/EpicGames.Horde/Compute/ClusterId.cs
--- a/Engine/Source/Programs/Shared/EpicGames.Horde/Compute/ClusterId.cs
+++ b/Engine/Source/Programs/Shared/EpicGames.Horde/Compute/ClusterId.cs
@@ -56,7 +56,7 @@ namespace EpicGames.Horde.Compute
 	/// <summary>
 	/// Compact binary converter for ClusterId
 	/// </summary>
-	sealed class ClusterIdCbConverter : CbConverterBase<ClusterId>
+	sealed class ClusterIdCbConverter : CbConverter<ClusterId>
 	{
 		/// <inheritdoc/>
 		public override ClusterId Read(CbField field) => new ClusterId(field.AsString());
diff --git a/Engine/Source/Programs/Shared/EpicGames.Horde/Storage/RefName.cs b/Engine/Source/Programs/Shared/EpicGames.Horde/Storage/RefName.cs
--- a/Engine/Source/Programs/Shared/EpicGames.Horde/Storage/RefName.cs
+++ b/Engine/Source/Programs/Shared/EpicGames.Horde/Storage/RefName.cs
@@ -184,7 +184,7 @@ namespace EpicGames.Horde.Storage
 	/// <summary>
 	/// Type converter to compact binary
 	/// </summary>
-	sealed class RefNameCbConverter : CbConverterBase<RefName>
+	sealed class RefNameCbConverter : CbConverter<RefName>
 	{
 		/// <inheritdoc/>
 		public override RefName Read(CbField field) => new RefName(field.AsUtf8String());
diff --git a/Engine/Source/Programs/Shared/EpicGames.Horde/StringIdConverter.cs b/Engine/Source/Programs/Shared/EpicGames.Horde/StringIdConverter.cs
--- a/Engine/Source/Programs/Shared/EpicGames.Horde/StringIdConverter.cs
+++ b/Engine/Source/Programs/Shared/EpicGames.Horde/StringIdConverter.cs
@@ -48,7 +48,7 @@ namespace EpicGames.Horde
 	/// <summary>
 	/// Converter to compact binary objects
 	/// </summary>
-	public sealed class StringIdCbConverter<TValue, TConverter> : CbConverterBase<TValue> where TValue : struct where TConverter : StringIdConverter<TValue>, new()
+	public sealed class StringIdCbConverter<TValue, TConverter> : CbConverter<TValue> where TValue : struct where TConverter : StringIdConverter<TValue>, new()
 	{
 		readonly TConverter _converter = new TConverter();
 
diff --git a/Engine/Source/Programs/Shared/EpicGames.Serialization/CbConverter.cs b/Engine/Source/Programs/Shared/EpicGames.Serialization/CbConverter.cs
--- a/Engine/Source/Programs/Shared/EpicGames.Serialization/CbConverter.cs
+++ b/Engine/Source/Programs/Shared/EpicGames.Serialization/CbConverter.cs
@@ -33,21 +33,21 @@ namespace EpicGames.Serialization
 	/// <summary>
 	/// Base class for all converters. Deriving from <see cref="ICbConverter{T}"/> is more efficient
 	/// </summary>
-	public interface ICbConverter
+	public abstract partial class CbConverter
 	{
 		/// <summary>
 		/// Reads an object from a field
 		/// </summary>
 		/// <param name="field"></param>
 		/// <returns></returns>
-		object? ReadObject(CbField field);
+		public abstract object? ReadObject(CbField field);
 
 		/// <summary>
 		/// Writes an object to a field
 		/// </summary>
 		/// <param name="writer"></param>
 		/// <param name="value"></param>
-		void WriteObject(CbWriter writer, object? value);
+		public abstract void WriteObject(CbWriter writer, object? value);
 
 		/// <summary>
 		/// Writes an object to a named field, if not equal to the default value
@@ -55,27 +55,36 @@ namespace EpicGames.Serialization
 		/// <param name="writer"></param>
 		/// <param name="name"></param>
 		/// <param name="value"></param>
-		void WriteNamedObject(CbWriter writer, CbFieldName name, object? value);
+		public abstract void WriteNamedObject(CbWriter writer, CbFieldName name, object? value);
 	}
 
 	/// <summary>
 	/// Converter for a particular type
 	/// </summary>
-	public interface ICbConverter<T> : ICbConverter
+	public abstract class CbConverter<T> : CbConverter
 	{
+		/// <inheritdoc/>
+		public override sealed object? ReadObject(CbField field) => Read(field);
+
+		/// <inheritdoc/>
+		public override sealed void WriteObject(CbWriter writer, object? value) => Write(writer, (T)value!);
+
+		/// <inheritdoc/>
+		public override sealed void WriteNamedObject(CbWriter writer, CbFieldName name, object? value) => WriteNamed(writer, name, (T)value!);
+
 		/// <summary>
 		/// Reads an object from a field
 		/// </summary>
 		/// <param name="field"></param>
 		/// <returns></returns>
-		T Read(CbField field);
+		public abstract T Read(CbField field);
 
 		/// <summary>
 		/// Writes an object to a field
 		/// </summary>
 		/// <param name="writer"></param>
 		/// <param name="value"></param>
-		void Write(CbWriter writer, T value);
+		public abstract void Write(CbWriter writer, T value);
 
 		/// <summary>
 		/// Writes an object to a named field, if not equal to the default value
@@ -83,7 +92,7 @@ namespace EpicGames.Serialization
 		/// <param name="writer"></param>
 		/// <param name="name"></param>
 		/// <param name="value"></param>
-		void WriteNamed(CbWriter writer, CbFieldName name, T value);
+		public abstract void WriteNamed(CbWriter writer, CbFieldName name, T value);
 	}
 
 	/// <summary>
@@ -112,7 +121,7 @@ namespace EpicGames.Serialization
 	/// </summary>
 	public static class CbConverterMethods
 	{
-		class CbConverterMethodsWrapper<TConverter, T> : ICbConverterMethods where TConverter : class, ICbConverter<T>
+		class CbConverterMethodsWrapper<TConverter, T> : ICbConverterMethods where TConverter : CbConverter<T>
 		{
 			static TConverter s_staticConverter = null!;
 
@@ -138,7 +147,7 @@ namespace EpicGames.Serialization
 		static readonly Dictionary<PropertyInfo, ICbConverterMethods> s_propertyToMethods = new Dictionary<PropertyInfo, ICbConverterMethods>();
 		static readonly Dictionary<Type, ICbConverterMethods> s_typeToMethods = new Dictionary<Type, ICbConverterMethods>();
 
-		static ICbConverterMethods CreateWrapper(Type type, ICbConverter converter)
+		static ICbConverterMethods CreateWrapper(Type type, CbConverter converter)
 		{
 			Type converterMethodsType = typeof(CbConverterMethodsWrapper<,>).MakeGenericType(converter.GetType(), type);
 			return (ICbConverterMethods)Activator.CreateInstance(converterMethodsType, new object[] { converter })!;
@@ -154,7 +163,7 @@ namespace EpicGames.Serialization
 			ICbConverterMethods? methods;
 			if (!s_propertyToMethods.TryGetValue(property, out methods))
 			{
-				ICbConverter converter = CbConverter.GetConverter(property);
+				CbConverter converter = CbConverter.GetConverter(property);
 				methods = (converter as ICbConverterMethods) ?? CreateWrapper(property.PropertyType, converter);
 				s_propertyToMethods.Add(property, methods);
 			}
@@ -171,7 +180,7 @@ namespace EpicGames.Serialization
 			ICbConverterMethods? methods;
 			if (!s_typeToMethods.TryGetValue(type, out methods))
 			{
-				ICbConverter converter = CbConverter.GetConverter(type);
+				CbConverter converter = CbConverter.GetConverter(type);
 				methods = (converter as ICbConverterMethods) ?? CreateWrapper(type, converter);
 				s_typeToMethods.Add(type, methods);
 			}
@@ -179,31 +188,6 @@ namespace EpicGames.Serialization
 		}
 	}
 
-	/// <summary>
-	/// Base class for converter implementations 
-	/// </summary>
-	/// <typeparam name="T"></typeparam>
-	public abstract class CbConverterBase<T> : ICbConverter<T>
-	{
-		/// <inheritdoc/>
-		public object? ReadObject(CbField field) => Read(field);
-
-		/// <inheritdoc/>
-		public void WriteObject(CbWriter writer, object? value) => Write(writer, (T)value!);
-
-		/// <inheritdoc/>
-		public void WriteNamedObject(CbWriter writer, CbFieldName name, object? value) => WriteNamed(writer, name, (T)value!);
-
-		/// <inheritdoc/>
-		public abstract T Read(CbField field);
-
-		/// <inheritdoc/>
-		public abstract void Write(CbWriter writer, T value);
-
-		/// <inheritdoc/>
-		public abstract void WriteNamed(CbWriter writer, CbFieldName name, T value);
-	}
-
 	/// <summary>
 	/// Base class for converter factories
 	/// </summary>
@@ -214,13 +198,13 @@ namespace EpicGames.Serialization
 		/// </summary>
 		/// <param name="type">The type to create a converter for</param>
 		/// <returns>The converter instance, or null if this factory does not support the given type</returns>
-		public abstract ICbConverter? CreateConverter(Type type);
+		public abstract CbConverter? CreateConverter(Type type);
 	}
 
 	/// <summary>
 	/// Utility functions for creating converters
 	/// </summary>
-	public static class CbConverter
+	public partial class CbConverter
 	{
 		/// <summary>
 		/// Object used for locking access to shared objects
@@ -230,12 +214,12 @@ namespace EpicGames.Serialization
 		/// <summary>
 		/// Cache of property to converter
 		/// </summary>
-		static readonly Dictionary<PropertyInfo, ICbConverter> s_propertyToConverter = new Dictionary<PropertyInfo, ICbConverter>();
+		static readonly Dictionary<PropertyInfo, CbConverter> s_propertyToConverter = new Dictionary<PropertyInfo, CbConverter>();
 
 		/// <summary>
 		/// Cache of type to converter
 		/// </summary>
-		public static Dictionary<Type, ICbConverter> TypeToConverter { get; } = new Dictionary<Type, ICbConverter>()
+		public static Dictionary<Type, CbConverter> TypeToConverter { get; } = new Dictionary<Type, CbConverter>()
 		{
 			[typeof(bool)] = new CbPrimitiveConverter<bool>(x => x.AsBool(), (w, v) => w.WriteBoolValue(v), (w, n, v) => w.WriteBool(n, v)),
 			[typeof(int)] = new CbPrimitiveConverter<int>(x => x.AsInt32(), (w, v) => w.WriteIntegerValue(v), (w, n, v) => w.WriteInteger(n, v)),
@@ -278,26 +262,26 @@ namespace EpicGames.Serialization
 			/// <summary>
 			/// The converter instance
 			/// </summary>
-			public static ICbConverter<T> Instance { get; } = CreateConverter();
+			public static CbConverter<T> Instance { get; } = CreateConverter();
 
 			/// <summary>
 			/// Create a typed converter
 			/// </summary>
 			/// <returns></returns>
-			static ICbConverter<T> CreateConverter()
+			static CbConverter<T> CreateConverter()
 			{
-				ICbConverter converter = GetConverter(typeof(T));
-				return (converter as ICbConverter<T>) ?? new CbConverterWrapper(converter);
+				CbConverter converter = GetConverter(typeof(T));
+				return (converter as CbConverter<T>) ?? new CbConverterWrapper(converter);
 			}
 
 			/// <summary>
 			/// Wrapper class to convert an untyped converter into a typed one
 			/// </summary>
-			class CbConverterWrapper : CbConverterBase<T>
+			class CbConverterWrapper : CbConverter<T>
 			{
-				readonly ICbConverter _inner;
+				readonly CbConverter _inner;
 
-				public CbConverterWrapper(ICbConverter inner) => _inner = inner;
+				public CbConverterWrapper(CbConverter inner) => _inner = inner;
 
 				/// <inheritdoc/>
 				public override T Read(CbField field) => (T)_inner.ReadObject(field)!;
@@ -315,7 +299,7 @@ namespace EpicGames.Serialization
 		/// </summary>
 		/// <param name="property"></param>
 		/// <returns></returns>
-		public static ICbConverter GetConverter(PropertyInfo property)
+		public static CbConverter GetConverter(PropertyInfo property)
 		{
 			CbConverterAttribute? converterAttribute = property.GetCustomAttribute<CbConverterAttribute>();
 			if (converterAttribute != null)
@@ -323,10 +307,10 @@ namespace EpicGames.Serialization
 				Type converterType = converterAttribute.ConverterType;
 				lock (s_lockObject)
 				{
-					ICbConverter? converter;
+					CbConverter? converter;
 					if (!s_propertyToConverter.TryGetValue(property, out converter))
 					{
-						converter = (ICbConverter?)Activator.CreateInstance(converterType)!;
+						converter = (CbConverter?)Activator.CreateInstance(converterType)!;
 						s_propertyToConverter.Add(property, converter);
 					}
 					return converter;
@@ -340,9 +324,9 @@ namespace EpicGames.Serialization
 		/// </summary>
 		/// <param name="type"></param>
 		/// <returns></returns>
-		public static ICbConverter GetConverter(Type type)
+		public static CbConverter GetConverter(Type type)
 		{
-			ICbConverter? converter;
+			CbConverter? converter;
 			lock (s_lockObject)
 			{
 				if (!TypeToConverter.TryGetValue(type, out converter))
@@ -355,7 +339,7 @@ namespace EpicGames.Serialization
 						{
 							converterType = converterType.MakeGenericType(type.GetGenericArguments());
 						}
-						converter = (ICbConverter?)Activator.CreateInstance(converterType)!;
+						converter = (CbConverter?)Activator.CreateInstance(converterType)!;
 					}
 					else
 					{
@@ -380,7 +364,7 @@ namespace EpicGames.Serialization
 		/// </summary>
 		/// <typeparam name="T">Type to retrieve the converter for</typeparam>
 		/// <returns></returns>
-		public static ICbConverter<T> GetConverter<T>()
+		public static CbConverter<T> GetConverter<T>()
 		{
 			try
 			{
diff --git a/Engine/Source/Programs/Shared/EpicGames.Serialization/Converters/CbArrayConverter.cs b/Engine/Source/Programs/Shared/EpicGames.Serialization/Converters/CbArrayConverter.cs
--- a/Engine/Source/Programs/Shared/EpicGames.Serialization/Converters/CbArrayConverter.cs
+++ b/Engine/Source/Programs/Shared/EpicGames.Serialization/Converters/CbArrayConverter.cs
@@ -10,7 +10,7 @@ namespace EpicGames.Serialization.Converters
 	/// Converter for array types
 	/// </summary>
 	/// <typeparam name="T"></typeparam>
-	class CbArrayConverter<T> : CbConverterBase<T[]>
+	class CbArrayConverter<T> : CbConverter<T[]>
 	{
 		/// <inheritdoc/>
 		public override T[] Read(CbField field)
@@ -71,13 +71,13 @@ namespace EpicGames.Serialization.Converters
 	class CbArrayConverterFactory : CbConverterFactory
 	{
 		/// <inheritdoc/>
-		public override ICbConverter? CreateConverter(Type type)
+		public override CbConverter? CreateConverter(Type type)
 		{
 			if (type.IsArray)
 			{
 				Type elementType = type.GetElementType()!;
 				Type converterType = typeof(CbArrayConverter<>).MakeGenericType(elementType);
-				return (ICbConverter)Activator.CreateInstance(converterType)!;
+				return (CbConverter)Activator.CreateInstance(converterType)!;
 			}
 			return null;
 		}
diff --git a/Engine/Source/Programs/Shared/EpicGames.Serialization/Converters/CbClassConverter.cs b/Engine/Source/Programs/Shared/EpicGames.Serialization/Converters/CbClassConverter.cs
--- a/Engine/Source/Programs/Shared/EpicGames.Serialization/Converters/CbClassConverter.cs
+++ b/Engine/Source/Programs/Shared/EpicGames.Serialization/Converters/CbClassConverter.cs
@@ -493,7 +493,7 @@ namespace EpicGames.Serialization.Converters
 		}
 	}
 
-	class CbClassConverter<T> : CbConverterBase<T>, ICbConverterMethods where T : class
+	class CbClassConverter<T> : CbConverter<T>, ICbConverterMethods where T : class
 	{
 		readonly CbClassConverterMethods _methods;
 
@@ -527,15 +527,15 @@ namespace EpicGames.Serialization.Converters
 
 	class CbClassConverterFactory : CbConverterFactory
 	{
-		public override ICbConverter? CreateConverter(Type type)
+		public override CbConverter? CreateConverter(Type type)
 		{
-			ICbConverter? converter = null;
+			CbConverter? converter = null;
 			if (type.IsClass)
 			{
 				Type converterType = typeof(CbClassConverter<>).MakeGenericType(type);
 				try
 				{
-					converter = (ICbConverter?)Activator.CreateInstance(converterType);
+					converter = (CbConverter?)Activator.CreateInstance(converterType);
 				}
 				catch (TargetInvocationException ex) when (ex.InnerException is not null)
 				{
diff --git a/Engine/Source/Programs/Shared/EpicGames.Serialization/Converters/CbDictionaryConverter.cs b/Engine/Source/Programs/Shared/EpicGames.Serialization/Converters/CbDictionaryConverter.cs
--- a/Engine/Source/Programs/Shared/EpicGames.Serialization/Converters/CbDictionaryConverter.cs
+++ b/Engine/Source/Programs/Shared/EpicGames.Serialization/Converters/CbDictionaryConverter.cs
@@ -9,7 +9,7 @@ namespace EpicGames.Serialization.Converters
 	/// <summary>
 	/// Converter for dictionary types
 	/// </summary>
-	class CbDictionaryConverter<TKey, TValue> : CbConverterBase<Dictionary<TKey, TValue>> where TKey : notnull
+	class CbDictionaryConverter<TKey, TValue> : CbConverter<Dictionary<TKey, TValue>> where TKey : notnull
 	{
 		/// <inheritdoc/>
 		public override Dictionary<TKey, TValue> Read(CbField field)
@@ -90,12 +90,12 @@ namespace EpicGames.Serialization.Converters
 	class CbDictionaryConverterFactory : CbConverterFactory
 	{
 		/// <inheritdoc/>
-		public override ICbConverter? CreateConverter(Type type)
+		public override CbConverter? CreateConverter(Type type)
 		{
 			if (type.IsGenericType && type.GetGenericTypeDefinition() == typeof(Dictionary<,>))
 			{
 				Type converterType = typeof(CbDictionaryConverter<,>).MakeGenericType(type.GenericTypeArguments);
-				return (ICbConverter)Activator.CreateInstance(converterType)!;
+				return (CbConverter)Activator.CreateInstance(converterType)!;
 			}
 			return null;
 		}
diff --git a/Engine/Source/Programs/Shared/EpicGames.Serialization/Converters/CbEnumConverter.cs b/Engine/Source/Programs/Shared/EpicGames.Serialization/Converters/CbEnumConverter.cs
--- a/Engine/Source/Programs/Shared/EpicGames.Serialization/Converters/CbEnumConverter.cs
+++ b/Engine/Source/Programs/Shared/EpicGames.Serialization/Converters/CbEnumConverter.cs
@@ -8,7 +8,7 @@ using System.Reflection.Emit;
 
 namespace EpicGames.Serialization.Converters
 {
-	class CbEnumConverter<T> : CbConverterBase<T>, ICbConverterMethods where T : Enum
+	class CbEnumConverter<T> : CbConverter<T>, ICbConverterMethods where T : Enum
 	{
 		readonly Func<CbField, T> _readFunc;
 		readonly Action<CbWriter, T> _writeFunc;
@@ -94,13 +94,13 @@ namespace EpicGames.Serialization.Converters
 
 	class CbEnumConverterFactory : CbConverterFactory
 	{
-		public override ICbConverter? CreateConverter(Type type)
+		public override CbConverter? CreateConverter(Type type)
 		{
-			ICbConverter? converter = null;
+			CbConverter? converter = null;
 			if (type.IsEnum)
 			{
 				Type converterType = typeof(CbEnumConverter<>).MakeGenericType(type);
-				converter = (ICbConverter?)Activator.CreateInstance(converterType);
+				converter = (CbConverter?)Activator.CreateInstance(converterType);
 			}
 			return converter;
 		}
diff --git a/Engine/Source/Programs/Shared/EpicGames.Serialization/Converters/CbFieldConverter.cs b/Engine/Source/Programs/Shared/EpicGames.Serialization/Converters/CbFieldConverter.cs
--- a/Engine/Source/Programs/Shared/EpicGames.Serialization/Converters/CbFieldConverter.cs
+++ b/Engine/Source/Programs/Shared/EpicGames.Serialization/Converters/CbFieldConverter.cs
@@ -7,7 +7,7 @@ namespace EpicGames.Serialization.Converters
 	/// <summary>
 	/// Converter for raw CbField types
 	/// </summary>
-	class CbFieldConverter : CbConverterBase<CbField>
+	class CbFieldConverter : CbConverter<CbField>
 	{
 		/// <inheritdoc/>
 		public override CbField Read(CbField field)
@@ -31,7 +31,7 @@ namespace EpicGames.Serialization.Converters
 	/// <summary>
 	/// Converter for raw CbObject types
 	/// </summary>
-	class CbObjectConverter : CbConverterBase<CbObject>
+	class CbObjectConverter : CbConverter<CbObject>
 	{
 		/// <inheritdoc/>
 		public override CbObject Read(CbField field)
diff --git a/Engine/Source/Programs/Shared/EpicGames.Serialization/Converters/CbListConverter.cs b/Engine/Source/Programs/Shared/EpicGames.Serialization/Converters/CbListConverter.cs
--- a/Engine/Source/Programs/Shared/EpicGames.Serialization/Converters/CbListConverter.cs
+++ b/Engine/Source/Programs/Shared/EpicGames.Serialization/Converters/CbListConverter.cs
@@ -10,7 +10,7 @@ namespace EpicGames.Serialization.Converters
 	/// Converter for list types
 	/// </summary>
 	/// <typeparam name="T"></typeparam>
-	class CbListConverter<T> : CbConverterBase<List<T>>
+	class CbListConverter<T> : CbConverter<List<T>>
 	{
 		/// <inheritdoc/>
 		public override List<T> Read(CbField field)
@@ -70,7 +70,7 @@ namespace EpicGames.Serialization.Converters
 	/// <summary>
 	/// Specialization for serializing string lists
 	/// </summary>
-	sealed class CbStringListConverter : CbConverterBase<List<Utf8String>>
+	sealed class CbStringListConverter : CbConverter<List<Utf8String>>
 	{
 		/// <inheritdoc/>
 		public override List<Utf8String> Read(CbField field)
@@ -115,7 +115,7 @@ namespace EpicGames.Serialization.Converters
 	class CbListConverterFactory : CbConverterFactory
 	{
 		/// <inheritdoc/>
-		public override ICbConverter? CreateConverter(Type type)
+		public override CbConverter? CreateConverter(Type type)
 		{
 			if (type.IsGenericType && type.GetGenericTypeDefinition() == typeof(List<>))
 			{
@@ -126,7 +126,7 @@ namespace EpicGames.Serialization.Converters
 				else
 				{
 					Type converterType = typeof(CbListConverter<>).MakeGenericType(type.GenericTypeArguments);
-					return (ICbConverter)Activator.CreateInstance(converterType)!;
+					return (CbConverter)Activator.CreateInstance(converterType)!;
 				}
 			}
 			return null;
diff --git a/Engine/Source/Programs/Shared/EpicGames.Serialization/Converters/CbNullableConverter.cs b/Engine/Source/Programs/Shared/EpicGames.Serialization/Converters/CbNullableConverter.cs
--- a/Engine/Source/Programs/Shared/EpicGames.Serialization/Converters/CbNullableConverter.cs
+++ b/Engine/Source/Programs/Shared/EpicGames.Serialization/Converters/CbNullableConverter.cs
@@ -9,7 +9,7 @@ namespace EpicGames.Serialization.Converters
 	/// Converter for list types
 	/// </summary>
 	/// <typeparam name="T"></typeparam>
-	class CbNullableConverter<T> : CbConverterBase<Nullable<T>> where T : struct
+	class CbNullableConverter<T> : CbConverter<Nullable<T>> where T : struct
 	{
 		/// <inheritdoc/>
 		public override T? Read(CbField field)
@@ -42,12 +42,12 @@ namespace EpicGames.Serialization.Converters
 	class CbNullableConverterFactory : CbConverterFactory
 	{
 		/// <inheritdoc/>
-		public override ICbConverter? CreateConverter(Type type)
+		public override CbConverter? CreateConverter(Type type)
 		{
 			if (type.IsGenericType && type.GetGenericTypeDefinition() == typeof(Nullable<>))
 			{
 				Type converterType = typeof(CbNullableConverter<>).MakeGenericType(type.GenericTypeArguments);
-				return (ICbConverter)Activator.CreateInstance(converterType)!;
+				return (CbConverter)Activator.CreateInstance(converterType)!;
 			}
 			return null;
 		}
diff --git a/Engine/Source/Programs/Shared/EpicGames.Serialization/Converters/CbPrimitiveConverter.cs b/Engine/Source/Programs/Shared/EpicGames.Serialization/Converters/CbPrimitiveConverter.cs
--- a/Engine/Source/Programs/Shared/EpicGames.Serialization/Converters/CbPrimitiveConverter.cs
+++ b/Engine/Source/Programs/Shared/EpicGames.Serialization/Converters/CbPrimitiveConverter.cs
@@ -7,7 +7,7 @@ using System.Reflection;
 
 namespace EpicGames.Serialization.Converters
 {
-	class CbPrimitiveConverter<T> : CbConverterBase<T>, ICbConverterMethods
+	class CbPrimitiveConverter<T> : CbConverter<T>, ICbConverterMethods
 	{
 		public MethodInfo ReadMethod { get; }
 		public Func<CbField, T> ReadFunc { get; }
diff --git a/Engine/Source/Programs/Shared/EpicGames.Serialization/Converters/CbStringConverter.cs b/Engine/Source/Programs/Shared/EpicGames.Serialization/Converters/CbStringConverter.cs
--- a/Engine/Source/Programs/Shared/EpicGames.Serialization/Converters/CbStringConverter.cs
+++ b/Engine/Source/Programs/Shared/EpicGames.Serialization/Converters/CbStringConverter.cs
@@ -6,7 +6,7 @@ using System.ComponentModel;
 
 namespace EpicGames.Serialization.Converters
 {
-	class CbStringConverter<T> : CbConverterBase<T>
+	class CbStringConverter<T> : CbConverter<T>
 	{
 		readonly TypeConverter _typeConverter;
 
@@ -33,7 +33,7 @@ namespace EpicGames.Serialization.Converters
 
 	class CbStringConverterFactory : CbConverterFactory
 	{
-		public override ICbConverter? CreateConverter(Type type)
+		public override CbConverter? CreateConverter(Type type)
 		{
 			TypeConverter? frameworkTypeConverter = TypeDescriptor.GetConverter(type);
 			if (frameworkTypeConverter == null || !frameworkTypeConverter.CanConvertFrom(typeof(string)) || !frameworkTypeConverter.CanConvertTo(typeof(string)))
@@ -42,7 +42,7 @@ namespace EpicGames.Serialization.Converters
 			}
 
 			Type converterType = typeof(CbStringConverter<>).MakeGenericType(type);
-			return (ICbConverter)Activator.CreateInstance(converterType, frameworkTypeConverter)!;
+			return (CbConverter)Activator.CreateInstance(converterType, frameworkTypeConverter)!;
 		}
 	}
 }
diff --git a/Engine/Source/Programs/UnrealCloudDDC/Jupiter.Common/Implementation/BlobId.cs b/Engine/Source/Programs/UnrealCloudDDC/Jupiter.Common/Implementation/BlobId.cs
--- a/Engine/Source/Programs/UnrealCloudDDC/Jupiter.Common/Implementation/BlobId.cs
+++ b/Engine/Source/Programs/UnrealCloudDDC/Jupiter.Common/Implementation/BlobId.cs
@@ -213,7 +213,7 @@ namespace Jupiter.Implementation
 		}
 	}
 
-	public class BlobIdentifierCbConverter : CbConverterBase<BlobId>
+	public class BlobIdentifierCbConverter : CbConverter<BlobId>
 	{
 		public override BlobId Read(CbField field) => new BlobId(field.AsHash().ToByteArray());
 
diff --git a/Engine/Source/Programs/UnrealCloudDDC/Jupiter.Common/Implementation/BucketId.cs b/Engine/Source/Programs/UnrealCloudDDC/Jupiter.Common/Implementation/BucketId.cs
--- a/Engine/Source/Programs/UnrealCloudDDC/Jupiter.Common/Implementation/BucketId.cs
+++ b/Engine/Source/Programs/UnrealCloudDDC/Jupiter.Common/Implementation/BucketId.cs
@@ -75,7 +75,7 @@ namespace EpicGames.Horde.Storage
 		public override object ConvertFrom(ITypeDescriptorContext? context, CultureInfo? culture, object? value) => new BucketId((string)value!);
 	}
 
-	sealed class BucketIdCbConverter : CbConverterBase<BucketId>
+	sealed class BucketIdCbConverter : CbConverter<BucketId>
 	{
 		public override BucketId Read(CbField field) => new BucketId(field.AsString());
 
diff --git a/Engine/Source/Programs/UnrealCloudDDC/Jupiter.Common/Implementation/ContentHash.cs b/Engine/Source/Programs/UnrealCloudDDC/Jupiter.Common/Implementation/ContentHash.cs
--- a/Engine/Source/Programs/UnrealCloudDDC/Jupiter.Common/Implementation/ContentHash.cs
+++ b/Engine/Source/Programs/UnrealCloudDDC/Jupiter.Common/Implementation/ContentHash.cs
@@ -110,7 +110,7 @@ namespace Jupiter.Implementation
 		}
 	}
 	
-	public class ContentHashCbConverter : CbConverterBase<ContentHash>
+	public class ContentHashCbConverter : CbConverter<ContentHash>
 	{
 		public override ContentHash Read(CbField field) => new ContentHash(field.AsHash().ToByteArray());
 
diff --git a/Engine/Source/Programs/UnrealCloudDDC/Jupiter.Common/Implementation/RefId.cs b/Engine/Source/Programs/UnrealCloudDDC/Jupiter.Common/Implementation/RefId.cs
--- a/Engine/Source/Programs/UnrealCloudDDC/Jupiter.Common/Implementation/RefId.cs
+++ b/Engine/Source/Programs/UnrealCloudDDC/Jupiter.Common/Implementation/RefId.cs
@@ -129,7 +129,7 @@ namespace Jupiter.Implementation
 		}
 	}
 
-	sealed class IoHashKeyCbConverter : CbConverterBase<RefId>
+	sealed class IoHashKeyCbConverter : CbConverter<RefId>
 	{
 		public override RefId Read(CbField field) => new RefId(field.AsString());
 
