Remove reference to deleted type from XML comment.

diff --git a/Engine/Source/Programs/Shared/EpicGames.Serialization/CbConverter.cs b/Engine/Source/Programs/Shared/EpicGames.Serialization/CbConverter.cs
--- a/Engine/Source/Programs/Shared/EpicGames.Serialization/CbConverter.cs
+++ b/Engine/Source/Programs/Shared/EpicGames.Serialization/CbConverter.cs
@@ -31,7 +31,7 @@ namespace EpicGames.Serialization
 	}
 
 	/// <summary>
-	/// Base class for all converters. Deriving from <see cref="ICbConverter{T}"/> is more efficient
+	/// Base class for all converters.
 	/// </summary>
 	public abstract partial class CbConverter
 	{
