Fixing issue introduced into UHT that allowed unsized enumerations to be used as blueprint types.  Will once again be limited to uint8.

diff --git a/Engine/Source/Programs/Shared/EpicGames.UHT/Parsers/UhtEnumParser.cs b/Engine/Source/Programs/Shared/EpicGames.UHT/Parsers/UhtEnumParser.cs
--- a/Engine/Source/Programs/Shared/EpicGames.UHT/Parsers/UhtEnumParser.cs
+++ b/Engine/Source/Programs/Shared/EpicGames.UHT/Parsers/UhtEnumParser.cs
@@ -78,7 +78,7 @@ namespace EpicGames.UHT.Parsers
 					{
 						ParseUnderlyingType(topScope, enumObject);
 
-						if (enumObject.UnderlyingType != UhtEnumUnderlyingType.Unspecified && enumObject.UnderlyingType != UhtEnumUnderlyingType.Uint8 && enumObject.MetaData.ContainsKey("BlueprintType"))
+						if (enumObject.UnderlyingType != UhtEnumUnderlyingType.Uint8 && enumObject.MetaData.ContainsKey("BlueprintType"))
 						{
 							topScope.TokenReader.LogError("Invalid BlueprintType enum base - currently only uint8 supported");
 						}
diff --git a/Engine/Source/Programs/Shared/EpicGames.UHT/Types/Properties/UhtByteProperty.cs b/Engine/Source/Programs/Shared/EpicGames.UHT/Types/Properties/UhtByteProperty.cs
--- a/Engine/Source/Programs/Shared/EpicGames.UHT/Types/Properties/UhtByteProperty.cs
+++ b/Engine/Source/Programs/Shared/EpicGames.UHT/Types/Properties/UhtByteProperty.cs
@@ -150,7 +150,7 @@ namespace EpicGames.UHT.Types
 
 			if (function.FunctionFlags.HasAnyFlags(EFunctionFlags.BlueprintEvent | EFunctionFlags.BlueprintCallable))
 			{
-				if (Enum != null && Enum.CppForm == UhtEnumCppForm.EnumClass && Enum.UnderlyingType != UhtEnumUnderlyingType.Uint8 && Enum.UnderlyingType != UhtEnumUnderlyingType.Unspecified)
+				if (Enum != null && Enum.CppForm == UhtEnumCppForm.EnumClass && Enum.UnderlyingType != UhtEnumUnderlyingType.Uint8)
 				{
 					this.LogError("Invalid enum param for Blueprints - currently only uint8 supported");
 				}
