Serialize all public properties to compact binary by default, in line with semantics for JsonSerializer.

diff --git a/Engine/Source/Programs/Shared/EpicGames.Serialization.Tests/SerializerTests.cs b/Engine/Source/Programs/Shared/EpicGames.Serialization.Tests/SerializerTests.cs
--- a/Engine/Source/Programs/Shared/EpicGames.Serialization.Tests/SerializerTests.cs
+++ b/Engine/Source/Programs/Shared/EpicGames.Serialization.Tests/SerializerTests.cs
@@ -13,20 +13,28 @@ namespace EpicGames.Serialization.Tests
 	[TestClass]
 	public class SerializerTests
 	{
+		enum TestEnum
+		{
+			Value1,
+			Value2,
+			Value3
+		}
+
 		class TestObject : IEquatable<TestObject>
 		{
-			[CbField]
 			public int A { get; set; }
-
-			[CbField]
 			public string B { get; set; } = String.Empty;
-
-			[CbField]
 			public IoHash C { get; set; }
-
-			[CbField]
 			public List<TestObject> Children { get; set; } = new List<TestObject>();
 
+			[CbField("D")]
+			public string PropertyWithCustomName { get; set; } = "hello";
+
+			[CbIgnore]
+			public string IgnoredProperty { get; set; } = "world";
+
+			public TestEnum Enum { get; set; } = TestEnum.Value3;
+
 			public override bool Equals(object? obj) => obj is TestObject other && Equals(other);
 
 			public bool Equals(TestObject? other) => other != null && A == other.A && B == other.B && C == other.C && Enumerable.SequenceEqual(Children, other.Children);
@@ -47,6 +55,12 @@ namespace EpicGames.Serialization.Tests
 
 			Assert.AreEqual(obj1, obj2);
 			Assert.IsTrue(mem1.Span.SequenceEqual(mem2.Span));
+
+			Assert.AreEqual(cbObj1.Find("D").AsString(), "hello");
+			Assert.IsFalse(cbObj1.Find(nameof(TestObject.PropertyWithCustomName)).HasValue());
+			Assert.IsFalse(cbObj1.Find(nameof(TestObject.IgnoredProperty)).HasValue());
+
+			Assert.AreEqual(cbObj2.Find("Enum").AsInt32(), (int)TestEnum.Value3);
 		}
 
 		[TestMethod]
diff --git a/Engine/Source/Programs/Shared/EpicGames.Serialization/CbSerializer.cs b/Engine/Source/Programs/Shared/EpicGames.Serialization/CbSerializer.cs
--- a/Engine/Source/Programs/Shared/EpicGames.Serialization/CbSerializer.cs
+++ b/Engine/Source/Programs/Shared/EpicGames.Serialization/CbSerializer.cs
@@ -42,6 +42,14 @@ namespace EpicGames.Serialization
 		}
 	}
 
+	/// <summary>
+	/// Attribute used to mark that a property should not be serialized to compact binary
+	/// </summary>
+	[AttributeUsage(AttributeTargets.Property)]
+	public sealed class CbIgnoreAttribute : Attribute
+	{
+	}
+
 	/// <summary>
 	/// Attribute used to indicate that this object is the base for a class hierarchy. Each derived class must have a [CbDiscriminator] attribute.
 	/// </summary>
diff --git a/Engine/Source/Programs/Shared/EpicGames.Serialization/Converters/CbClassConverter.cs b/Engine/Source/Programs/Shared/EpicGames.Serialization/Converters/CbClassConverter.cs
--- a/Engine/Source/Programs/Shared/EpicGames.Serialization/Converters/CbClassConverter.cs
+++ b/Engine/Source/Programs/Shared/EpicGames.Serialization/Converters/CbClassConverter.cs
@@ -463,11 +463,15 @@ namespace EpicGames.Serialization.Converters
 			List<(Utf8String, PropertyInfo)> propertyList = new List<(Utf8String, PropertyInfo)>();
 			foreach (PropertyInfo property in type.GetProperties(BindingFlags.Public | BindingFlags.NonPublic | BindingFlags.Instance))
 			{
-				CbFieldAttribute? attribute = property.GetCustomAttribute<CbFieldAttribute>();
-				if (attribute != null)
+				CbIgnoreAttribute? ignoreAttribute = property.GetCustomAttribute<CbIgnoreAttribute>();
+				if (ignoreAttribute == null)
 				{
-					Utf8String name = new Utf8String(attribute.Name ?? property.Name);
-					propertyList.Add((name, property));
+					CbFieldAttribute? attribute = property.GetCustomAttribute<CbFieldAttribute>();
+					if (attribute != null || (property.GetGetMethod()?.IsPublic ?? false))
+					{
+						Utf8String name = new Utf8String(attribute?.Name ?? property.Name);
+						propertyList.Add((name, property));
+					}
 				}
 			}
 			if (propertyList.Count == 0 && type.GetCustomAttribute<CbObjectAttribute>() == null)
