Horde: Move NeutralResourcesLanguage assembly attribute into an AssemblyInfo.cs file.

diff --git a/Engine/Source/Programs/Shared/EpicGames.Horde.Tests/LogTests.cs b/Engine/Source/Programs/Shared/EpicGames.Horde.Tests/LogTests.cs
--- a/Engine/Source/Programs/Shared/EpicGames.Horde.Tests/LogTests.cs
+++ b/Engine/Source/Programs/Shared/EpicGames.Horde.Tests/LogTests.cs
@@ -4,7 +4,6 @@ using System;
 using System.Collections.Generic;
 using System.Linq;
 using System.Text;
-using System.Resources;
 using System.Threading;
 using System.Threading.Tasks;
 using EpicGames.Horde.Logs;
@@ -14,8 +13,6 @@ using EpicGames.Horde.Tests.Properties;
 using Microsoft.Extensions.Logging.Abstractions;
 using Microsoft.VisualStudio.TestTools.UnitTesting;
 
-[assembly: NeutralResourcesLanguageAttribute("en-US")]
-
 namespace EpicGames.Horde.Tests
 {
 	[TestClass]
diff --git a/Engine/Source/Programs/Shared/EpicGames.Horde.Tests/Properties/AssemblyInfo.cs b/Engine/Source/Programs/Shared/EpicGames.Horde.Tests/Properties/AssemblyInfo.cs
new file mode 100644
--- /dev/null
+++ b/Engine/Source/Programs/Shared/EpicGames.Horde.Tests/Properties/AssemblyInfo.cs
@@ -0,0 +1,5 @@
+// Copyright Epic Games, Inc. All Rights Reserved.
+
+using System.Resources;
+
+[assembly: NeutralResourcesLanguage("en-US")]
