[OidcToken]
* Added system wide lock around TokenStore to prevent multiple OidcToken.exe from reading/writing the same file at the same time.

diff --git a/Engine/Source/Programs/Shared/EpicGames.OIDC/WindowsTokenStore.cs b/Engine/Source/Programs/Shared/EpicGames.OIDC/WindowsTokenStore.cs
--- a/Engine/Source/Programs/Shared/EpicGames.OIDC/WindowsTokenStore.cs
+++ b/Engine/Source/Programs/Shared/EpicGames.OIDC/WindowsTokenStore.cs
@@ -8,6 +8,7 @@ using System.Runtime.InteropServices;
 using System.Text;
 using System.Text.Json;
 using System.Text.Json.Serialization;
+using System.Threading;
 
 #pragma warning disable CS1591 // Missing XML documentation on public types
 
@@ -36,8 +37,20 @@ namespace EpicGames.OIDC
 
 		private readonly Dictionary<string, byte[]> _providerToRefreshToken = new Dictionary<string, byte[]>();
 
+		private Mutex _mutex;
+
 		public WindowsTokenStore()
 		{
+			_mutex = new Mutex(false, "oidcTokenStoreDat");
+
+			try
+			{
+				_mutex.WaitOne();
+			} catch (AbandonedMutexException)
+			{
+
+			}
+
 			ReadStoreFromDisk();
 		}
 
