Horde: Allow customizing blob serializers to use in different contexts, allowing us to lock serialization formats when talking to remote agents of a particular version.

To support this, the Node base class and serialization methods have been removed, and serialization is performed through a BlobConverter<T> instance created and cached according to a BlobSerializationOptions object. Custom instances of serializers with different configurations may be registered with this object, allowing different code paths to use different serialization settings.

diff --git a/Engine/Source/Programs/AutomationTool/BuildGraph/Tasks/CreateArtifactTask.cs b/Engine/Source/Programs/AutomationTool/BuildGraph/Tasks/CreateArtifactTask.cs
--- a/Engine/Source/Programs/AutomationTool/BuildGraph/Tasks/CreateArtifactTask.cs
+++ b/Engine/Source/Programs/AutomationTool/BuildGraph/Tasks/CreateArtifactTask.cs
@@ -125,10 +125,10 @@ namespace AutomationTool.Tasks
 					int totalCount = files.Count;
 					long totalSize = files.Sum(x => x.Length);
 
-					HashedNodeRef<DirectoryNode> outputNodeRef = await writer.WriteFilesAsync(baseDir.ToDirectoryInfo(), files, progress: new CopyStatsLoggerWithTotals(totalCount, totalSize, Logger));
+					IBlobHandle<DirectoryNode> outputNodeRef = await writer.WriteFilesAsync(baseDir.ToDirectoryInfo(), files, progress: new CopyStatsLoggerWithTotals(totalCount, totalSize, Logger));
 					await writer.FlushAsync();
 
-					await client.WriteRefAsync(response.RefName, outputNodeRef.Handle);
+					await client.WriteRefAsync(response.RefName, outputNodeRef);
 				}
 			}
 
diff --git a/Engine/Source/Programs/AutomationTool/BuildGraph/Tasks/DeployToolTask.cs b/Engine/Source/Programs/AutomationTool/BuildGraph/Tasks/DeployToolTask.cs
--- a/Engine/Source/Programs/AutomationTool/BuildGraph/Tasks/DeployToolTask.cs
+++ b/Engine/Source/Programs/AutomationTool/BuildGraph/Tasks/DeployToolTask.cs
@@ -174,7 +174,8 @@ namespace AutomationTool.Tasks
 				{
 					throw new AutomationException("Either File=... or Directory=... must be specified");
 				}
-				handle = await treeWriter.FlushAsync(sandbox);
+				handle = await treeWriter.WriteBlobAsync(sandbox);
+				await treeWriter.FlushAsync();
 			}
 
 			CreateDeploymentRequest request = new CreateDeploymentRequest();
diff --git a/Engine/Source/Programs/Horde/Horde.Agent.Tests/JsonRpcLoggerTests.cs b/Engine/Source/Programs/Horde/Horde.Agent.Tests/JsonRpcLoggerTests.cs
--- a/Engine/Source/Programs/Horde/Horde.Agent.Tests/JsonRpcLoggerTests.cs
+++ b/Engine/Source/Programs/Horde/Horde.Agent.Tests/JsonRpcLoggerTests.cs
@@ -157,16 +157,16 @@ namespace Horde.Agent.Tests
 						logger.LogInformation("Testing {Number}", idx);
 					}
 				}
-				file = await sink.Target!.ReadNodeAsync<LogNode>();
+				file = await sink.Target!.ReadBlobAsync<LogNode>();
 			}
 
 			// Check the index text
 			List<Utf8String> extractedIndexText = new List<Utf8String>();
 
-			LogIndexNode index = await file.IndexRef.ExpandAsync();
+			LogIndexNode index = await file.IndexRef.ReadBlobAsync();
 			foreach (LogChunkRef block in index.PlainTextChunkRefs)
 			{
-				LogChunkNode text = await block.ExpandAsync();
+				LogChunkNode text = await block.Target.ReadBlobAsync();
 				extractedIndexText.AddRange(text.Lines);
 			}
 
@@ -181,7 +181,7 @@ namespace Horde.Agent.Tests
 
 			foreach (LogChunkRef blockRef in file.TextChunkRefs)
 			{
-				LogChunkNode blockText = await blockRef.ExpandAsync();
+				LogChunkNode blockText = await blockRef.Target.ReadBlobAsync();
 				foreach (Utf8String line in blockText.Lines)
 				{
 					LogEvent logEvent = LogEvent.Read(line.Span);
diff --git a/Engine/Source/Programs/Horde/Horde.Agent/Execution/JobExecutor.cs b/Engine/Source/Programs/Horde/Horde.Agent/Execution/JobExecutor.cs
--- a/Engine/Source/Programs/Horde/Horde.Agent/Execution/JobExecutor.cs
+++ b/Engine/Source/Programs/Horde/Horde.Agent/Execution/JobExecutor.cs
@@ -518,19 +518,19 @@ namespace Horde.Agent.Execution
 
 					Stopwatch timer = Stopwatch.StartNew();
 
-					NodeRef<DirectoryNode> rootNodeRef;
+					IBlobHandle<DirectoryNode> rootNodeRef;
 					await using (IStorageWriter treeWriter = CreateStorageWriter(storage, artifact.RefName, logger))
 					{
 						DirectoryNode buildGraphNode = new DirectoryNode();
 						await buildGraphNode.AddFilesAsync(workspaceDir, buildGraphFiles, treeWriter, cancellationToken: cancellationToken);
-						HashedNodeRef<DirectoryNode> outputNodeRef = await treeWriter.WriteHashedNodeAsync(buildGraphNode, cancellationToken);
+						IBlobHandle<DirectoryNode> outputNodeRef = await treeWriter.WriteBlobAsync(buildGraphNode, cancellationToken);
 
 						DirectoryNode rootNode = new DirectoryNode();
 						rootNode.AddDirectory(new DirectoryEntry(BuildGraphTempStorageDir, buildGraphNode.Length, outputNodeRef));
 
-						rootNodeRef = await treeWriter.WriteNodeAsync(rootNode, cancellationToken);
+						rootNodeRef = await treeWriter.WriteBlobAsync(rootNode, cancellationToken);
 					}
-					await storage.WriteRefTargetAsync(artifact.RefName, rootNodeRef.Handle, new RefOptions(), cancellationToken);
+					await storage.WriteRefTargetAsync(artifact.RefName, rootNodeRef, new RefOptions(), cancellationToken);
 
 					logger.LogInformation("Upload took {Time:n1}s", timer.Elapsed.TotalSeconds);
 				}
@@ -769,10 +769,10 @@ namespace Horde.Agent.Execution
 					using IStorageClient storage = CreateStorageClient(namespaceId, artifact.Token);
 
 					DirectoryNode node = await storage.ReadRefAsync<DirectoryNode>(refName, cancellationToken: cancellationToken);
-					DirectoryNode? buildGraphDir = await node.TryOpenDirectoryAsync(BuildGraphTempStorageDir, cancellationToken);
+					DirectoryNode? buildGraphDir = await node.TryOpenDirectoryAsync(BuildGraphTempStorageDir, cancellationToken: cancellationToken);
 					if (buildGraphDir != null)
 					{
-						await buildGraphDir.CopyToDirectoryAsync(new DirectoryInfo(workspaceDir.FullName), logger, cancellationToken);
+						await buildGraphDir.CopyToDirectoryAsync(new DirectoryInfo(workspaceDir.FullName), null, logger, cancellationToken);
 						logger.LogInformation("Copying preprocessed script from {BuildGraphFolderName} into {OutputDir}", BuildGraphTempStorageDir, workspaceDir);
 					}
 					else
@@ -848,14 +848,14 @@ namespace Horde.Agent.Execution
 
 				using IStorageClient storage = CreateStorageClient(new NamespaceId(artifact.NamespaceId), artifact.Token);
 
-				NodeRef<DirectoryNode> rootRef;
+				IBlobHandle<DirectoryNode> rootRef;
 				await using (IStorageWriter writer = CreateStorageWriter(storage, new RefName(artifact.RefName), logger))
 				{
 					try
 					{
 						DirectoryNode dir = new DirectoryNode();
 						await dir.AddFilesAsync(baseDir, files, writer, progress: new CopyStatsLogger(logger), cancellationToken: cancellationToken);
-						rootRef = await writer.WriteNodeAsync(dir, cancellationToken);
+						rootRef = await writer.WriteBlobAsync(dir, cancellationToken: cancellationToken);
 					}
 					catch (Exception ex)
 					{
@@ -863,7 +863,7 @@ namespace Horde.Agent.Execution
 						throw;
 					}
 				}
-				await storage.WriteRefTargetAsync(new RefName(artifact.RefName), rootRef.Handle, cancellationToken: cancellationToken);
+				await storage.WriteRefTargetAsync(new RefName(artifact.RefName), rootRef, cancellationToken: cancellationToken);
 
 				Logger.LogInformation("Uploaded artifact {ArtifactId}", artifact.Id);
 			}
@@ -1051,7 +1051,7 @@ namespace Horde.Agent.Execution
 				// Upload the data
 				Stopwatch timer = Stopwatch.StartNew();
 
-				NodeRef<DirectoryNode> outputNodeRef;
+				IBlobHandle<DirectoryNode> outputNodeRef;
 				await using (IStorageWriter treeWriter = CreateStorageWriter(storage, artifact.RefName, logger))
 				{
 					DirectoryNode outputNode = new DirectoryNode();
@@ -1088,11 +1088,11 @@ namespace Horde.Agent.Execution
 						outputNode.AddFile(await TempStorage.ArchiveTagAsync(manifestDir, step.Name, outputName, workspaceDir, files, storageBlocks.ToArray(), treeWriter, logger, cancellationToken));
 					}
 
-					outputNodeRef = await treeWriter.WriteNodeAsync(outputNode, cancellationToken);
+					outputNodeRef = await treeWriter.WriteBlobAsync(outputNode, cancellationToken: cancellationToken);
 				}
 
 				// Write the final node
-				await storage.WriteRefTargetAsync(artifact.RefName, outputNodeRef.Handle, new RefOptions(), cancellationToken: cancellationToken);
+				await storage.WriteRefTargetAsync(artifact.RefName, outputNodeRef, new RefOptions(), cancellationToken: cancellationToken);
 				logger.LogInformation("Upload took {Time:n1}s", timer.Elapsed.TotalSeconds);
 			}
 
diff --git a/Engine/Source/Programs/Horde/Horde.Agent/Utility/JsonRpcLogSink.cs b/Engine/Source/Programs/Horde/Horde.Agent/Utility/JsonRpcLogSink.cs
--- a/Engine/Source/Programs/Horde/Horde.Agent/Utility/JsonRpcLogSink.cs
+++ b/Engine/Source/Programs/Horde/Horde.Agent/Utility/JsonRpcLogSink.cs
@@ -239,8 +239,8 @@ namespace Horde.Agent.Utility
 
 			if (request.Flush || _bufferLength > FlushLength)
 			{
-				HashedNodeRef<LogNode> target = await _builder.FlushAsync(_writer, request.Flush, cancellationToken);
-				await UpdateLogAsync(target.Handle, _builder.LineCount, request.Flush, cancellationToken);
+				IBlobHandle<LogNode> target = await _builder.FlushAsync(_writer, request.Flush, null, cancellationToken);
+				await UpdateLogAsync(target, _builder.LineCount, request.Flush, cancellationToken);
 				_bufferLength = 0;
 			}
 
diff --git a/Engine/Source/Programs/Horde/Horde.Agent/Utility/TempStorage.cs b/Engine/Source/Programs/Horde/Horde.Agent/Utility/TempStorage.cs
--- a/Engine/Source/Programs/Horde/Horde.Agent/Utility/TempStorage.cs
+++ b/Engine/Source/Programs/Horde/Horde.Agent/Utility/TempStorage.cs
@@ -701,7 +701,7 @@ namespace Horde.Storage.Utility
 			DirectoryNode rootNode = new DirectoryNode();
 			await rootNode.AddFilesAsync(workspaceDir, archiveFiles, writer, progress: new CopyStatsLogger(logger), cancellationToken: cancellationToken);
 
-			HashedNodeRef<DirectoryNode> rootNodeRef = await writer.WriteHashedNodeAsync(rootNode, cancellationToken);
+			IBlobHandle<DirectoryNode> rootNodeRef = await writer.WriteBlobAsync(rootNode, cancellationToken: cancellationToken);
 			return new DirectoryEntry(blockDirectoryName, rootNode.Length, rootNodeRef);
 		}
 
@@ -765,8 +765,8 @@ namespace Horde.Storage.Utility
 				Stopwatch timer = Stopwatch.StartNew();
 
 				// Add all the files and flush the ref
-				DirectoryNode rootDirNode = await rootDirEntry.ExpandAsync(cancellationToken);
-				await rootDirNode.CopyToDirectoryAsync(rootDir.ToDirectoryInfo(), new CopyStatsLogger(logger), logger, cancellationToken);
+				DirectoryNode rootDirNode = await rootDirEntry.Handle.ReadBlobAsync(cancellationToken: cancellationToken);
+				await rootDirNode.CopyToDirectoryAsync(rootDir.ToDirectoryInfo(), new CopyStatsLogger(logger), null, logger, cancellationToken);
 
 				StorageStats deltaStats = StorageStats.GetDelta(initialStats, storageClient.GetStats());
 				logger.LogInformation("{Stats}", $"Elapsed: {(int)timer.Elapsed.TotalSeconds}s, {String.Join(", ", deltaStats.Values.Select(x => $"{x.Item1}: {x.Item2:n0}"))}");
diff --git a/Engine/Source/Programs/Horde/Horde.Server.Tests/BlobStoreTests.cs b/Engine/Source/Programs/Horde/Horde.Server.Tests/BlobStoreTests.cs
--- a/Engine/Source/Programs/Horde/Horde.Server.Tests/BlobStoreTests.cs
+++ b/Engine/Source/Programs/Horde/Horde.Server.Tests/BlobStoreTests.cs
@@ -65,7 +65,7 @@ namespace Horde.Server.Tests
 			await using IStorageWriter writer = store.CreateWriter();
 			blob.Data.CopyTo(writer.GetOutputBuffer(0, blob.Data.Length));
 
-			IBlobHandle handle = await writer.WriteBlobAsync(s_blobType, blob.Data.Length, blob.References.ConvertAll(x => store.CreateBlobHandle(x)));
+			IBlobHandle handle = await writer.WriteBlobAsync(s_blobType, blob.Data.Length, blob.References.ConvertAll(x => store.CreateBlobHandle(x)), Array.Empty<AliasInfo>());
 			await handle.FlushAsync();
 
 			return handle.GetLocator();
diff --git a/Engine/Source/Programs/Horde/Horde.Server.Tests/GcServiceTests.cs b/Engine/Source/Programs/Horde/Horde.Server.Tests/GcServiceTests.cs
--- a/Engine/Source/Programs/Horde/Horde.Server.Tests/GcServiceTests.cs
+++ b/Engine/Source/Programs/Horde/Horde.Server.Tests/GcServiceTests.cs
@@ -70,8 +70,8 @@ namespace Horde.Server.Tests
 				if (nodes.Add(root))
 				{
 					IBlobHandle handle = store.CreateBlobHandle(root);
-					IReadOnlyList<IBlobHandle> refs = await handle.ReadImportsAsync();
-					await FindNodesAsync(store, refs.ConvertAll(x => x.GetLocator()), nodes);
+					BlobData data = await handle.ReadBlobDataAsync();
+					await FindNodesAsync(store, data.Refs.ConvertAll(x => x.GetLocator()), nodes);
 				}
 			}
 		}
@@ -109,7 +109,7 @@ namespace Horde.Server.Tests
 				await using (IStorageWriter writer = store.CreateWriter("gctest"))
 				{
 					List<IBlobHandle> imports = children[idx].ConvertAll(x => store.CreateBlobHandle(locators[x]));
-					handle = await writer.WriteBlobAsync(blobType, 0, imports);
+					handle = await writer.WriteBlobAsync(blobType, 0, imports, Array.Empty<AliasInfo>());
 				}
 				locators[idx] = handle.GetLocator();
 			}
diff --git a/Engine/Source/Programs/Horde/Horde.Server/Artifacts/ArtifactsController.cs b/Engine/Source/Programs/Horde/Horde.Server/Artifacts/ArtifactsController.cs
--- a/Engine/Source/Programs/Horde/Horde.Server/Artifacts/ArtifactsController.cs
+++ b/Engine/Source/Programs/Horde/Horde.Server/Artifacts/ArtifactsController.cs
@@ -276,7 +276,7 @@ namespace Horde.Server.Artifacts
 			DirectoryNode directoryNode;
 			try
 			{
-				directoryNode = await storageClient.ReadRefAsync<DirectoryNode>(artifact.RefName, DateTime.UtcNow.AddHours(1.0), cancellationToken);
+				directoryNode = await storageClient.ReadRefAsync<DirectoryNode>(artifact.RefName, DateTime.UtcNow.AddHours(1.0), cancellationToken: cancellationToken);
 			}
 			catch (RefNameNotFoundException)
 			{
@@ -292,7 +292,7 @@ namespace Horde.Server.Artifacts
 					{
 						return NotFound();
 					}
-					directoryNode = await nextDirectoryEntry.ExpandAsync(cancellationToken);
+					directoryNode = await nextDirectoryEntry.Handle.ReadBlobAsync(cancellationToken: cancellationToken);
 				}
 			}
 
@@ -308,9 +308,9 @@ namespace Horde.Server.Artifacts
 				response.Directories = new List<GetArtifactDirectoryEntryResponse>();
 				foreach (DirectoryEntry subDirectoryEntry in directoryNode.Directories)
 				{
-					DirectoryNode subDirectoryNode = await subDirectoryEntry.ExpandAsync(cancellationToken);
+					DirectoryNode subDirectoryNode = await subDirectoryEntry.Handle.ReadBlobAsync(cancellationToken: cancellationToken);
 
-					GetArtifactDirectoryEntryResponse subDirectoryEntryResponse = new GetArtifactDirectoryEntryResponse(subDirectoryEntry.Name.ToString(), subDirectoryEntry.Length, subDirectoryEntry.Hash);
+					GetArtifactDirectoryEntryResponse subDirectoryEntryResponse = new GetArtifactDirectoryEntryResponse(subDirectoryEntry.Name.ToString(), subDirectoryEntry.Length, subDirectoryEntry.Handle.Hash);
 					if (depth == 0)
 					{
 						if (subDirectoryNode.Directories.Count + subDirectoryNode.Files.Count < 16)
@@ -379,9 +379,9 @@ namespace Horde.Server.Artifacts
 			}
 
 			using IStorageClient storageClient = _storageClientFactory.CreateClient(artifact.NamespaceId);
-			DirectoryNode directory = await storageClient.ReadRefAsync<DirectoryNode>(artifact.RefName, DateTime.UtcNow.AddHours(1.0), cancellationToken);
+			DirectoryNode directory = await storageClient.ReadRefAsync<DirectoryNode>(artifact.RefName, DateTime.UtcNow.AddHours(1.0), cancellationToken: cancellationToken);
 
-			FileEntry? fileEntry = await directory.GetFileEntryByPathAsync(path, cancellationToken);
+			FileEntry? fileEntry = await directory.GetFileEntryByPathAsync(path, cancellationToken: cancellationToken);
 			if (fileEntry == null)
 			{
 				return NotFound($"Unable to find file {path}");
@@ -454,7 +454,7 @@ namespace Horde.Server.Artifacts
 			IStorageClient storageClient = _storageClientFactory.CreateClient(artifact.NamespaceId);
 			try
 			{
-				DirectoryNode directory = await storageClient.ReadRefAsync<DirectoryNode>(artifact.RefName, DateTime.UtcNow.AddHours(1.0), cancellationToken);
+				DirectoryNode directory = await storageClient.ReadRefAsync<DirectoryNode>(artifact.RefName, DateTime.UtcNow.AddHours(1.0), cancellationToken: cancellationToken);
 
 				Stream stream = directory.AsZipStream(filter).WrapOwnership(storageClient);
 				return new FileStreamResult(stream, "application/zip") { FileDownloadName = $"{artifact.RefName}.zip" };
diff --git a/Engine/Source/Programs/Horde/Horde.Server/Commands/Generate/ToolDataCommand.cs b/Engine/Source/Programs/Horde/Horde.Server/Commands/Generate/ToolDataCommand.cs
--- a/Engine/Source/Programs/Horde/Horde.Server/Commands/Generate/ToolDataCommand.cs
+++ b/Engine/Source/Programs/Horde/Horde.Server/Commands/Generate/ToolDataCommand.cs
@@ -53,14 +53,14 @@ namespace Horde.Server.Commands.Generate
 			{
 				using (IStorageClient client = BundleStorageClient.CreateFromDirectory(bundleDir, bundleCache, logger))
 				{
-					NodeRef<DirectoryNode> dirNodeRef;
+					IBlobHandle<DirectoryNode> dirNodeRef;
 					await using (IStorageWriter writer = client.CreateWriter(refName))
 					{
 						DirectoryNode dirNode = new DirectoryNode();
 						await dirNode.AddFilesAsync(InputDir.ToDirectoryInfo(), writer);
-						dirNodeRef = await writer.WriteNodeAsync(dirNode);
+						dirNodeRef = await writer.WriteBlobAsync(dirNode);
 					}
-					await client.WriteRefAsync(refName, dirNodeRef.Handle);
+					await client.WriteRefAsync(refName, dirNodeRef);
 				}
 			}
 
diff --git a/Engine/Source/Programs/Horde/Horde.Server/Ddc/BlobService.cs b/Engine/Source/Programs/Horde/Horde.Server/Ddc/BlobService.cs
--- a/Engine/Source/Programs/Horde/Horde.Server/Ddc/BlobService.cs
+++ b/Engine/Source/Programs/Horde/Horde.Server/Ddc/BlobService.cs
@@ -103,7 +103,7 @@ namespace Horde.Server.Ddc
 				throw new BlobNotFoundException(ns, blob);
 			}
 
-			using BlobData data = await alias.Target.ReadAsync(cancellationToken);
+			using BlobData data = await alias.Target.ReadBlobDataAsync(cancellationToken);
 			return new BlobContents(data.Data.ToArray());
 		}
 
@@ -152,7 +152,7 @@ namespace Horde.Server.Ddc
 				using Stream stream = content.GetStream();
 				await stream.ReadFixedLengthBytesAsync(memory, cancellationToken);
 
-				blobHandle = await writer.WriteBlobAsync(s_rawBlobType, memory.Length, Array.Empty<IBlobHandle>(), cancellationToken);
+				blobHandle = await writer.WriteBlobAsync(s_rawBlobType, memory.Length, Array.Empty<IBlobHandle>(), Array.Empty<AliasInfo>(), cancellationToken);
 				await writer.FlushAsync(cancellationToken);
 			}
 
diff --git a/Engine/Source/Programs/Horde/Horde.Server/Ddc/RefService.cs b/Engine/Source/Programs/Horde/Horde.Server/Ddc/RefService.cs
--- a/Engine/Source/Programs/Horde/Horde.Server/Ddc/RefService.cs
+++ b/Engine/Source/Programs/Horde/Horde.Server/Ddc/RefService.cs
@@ -49,7 +49,7 @@ namespace Horde.Server.Ddc
 			}
 
 			IBlobHandle blobHandle = blobAlias.Target;
-			using BlobData blobContents = await blobHandle.ReadAsync(cancellationToken);
+			using BlobData blobContents = await blobHandle.ReadBlobDataAsync(cancellationToken);
 			CbObject payload = new CbObject(blobContents.Data);
 
 			BlobId[] referencedBlobs = Array.Empty<BlobId>();
@@ -78,7 +78,7 @@ namespace Horde.Server.Ddc
 				// TODO: We resolved all these blobs above... Need to just have GetReferencedBlobs just return the appropriate handles directly.
 				RefName refName = GetRefName(bucket, key);
 
-				NodeRef<DdcRefNode> refNodeRef;
+				IBlobHandle<DdcRefNode> refNodeRef;
 				await using (IStorageWriter writer = storageClient.CreateWriter(refName))
 				{
 					DdcRefNode refNode = new DdcRefNode(blobHash.Hash);
@@ -88,7 +88,7 @@ namespace Horde.Server.Ddc
 						BlobAlias? alias = await storageClient.FindAliasAsync(BlobService.GetAlias(referencedBlob), cancellationToken);
 						refNode.References.Add((referencedBlob.Hash, alias!.Target));
 					}
-					refNodeRef = await writer.WriteNodeAsync(refNode, cancellationToken);
+					refNodeRef = await writer.WriteBlobAsync(refNode, cancellationToken: cancellationToken);
 				}
 
 				await storageClient.WriteRefTargetAsync(refName, refNodeRef, cancellationToken: cancellationToken);
@@ -138,7 +138,7 @@ namespace Horde.Server.Ddc
 				throw new RefNotFoundException(ns, bucket, key);
 			}
 
-			BlobData data = await node.References.First(x => x.Hash == node.RootHash).Handle.ReadAsync(cancellationToken);
+			BlobData data = await node.References.First(x => x.Hash == node.RootHash).Handle.ReadBlobDataAsync(cancellationToken);
 			BlobContents contents = new BlobContents(data.Data.ToArray());
 
 			RefRecord record = new RefRecord(ns, bucket, key, DateTime.UtcNow, null, new BlobId(node.RootHash), true);
diff --git a/Engine/Source/Programs/Horde/Horde.Server/Logs/LogFileService.cs b/Engine/Source/Programs/Horde/Horde.Server/Logs/LogFileService.cs
--- a/Engine/Source/Programs/Horde/Horde.Server/Logs/LogFileService.cs
+++ b/Engine/Source/Programs/Horde/Horde.Server/Logs/LogFileService.cs
@@ -519,7 +519,7 @@ namespace Horde.Server.Logs
 
 						// Get the chunk data
 						LogChunkRef chunk = _rootNode.TextChunkRefs[_chunkIdx];
-						LogChunkNode chunkNode = await chunk.ExpandAsync(cancellationToken);
+						LogChunkNode chunkNode = await chunk.Target.ReadBlobAsync(cancellationToken: cancellationToken);
 
 						// Get the source data
 						_sourceBuffer = chunkNode.Data;
@@ -625,7 +625,7 @@ namespace Horde.Server.Logs
 					for (; index < maxIndex && chunkIdx < root.TextChunkRefs.Count; chunkIdx++)
 					{
 						LogChunkRef chunk = root.TextChunkRefs[chunkIdx];
-						LogChunkNode chunkData = await chunk.ExpandAsync(cancellationToken);
+						LogChunkNode chunkData = await chunk.Target.ReadBlobAsync(cancellationToken: cancellationToken);
 
 						for (; index < maxIndex && index < chunk.LineIndex; index++)
 						{
@@ -1025,7 +1025,7 @@ namespace Horde.Server.Logs
 						long lastChunkLength = lastChunk.Length;
 						if (lastChunkLength <= 0)
 						{
-							LogChunkNode lastChunkNode = await lastChunk.ExpandAsync(cancellationToken);
+							LogChunkNode lastChunkNode = await lastChunk.Target.ReadBlobAsync(cancellationToken: cancellationToken);
 							lastChunkLength = lastChunkNode.Length;
 						}
 						length = Math.Min(length, (lastChunk.Offset + lastChunkLength) - offset);
@@ -1161,7 +1161,7 @@ namespace Horde.Server.Logs
 
 				int chunkIdx = root.TextChunkRefs.GetChunkForLine(lineIdx);
 				LogChunkRef chunk = root.TextChunkRefs[chunkIdx];
-				LogChunkNode chunkData = await chunk.ExpandAsync(cancellationToken);
+				LogChunkNode chunkData = await chunk.Target.ReadBlobAsync(cancellationToken: cancellationToken);
 
 				if (lineIdx < chunk.LineIndex)
 				{
@@ -1820,8 +1820,8 @@ namespace Horde.Server.Logs
 				LogNode? root = await storageClient.ReadRefAsync<LogNode>(logFile.RefName, cancellationToken: cancellationToken);
 				if(root != null)
 				{
-					LogIndexNode index = await root.IndexRef.ExpandAsync(cancellationToken);
-					await foreach (int lineIdx in index.SearchAsync(firstLine, searchText, searchStats, cancellationToken))
+					LogIndexNode index = await root.IndexRef.ReadBlobAsync(cancellationToken: cancellationToken);
+					await foreach (int lineIdx in index.SearchAsync(firstLine, searchText, searchStats, cancellationToken: cancellationToken))
 					{
 						yield return lineIdx;
 					}
diff --git a/Engine/Source/Programs/Horde/Horde.Server/Storage/StorageController.cs b/Engine/Source/Programs/Horde/Horde.Server/Storage/StorageController.cs
--- a/Engine/Source/Programs/Horde/Horde.Server/Storage/StorageController.cs
+++ b/Engine/Source/Programs/Horde/Horde.Server/Storage/StorageController.cs
@@ -367,7 +367,7 @@ namespace Horde.Server.Storage
 				return false;
 			}
 		}
-
+		/*
 		/// <summary>
 		/// Gets information about a particular bundle in storage
 		/// </summary>
@@ -377,7 +377,7 @@ namespace Horde.Server.Storage
 		/// <returns></returns>
 		[HttpGet]
 		[Route("/api/v1/storage/{namespaceId}/nodes/{*locator}")]
-		public async Task<ActionResult<object>> GetNodeAsync(NamespaceId namespaceId, BlobLocator locator, CancellationToken cancellationToken = default)
+		public Task<ActionResult<object>> GetNodeAsync(NamespaceId namespaceId, BlobLocator locator, CancellationToken cancellationToken = default)
 		{
 			NamespaceConfig? namespaceConfig;
 			if (!_globalConfig.Value.Storage.TryGetNamespace(namespaceId, out namespaceConfig))
@@ -395,7 +395,7 @@ namespace Horde.Server.Storage
 
 			object content;
 
-			using BlobData blobData = await storageClient.CreateBlobHandle(locator).ReadAsync(cancellationToken);
+			using BlobData blobData = await storageClient.CreateBlobHandle(locator).ReadDataAsync(cancellationToken);
 
 			Node node = Node.Deserialize(blobData);
 			switch (node)
@@ -424,7 +424,7 @@ namespace Horde.Server.Storage
 
 			return new { type = blobData.Type.Guid, @class = node.GetType().Name, content = content };
 		}
-
+		*/
 		static string GetNodeLink(string linkBase, IBlobHandle handle) => GetNodeLink(linkBase, handle.GetLocator());
 		
 		static string GetNodeLink(string linkBase, BlobLocator locator) => $"{linkBase}/nodes/{locator}";
diff --git a/Engine/Source/Programs/Horde/Horde.Server/Storage/StorageService.cs b/Engine/Source/Programs/Horde/Horde.Server/Storage/StorageService.cs
--- a/Engine/Source/Programs/Horde/Horde.Server/Storage/StorageService.cs
+++ b/Engine/Source/Programs/Horde/Horde.Server/Storage/StorageService.cs
@@ -91,7 +91,7 @@ namespace Horde.Server.Storage
 			public ValueTask<IReadOnlyList<IBlobHandle>> ReadImportsAsync(CancellationToken cancellationToken = default) => new ValueTask<IReadOnlyList<IBlobHandle>>(Array.Empty<IBlobHandle>());
 
 			/// <inheritdoc/>
-			public async ValueTask<BlobData> ReadAsync(CancellationToken cancellationToken = default)
+			public async ValueTask<BlobData> ReadBlobDataAsync(CancellationToken cancellationToken = default)
 			{
 				IReadOnlyMemoryOwner<byte> obj = await _backend.ReadAsync(_path, cancellationToken);
 				return new BlobDataWithOwner(BlobType.Leaf, obj.Memory, Array.Empty<IBlobHandle>(), obj);
@@ -176,7 +176,7 @@ namespace Horde.Server.Storage
 			public async ValueTask<BlobData> ReadBlobAsync(BlobLocator locator, CancellationToken cancellationToken = default)
 			{
 				IBlobHandle handle = CreateBlobHandle(locator);
-				return await handle.ReadAsync(cancellationToken);
+				return await handle.ReadBlobDataAsync(cancellationToken);
 			}
 
 			/// <inheritdoc/>
@@ -696,10 +696,11 @@ namespace Horde.Server.Storage
 						if (scopedState.Value.Namespaces.TryGetValue(blobInfo.NamespaceId, out namespaceInfo))
 						{
 							IBlobHandle handle = namespaceInfo.Client.CreateBlobHandle(blobInfo.Locator);
-							IReadOnlyList<IBlobHandle> imports = await handle.ReadImportsAsync(cancellationToken);
+
+							BlobData data = await handle.ReadBlobDataAsync(cancellationToken);
 
 							List<ObjectId> importInfoIds = new List<ObjectId>();
-							foreach (IBlobHandle import in imports)
+							foreach (IBlobHandle import in data.Refs)
 							{
 								BlobLocator importLocator = import.GetLocator();
 								string importPath = importLocator.BaseLocator.ToString();
diff --git a/Engine/Source/Programs/Horde/Horde.Server/Tools/ToolsController.cs b/Engine/Source/Programs/Horde/Horde.Server/Tools/ToolsController.cs
--- a/Engine/Source/Programs/Horde/Horde.Server/Tools/ToolsController.cs
+++ b/Engine/Source/Programs/Horde/Horde.Server/Tools/ToolsController.cs
@@ -330,7 +330,7 @@ namespace Horde.Server.Tools
 			IStorageClient client = _toolCollection.CreateStorageClient(tool);
 			try
 			{
-				DirectoryNode node = await client.ReadRefAsync<DirectoryNode>(deployment.RefName, DateTime.UtcNow - TimeSpan.FromDays(2.0), cancellationToken);
+				DirectoryNode node = await client.ReadRefAsync<DirectoryNode>(deployment.RefName, DateTime.UtcNow - TimeSpan.FromDays(2.0), cancellationToken: cancellationToken);
 
 				if (node.Directories.Count == 0 && node.Files.Count == 1 && action != GetToolAction.Zip)
 				{
diff --git a/Engine/Source/Programs/Horde/Horde/Commands/Bundles/BundleCreate.cs b/Engine/Source/Programs/Horde/Horde/Commands/Bundles/BundleCreate.cs
--- a/Engine/Source/Programs/Horde/Horde/Commands/Bundles/BundleCreate.cs
+++ b/Engine/Source/Programs/Horde/Horde/Commands/Bundles/BundleCreate.cs
@@ -87,19 +87,19 @@ namespace Horde.Commands.Bundles
 				List<FileInfo> fileInfos = files.ConvertAll(x => x.ToFileInfo());
 				CopyStatsLoggerWithTotals copyStatsLogger = new CopyStatsLoggerWithTotals(files.Count, fileInfos.Sum(x => x.Length), logger);
 
-				HashedNodeRef<DirectoryNode> nodeRef = await writer.WriteFilesAsync(baseDir.ToDirectoryInfo(), fileInfos, options, copyStatsLogger, CancellationToken.None);
+				IBlobHandle<DirectoryNode> nodeRef = await writer.WriteFilesAsync(baseDir.ToDirectoryInfo(), fileInfos, options, copyStatsLogger, CancellationToken.None);
 
 				await writer.FlushAsync();
 
 				if (File != null)
 				{
 					logger.LogInformation("Writing {File}", File);
-					await FileReference.WriteAllTextAsync(File, nodeRef.Handle.GetLocator().ToString()!);
+					await FileReference.WriteAllTextAsync(File, nodeRef.GetLocator().ToString()!);
 				}
 				else
 				{
 					logger.LogInformation("Writing ref {Ref}", Ref);
-					await store.WriteRefTargetAsync(new RefName(Ref!), nodeRef.Handle);
+					await store.WriteRefTargetAsync(new RefName(Ref!), nodeRef);
 				}
 			}
 
diff --git a/Engine/Source/Programs/Horde/Horde/Commands/Bundles/BundleDump.cs b/Engine/Source/Programs/Horde/Horde/Commands/Bundles/BundleDump.cs
--- a/Engine/Source/Programs/Horde/Horde/Commands/Bundles/BundleDump.cs
+++ b/Engine/Source/Programs/Horde/Horde/Commands/Bundles/BundleDump.cs
@@ -29,15 +29,7 @@ namespace Horde.Commands.Bundles
 			string[] types = new string[header.Types.Count];
 			for (int idx = 0; idx < header.Types.Count; idx++)
 			{
-				string name;
-				if (Node.TryGetConcreteType(header.Types[idx].Guid, out Type? type))
-				{
-					name = type.Name;
-				}
-				else
-				{
-					name = header.Types[idx].Guid.ToString();
-				}
+				string name = header.Types[idx].Guid.ToString();
 				types[idx] = $"{name} v{header.Types[idx].Version}";
 			}
 
diff --git a/Engine/Source/Programs/Horde/Horde/Commands/Bundles/BundleExtract.cs b/Engine/Source/Programs/Horde/Horde/Commands/Bundles/BundleExtract.cs
--- a/Engine/Source/Programs/Horde/Horde/Commands/Bundles/BundleExtract.cs
+++ b/Engine/Source/Programs/Horde/Horde/Commands/Bundles/BundleExtract.cs
@@ -71,8 +71,8 @@ namespace Horde.Commands.Bundles
 		{
 			Stopwatch timer = Stopwatch.StartNew();
 
-			DirectoryNode node = await handle.ReadNodeAsync<DirectoryNode>();
-			await node.CopyToDirectoryAsync(OutputDir.ToDirectoryInfo(), new CopyStatsLogger(logger), logger, CancellationToken.None);
+			DirectoryNode node = await handle.ReadBlobAsync<DirectoryNode>();
+			await node.CopyToDirectoryAsync(OutputDir.ToDirectoryInfo(), new CopyStatsLogger(logger), null, logger, CancellationToken.None);
 
 			logger.LogInformation("Elapsed: {Time}s", timer.Elapsed.TotalSeconds);
 
diff --git a/Engine/Source/Programs/Horde/Horde/Commands/Compute/ComputeRun.cs b/Engine/Source/Programs/Horde/Horde/Commands/Compute/ComputeRun.cs
--- a/Engine/Source/Programs/Horde/Horde/Commands/Compute/ComputeRun.cs
+++ b/Engine/Source/Programs/Horde/Horde/Commands/Compute/ComputeRun.cs
@@ -153,10 +153,10 @@ namespace Horde.Commands.Compute
 		{
 			await using IStorageWriter writer = storage.CreateWriter();
 
-			HashedNodeRef<DirectoryNode> sandbox = await writer.WriteFilesAsync(taskFile.Directory, cancellationToken: cancellationToken);
+			IBlobHandle<DirectoryNode> sandbox = await writer.WriteFilesAsync(taskFile.Directory, cancellationToken: cancellationToken);
 			await writer.FlushAsync(cancellationToken);
 
-			return sandbox.Handle.GetLocator();
+			return sandbox.GetLocator();
 		}
 	}
 }
diff --git a/Engine/Source/Programs/Horde/Horde/Commands/Tools/ToolUpload.cs b/Engine/Source/Programs/Horde/Horde/Commands/Tools/ToolUpload.cs
--- a/Engine/Source/Programs/Horde/Horde/Commands/Tools/ToolUpload.cs
+++ b/Engine/Source/Programs/Horde/Horde/Commands/Tools/ToolUpload.cs
@@ -40,13 +40,13 @@ namespace Horde.Commands
 		{
 			using IStorageClient storageClient = _httpStorageClientFactory.CreateClientWithPath($"api/v1/tools/{ToolId}");
 
-			HashedNodeRef<DirectoryNode> target;
+			IBlobHandle<DirectoryNode> target;
 			await using (IStorageWriter writer = storageClient.CreateWriter())
 			{
 				target = await writer.WriteFilesAsync(InputDir);
 			}
 
-			ToolDeploymentId deploymentId = await _hordeHttpClient.CreateToolDeploymentAsync(ToolId, Version, null, null, target.Handle.GetLocator());
+			ToolDeploymentId deploymentId = await _hordeHttpClient.CreateToolDeploymentAsync(ToolId, Version, null, null, target.GetLocator());
 			logger.LogInformation("Created deployment {DeploymentId}", deploymentId);
 
 			return 0;
diff --git a/Engine/Source/Programs/Horde/Horde/Commands/Vcs/VcsBase.cs b/Engine/Source/Programs/Horde/Horde/Commands/Vcs/VcsBase.cs
--- a/Engine/Source/Programs/Horde/Horde/Commands/Vcs/VcsBase.cs
+++ b/Engine/Source/Programs/Horde/Horde/Commands/Vcs/VcsBase.cs
@@ -352,7 +352,7 @@ namespace Horde.Commands.Vcs
 					{
 						return null;
 					}
-					tip = await tip.Parent.ExpandAsync();
+					tip = await tip.Parent.ReadBlobAsync();
 				}
 			}
 			return tip;
diff --git a/Engine/Source/Programs/Horde/Horde/Commands/Vcs/VcsCheckout.cs b/Engine/Source/Programs/Horde/Horde/Commands/Vcs/VcsCheckout.cs
--- a/Engine/Source/Programs/Horde/Horde/Commands/Vcs/VcsCheckout.cs
+++ b/Engine/Source/Programs/Horde/Horde/Commands/Vcs/VcsCheckout.cs
@@ -73,7 +73,7 @@ namespace Horde.Commands.Vcs
 
 			DirectoryState newState = new DirectoryState();
 
-			DirectoryNode directoryNode = await directoryRef.ExpandAsync();
+			DirectoryNode directoryNode = await directoryRef.Handle.ReadBlobAsync();
 			foreach ((string name, DirectoryEntry? subDirEntry, DirectoryState? subDirState) in EnumerableExtensions.Zip(directoryNode.NameToDirectory, directoryState?.Directories))
 			{
 				DirectoryReference subDirPath = DirectoryReference.Combine(dirPath, name.ToString());
@@ -110,7 +110,7 @@ namespace Horde.Commands.Vcs
 				}
 			}
 
-			newState.Hash = directoryRef.Hash;
+			newState.Hash = directoryRef.Handle.Hash;
 			return newState;
 		}
 
diff --git a/Engine/Source/Programs/Horde/Horde/Commands/Vcs/VcsLog.cs b/Engine/Source/Programs/Horde/Horde/Commands/Vcs/VcsLog.cs
--- a/Engine/Source/Programs/Horde/Horde/Commands/Vcs/VcsLog.cs
+++ b/Engine/Source/Programs/Horde/Horde/Commands/Vcs/VcsLog.cs
@@ -43,7 +43,7 @@ namespace Horde.Commands.Vcs
 						break;
 					}
 
-					tip = await tip.Parent.ExpandAsync();
+					tip = await tip.Parent.ReadBlobAsync();
 				}
 			}
 
diff --git a/Engine/Source/Programs/Horde/Horde/Commands/Workspace/WorkspaceSync.cs b/Engine/Source/Programs/Horde/Horde/Commands/Workspace/WorkspaceSync.cs
--- a/Engine/Source/Programs/Horde/Horde/Commands/Workspace/WorkspaceSync.cs
+++ b/Engine/Source/Programs/Horde/Horde/Commands/Workspace/WorkspaceSync.cs
@@ -86,7 +86,7 @@ namespace Horde.Commands.Workspace
 			Stopwatch timer = Stopwatch.StartNew();
 			logger.LogInformation("Syncing into layer '{LayerId}'...", LayerId);
 
-			DirectoryNode contents = await handle.ReadNodeAsync<DirectoryNode>();
+			DirectoryNode contents = await handle.ReadBlobAsync<DirectoryNode>();
 			await workspace.SyncAsync(LayerId, contents, cancellationToken);
 			await workspace.SaveAsync(cancellationToken);
 
diff --git a/Engine/Source/Programs/Horde/Samples/RemoteClient/RemoteClientApp.cs b/Engine/Source/Programs/Horde/Samples/RemoteClient/RemoteClientApp.cs
--- a/Engine/Source/Programs/Horde/Samples/RemoteClient/RemoteClientApp.cs
+++ b/Engine/Source/Programs/Horde/Samples/RemoteClient/RemoteClientApp.cs
@@ -116,9 +116,9 @@ namespace RemoteClient
 
 				await using (IStorageWriter writer = storage.CreateWriter())
 				{
-					HashedNodeRef<DirectoryNode> sandbox = await writer.WriteFilesAsync(uploadDir);
+					IBlobHandle<DirectoryNode> sandbox = await writer.WriteFilesAsync(uploadDir);
 					await writer.FlushAsync();
-					await channel.UploadFilesAsync("", sandbox.Handle.GetLocator(), storage);
+					await channel.UploadFilesAsync("", sandbox.GetLocator(), storage);
 				}
 
 				// Run the task remotely in the background and echo the output to the console
diff --git a/Engine/Source/Programs/Shared/EpicGames.Horde.Tests/BlobDataTests.cs b/Engine/Source/Programs/Shared/EpicGames.Horde.Tests/BlobDataTests.cs
--- a/Engine/Source/Programs/Shared/EpicGames.Horde.Tests/BlobDataTests.cs
+++ b/Engine/Source/Programs/Shared/EpicGames.Horde.Tests/BlobDataTests.cs
@@ -21,7 +21,7 @@ namespace EpicGames.Horde.Tests
 
 			public DummyHandle(string locator) => _locator = new BlobLocator(new Utf8String(locator));
 			public ValueTask FlushAsync(CancellationToken cancellationToken = default) => default;
-			public ValueTask<BlobData> ReadAsync(CancellationToken cancellationToken = default) => throw new NotImplementedException();
+			public ValueTask<BlobData> ReadBlobDataAsync(CancellationToken cancellationToken = default) => throw new NotImplementedException();
 
 			public bool TryAppendIdentifier(Utf8StringBuilder builder)
 			{
diff --git a/Engine/Source/Programs/Shared/EpicGames.Horde.Tests/BundlePacketTests.cs b/Engine/Source/Programs/Shared/EpicGames.Horde.Tests/BundlePacketTests.cs
--- a/Engine/Source/Programs/Shared/EpicGames.Horde.Tests/BundlePacketTests.cs
+++ b/Engine/Source/Programs/Shared/EpicGames.Horde.Tests/BundlePacketTests.cs
@@ -29,7 +29,7 @@ namespace EpicGames.Horde.Tests
 				_fragment = fragment;
 			}
 
-			public ValueTask<BlobData> ReadAsync(CancellationToken cancellationToken = default) => throw new NotImplementedException();
+			public ValueTask<BlobData> ReadBlobDataAsync(CancellationToken cancellationToken = default) => throw new NotImplementedException();
 			public ValueTask FlushAsync(CancellationToken cancellationToken = default) => throw new NotImplementedException();
 
 			public bool TryAppendIdentifier(Utf8StringBuilder builder)
diff --git a/Engine/Source/Programs/Shared/EpicGames.Horde.Tests/BundleStoreTests.cs b/Engine/Source/Programs/Shared/EpicGames.Horde.Tests/BundleStoreTests.cs
--- a/Engine/Source/Programs/Shared/EpicGames.Horde.Tests/BundleStoreTests.cs
+++ b/Engine/Source/Programs/Shared/EpicGames.Horde.Tests/BundleStoreTests.cs
@@ -18,6 +18,7 @@ using Microsoft.Extensions.Logging.Abstractions;
 using System.Reflection;
 using EpicGames.Horde.Storage.Bundles;
 using EpicGames.Horde.Storage.Bundles.V1;
+using EpicGames.Serialization;
 
 namespace EpicGames.Horde.Tests
 {
@@ -32,29 +33,31 @@ namespace EpicGames.Horde.Tests
 			Assert.IsTrue(a.SequenceEqual(b));
 		}
 
-		[BlobType("{F63606D4-4061-5DBB-446F-55A69E22654F}")]
-		class TextNode : Node
+		[BlobConverter(typeof(TextNodeConverter))]
+		class TextNode
 		{
 			public string Text { get; }
 
 			public TextNode(string text) => Text = text;
+		}
+
+		class TextNodeConverter : BlobConverter<TextNode>
+		{
+			public static BlobType BlobType { get; } = new BlobType("{F63606D4-4061-5DBB-446F-55A69E22654F}", 1);
 
-			public TextNode(IBlobReader reader)
+			public override TextNode Read(IBlobReader reader, BlobSerializerOptions options)
 			{
-				Text = reader.ReadString();
+				string text = reader.ReadString();
+				return new TextNode(text);
 			}
 
-			public override void Serialize(IBlobWriter writer)
+			public override BlobType Write(IBlobWriter writer, TextNode value, BlobSerializerOptions options)
 			{
-				writer.WriteString(Text);
+				writer.WriteString(value.Text);
+				return BlobType;
 			}
 		}
 
-		public BundleStoreTests()
-		{
-			Node.RegisterTypesFromAssembly(Assembly.GetExecutingAssembly());
-		}
-
 		static async Task<byte[]> CreateBundleNormalAsync()
 		{
 			using MemoryStorageClient memoryStore = new MemoryStorageClient();
@@ -62,10 +65,11 @@ namespace EpicGames.Horde.Tests
 			await using IStorageWriter writer = store.CreateWriter(options: new BundleOptions { MaxVersion = BundleVersion.ImportHashes, CompressionFormat = BundleCompressionFormat.None });
 
 			TextNode node = new TextNode("Hello world");
-			IBlobHandle handle = await writer.FlushAsync(node, CancellationToken.None);
+			IBlobHandle<TextNode> handle = await writer.WriteBlobAsync(node);
+			await writer.FlushAsync();
 
 			IBlobHandle bundleHandle = store.CreateBlobHandle(handle.GetLocator().BaseLocator);
-			using BlobData blobData = await bundleHandle.ReadAsync();
+			using BlobData blobData = await bundleHandle.ReadBlobDataAsync();
 
 			return blobData.Data.ToArray();
 		}
@@ -121,28 +125,35 @@ namespace EpicGames.Horde.Tests
 			Assert.AreEqual(1, blobStore.Refs.Count);
 		}
 
-		[BlobType("{F63606D4-4061-5DBB-446F-55A69F22654F}")]
-		class SimpleNode : Node
+		[BlobConverter(typeof(SimpleNodeConverter))]
+		class SimpleNode
 		{
+			public static BlobType BlobType { get; } = new BlobType("{F63606D4-4061-5DBB-446F-55A69F22654F}", 1);
+
 			public ReadOnlySequence<byte> Data { get; }
-			public IReadOnlyList<HashedNodeRef<SimpleNode>> Refs { get; }
+			public IReadOnlyList<IBlobHandle<SimpleNode>> Refs { get; }
 
-			public SimpleNode(ReadOnlySequence<byte> data, IReadOnlyList<HashedNodeRef<SimpleNode>> refs)
+			public SimpleNode(ReadOnlySequence<byte> data, IReadOnlyList<IBlobHandle<SimpleNode>> refs)
 			{
 				Data = data;
 				Refs = refs;
 			}
+		}
 
-			public SimpleNode(IBlobReader reader)
+		class SimpleNodeConverter : BlobConverter<SimpleNode>
+		{
+			public override SimpleNode Read(IBlobReader reader, BlobSerializerOptions options)
 			{
-				Data = new ReadOnlySequence<byte>(reader.ReadVariableLengthBytes());
-				Refs = reader.ReadVariableLengthArray(() => reader.ReadHashedNodeRef<SimpleNode>());
+				ReadOnlyMemory<byte> data = reader.ReadVariableLengthBytes();
+				IReadOnlyList<IBlobHandle<SimpleNode>> refs = reader.ReadVariableLengthArray(() => reader.ReadBlobHandle<SimpleNode>());
+				return new SimpleNode(new ReadOnlySequence<byte>(data), refs);
 			}
 
-			public override void Serialize(IBlobWriter writer)
+			public override BlobType Write(IBlobWriter writer, SimpleNode value, BlobSerializerOptions options)
 			{
-				writer.WriteVariableLengthBytes(Data);
-				writer.WriteVariableLengthArray(Refs, x => writer.WriteHashedNodeRef(x));
+				writer.WriteVariableLengthBytes(value.Data);
+				writer.WriteVariableLengthArray(value.Refs, x => writer.WriteBlobHandle(x));
+				return SimpleNode.BlobType;
 			}
 		}
 
@@ -152,13 +163,13 @@ namespace EpicGames.Horde.Tests
 			{
 				await using IStorageWriter writer = store.CreateWriter("test", options);
 
-				SimpleNode node1 = new SimpleNode(new ReadOnlySequence<byte>(new byte[] { 1 }), Array.Empty<HashedNodeRef<SimpleNode>>());
-				SimpleNode node2 = new SimpleNode(new ReadOnlySequence<byte>(new byte[] { 2 }), new[] { await writer.WriteHashedNodeAsync(node1) });
-				SimpleNode node3 = new SimpleNode(new ReadOnlySequence<byte>(new byte[] { 3 }), new[] { await writer.WriteHashedNodeAsync(node2) });
-				SimpleNode node4 = new SimpleNode(new ReadOnlySequence<byte>(new byte[] { 4 }), Array.Empty<HashedNodeRef<SimpleNode>>());
+				SimpleNode node1 = new SimpleNode(new ReadOnlySequence<byte>(new byte[] { 1 }), Array.Empty<IBlobHandle<SimpleNode>>());
+				SimpleNode node2 = new SimpleNode(new ReadOnlySequence<byte>(new byte[] { 2 }), new[] { await writer.WriteBlobAsync(node1) });
+				SimpleNode node3 = new SimpleNode(new ReadOnlySequence<byte>(new byte[] { 3 }), new[] { await writer.WriteBlobAsync(node2) });
+				SimpleNode node4 = new SimpleNode(new ReadOnlySequence<byte>(new byte[] { 4 }), Array.Empty<IBlobHandle<SimpleNode>>());
 
-				SimpleNode root = new SimpleNode(new ReadOnlySequence<byte>(new byte[] { 5 }), new[] { await writer.WriteHashedNodeAsync(node4), await writer.WriteHashedNodeAsync(node3) });
-				NodeRef<SimpleNode> rootRef = await writer.WriteNodeAsync(root);
+				SimpleNode root = new SimpleNode(new ReadOnlySequence<byte>(new byte[] { 5 }), new[] { await writer.WriteBlobAsync(node4), await writer.WriteBlobAsync(node3) });
+				IBlobHandle<SimpleNode> rootRef = await writer.WriteBlobAsync(root);
 				await writer.FlushAsync();
 
 				await store.WriteRefTargetAsync(new RefName("test"), rootRef);
@@ -179,31 +190,31 @@ namespace EpicGames.Horde.Tests
 			SimpleNode node5 = root;
 			byte[] data5 = node5.Data.ToArray();
 			Assert.IsTrue(data5.SequenceEqual(new byte[] { 5 }));
-			IReadOnlyList<HashedNodeRef<SimpleNode>> refs5 = node5.Refs;
+			IReadOnlyList<IBlobHandle<SimpleNode>> refs5 = node5.Refs;
 			Assert.AreEqual(2, refs5.Count);
 
-			SimpleNode node4 = await refs5[0].ExpandAsync();
+			SimpleNode node4 = await refs5[0].ReadBlobAsync();
 			byte[] data4 = node4.Data.ToArray();
 			Assert.IsTrue(data4.SequenceEqual(new byte[] { 4 }));
-			IReadOnlyList<HashedNodeRef<SimpleNode>> refs4 = node4.Refs;
+			IReadOnlyList<IBlobHandle<SimpleNode>> refs4 = node4.Refs;
 			Assert.AreEqual(0, refs4.Count);
 
-			SimpleNode node3 = await refs5[1].ExpandAsync();
+			SimpleNode node3 = await refs5[1].ReadBlobAsync();
 			byte[] data3 = node3.Data.ToArray();
 			Assert.IsTrue(data3.SequenceEqual(new byte[] { 3 }));
-			IReadOnlyList<HashedNodeRef<SimpleNode>> refs3 = node3.Refs;
+			IReadOnlyList<IBlobHandle<SimpleNode>> refs3 = node3.Refs;
 			Assert.AreEqual(1, refs3.Count);
 
-			SimpleNode node2 = await refs3[0].ExpandAsync();
+			SimpleNode node2 = await refs3[0].ReadBlobAsync();
 			byte[] data2 = node2.Data.ToArray();
 			Assert.IsTrue(data2.SequenceEqual(new byte[] { 2 }));
-			IReadOnlyList<HashedNodeRef<SimpleNode>> refs2 = node2.Refs;
+			IReadOnlyList<IBlobHandle<SimpleNode>> refs2 = node2.Refs;
 			Assert.AreEqual(1, refs2.Count);
 
-			SimpleNode node1 = await refs2[0].ExpandAsync();
+			SimpleNode node1 = await refs2[0].ReadBlobAsync();
 			byte[] data1 = node1.Data.ToArray();
 			Assert.IsTrue(data1.SequenceEqual(new byte[] { 1 }));
-			IReadOnlyList<HashedNodeRef<SimpleNode>> refs1 = node1.Refs;
+			IReadOnlyList<IBlobHandle<SimpleNode>> refs1 = node1.Refs;
 			Assert.AreEqual(0, refs1.Count);
 		}
 
@@ -214,12 +225,12 @@ namespace EpicGames.Horde.Tests
 
 			RefName refName = new RefName("test");
 
-			NodeRef<SimpleNode> inputRef;
+			IBlobHandle<SimpleNode> inputRef;
 			await using (IStorageWriter writer = store.CreateWriter(refName))
 			{
-				inputRef = await writer.WriteNodeAsync(new SimpleNode(new ReadOnlySequence<byte>(new byte[] { (byte)123 }), Array.Empty<HashedNodeRef<SimpleNode>>()));
+				inputRef = await writer.WriteBlobAsync(new SimpleNode(new ReadOnlySequence<byte>(new byte[] { (byte)123 }), Array.Empty<IBlobHandle<SimpleNode>>()));
 			}
-			await store.WriteRefAsync(refName, inputRef.Handle);
+			await store.WriteRefAsync(refName, inputRef);
 
 			SimpleNode node = await store.ReadRefAsync<SimpleNode>(refName);
 
@@ -236,19 +247,19 @@ namespace EpicGames.Horde.Tests
 				await using (IStorageWriter writer = store.CreateWriter(new RefName("test")))
 				{
 					DirectoryNode world = new DirectoryNode();
-					HashedNodeRef<DirectoryNode> worldRef = await writer.WriteHashedNodeAsync(world);
+					IBlobHandle<DirectoryNode> worldRef = await writer.WriteBlobAsync(world);
 
 					DirectoryNode hello = new DirectoryNode();
 					hello.AddDirectory(new DirectoryEntry("world", 0, worldRef));
-					HashedNodeRef<DirectoryNode> helloRef = await writer.WriteHashedNodeAsync(hello);
+					IBlobHandle<DirectoryNode> helloRef = await writer.WriteBlobAsync(hello);
 
 					DirectoryNode root = new DirectoryNode(DirectoryFlags.None);
 					root.AddDirectory(new DirectoryEntry("hello", 0, helloRef));
-					HashedNodeRef<DirectoryNode> rootRef = await writer.WriteHashedNodeAsync(root);
+					IBlobHandle<DirectoryNode> rootRef = await writer.WriteBlobAsync(root);
 
 					await writer.FlushAsync();
 
-					await store.WriteRefTargetAsync(new RefName("test"), rootRef.Handle);
+					await store.WriteRefTargetAsync(new RefName("test"), rootRef);
 				}
 			}
 
@@ -264,11 +275,11 @@ namespace EpicGames.Horde.Tests
 			Assert.AreEqual(1, root.Directories.Count);
 			Assert.AreEqual("hello", root.Directories.First().Name);
 
-			DirectoryNode hello = await root.Directories.First().ExpandAsync(CancellationToken.None);
+			DirectoryNode hello = await root.Directories.First().Handle.ReadBlobAsync();
 			Assert.AreEqual(1, hello.Directories.Count);
 			Assert.AreEqual("world", hello.Directories.First().Name);
 
-			DirectoryNode world = await hello.Directories.First().ExpandAsync(CancellationToken.None);
+			DirectoryNode world = await hello.Directories.First().Handle.ReadBlobAsync();
 			Assert.AreEqual(0, world.Directories.Count);
 		}
 
@@ -292,7 +303,7 @@ namespace EpicGames.Horde.Tests
 				DirectoryNode root = new DirectoryNode();
 				await root.UpdateAsync(fileUpdates, writer);
 
-				HashedNodeRef<DirectoryNode> rootRef = await writer.WriteHashedNodeAsync(root);
+				IBlobHandle<DirectoryNode> rootRef = await writer.WriteBlobAsync(root);
 				await store.WriteRefTargetAsync(new RefName("test"), rootRef);
 
 				await CheckFileTreeAsync(root);
@@ -310,12 +321,12 @@ namespace EpicGames.Horde.Tests
 			Assert.AreEqual(1, root.Directories.Count);
 			Assert.AreEqual("hello", root.Directories.First().Name);
 
-			DirectoryNode hello = await root.Directories.First().ExpandAsync();
+			DirectoryNode hello = await root.Directories.First().Handle.ReadBlobAsync();
 			Assert.AreEqual(0, hello.Directories.Count);
 			Assert.AreEqual(1, hello.Files.Count);
 			Assert.AreEqual("world", hello.Files.First().Name);
 
-			ChunkedDataNode world = await hello.Files.First().Target.ExpandAsync();
+			ChunkedDataNode world = await hello.Files.First().Target.ReadBlobAsync();
 
 			byte[] worldData = await GetFileDataAsync(world);
 			Assert.IsTrue(worldData.SequenceEqual(Encoding.UTF8.GetBytes("world")));
@@ -351,10 +362,10 @@ namespace EpicGames.Horde.Tests
 
 			// Check we can read it back in
 			{
-				ChunkedDataNode newRoot = await nodeRef.ExpandAsync();
+				ChunkedDataNode newRoot = await nodeRef.ReadBlobAsync();
 
 				using MemoryStream stream = new MemoryStream();
-				await newRoot.CopyToStreamAsync(stream, CancellationToken.None);
+				await newRoot.CopyToStreamAsync(stream);
 
 				byte[] output = stream.ToArray();
 				Assert.IsTrue(chunk.SequenceEqual(output));
@@ -394,7 +405,7 @@ namespace EpicGames.Horde.Tests
 				root = new DirectoryNode(DirectoryFlags.None);
 				root.AddFile("test", FileEntryFlags.None, fileWriter.Length, chunkedData);
 
-				HashedNodeRef<DirectoryNode> rootRef = await writer.WriteHashedNodeAsync(root);
+				IBlobHandle<DirectoryNode> rootRef = await writer.WriteBlobAsync(root);
 				await store.WriteRefTargetAsync(new RefName("test"), rootRef);
 
 				await CheckLargeFileTreeAsync(root, data);
@@ -418,7 +429,7 @@ namespace EpicGames.Horde.Tests
 
 		static async Task GetUniqueBlobsAsync(IBlobHandle handle, Dictionary<BlobLocator, long> locatorToSize)
 		{
-			using BlobData data = await handle.ReadAsync();
+			using BlobData data = await handle.ReadBlobDataAsync();
 			locatorToSize[handle.GetLocator()] = data.Data.Length;
 
 			foreach (IBlobHandle reference in data.Refs)
@@ -435,8 +446,8 @@ namespace EpicGames.Horde.Tests
 
 			foreach ((FileEntry oldFileEntry, FileEntry newFileEntry) in oldNode.Files.Zip(newNode.Files))
 			{
-				ChunkedDataNode oldFile = await oldFileEntry.Target.ExpandAsync();
-				ChunkedDataNode newFile = await newFileEntry.Target.ExpandAsync();
+				ChunkedDataNode oldFile = await oldFileEntry.Target.ReadBlobAsync();
+				ChunkedDataNode newFile = await newFileEntry.Target.ReadBlobAsync();
 				await CompareTreesAsync(oldFile, newFile);
 			}
 		}
@@ -451,8 +462,8 @@ namespace EpicGames.Horde.Tests
 				int index = 0;
 				foreach ((ChunkedDataNodeRef oldFileRef, ChunkedDataNodeRef newFileRef) in oldInteriorNode.Children.Zip(newInteriorNode.Children))
 				{
-					ChunkedDataNode oldFile = await oldFileRef.ExpandAsync();
-					ChunkedDataNode newFile = await newFileRef.ExpandAsync();
+					ChunkedDataNode oldFile = await oldFileRef.ReadBlobAsync();
+					ChunkedDataNode newFile = await newFileRef.ReadBlobAsync();
 					await CompareTreesAsync(oldFile, newFile);
 					index++;
 				}
@@ -473,7 +484,7 @@ namespace EpicGames.Horde.Tests
 			Assert.AreEqual(0, root.Directories.Count);
 			Assert.AreEqual(1, root.Files.Count);
 
-			ChunkedDataNode world = await root.Files.First().Target.ExpandAsync(CancellationToken.None);
+			ChunkedDataNode world = await root.Files.First().Target.ReadBlobAsync();
 
 			int length = await CheckFileDataAsync(world, data);
 			Assert.AreEqual(data.Length, length);
@@ -491,7 +502,7 @@ namespace EpicGames.Horde.Tests
 			{
 				foreach (ChunkedDataNodeRef childRef in interiorFileNode.Children)
 				{
-					ChunkedDataNode child = await childRef.ExpandAsync();
+					ChunkedDataNode child = await childRef.ReadBlobAsync();
 					offset += await CheckFileDataAsync(child, data.Slice(offset));
 				}
 			}
@@ -506,9 +517,9 @@ namespace EpicGames.Horde.Tests
 		{
 			using (MemoryStream stream = new MemoryStream())
 			{
-				await fileNode.CopyToStreamAsync(stream, CancellationToken.None);
+				await fileNode.CopyToStreamAsync(stream);
 				return stream.ToArray();
 			}
 		}
 	}
-}
+}
\ No newline at end of file
diff --git a/Engine/Source/Programs/Shared/EpicGames.Horde.Tests/BundleTests.cs b/Engine/Source/Programs/Shared/EpicGames.Horde.Tests/BundleTests.cs
--- a/Engine/Source/Programs/Shared/EpicGames.Horde.Tests/BundleTests.cs
+++ b/Engine/Source/Programs/Shared/EpicGames.Horde.Tests/BundleTests.cs
@@ -84,35 +84,35 @@ namespace EpicGames.Horde.Tests
 			byte[] data = CreateBuffer(1024);
 
 			nodeRef = (await fileNodeWriter.CreateAsync(data.AsMemory(0, 7), CancellationToken.None)).Root;
-			node = await nodeRef.ExpandAsync();
+			node = await nodeRef.ReadBlobAsync();
 			Assert.IsTrue(node is LeafChunkedDataNode);
 			Assert.AreEqual(7, ((LeafChunkedDataNode)node).Data.Length);
 			await TestBufferlessReadsAsync(nodeRef, data.AsMemory(0, 7));
 
 			nodeRef = (await fileNodeWriter.CreateAsync(data.AsMemory(0, 8), CancellationToken.None)).Root;
-			node = await nodeRef.ExpandAsync();
+			node = await nodeRef.ReadBlobAsync();
 			Assert.IsTrue(node is LeafChunkedDataNode);
 			Assert.AreEqual(8, ((LeafChunkedDataNode)node).Data.Length);
 			await TestBufferlessReadsAsync(nodeRef, data.AsMemory(0, 8));
 
 			nodeRef = (await fileNodeWriter.CreateAsync(data.AsMemory(0, 9), CancellationToken.None)).Root;
-			node = await nodeRef.ExpandAsync();
+			node = await nodeRef.ReadBlobAsync();
 			Assert.IsTrue(node is InteriorChunkedDataNode);
 			Assert.AreEqual(2, ((InteriorChunkedDataNode)node).Children.Count);
 			await TestBufferlessReadsAsync(nodeRef, data.AsMemory(0, 9));
 
-			ChunkedDataNode? childNode1 = await ((InteriorChunkedDataNode)node).Children[0].ExpandAsync();
+			ChunkedDataNode? childNode1 = await ((InteriorChunkedDataNode)node).Children[0].ReadBlobAsync();
 			Assert.IsNotNull(childNode1);
 			Assert.IsTrue(childNode1 is LeafChunkedDataNode);
 			Assert.AreEqual(8, ((LeafChunkedDataNode)childNode1!).Data.Length);
 
-			ChunkedDataNode? childNode2 = await ((InteriorChunkedDataNode)node).Children[1].ExpandAsync();
+			ChunkedDataNode? childNode2 = await ((InteriorChunkedDataNode)node).Children[1].ReadBlobAsync();
 			Assert.IsNotNull(childNode2);
 			Assert.IsTrue(childNode2 is LeafChunkedDataNode);
 			Assert.AreEqual(1, ((LeafChunkedDataNode)childNode2!).Data.Length);
 
 			nodeRef = (await fileNodeWriter.CreateAsync(data, CancellationToken.None)).Root;
-			node = await nodeRef.ExpandAsync();
+			node = await nodeRef.ReadBlobAsync();
 			Assert.IsTrue(node is InteriorChunkedDataNode);
 			await TestBufferlessReadsAsync(nodeRef, data);
 		}
@@ -194,19 +194,19 @@ namespace EpicGames.Horde.Tests
 		{
 			IStorageClient store = _storage;
 
-			HashedNodeRef<DirectoryNode> rootRef;
+			IBlobHandle<DirectoryNode> rootRef;
 			await using (IStorageWriter writer = store.CreateWriter())
 			{
 				DirectoryNode world = new DirectoryNode();
-				HashedNodeRef<DirectoryNode> worldRef = await writer.WriteHashedNodeAsync(world);
+				IBlobHandle<DirectoryNode> worldRef = await writer.WriteBlobAsync(world);
 
 				DirectoryNode hello = new DirectoryNode();
 				hello.AddDirectory(new DirectoryEntry("world", 0, worldRef));
-				HashedNodeRef<DirectoryNode> helloRef = await writer.WriteHashedNodeAsync(hello);
+				IBlobHandle<DirectoryNode> helloRef = await writer.WriteBlobAsync(hello);
 
 				DirectoryNode root = new DirectoryNode();
 				root.AddDirectory(new DirectoryEntry("hello", 0, helloRef));
-				rootRef = await writer.WriteHashedNodeAsync(root);
+				rootRef = await writer.WriteBlobAsync(root);
 
 				await writer.FlushAsync();
 			}
@@ -221,7 +221,7 @@ namespace EpicGames.Horde.Tests
 			// Check the ref
 			IBlobHandle refTarget =  await store.ReadRefTargetAsync(refName);
 			IBlobHandle bundleTarget = store.CreateBlobHandle(refTarget.GetLocator().BaseLocator);
-			using BlobData bundleData = await bundleTarget.ReadAsync();
+			using BlobData bundleData = await bundleTarget.ReadBlobDataAsync();
 
 // This is specific to V1 data
 //			BundleHeader bundleHeader = BundleHeader.Read(bundleData.Data);
@@ -233,12 +233,12 @@ namespace EpicGames.Horde.Tests
 
 			Assert.AreEqual(0, newRoot.Files.Count);
 			Assert.AreEqual(1, newRoot.Directories.Count);
-			DirectoryNode? outputNode = await newRoot.TryOpenDirectoryAsync("hello", CancellationToken.None);
+			DirectoryNode? outputNode = await newRoot.TryOpenDirectoryAsync("hello");
 			Assert.IsNotNull(outputNode);
 
 			Assert.AreEqual(0, outputNode!.Files.Count);
 			Assert.AreEqual(1, outputNode!.Directories.Count);
-			DirectoryNode? outputNode2 = await outputNode.TryOpenDirectoryAsync("world", CancellationToken.None);
+			DirectoryNode? outputNode2 = await outputNode.TryOpenDirectoryAsync("world");
 			Assert.IsNotNull(outputNode2);
 
 			Assert.AreEqual(0, outputNode2!.Files.Count);
@@ -260,12 +260,12 @@ namespace EpicGames.Horde.Tests
 				for (int idx = 1; idx <= 3; idx++)
 				{
 					DirectoryNode node = new DirectoryNode();
-					HashedNodeRef<DirectoryNode> nodeRef = await writer.WriteHashedNodeAsync(node);
+					IBlobHandle<DirectoryNode> nodeRef = await writer.WriteBlobAsync(node);
 					root.AddDirectory(new DirectoryEntry($"node{idx}", 0, nodeRef));
 				}
 
 				RefName refName = new RefName("ref");
-				HashedNodeRef<DirectoryNode> rootRef = await writer.WriteHashedNodeAsync(root);
+				IBlobHandle<DirectoryNode> rootRef = await writer.WriteBlobAsync(root);
 				await _storage.WriteRefTargetAsync(refName, rootRef);
 			}
 
@@ -284,14 +284,14 @@ namespace EpicGames.Horde.Tests
 			{
 				await using (IStorageWriter writer = _storage.CreateWriter(options: options))
 				{
-					HashedNodeRef<DirectoryNode> rootRef = await writer.WriteHashedNodeAsync(new DirectoryNode());
+					IBlobHandle<DirectoryNode> rootRef = await writer.WriteBlobAsync(new DirectoryNode());
 					for (int idx = 4; idx >= 1; idx--)
 					{
 						DirectoryNode next = new DirectoryNode();
 						next.AddDirectory(new DirectoryEntry($"node{idx}", 0, rootRef));
-						rootRef = await writer.WriteHashedNodeAsync(next);
+						rootRef = await writer.WriteBlobAsync(next);
 					}
-					await _storage.WriteRefTargetAsync(refName, rootRef.Handle);
+					await _storage.WriteRefTargetAsync(refName, rootRef);
 				}
 
 				Assert.AreEqual(1, _memoryStore.Refs.Count);
@@ -301,16 +301,16 @@ namespace EpicGames.Horde.Tests
 			{
 				DirectoryNode root = await _storage.ReadRefAsync<DirectoryNode>(refName);
 
-				DirectoryNode? newNode1 = await root.TryOpenDirectoryAsync("node1", CancellationToken.None);
+				DirectoryNode? newNode1 = await root.TryOpenDirectoryAsync("node1");
 				Assert.IsNotNull(newNode1);
 
-				DirectoryNode? newNode2 = await newNode1!.TryOpenDirectoryAsync("node2", CancellationToken.None);
+				DirectoryNode? newNode2 = await newNode1!.TryOpenDirectoryAsync("node2");
 				Assert.IsNotNull(newNode2);
 
-				DirectoryNode? newNode3 = await newNode2!.TryOpenDirectoryAsync("node3", CancellationToken.None);
+				DirectoryNode? newNode3 = await newNode2!.TryOpenDirectoryAsync("node3");
 				Assert.IsNotNull(newNode3);
 
-				DirectoryNode? newNode4 = await newNode3!.TryOpenDirectoryAsync("node4", CancellationToken.None);
+				DirectoryNode? newNode4 = await newNode3!.TryOpenDirectoryAsync("node4");
 				Assert.IsNotNull(newNode4);
 			}
 		}
diff --git a/Engine/Source/Programs/Shared/EpicGames.Horde.Tests/ChunkingTests.cs b/Engine/Source/Programs/Shared/EpicGames.Horde.Tests/ChunkingTests.cs
--- a/Engine/Source/Programs/Shared/EpicGames.Horde.Tests/ChunkingTests.cs
+++ b/Engine/Source/Programs/Shared/EpicGames.Horde.Tests/ChunkingTests.cs
@@ -65,8 +65,8 @@ namespace EpicGames.Horde.Tests
 				DirectoryNode directory = new DirectoryNode();
 				directory.AddFile("test.foo", FileEntryFlags.None, 0, chunkedData);
 
-				HashedNodeRef<DirectoryNode> directoryRef = await writer.WriteHashedNodeAsync(directory);
-				await store.WriteRefTargetAsync(RefName, directoryRef.Handle);
+				IBlobHandle handle = await writer.WriteBlobAsync(directory);
+				await store.WriteRefTargetAsync(RefName, handle);
 			}
 		}
 
@@ -120,12 +120,12 @@ namespace EpicGames.Horde.Tests
 				handle = (await fileWriter.FlushAsync(CancellationToken.None)).Root;
 			}
 
-			ChunkedDataNode root = await handle.ExpandAsync();
+			ChunkedDataNode root = await handle.ReadBlobAsync();
 
 			byte[] result;
 			using (MemoryStream stream = new MemoryStream())
 			{
-				await root.CopyToStreamAsync(stream, CancellationToken.None);
+				await root.CopyToStreamAsync(stream);
 				result = stream.ToArray();
 			}
 
@@ -157,7 +157,7 @@ namespace EpicGames.Horde.Tests
 				int childCount = interiorNode.Children.Count;
 				for (int idx = 0; idx < childCount; idx++)
 				{
-					ChunkedDataNode childNode = await interiorNode.Children[idx].ExpandAsync(CancellationToken.None);
+					ChunkedDataNode childNode = await interiorNode.Children[idx].ReadBlobAsync();
 					await CheckSizesAsync(childNode, options, idx == childCount - 1);
 				}
 			}
diff --git a/Engine/Source/Programs/Shared/EpicGames.Horde.Tests/ComputeSocketTests.cs b/Engine/Source/Programs/Shared/EpicGames.Horde.Tests/ComputeSocketTests.cs
--- a/Engine/Source/Programs/Shared/EpicGames.Horde.Tests/ComputeSocketTests.cs
+++ b/Engine/Source/Programs/Shared/EpicGames.Horde.Tests/ComputeSocketTests.cs
@@ -191,12 +191,14 @@ namespace EpicGames.Horde.Tests
 						DirectoryNode directory = new DirectoryNode();
 						directory.AddFile("hello.txt", FileEntryFlags.None, data.Length, chunkedData);
 
-						HashedNodeRef<DirectoryNode> directoryRef = await treeWriter.WriteHashedNodeAsync(directory, cancellationToken);
+						IBlobHandle<DirectoryNode> directoryRef = await treeWriter.WriteBlobAsync(directory, cancellationToken: cancellationToken);
 
 						DirectoryNode root = new DirectoryNode();
 						root.AddDirectory(new DirectoryEntry("subdir", directory.Length, directoryRef));
 
-						IBlobHandle handle = await treeWriter.FlushAsync(root, cancellationToken);
+						IBlobHandle<DirectoryNode> handle = await treeWriter.WriteBlobAsync(root, cancellationToken: cancellationToken);
+						await treeWriter.FlushAsync(cancellationToken);
+
 						await channel.UploadFilesAsync("", handle.GetLocator(), storage, cancellationToken);
 
 						Assert.IsTrue(FileReference.Exists(file));
diff --git a/Engine/Source/Programs/Shared/EpicGames.Horde.Tests/LogTests.cs b/Engine/Source/Programs/Shared/EpicGames.Horde.Tests/LogTests.cs
--- a/Engine/Source/Programs/Shared/EpicGames.Horde.Tests/LogTests.cs
+++ b/Engine/Source/Programs/Shared/EpicGames.Horde.Tests/LogTests.cs
@@ -69,18 +69,18 @@ namespace EpicGames.Horde.Tests
 			}
 
 			// Flush it to storage, and read the finished log node
-			HashedNodeRef<LogNode> logRef;
+			IBlobHandle<LogNode> logRef;
 			await using (IStorageWriter writer = store.CreateWriter())
 			{
-				logRef = await builder.FlushAsync(writer, true, CancellationToken.None);
+				logRef = await builder.FlushAsync(writer, true, null, CancellationToken.None);
 			}
-			LogNode log = await logRef.ExpandAsync();
+			LogNode log = await logRef.ReadBlobAsync();
 
 			// Read the data back out and check it's the same
 			byte[] readData = new byte[_data.Length];
 
 			int writeOffset = 0;
-			await foreach (ReadOnlyMemory<byte> line in log.ReadLogLinesAsync(0, CancellationToken.None))
+			await foreach (ReadOnlyMemory<byte> line in log.ReadLogLinesAsync(0))
 			{
 				line.CopyTo(readData.AsMemory(writeOffset));
 				writeOffset += line.Length;
@@ -94,7 +94,7 @@ namespace EpicGames.Horde.Tests
 			Assert.AreEqual(equalSize, readOffset);
 
 			// Test some searches
-			LogIndexNode index = await log.IndexRef.ExpandAsync();
+			LogIndexNode index = await log.IndexRef.ReadBlobAsync();
 			await SearchLogDataTestAsync(index);
 		}
 
@@ -118,15 +118,15 @@ namespace EpicGames.Horde.Tests
 				builder.WriteData(Encoding.UTF8.GetBytes(lines[lineIdx]));
 			}
 
-			HashedNodeRef<LogNode> rootNodeRef;
+			IBlobHandle<LogNode> rootNodeRef;
 			await using (IStorageWriter writer = store.CreateWriter())
 			{
-				rootNodeRef = await builder.FlushAsync(writer, true, CancellationToken.None);
+				rootNodeRef = await builder.FlushAsync(writer, true, null, CancellationToken.None);
 			}
 
 			// Read it back in and test the index
-			LogNode rootNode = await rootNodeRef.ExpandAsync();
-			LogIndexNode index = await rootNode.IndexRef.ExpandAsync();
+			LogNode rootNode = await rootNodeRef.ReadBlobAsync();
+			LogIndexNode index = await rootNode.IndexRef.ReadBlobAsync();
 		
 			for (int lineIdx = 0; lineIdx < lines.Length; lineIdx++)
 			{
@@ -137,7 +137,7 @@ namespace EpicGames.Horde.Tests
 						string str = lines[lineIdx].Substring(strOfs, strLen);
 
 						SearchStats stats = new SearchStats();
-						List<int> results = await index.SearchAsync(0, new SearchTerm(str), stats, CancellationToken.None).ToListAsync();
+						List<int> results = await index.SearchAsync(0, new SearchTerm(str), stats).ToListAsync();
 						Assert.AreEqual(1, results.Count);
 						Assert.AreEqual(lineIdx, results[0]);
 
@@ -205,7 +205,7 @@ namespace EpicGames.Horde.Tests
 		static async Task SearchLogDataTestAsync(LogIndexNode index, string text, int firstLine, int count, int[] expectedLines)
 		{
 			SearchStats stats = new SearchStats();
-			List<int> lines = await index.SearchAsync(firstLine, new SearchTerm(text), stats, CancellationToken.None).Take(count).ToListAsync();
+			List<int> lines = await index.SearchAsync(firstLine, new SearchTerm(text), stats, null, CancellationToken.None).Take(count).ToListAsync();
 			Assert.IsTrue(lines.SequenceEqual(expectedLines));
 		}
 	}
diff --git a/Engine/Source/Programs/Shared/EpicGames.Horde.Tests/StorageClientTests.cs b/Engine/Source/Programs/Shared/EpicGames.Horde.Tests/StorageClientTests.cs
--- a/Engine/Source/Programs/Shared/EpicGames.Horde.Tests/StorageClientTests.cs
+++ b/Engine/Source/Programs/Shared/EpicGames.Horde.Tests/StorageClientTests.cs
@@ -1,6 +1,8 @@
 // Copyright Epic Games, Inc. All Rights Reserved.
 
 using System;
+using System.Collections.Generic;
+using System.Threading;
 using System.Threading.Tasks;
 using EpicGames.Core;
 using EpicGames.Horde.Storage;
@@ -15,37 +17,48 @@ namespace EpicGames.Horde.Tests
 	[TestClass]
 	public class StorageClientTests
 	{
-		[BlobType("{99601905-4F6E-A089-03D6-F187711BAFEE}", 1)]
-		class TestNode : Node
+		[BlobConverter(typeof(TestNodeConverter))]
+		class TestNode
 		{
+			public static BlobType BlobType = new BlobType("{99601905-4F6E-A089-03D6-F187711BAFEE}", 1);
+
 			public int Value { get; }
 			public byte[] Padding { get; set; } = Array.Empty<byte>();
-			public NodeRef<TestNode>[] Refs { get; }
+			public IBlobHandle<TestNode>[] Refs { get; }
 
-			public TestNode(int value, params NodeRef<TestNode>[] refs)
+			public TestNode(int value, params IBlobHandle<TestNode>[] refs)
 			{
 				Value = value;
 				Refs = refs;
 			}
 
-			public TestNode(IBlobReader reader)
+			public TestNode(int value, byte[] padding, IBlobHandle<TestNode>[] refs)
 			{
-				Value = reader.ReadInt32();
-				Padding = reader.ReadVariableLengthBytes().ToArray();
-				Refs = reader.ReadVariableLengthArray(() => reader.ReadNodeRef<TestNode>());
+				Value = value;
+				Padding = padding;
+				Refs = refs;
 			}
+		}
 
-			public override void Serialize(IBlobWriter writer)
+		class TestNodeConverter : BlobConverter<TestNode>
+		{
+			public override TestNode Read(IBlobReader reader, BlobSerializerOptions options)
 			{
-				writer.WriteInt32(Value);
-				writer.WriteVariableLengthBytes(Padding);
-				writer.WriteVariableLengthArray(Refs, x => writer.WriteNodeRef(x));
+				int value = reader.ReadInt32();
+				byte[] padding = reader.ReadVariableLengthBytes().ToArray();
+				IBlobHandle<TestNode>[] refs = reader.ReadVariableLengthArray(() => reader.ReadBlobHandle<TestNode>());
+
+				return new TestNode(value, padding, refs);
 			}
-		}
 
-		static StorageClientTests()
-		{
-			Node.RegisterType<TestNode>();
+			public override BlobType Write(IBlobWriter writer, TestNode value, BlobSerializerOptions options)
+			{
+				writer.WriteInt32(value.Value);
+				writer.WriteVariableLengthBytes(value.Padding);
+				writer.WriteVariableLengthArray(value.Refs, x => writer.WriteBlobHandle(x));
+
+				return TestNode.BlobType;
+			}
 		}
 
 		[TestMethod]
@@ -66,12 +79,12 @@ namespace EpicGames.Horde.Tests
 
 		static async Task TestBasicAsync(IStorageClient store)
 		{
-			NodeRef nodeRef;
+			IBlobHandle<TestNode> nodeRef;
 			await using (IStorageWriter writer = store.CreateWriter())
 			{
-				nodeRef = await writer.WriteHashedNodeAsync(new TestNode(123));
+				nodeRef = await writer.WriteBlobAsync(new TestNode(123));
 			}
-			await store.WriteRefAsync("hello", nodeRef.Handle);
+			await store.WriteRefAsync("hello", nodeRef);
 
 			TestNode output = await store.ReadRefAsync<TestNode>("hello");
 			Assert.AreEqual(123, output.Value);
@@ -85,19 +98,19 @@ namespace EpicGames.Horde.Tests
 			using MemoryStorageClient memoryStore = new MemoryStorageClient();
 			using BundleStorageClient store = new BundleStorageClient(memoryStore, cache, NullLogger.Instance);
 
-			NodeRef<TestNode> nodeRef2;
+			IBlobHandle<TestNode> nodeRef2;
 			await using (IStorageWriter writer = store.CreateWriter())
 			{
-				NodeRef<TestNode> nodeRef1 = await writer.WriteNodeAsync(new TestNode(123));
-				nodeRef2 = await writer.WriteNodeAsync(new TestNode(456, nodeRef1));
+				IBlobHandle<TestNode> nodeRef1 = await writer.WriteBlobAsync(new TestNode(123));
+				nodeRef2 = await writer.WriteBlobAsync(new TestNode(456, nodeRef1));
 			}
-			await store.WriteRefAsync("hello", nodeRef2.Handle);
+			await store.WriteRefAsync("hello", nodeRef2);
 
 			TestNode output2 = await store.ReadRefAsync<TestNode>("hello");
 			Assert.AreEqual(456, output2.Value);
 			Assert.AreEqual(1, output2.Refs.Length);
 
-			TestNode output1 = await output2.Refs[0].ExpandAsync();
+			TestNode output1 = await output2.Refs[0].ReadBlobAsync();
 			Assert.AreEqual(0, output1.Refs.Length);
 			Assert.AreEqual(123, output1.Value);
 		}
@@ -111,23 +124,23 @@ namespace EpicGames.Horde.Tests
 			using BundleStorageClient store = new BundleStorageClient(memoryStore, cache, NullLogger.Instance);
 
 			await using IStorageWriter writer = store.CreateWriter(options: new BundleOptions { MinCompressionPacketSize = 100, MaxBlobSize = 1024 * 1024, MaxVersion = BundleVersion.LatestV2 });
-			NodeRef<TestNode> nodeRef1 = await writer.WriteNodeAsync(new TestNode(123) { Padding = new byte[1024] });
+			IBlobHandle<TestNode> nodeRef1 = await writer.WriteBlobAsync(new TestNode(123) { Padding = new byte[1024] });
 			await writer.FlushAsync();
-			NodeRef<TestNode> nodeRef2 = await writer.WriteNodeAsync(new TestNode(456, nodeRef1) { Padding = new byte[1024] });
-			NodeRef<TestNode> nodeRef3 = await writer.WriteNodeAsync(new TestNode(789, nodeRef2));
+			IBlobHandle<TestNode> nodeRef2 = await writer.WriteBlobAsync(new TestNode(456, nodeRef1) { Padding = new byte[1024] });
+			IBlobHandle<TestNode> nodeRef3 = await writer.WriteBlobAsync(new TestNode(789, nodeRef2));
 
 			// nodeRef1 is in a flushed bundle
-			BundleWriter.PendingExportHandle export1 = (BundleWriter.PendingExportHandle)nodeRef1.Handle;
+			BundleWriter.PendingExportHandle export1 = (BundleWriter.PendingExportHandle)nodeRef1.Unwrap();
 			BundleWriter.PendingPacketHandle packet1 = (BundleWriter.PendingPacketHandle)export1.Outer!;
 			BundleWriter.PendingBundleHandle bundle1 = (BundleWriter.PendingBundleHandle)packet1.Outer!;
 			Assert.IsNotNull(packet1.FlushedHandle);
 			Assert.IsNotNull(bundle1.FlushedHandle);
 
-			TestNode node1 = await export1.ReadNodeAsync<TestNode>();
+			TestNode node1 = await export1.ReadBlobAsync<TestNode>();
 			Assert.AreEqual(123, node1.Value);
 
 			// nodeRef2 is in a flushed packet, unflushed bundle
-			BundleWriter.PendingExportHandle export2 = (BundleWriter.PendingExportHandle)nodeRef2.Handle;
+			BundleWriter.PendingExportHandle export2 = (BundleWriter.PendingExportHandle)nodeRef2.Unwrap();
 
 			BundleWriter.PendingPacketHandle packet2 = (BundleWriter.PendingPacketHandle)export2.Outer!;
 			Assert.AreNotEqual(packet1, packet2);
@@ -137,11 +150,11 @@ namespace EpicGames.Horde.Tests
 			Assert.AreNotEqual(bundle1, bundle2);
 			Assert.IsNull(bundle2.FlushedHandle);
 
-			TestNode node2 = await export2.ReadNodeAsync<TestNode>();
+			TestNode node2 = await export2.ReadBlobAsync<TestNode>();
 			Assert.AreEqual(456, node2.Value);
 
 			// nodeRef3 is in an unflushed packet, unflushed bundle
-			BundleWriter.PendingExportHandle export3 = (BundleWriter.PendingExportHandle)nodeRef3.Handle;
+			BundleWriter.PendingExportHandle export3 = (BundleWriter.PendingExportHandle)nodeRef3.Unwrap();
 
 			BundleWriter.PendingPacketHandle packet3 = (BundleWriter.PendingPacketHandle)export3.Outer!;
 			Assert.AreNotEqual(packet2, packet3);
@@ -150,7 +163,7 @@ namespace EpicGames.Horde.Tests
 			BundleWriter.PendingBundleHandle bundle3 = (BundleWriter.PendingBundleHandle)packet3.Outer!;
 			Assert.AreEqual(bundle2, bundle3);
 
-			TestNode node3 = await export3.ReadNodeAsync<TestNode>();
+			TestNode node3 = await export3.ReadBlobAsync<TestNode>();
 			Assert.AreEqual(789, node3.Value);
 		}
 	}
diff --git a/Engine/Source/Programs/Shared/EpicGames.Horde/Compute/AgentMessage.cs b/Engine/Source/Programs/Shared/EpicGames.Horde/Compute/AgentMessage.cs
--- a/Engine/Source/Programs/Shared/EpicGames.Horde/Compute/AgentMessage.cs
+++ b/Engine/Source/Programs/Shared/EpicGames.Horde/Compute/AgentMessage.cs
@@ -371,11 +371,6 @@ namespace EpicGames.Horde.Compute
 		/// <summary>
 		/// Creates a sandbox on the remote machine
 		/// </summary>
-		/// <param name="channel">Current channel</param>
-		/// <param name="path">Root directory to extract files within the sandbox</param>
-		/// <param name="locator">Location of a <see cref="DirectoryNode"/> describing contents of the sandbox</param>
-		/// <param name="storage">Storage for the sandbox data</param>
-		/// <param name="cancellationToken">Cancellation token for the operation</param>
 		public static async Task UploadFilesAsync(this AgentMessageChannel channel, string path, BlobLocator locator, IStorageClient storage, CancellationToken cancellationToken = default)
 		{
 			using (IAgentMessageBuilder request = await channel.CreateMessageAsync(AgentMessageType.WriteFiles, cancellationToken))
diff --git a/Engine/Source/Programs/Shared/EpicGames.Horde/Compute/AgentMessageHandler.cs b/Engine/Source/Programs/Shared/EpicGames.Horde/Compute/AgentMessageHandler.cs
--- a/Engine/Source/Programs/Shared/EpicGames.Horde/Compute/AgentMessageHandler.cs
+++ b/Engine/Source/Programs/Shared/EpicGames.Horde/Compute/AgentMessageHandler.cs
@@ -107,7 +107,7 @@ namespace EpicGames.Horde.Compute
 						case AgentMessageType.WriteFiles:
 							{
 								UploadFilesMessage writeFiles = message.ParseUploadFilesMessage();
-								await WriteFilesAsync(channel, writeFiles.Name, writeFiles.Locator, cancellationToken);
+								await WriteFilesAsync(channel, writeFiles.Name, writeFiles.Locator, cancellationToken: cancellationToken);
 							}
 							break;
 						case AgentMessageType.DeleteFiles:
@@ -170,14 +170,14 @@ namespace EpicGames.Horde.Compute
 			}
 		}
 
-		async Task WriteFilesAsync(AgentMessageChannel channel, string path, BlobLocator locator, CancellationToken cancellationToken)
+		async Task WriteFilesAsync(AgentMessageChannel channel, string path, BlobLocator locator, BlobSerializerOptions? options = null, CancellationToken cancellationToken = default)
 		{
 			using AgentStorageClient innerStore = new AgentStorageClient(channel);
 			await using BundleCache cache = new BundleCache(new BundleCacheOptions { HeaderCacheSize = 10 * 1024 * 1024, PacketCacheSize = 128 * 1024 * 1024 });
 			using BundleStorageClient store = new BundleStorageClient(innerStore, cache, _logger);
 
 			IBlobHandle handle = store.CreateBlobHandle(locator);
-			DirectoryNode directoryNode = await handle.ReadNodeAsync<DirectoryNode>(cancellationToken);
+			DirectoryNode directoryNode = await handle.ReadBlobAsync<DirectoryNode>(options, cancellationToken);
 
 			DirectoryReference outputDir = DirectoryReference.Combine(_sandboxDir, path);
 			if (!outputDir.IsUnderDirectory(_sandboxDir))
@@ -185,7 +185,7 @@ namespace EpicGames.Horde.Compute
 				throw new InvalidOperationException("Cannot write files outside sandbox");
 			}
 
-			await directoryNode.CopyToDirectoryAsync(outputDir.ToDirectoryInfo(), _logger, cancellationToken);
+			await directoryNode.CopyToDirectoryAsync(outputDir.ToDirectoryInfo(), options, _logger, cancellationToken);
 
 			using (IAgentMessageBuilder message = await channel.CreateMessageAsync(AgentMessageType.WriteFilesResponse, cancellationToken))
 			{
diff --git a/Engine/Source/Programs/Shared/EpicGames.Horde/Logs/LogChunkNode.cs b/Engine/Source/Programs/Shared/EpicGames.Horde/Logs/LogChunkNode.cs
--- a/Engine/Source/Programs/Shared/EpicGames.Horde/Logs/LogChunkNode.cs
+++ b/Engine/Source/Programs/Shared/EpicGames.Horde/Logs/LogChunkNode.cs
@@ -14,9 +14,14 @@ namespace EpicGames.Horde.Logs
 	/// <summary>
 	/// Read-only buffer for log text, with indexed line offsets.
 	/// </summary>
-	[BlobType("{7020B6CA-4174-0F72-AA06-AAB630EFA360}", 1)]
-	public class LogChunkNode : Node
+	[BlobConverter(typeof(LogChunkConverter))]
+	public class LogChunkNode
 	{
+		/// <summary>
+		/// Type of this blob when serialized
+		/// </summary>
+		public static BlobType BlobType { get; } = new BlobType("{7020B6CA-4174-0F72-AA06-AAB630EFA360}", 1);
+
 		/// <summary>
 		/// Provides access to the lines for this chunk through a list interface
 		/// </summary>
@@ -107,21 +112,6 @@ namespace EpicGames.Horde.Logs
 			LineOffsets = lineOffsets;
 		}
 
-		/// <summary>
-		/// Deserializing constructor
-		/// </summary>
-		/// <param name="reader">Reader to pull data from</param>
-		public LogChunkNode(IBlobReader reader)
-			: this(reader.ReadVariableLengthBytes())
-		{
-		}
-
-		/// <inheritdoc/>
-		public override void Serialize(IBlobWriter writer)
-		{
-			writer.WriteVariableLengthBytes(Data.Span);
-		}
-
 		/// <summary>
 		/// Accessor for an individual line
 		/// </summary>
@@ -182,10 +172,30 @@ namespace EpicGames.Horde.Logs
 		}
 	}
 
+	/// <summary>
+	/// Converter from log chunks to blobs
+	/// </summary>
+	class LogChunkConverter : BlobConverter<LogChunkNode>
+	{
+		/// <inheritdoc/>
+		public override LogChunkNode Read(IBlobReader reader, BlobSerializerOptions options)
+		{
+			byte[] data = reader.ReadVariableLengthBytes().ToArray();
+			return new LogChunkNode(data);
+		}
+
+		/// <inheritdoc/>
+		public override BlobType Write(IBlobWriter writer, LogChunkNode value, BlobSerializerOptions options)
+		{
+			writer.WriteVariableLengthBytes(value.Data.Span);
+			return LogChunkNode.BlobType;
+		}
+	}
+
 	/// <summary>
 	/// Reference to a chunk of text, with information about its placement in the larger log file
 	/// </summary>
-	public class LogChunkRef : HashedNodeRef<LogChunkNode>
+	public class LogChunkRef
 	{
 		/// <summary>
 		/// First line within the file
@@ -207,6 +217,11 @@ namespace EpicGames.Horde.Logs
 		/// </summary>
 		public int Length { get; }
 
+		/// <summary>
+		/// Handle to the target chunk
+		/// </summary>
+		public IBlobHandle<LogChunkNode> Target { get; }
+
 		/// <summary>
 		/// Constructor
 		/// </summary>
@@ -215,13 +230,13 @@ namespace EpicGames.Horde.Logs
 		/// <param name="offset">Offset within the log file</param>
 		/// <param name="length">Length of the chunk</param>
 		/// <param name="target">Referenced log text</param>
-		public LogChunkRef(int lineIndex, int lineCount, long offset, int length, HashedNodeRef<LogChunkNode> target)
-			: base(target)
+		public LogChunkRef(int lineIndex, int lineCount, long offset, int length, IBlobHandle<LogChunkNode> target)
 		{
 			LineIndex = lineIndex;
 			LineCount = lineCount;
 			Offset = offset;
 			Length = length;
+			Target = target;
 		}
 
 		/// <summary>
@@ -229,8 +244,8 @@ namespace EpicGames.Horde.Logs
 		/// </summary>
 		/// <param name="reader"></param>
 		public LogChunkRef(IBlobReader reader)
-			: base(reader)
 		{
+			Target = reader.ReadBlobHandle<LogChunkNode>();
 			LineIndex = (int)reader.ReadUnsignedVarInt();
 			LineCount = (int)reader.ReadUnsignedVarInt();
 			Offset = (long)reader.ReadUnsignedVarInt();
@@ -238,10 +253,9 @@ namespace EpicGames.Horde.Logs
 		}
 
 		/// <inheritdoc/>
-		public override void Serialize(IBlobWriter writer)
+		public void Serialize(IBlobWriter writer)
 		{
-			base.Serialize(writer);
-
+			writer.WriteBlobHandle(Target);
 			writer.WriteUnsignedVarInt(LineIndex);
 			writer.WriteUnsignedVarInt(LineCount);
 			writer.WriteUnsignedVarInt((ulong)Offset);
diff --git a/Engine/Source/Programs/Shared/EpicGames.Horde/Logs/LogIndexNode.cs b/Engine/Source/Programs/Shared/EpicGames.Horde/Logs/LogIndexNode.cs
--- a/Engine/Source/Programs/Shared/EpicGames.Horde/Logs/LogIndexNode.cs
+++ b/Engine/Source/Programs/Shared/EpicGames.Horde/Logs/LogIndexNode.cs
@@ -18,7 +18,7 @@ namespace EpicGames.Horde.Logs
 	/// Index for a log file.
 	/// 
 	/// The index consists of a sequence of compressed, plain text blocks (see <see cref="LogChunkRef"/>, and a 
-	/// set of (ngram, block index) pairs encoded as 64-bit integers (see <see cref="NgramSet"/>). 
+	/// set of (ngram, block index) pairs encoded as 64-bit integers (see <see cref="Logs.NgramSet"/>). 
 	/// 
 	/// Each ngram is a 1-4 byte sequence of utf8 bytes, padded out to 32-bits (see <see cref="Ngram"/>).
 	/// 
@@ -28,23 +28,18 @@ namespace EpicGames.Horde.Logs
 	/// Since alignment of ngrams may not match alignment of ngrams in the search term, we offset the search term by
 	/// 1-4 bytes and include the union of blocks matching at any offset.
 	/// </summary>
-	[BlobType("{BAE1A00E-474E-FD63-8180-04A8F93401E2}", 1)]
-	public class LogIndexNode : Node
+	[BlobConverter(typeof(LogIndexNodeConverter))]
+	public class LogIndexNode
 	{
-		/// <summary>
-		/// Version number for serialized data
-		/// </summary>
-		const int CurrentVersion = 1;
-
 		/// <summary>
 		/// Index for tokens into the block list
 		/// </summary>
-		readonly NgramSet _ngramSet;
+		public NgramSet NgramSet { get; }
 
 		/// <summary>
 		/// Number of bits in the index devoted to the block index
 		/// </summary>
-		readonly int _numChunkBits;
+		public int NumChunkBits { get; }
 
 		/// <summary>
 		/// List of text blocks
@@ -79,45 +74,20 @@ namespace EpicGames.Horde.Logs
 		/// <param name="plainTextChunkRefs">Plain text chunks for this log file</param>
 		public LogIndexNode(NgramSet ngramSet, int numChunkBits, LogChunkRef[] plainTextChunkRefs)
 		{
-			_ngramSet = ngramSet;
-			_numChunkBits = numChunkBits;
+			NgramSet = ngramSet;
+			NumChunkBits = numChunkBits;
 			_plainTextChunkRefs = plainTextChunkRefs;
 		}
 
-		/// <summary>
-		/// Deserialization constructor
-		/// </summary>
-		/// <param name="reader">Reader for data</param>
-		public LogIndexNode(IBlobReader reader)
-		{
-			int version = (int)reader.ReadUnsignedVarInt();
-			if (version != CurrentVersion)
-			{
-				throw new InvalidDataException("Incorrect version number for log data");
-			}
-
-			_ngramSet = reader.ReadNgramSet();
-			_numChunkBits = (int)reader.ReadUnsignedVarInt();
-			_plainTextChunkRefs = reader.ReadVariableLengthArray(() => new LogChunkRef(reader));
-		}
-
-		/// <inheritdoc/>
-		public override void Serialize(IBlobWriter writer)
-		{
-			writer.WriteUnsignedVarInt(CurrentVersion);
-			writer.WriteNgramSet(_ngramSet);
-			writer.WriteUnsignedVarInt(_numChunkBits);
-			writer.WriteVariableLengthArray(_plainTextChunkRefs, x => writer.WriteHashedNodeRef(x));
-		}
-
 		/// <summary>
 		/// Appends a set of text blocks to this index
 		/// </summary>
 		/// <param name="writer">Writer for output nodes</param>
 		/// <param name="appendPlainTextChunks">Text blocks to append</param>
+		/// <param name="options">Options to control serialization</param>
 		/// <param name="cancellationToken"></param>
 		/// <returns>New log index with the given blocks appended</returns>
-		public async ValueTask<LogIndexNode> AppendAsync(IStorageWriter writer, IReadOnlyList<LogChunkNode> appendPlainTextChunks, CancellationToken cancellationToken)
+		public async ValueTask<LogIndexNode> AppendAsync(IStorageWriter writer, IReadOnlyList<LogChunkNode> appendPlainTextChunks, BlobSerializerOptions? options, CancellationToken cancellationToken)
 		{
 			using IScope scope = GlobalTracer.Instance.BuildSpan("LogIndex.Append").StartActive();
 
@@ -133,7 +103,7 @@ namespace EpicGames.Horde.Logs
 			for (int idx = 0; idx < appendPlainTextChunks.Count; idx++)
 			{
 				LogChunkNode newChunk = appendPlainTextChunks[idx];
-				HashedNodeRef<LogChunkNode> newChunkRef = await writer.WriteHashedNodeAsync(newChunk, cancellationToken);
+				IBlobHandle<LogChunkNode> newChunkRef = await writer.WriteBlobAsync(newChunk, options, cancellationToken);
 				newChunks[_plainTextChunkRefs.Length + idx] = new LogChunkRef(lineIndex, newChunk.LineCount, offset, newChunk.Length, newChunkRef);
 				lineIndex += newChunk.LineCount;
 				offset += newChunk.Length;
@@ -149,10 +119,10 @@ namespace EpicGames.Horde.Logs
 			// Add the existing index data into a new ngram set
 			NgramSetBuilder newNgramSetBuilder = new NgramSetBuilder();
 
-			ulong blockMask = ((1UL << _numChunkBits) - 1);
-			foreach (ulong value in _ngramSet)
+			ulong blockMask = ((1UL << NumChunkBits) - 1);
+			foreach (ulong value in NgramSet)
 			{
-				ulong token = value >> _numChunkBits;
+				ulong token = value >> NumChunkBits;
 				ulong blockIdx = value & blockMask;
 
 				ulong newToken = (token << newNumBlockBits) | (blockIdx);
@@ -178,9 +148,10 @@ namespace EpicGames.Horde.Logs
 		/// <param name="firstLineIndex">First line index to search from</param>
 		/// <param name="text">Text to search for</param>
 		/// <param name="stats">Receives stats for the search</param>
+		/// <param name="options">Options for serialization</param>
 		/// <param name="cancellationToken">Cancellation token for the operation</param>
 		/// <returns>List of line numbers for the text</returns>
-		public async IAsyncEnumerable<int> SearchAsync(int firstLineIndex, SearchTerm text, SearchStats stats, [EnumeratorCancellation] CancellationToken cancellationToken)
+		public async IAsyncEnumerable<int> SearchAsync(int firstLineIndex, SearchTerm text, SearchStats stats, BlobSerializerOptions? options = null, [EnumeratorCancellation] CancellationToken cancellationToken = default)
 		{
 			int lastBlockCount = 0;
 			foreach (int blockIdx in EnumeratePossibleChunks(text.Bytes, firstLineIndex))
@@ -194,7 +165,7 @@ namespace EpicGames.Horde.Logs
 				lastBlockCount = blockIdx + 1;
 
 				// Decompress the text
-				LogChunkNode chunk = await indexChunk.ExpandAsync(cancellationToken);
+				LogChunkNode chunk = await indexChunk.Target.ReadBlobAsync(options, cancellationToken);
 
 				// Find the initial offset within this block
 				int offset = 0;
@@ -228,7 +199,7 @@ namespace EpicGames.Horde.Logs
 				}
 
 				// If the last scanned bytes is zero, we didn't have any matches from this chunk
-				if(offset == 0 && _ngramSet != null)
+				if(offset == 0 && NgramSet != null)
 				{
 					stats.NumFalsePositiveBlocks++;
 				}
@@ -266,7 +237,7 @@ namespace EpicGames.Horde.Logs
 				List<Predicate<int>> predicates = new List<Predicate<int>>();
 
 				// If the index has a trie, tokenize the input and generate a list of enumerators and predicates.
-				if (_ngramSet != null)
+				if (NgramSet != null)
 				{
 					// Find a list of filters for matching chunks
 					HashSet<ulong> tokens = new HashSet<ulong>();
@@ -287,10 +258,10 @@ namespace EpicGames.Horde.Logs
 					// Create an enumerator for each token
 					foreach (ulong token in tokens)
 					{
-						ulong minValue = (token << _numChunkBits) | (ulong)(long)chunkIdx;
-						ulong maxValue = minValue + (1UL << _numChunkBits) - 1;
+						ulong minValue = (token << NumChunkBits) | (ulong)(long)chunkIdx;
+						ulong maxValue = minValue + (1UL << NumChunkBits) - 1;
 
-						IEnumerator<int> enumerator = _ngramSet.EnumerateRange(minValue, maxValue).Select(x => (int)(x - minValue)).GetEnumerator();
+						IEnumerator<int> enumerator = NgramSet.EnumerateRange(minValue, maxValue).Select(x => (int)(x - minValue)).GetEnumerator();
 						if (!enumerator.MoveNext())
 						{
 							yield break;
@@ -397,9 +368,9 @@ namespace EpicGames.Horde.Logs
 		/// <returns>True if the given block contains a token</returns>
 		bool ChunkContainsToken(int chunkIdx, ulong token, ulong tokenMask)
 		{
-			token = (token << _numChunkBits) | (uint)chunkIdx;
-			tokenMask = (tokenMask << _numChunkBits) | ((1UL << _numChunkBits) - 1);
-			return _ngramSet!.EnumerateValues((value, valueMask) => (value & tokenMask) == (token & valueMask)).Any();
+			token = (token << NumChunkBits) | (uint)chunkIdx;
+			tokenMask = (tokenMask << NumChunkBits) | ((1UL << NumChunkBits) - 1);
+			return NgramSet!.EnumerateValues((value, valueMask) => (value & tokenMask) == (token & valueMask)).Any();
 		}
 
 		/// <summary>
@@ -411,10 +382,50 @@ namespace EpicGames.Horde.Logs
 		/// <returns>True if the given block contains a token</returns>
 		IEnumerable<int> ChunksContainingNgram(ReadOnlySpan<byte> text, int offset, bool allowPartialMatch)
 		{
-			ulong token = Ngram.GetWindowedValue(text, offset) << _numChunkBits;
-			ulong tokenMask = Ngram.GetWindowedMask(text, offset, allowPartialMatch) << _numChunkBits;
-			ulong blockMask = (1UL << _numChunkBits) - 1;
-			return _ngramSet!.EnumerateValues((value, valueMask) => (value & tokenMask) == (token & valueMask)).Select(x => (int)(x & blockMask)).Distinct();
+			ulong token = Ngram.GetWindowedValue(text, offset) << NumChunkBits;
+			ulong tokenMask = Ngram.GetWindowedMask(text, offset, allowPartialMatch) << NumChunkBits;
+			ulong blockMask = (1UL << NumChunkBits) - 1;
+			return NgramSet!.EnumerateValues((value, valueMask) => (value & tokenMask) == (token & valueMask)).Select(x => (int)(x & blockMask)).Distinct();
+		}
+	}
+
+	class LogIndexNodeConverter : BlobConverter<LogIndexNode>
+	{
+		/// <summary>
+		/// Type for serializing index nodes
+		/// </summary>
+		public static BlobType BlobType { get; } = new BlobType("{BAE1A00E-474E-FD63-8180-04A8F93401E2}", 1);
+
+		/// <summary>
+		/// Version number for serialized data
+		/// </summary>
+		const int CurrentVersion = 1;
+
+		/// <inheritdoc/>
+		public override LogIndexNode Read(IBlobReader reader, BlobSerializerOptions options)
+		{
+			int version = (int)reader.ReadUnsignedVarInt();
+			if (version != CurrentVersion)
+			{
+				throw new InvalidDataException("Incorrect version number for log data");
+			}
+
+			NgramSet ngramSet = reader.ReadNgramSet();
+			int numChunkBits = (int)reader.ReadUnsignedVarInt();
+			LogChunkRef[] plainTextChunkRefs = reader.ReadVariableLengthArray(() => new LogChunkRef(reader));
+
+			return new LogIndexNode(ngramSet, numChunkBits, plainTextChunkRefs);
+		}
+
+		/// <inheritdoc/>
+		public override BlobType Write(IBlobWriter writer, LogIndexNode value, BlobSerializerOptions options)
+		{
+			writer.WriteUnsignedVarInt(CurrentVersion);
+			writer.WriteNgramSet(value.NgramSet);
+			writer.WriteUnsignedVarInt(value.NumChunkBits);
+			writer.WriteVariableLengthArray(value.PlainTextChunkRefs, x => x.Serialize(writer));
+
+			return BlobType;
 		}
 	}
 }
diff --git a/Engine/Source/Programs/Shared/EpicGames.Horde/Logs/LogNode.cs b/Engine/Source/Programs/Shared/EpicGames.Horde/Logs/LogNode.cs
--- a/Engine/Source/Programs/Shared/EpicGames.Horde/Logs/LogNode.cs
+++ b/Engine/Source/Programs/Shared/EpicGames.Horde/Logs/LogNode.cs
@@ -32,8 +32,8 @@ namespace EpicGames.Horde.Logs
 	/// <summary>
 	/// Represents an entire log
 	/// </summary>
-	[BlobType("{274DF8F7-4B4F-9E87-8C31-D58A33AD25DB}", 1)]
-	public class LogNode : Node
+	[BlobConverter(typeof(LogNodeConverter))]
+	public class LogNode
 	{
 		/// <summary>
 		/// Format for this log file
@@ -58,7 +58,7 @@ namespace EpicGames.Horde.Logs
 		/// <summary>
 		/// Index for this log
 		/// </summary>
-		public HashedNodeRef<LogIndexNode> IndexRef { get; }
+		public IBlobHandle<LogIndexNode> IndexRef { get; }
 
 		/// <summary>
 		/// Whether this log is complete
@@ -68,7 +68,7 @@ namespace EpicGames.Horde.Logs
 		/// <summary>
 		/// Deserializing constructor
 		/// </summary>
-		public LogNode(LogFormat format, int lineCount, long length, IReadOnlyList<LogChunkRef> textChunkRefs, HashedNodeRef<LogIndexNode> indexRef, bool complete)
+		public LogNode(LogFormat format, int lineCount, long length, IReadOnlyList<LogChunkRef> textChunkRefs, IBlobHandle<LogIndexNode> indexRef, bool complete)
 		{
 			Format = format;
 			LineCount = lineCount;
@@ -77,30 +77,42 @@ namespace EpicGames.Horde.Logs
 			IndexRef = indexRef;
 			Complete = complete;
 		}
+	}
 
+	/// <summary>
+	/// Serializer for <see cref="LogNode"/> types
+	/// </summary>
+	class LogNodeConverter : BlobConverter<LogNode>
+	{
 		/// <summary>
-		/// Deserializing constructor
+		/// Type of blob when serialized to storage
 		/// </summary>
-		/// <param name="reader">Reader to draw data from</param>
-		public LogNode(IBlobReader reader)
+		public static BlobType BlobType { get; } = new BlobType("{274DF8F7-4B4F-9E87-8C31-D58A33AD25DB}", 1);
+
+		/// <inheritdoc/>
+		public override LogNode Read(IBlobReader reader, BlobSerializerOptions options)
 		{
-			Format = (LogFormat)reader.ReadUInt8();
-			LineCount = (int)reader.ReadUnsignedVarInt();
-			Length = (long)reader.ReadUnsignedVarInt();
-			IndexRef = reader.ReadHashedNodeRef<LogIndexNode>();
-			TextChunkRefs = reader.ReadList(() => new LogChunkRef(reader));
-			Complete = reader.ReadBoolean();
+			LogFormat format = (LogFormat)reader.ReadUInt8();
+			int lineCount = (int)reader.ReadUnsignedVarInt();
+			long length = (long)reader.ReadUnsignedVarInt();
+			IBlobHandle<LogIndexNode> indexRef = reader.ReadBlobHandle<LogIndexNode>();
+			List<LogChunkRef> textChunkRefs = reader.ReadList(() => new LogChunkRef(reader));
+			bool complete = reader.ReadBoolean();
+
+			return new LogNode(format, lineCount, length, textChunkRefs, indexRef, complete);
 		}
 
 		/// <inheritdoc/>
-		public override void Serialize(IBlobWriter writer)
+		public override BlobType Write(IBlobWriter writer, LogNode value, BlobSerializerOptions options)
 		{
-			writer.WriteUInt8((byte)Format);
-			writer.WriteUnsignedVarInt(LineCount);
-			writer.WriteUnsignedVarInt((ulong)Length);
-			writer.WriteHashedNodeRef(IndexRef);
-			writer.WriteList(TextChunkRefs, x => writer.WriteHashedNodeRef(x));
-			writer.WriteBoolean(Complete);
+			writer.WriteUInt8((byte)value.Format);
+			writer.WriteUnsignedVarInt(value.LineCount);
+			writer.WriteUnsignedVarInt((ulong)value.Length);
+			writer.WriteBlobHandle(value.IndexRef);
+			writer.WriteList(value.TextChunkRefs, x => x.Serialize(writer));
+			writer.WriteBoolean(value.Complete);
+
+			return BlobType;
 		}
 	}
 
@@ -249,8 +261,9 @@ namespace EpicGames.Horde.Logs
 		/// </summary>
 		/// <param name="writer">Writer for the output nodes</param>
 		/// <param name="complete">Whether the log is complete</param>
+		/// <param name="options">Options for serialization</param>
 		/// <param name="cancellationToken">Cancellation token for the operation</param>
-		public async Task<HashedNodeRef<LogNode>> FlushAsync(IStorageWriter writer, bool complete, CancellationToken cancellationToken)
+		public async Task<IBlobHandle<LogNode>> FlushAsync(IStorageWriter writer, bool complete, BlobSerializerOptions? options, CancellationToken cancellationToken)
 		{
 			// Capture the new data that needs to be written
 			IReadOnlyList<LogChunkNode> writeTextChunks;
@@ -266,22 +279,22 @@ namespace EpicGames.Horde.Logs
 			}
 
 			// Flush any complete chunks to storage
-			LogIndexNode newIndex = await _index.AppendAsync(writer, writeIndexTextChunks, cancellationToken);
-			HashedNodeRef<LogIndexNode> newIndexRef = await writer.WriteHashedNodeAsync(newIndex, cancellationToken);
+			LogIndexNode newIndex = await _index.AppendAsync(writer, writeIndexTextChunks, null, cancellationToken);
+			IBlobHandle<LogIndexNode> newIndexRef = await writer.WriteBlobAsync(newIndex, options, cancellationToken);
 
 			List<LogChunkRef> newJsonChunkRefs = new List<LogChunkRef>(_root?.TextChunkRefs ?? Array.Empty<LogChunkRef>());
 			int lineCount = _root?.LineCount ?? 0;
 			long length = _root?.Length ?? 0;
 			foreach (LogChunkNode writeTextChunk in writeTextChunks)
 			{
-				HashedNodeRef<LogChunkNode> writeTextChunkRef = await writer.WriteHashedNodeAsync(writeTextChunk, cancellationToken);
+				IBlobHandle<LogChunkNode> writeTextChunkRef = await writer.WriteBlobAsync(writeTextChunk, options, cancellationToken);
 				newJsonChunkRefs.Add(new LogChunkRef(lineCount, writeTextChunk.LineCount, length, writeTextChunk.Length, writeTextChunkRef));
 				lineCount += writeTextChunk.LineCount;
 				length += writeTextChunk.Length;
 			}
 
 			LogNode newRoot = new LogNode(_format, lineCount, length, newJsonChunkRefs, newIndexRef, complete);
-			HashedNodeRef<LogNode> newRootRef = await writer.WriteHashedNodeAsync(newRoot, cancellationToken);
+			IBlobHandle<LogNode> newRootRef = await writer.WriteBlobAsync(newRoot, options, cancellationToken);
 
 			await writer.FlushAsync(cancellationToken);
 
@@ -307,13 +320,14 @@ namespace EpicGames.Horde.Logs
 		/// Reads lines from a line
 		/// </summary>
 		/// <param name="logNode">Log to read from</param>
+		/// <param name="options">Options for deserialization</param>
 		/// <param name="cancellationToken">Cancellation token</param>
 		/// <returns>Sequence of line buffers</returns>
-		public static async IAsyncEnumerable<ReadOnlyMemory<byte>> ReadLogAsync(this LogNode logNode, [EnumeratorCancellation] CancellationToken cancellationToken)
+		public static async IAsyncEnumerable<ReadOnlyMemory<byte>> ReadLogAsync(this LogNode logNode, BlobSerializerOptions? options = null, [EnumeratorCancellation] CancellationToken cancellationToken = default)
 		{
 			foreach (LogChunkRef textChunkRef in logNode.TextChunkRefs)
 			{
-				LogChunkNode textChunk = await textChunkRef.ExpandAsync(cancellationToken);
+				LogChunkNode textChunk = await textChunkRef.Target.ReadBlobAsync(options, cancellationToken);
 				yield return textChunk.Data;
 			}
 		}
@@ -323,16 +337,17 @@ namespace EpicGames.Horde.Logs
 		/// </summary>
 		/// <param name="logNode">Log to read from</param>
 		/// <param name="index">Zero-based index of the first line to read from</param>
+		/// <param name="options">Options for deserialization</param>
 		/// <param name="cancellationToken">Cancellation token</param>
 		/// <returns>Sequence of line buffers</returns>
-		public static async IAsyncEnumerable<ReadOnlyMemory<byte>> ReadLogLinesAsync(this LogNode logNode, int index, [EnumeratorCancellation] CancellationToken cancellationToken)
+		public static async IAsyncEnumerable<ReadOnlyMemory<byte>> ReadLogLinesAsync(this LogNode logNode, int index, BlobSerializerOptions? options = null, [EnumeratorCancellation] CancellationToken cancellationToken = default)
 		{
 			foreach (LogChunkRef textChunkRef in logNode.TextChunkRefs)
 			{
 				int lineIdx = Math.Max(index - textChunkRef.LineIndex, 0);
 				if (lineIdx < textChunkRef.LineCount)
 				{
-					LogChunkNode textChunk = await textChunkRef.ExpandAsync(cancellationToken);
+					LogChunkNode textChunk = await textChunkRef.Target.ReadBlobAsync(options, cancellationToken);
 
 					int offset = textChunk.LineOffsets[lineIdx];
 					for (; lineIdx < textChunk.LineCount; lineIdx++)
diff --git a/Engine/Source/Programs/Shared/EpicGames.Horde/Storage/BlobConverter.cs b/Engine/Source/Programs/Shared/EpicGames.Horde/Storage/BlobConverter.cs
new file mode 100644
--- /dev/null
+++ b/Engine/Source/Programs/Shared/EpicGames.Horde/Storage/BlobConverter.cs
@@ -0,0 +1,57 @@
+// Copyright Epic Games, Inc. All Rights Reserved.
+
+using System;
+using System.Threading;
+using System.Threading.Tasks;
+
+namespace EpicGames.Horde.Storage
+{
+	/// <summary>
+	/// Attribute used to denote the converter for an object type
+	/// </summary>
+	[AttributeUsage(AttributeTargets.Class)]
+	public sealed class BlobConverterAttribute : Attribute
+	{
+		/// <summary>
+		/// Type of the converter
+		/// </summary>
+		public Type ConverterType { get; }
+
+		/// <summary>
+		/// Constructor
+		/// </summary>
+		public BlobConverterAttribute(Type converterType) => ConverterType = converterType;
+	}
+
+	/// <summary>
+	/// Base class for converter types that can serialize blobs from native C# types. Semantics mirror <see cref="System.Text.Json.Serialization.JsonConverter"/>.
+	/// </summary>
+	public abstract class BlobConverter
+	{
+		/// <summary>
+		/// Determines if the converter can handle the given type
+		/// </summary>
+		/// <param name="typeToConvert">The type that needs to be converted</param>
+		/// <returns>True if this converter can handle the given type</returns>
+		public abstract bool CanConvert(Type typeToConvert);
+	}
+
+	/// <summary>
+	/// Serializer for typed values to blob data. Mimics the <see cref="System.Text.Json.Serialization.JsonConverter{T}"/> interface for familiarity.
+	/// </summary>
+	public abstract class BlobConverter<T> : BlobConverter
+	{
+		/// <inheritdoc/>
+		public override bool CanConvert(Type type) => type == typeof(T);
+
+		/// <summary>
+		/// Reads a strongly typed value
+		/// </summary>
+		public abstract T Read(IBlobReader reader, BlobSerializerOptions options);
+
+		/// <summary>
+		/// Writes a strongly typed value
+		/// </summary>
+		public abstract BlobType Write(IBlobWriter writer, T value, BlobSerializerOptions options);
+	}
+}
diff --git a/Engine/Source/Programs/Shared/EpicGames.Horde/Storage/BlobSerializer.cs b/Engine/Source/Programs/Shared/EpicGames.Horde/Storage/BlobSerializer.cs
new file mode 100644
--- /dev/null
+++ b/Engine/Source/Programs/Shared/EpicGames.Horde/Storage/BlobSerializer.cs
@@ -0,0 +1,242 @@
+// Copyright Epic Games, Inc. All Rights Reserved.
+
+using System;
+using System.Collections.Concurrent;
+using System.Collections.Generic;
+using System.Linq;
+using System.Reflection;
+using System.Text;
+using System.Threading;
+using System.Threading.Tasks;
+using EpicGames.Core;
+using EpicGames.Serialization;
+
+namespace EpicGames.Horde.Storage
+{
+	/// <summary>
+	/// Handles serialization of blobs using <see cref="BlobConverter"/> instances.
+	/// </summary>
+	public static class BlobSerializer
+	{
+		/// <summary>
+		/// Deserialize an object
+		/// </summary>
+		/// <typeparam name="T">Return type for deserialization</typeparam>
+		/// <param name="blobData">Data to deserialize from</param>
+		/// <param name="options">Options to control serialization</param>
+		public static T Deserialize<T>(BlobData blobData, BlobSerializerOptions? options = null)
+		{
+			options ??= BlobSerializerOptions.Default;
+			BlobReader reader = new BlobReader(blobData);
+			return options.GetConverter<T>().Read(reader, options);
+		}
+		/*
+		/// <summary>
+		/// Deserialize an object
+		/// </summary>
+		/// <typeparam name="T">Return type for the deserialized object</typeparam>
+		/// <param name="handle">Handle to the blob to deserialize</param>
+		/// <param name="options">Options to control serialization</param>
+		/// <param name="cancellationToken">Cancellation token for the operation</param>
+		public static async ValueTask<T> DeserializeAsync<T>(IBlobHandle handle, BlobSerializerOptions? options = null, CancellationToken cancellationToken = default)
+		{
+			BlobData data = await handle.ReadAsync(cancellationToken);
+			return Deserialize<T>(data, options);
+		}
+
+		/// <summary>
+		/// Deserialize an object
+		/// </summary>
+		/// <typeparam name="T">Return type for the deserialized object</typeparam>
+		/// <param name="handle">Handle to the blob to deserialize</param>
+		/// <param name="options">Options to control serialization</param>
+		/// <param name="cancellationToken">Cancellation token for the operation</param>
+		public static ValueTask<T> DeserializeAsync<T>(IBlobHandle<T> handle, BlobSerializerOptions? options = null, CancellationToken cancellationToken = default)
+		{
+			return DeserializeAsync<T>((IBlobHandle)handle, options, cancellationToken);
+		}*/
+
+		/// <summary>
+		/// Serialize an object into a blob
+		/// </summary>
+		/// <typeparam name="T">Type of object to serialize</typeparam>
+		/// <param name="writer">Writer for the blob data</param>
+		/// <param name="value">Object to serialize</param>
+		/// <param name="options">Options to control serialization</param>
+		/// <returns>Type of the serialized blob</returns>
+		public static BlobType Serialize<T>(IBlobWriter writer, T value, BlobSerializerOptions? options = null)
+		{
+			options ??= BlobSerializerOptions.Default;
+			return options.GetConverter<T>().Write(writer, value, options);
+		}
+		/*
+		/// <summary>
+		/// Deserialize an object
+		/// </summary>
+		/// <param name="writer">Writer for serialized data</param>
+		/// <param name="value">The object to serialize</param>
+		/// <param name="options">Options to control serialization</param>
+		/// <param name="cancellationToken">Cancellation token for the operation</param>
+		/// <returns>Handle to the serialized blob</returns>
+		public static async ValueTask<IBlobHandle<T>> SerializeAsync<T>(IStorageWriter writer, T value, BlobSerializerOptions? options = null, CancellationToken cancellationToken = default)
+		{
+			BlobWriter blobWriter = new BlobWriter(writer);
+			BlobType blobType = Serialize<T>(blobWriter, value, options);
+			return await writer.WriteBlobAsync<T>(blobType, blobWriter.Length, blobWriter.References, cancellationToken);
+		}*/
+	}
+
+	/// <summary>
+	/// Options for serializing blobs
+	/// </summary>
+	public class BlobSerializerOptions
+	{
+		/// <summary>
+		/// Default options instance
+		/// </summary>
+		public static BlobSerializerOptions Default { get; } = new BlobSerializerOptions();
+
+		/// <summary>
+		/// Known converter types
+		/// </summary>
+		public IList<BlobConverter> Converters { get; } = new List<BlobConverter>();
+
+		readonly ConcurrentDictionary<Type, BlobConverter> _cachedConverters = new ConcurrentDictionary<Type, BlobConverter>();
+		readonly static ConcurrentDictionary<Type, BlobConverter> s_cachedDefaultConverters = new ConcurrentDictionary<Type, BlobConverter>();
+		readonly static Func<Type, BlobConverter> s_createDefaultConverter = CreateDefaultConverter;
+
+		/// <summary>
+		/// Gets a converter for the given type
+		/// </summary>
+		public BlobConverter<T> GetConverter<T>()
+		{
+			return (BlobConverter<T>)_cachedConverters.GetOrAdd(typeof(T), CreateConverter);
+		}
+
+		BlobConverter CreateConverter(Type type)
+		{
+			foreach (BlobConverter converter in Converters)
+			{
+				if (converter.CanConvert(type))
+				{
+					return converter;
+				}
+			}
+			return s_cachedDefaultConverters.GetOrAdd(type, s_createDefaultConverter);
+		}
+
+		static BlobConverter CreateDefaultConverter(Type type)
+		{
+			BlobConverterAttribute? attribute = type.GetCustomAttribute<BlobConverterAttribute>();
+			if (attribute != null)
+			{
+				Type converterType = attribute.ConverterType;
+				if (converterType.IsGenericTypeDefinition)
+				{
+					converterType = converterType.MakeGenericType(type.GetGenericArguments());
+				}
+				return (BlobConverter)Activator.CreateInstance(converterType)!;
+			}
+			throw new NotSupportedException($"No converter is available to handle type {type.Name}");
+		}
+	}
+
+	/// <summary>
+	/// Extension methods for serializing blob types
+	/// </summary>
+	public static class BlobSerializerExtensions
+	{
+		/// <summary>
+		/// Deserialize an object
+		/// </summary>
+		/// <typeparam name="T">Return type for the deserialized object</typeparam>
+		/// <param name="handle">Handle to the blob to deserialize</param>
+		/// <param name="options">Options to control serialization</param>
+		/// <param name="cancellationToken">Cancellation token for the operation</param>
+		public static async ValueTask<T> ReadBlobAsync<T>(this IBlobHandle handle, BlobSerializerOptions? options = null, CancellationToken cancellationToken = default)
+		{
+			BlobData data = await handle.ReadBlobDataAsync(cancellationToken);
+			return BlobSerializer.Deserialize<T>(data, options);
+		}
+
+		/// <summary>
+		/// Deserialize an object
+		/// </summary>
+		/// <typeparam name="T">Return type for the deserialized object</typeparam>
+		/// <param name="handle">Handle to the blob to deserialize</param>
+		/// <param name="options">Options to control serialization</param>
+		/// <param name="cancellationToken">Cancellation token for the operation</param>
+		public static async ValueTask<T> ReadBlobAsync<T>(this IBlobHandle<T> handle, BlobSerializerOptions? options = null, CancellationToken cancellationToken = default)
+		{
+			BlobData data = await handle.ReadBlobDataAsync(cancellationToken);
+			return BlobSerializer.Deserialize<T>(data, options);
+		}
+
+		/// <summary>
+		/// Serialize an object to storage
+		/// </summary>
+		/// <param name="writer">Writer for serialized data</param>
+		/// <param name="value">The object to serialize</param>
+		/// <param name="cancellationToken">Cancellation token for the operation</param>
+		/// <returns>Handle to the serialized blob</returns>
+		public static ValueTask<IBlobHandle<T>> WriteBlobAsync<T>(this IStorageWriter writer, T value, CancellationToken cancellationToken)
+		{
+			return WriteBlobAsync<T>(writer, value, null, cancellationToken);
+		}
+
+		/// <summary>
+		/// Serialize an object to storage
+		/// </summary>
+		/// <param name="writer">Writer for serialized data</param>
+		/// <param name="value">The object to serialize</param>
+		/// <param name="options">Options to control serialization</param>
+		/// <param name="cancellationToken">Cancellation token for the operation</param>
+		/// <returns>Handle to the serialized blob</returns>
+		public static ValueTask<IBlobHandle<T>> WriteBlobAsync<T>(this IStorageWriter writer, T value, BlobSerializerOptions? options = null, CancellationToken cancellationToken = default)
+		{
+			BlobWriter blobWriter = new BlobWriter(writer);
+			BlobType blobType = BlobSerializer.Serialize<T>(blobWriter, value, options);
+			return writer.WriteBlobAsync<T>(blobType, blobWriter.Length, blobWriter.References, cancellationToken);
+		}
+
+		/// <summary>
+		/// Reads data for a ref from the store, along with the node's contents.
+		/// </summary>
+		/// <param name="store">Store instance to write to</param>
+		/// <param name="name">The ref name</param>
+		/// <param name="cacheTime">Minimum coherency for any cached value to be returned</param>
+		/// <param name="options">Options to control serialization</param>
+		/// <param name="cancellationToken">Cancellation token for the operation</param>
+		/// <returns>Node for the given ref, or null if it does not exist</returns>
+		public static async Task<TNode?> TryReadRefAsync<TNode>(this IStorageClient store, RefName name, DateTime cacheTime = default, BlobSerializerOptions? options = null, CancellationToken cancellationToken = default) where TNode : class
+		{
+			IBlobHandle? refTarget = await store.TryReadRefTargetAsync(name, cacheTime, cancellationToken);
+			if (refTarget == null)
+			{
+				return null;
+			}
+
+			using BlobData blobData = await refTarget.ReadBlobDataAsync(cancellationToken);
+			return BlobSerializer.Deserialize<TNode>(blobData, options);
+		}
+
+		/// <summary>
+		/// Reads a ref from the store, throwing an exception if it does not exist
+		/// </summary>
+		/// <param name="store">Store instance to write to</param>
+		/// <param name="name">Id for the ref</param>
+		/// <param name="cacheTime">Minimum coherency of any cached result</param>
+		/// <param name="options">Options to control serialization</param>
+		/// <param name="cancellationToken">Cancellation token for the operation</param>
+		/// <returns>The blob instance</returns>
+		public static async Task<TNode> ReadRefAsync<TNode>(this IStorageClient store, RefName name, DateTime cacheTime = default, BlobSerializerOptions? options = null, CancellationToken cancellationToken = default) where TNode : class
+		{
+			TNode? refValue = await store.TryReadRefAsync<TNode>(name, cacheTime, options, cancellationToken);
+			if (refValue == null)
+			{
+				throw new RefNameNotFoundException(name);
+			}
+			return refValue;
+		}
+	}
+}
diff --git a/Engine/Source/Programs/Shared/EpicGames.Horde/Storage/Bundles/V1/BundleWriter.cs b/Engine/Source/Programs/Shared/EpicGames.Horde/Storage/Bundles/V1/BundleWriter.cs
--- a/Engine/Source/Programs/Shared/EpicGames.Horde/Storage/Bundles/V1/BundleWriter.cs
+++ b/Engine/Source/Programs/Shared/EpicGames.Horde/Storage/Bundles/V1/BundleWriter.cs
@@ -65,7 +65,7 @@ namespace EpicGames.Horde.Storage.Bundles.V1
 		}
 
 		/// <inheritdoc/>
-		public ValueTask<BlobData> ReadAsync(CancellationToken cancellationToken = default) => _reader.ReadNodeDataAsync(BundleLocator, ExportIdx, cancellationToken);
+		public ValueTask<BlobData> ReadBlobDataAsync(CancellationToken cancellationToken = default) => _reader.ReadNodeDataAsync(BundleLocator, ExportIdx, cancellationToken);
 
 		/// <inheritdoc/>
 		public ValueTask FlushAsync(CancellationToken cancellationToken = default) => new ValueTask();
@@ -113,7 +113,7 @@ namespace EpicGames.Horde.Storage.Bundles.V1
 				Packet = packet;
 				Offset = offset;
 				Length = length;
-				Refs = refs.ToArray();
+				Refs = refs.Select(x => x.Unwrap()).ToArray();
 				Aliases = aliases.ToArray();
 
 				_pendingBundle = pendingBundle;
@@ -136,7 +136,7 @@ namespace EpicGames.Horde.Storage.Bundles.V1
 			}
 
 			/// <inheritdoc/>
-			public async ValueTask<BlobData> ReadAsync(CancellationToken cancellationToken = default)
+			public async ValueTask<BlobData> ReadBlobDataAsync(CancellationToken cancellationToken = default)
 			{
 				if (_flushedHandle == null)
 				{
@@ -150,7 +150,7 @@ namespace EpicGames.Horde.Storage.Bundles.V1
 					}
 				}
 
-				return await _flushedHandle!.ReadAsync(cancellationToken);
+				return await _flushedHandle!.ReadBlobDataAsync(cancellationToken);
 			}
 
 			/// <inheritdoc/>
@@ -755,7 +755,7 @@ namespace EpicGames.Horde.Storage.Bundles.V1
 			// Add dependencies on all bundles containing a dependent node
 			foreach (IBlobHandle reference in references)
 			{
-				PendingNode? pendingReference = reference as PendingNode;
+				PendingNode? pendingReference = reference.Unwrap() as PendingNode;
 				if (pendingReference?.PendingBundle != null)
 				{
 					currentBundle.AddDependencyOn(pendingReference.PendingBundle, TraceLogger);
@@ -796,16 +796,5 @@ namespace EpicGames.Horde.Storage.Bundles.V1
 			await CompleteAsync(cancellationToken);
 			await _writeQueue.FlushAsync(cancellationToken);
 		}
-
-		/// <summary>
-		/// Flushes all the current nodes to storage
-		/// </summary>
-		/// <param name="root">Root for the tree</param>
-		/// <param name="cancellationToken">Cancellation token for the operation</param>
-		/// <returns></returns>
-		public async Task<IBlobHandle> FlushAsync(Node root, CancellationToken cancellationToken = default)
-		{
-			return await NodeRefExtensions.FlushAsync(this, root, cancellationToken);
-		}
 	}
 }
diff --git a/Engine/Source/Programs/Shared/EpicGames.Horde/Storage/Bundles/V2/BundleWriter.cs b/Engine/Source/Programs/Shared/EpicGames.Horde/Storage/Bundles/V2/BundleWriter.cs
--- a/Engine/Source/Programs/Shared/EpicGames.Horde/Storage/Bundles/V2/BundleWriter.cs
+++ b/Engine/Source/Programs/Shared/EpicGames.Horde/Storage/Bundles/V2/BundleWriter.cs
@@ -35,7 +35,7 @@ namespace EpicGames.Horde.Storage.Bundles.V2
 			public ValueTask FlushAsync(CancellationToken cancellationToken = default) 
 				=> _packet.FlushAsync(cancellationToken);
 
-			public ValueTask<BlobData> ReadAsync(CancellationToken cancellationToken = default)
+			public ValueTask<BlobData> ReadBlobDataAsync(CancellationToken cancellationToken = default)
 				=> _packet.ReadExportAsync(_exportIdx, cancellationToken);
 
 			public bool TryAppendIdentifier(Utf8StringBuilder builder)
@@ -62,7 +62,7 @@ namespace EpicGames.Horde.Storage.Bundles.V2
 			public void CompletePacket(PacketHandle flushedHandle)
 				=> _flushedHandle = flushedHandle;
 
-			public ValueTask<BlobData> ReadAsync(CancellationToken cancellationToken = default)
+			public ValueTask<BlobData> ReadBlobDataAsync(CancellationToken cancellationToken = default)
 			{
 				lock (_bundle.LockObject)
 				{
@@ -71,7 +71,7 @@ namespace EpicGames.Horde.Storage.Bundles.V2
 						throw new NotSupportedException("Reading pending packets is not currently supported.");
 					}
 				}
-				return _flushedHandle!.ReadAsync(cancellationToken);
+				return _flushedHandle!.ReadBlobDataAsync(cancellationToken);
 			}
 
 			public async ValueTask<BlobData> ReadExportAsync(int exportIdx, CancellationToken cancellationToken = default)
@@ -286,8 +286,8 @@ namespace EpicGames.Horde.Storage.Bundles.V2
 			}
 
 			/// <inheritdoc/>
-			public ValueTask<BlobData> ReadAsync(CancellationToken cancellationToken)
-				=> GetFlushedHandle().ReadAsync(cancellationToken);
+			public ValueTask<BlobData> ReadBlobDataAsync(CancellationToken cancellationToken)
+				=> GetFlushedHandle().ReadBlobDataAsync(cancellationToken);
 
 			/// <inheritdoc/>
 			public async Task<Stream> OpenBodyAsync(int offset, int? length, CancellationToken cancellationToken = default)
diff --git a/Engine/Source/Programs/Shared/EpicGames.Horde/Storage/Bundles/V2/ExportHandle.cs b/Engine/Source/Programs/Shared/EpicGames.Horde/Storage/Bundles/V2/ExportHandle.cs
--- a/Engine/Source/Programs/Shared/EpicGames.Horde/Storage/Bundles/V2/ExportHandle.cs
+++ b/Engine/Source/Programs/Shared/EpicGames.Horde/Storage/Bundles/V2/ExportHandle.cs
@@ -68,7 +68,7 @@ namespace EpicGames.Horde.Storage.Bundles.V2
 		}
 
 		/// <inheritdoc/>
-		public ValueTask<BlobData> ReadAsync(CancellationToken cancellationToken = default)
+		public ValueTask<BlobData> ReadBlobDataAsync(CancellationToken cancellationToken = default)
 			=> _packet.ReadExportAsync(_exportIdx, cancellationToken);
 
 		/// <inheritdoc/>
diff --git a/Engine/Source/Programs/Shared/EpicGames.Horde/Storage/Bundles/V2/PacketHandle.cs b/Engine/Source/Programs/Shared/EpicGames.Horde/Storage/Bundles/V2/PacketHandle.cs
--- a/Engine/Source/Programs/Shared/EpicGames.Horde/Storage/Bundles/V2/PacketHandle.cs
+++ b/Engine/Source/Programs/Shared/EpicGames.Horde/Storage/Bundles/V2/PacketHandle.cs
@@ -95,7 +95,7 @@ namespace EpicGames.Horde.Storage.Bundles.V2
 		public ValueTask FlushAsync(CancellationToken cancellationToken = default) => default;
 
 		/// <inheritdoc/>
-		public async ValueTask<BlobData> ReadAsync(CancellationToken cancellationToken = default)
+		public async ValueTask<BlobData> ReadBlobDataAsync(CancellationToken cancellationToken = default)
 		{
 			try
 			{
diff --git a/Engine/Source/Programs/Shared/EpicGames.Horde/Storage/Clients/BundleStorageClient.cs b/Engine/Source/Programs/Shared/EpicGames.Horde/Storage/Clients/BundleStorageClient.cs
--- a/Engine/Source/Programs/Shared/EpicGames.Horde/Storage/Clients/BundleStorageClient.cs
+++ b/Engine/Source/Programs/Shared/EpicGames.Horde/Storage/Clients/BundleStorageClient.cs
@@ -64,7 +64,7 @@ namespace EpicGames.Horde.Storage.Clients
 					}
 					else
 					{
-						using BlobData blobData = await ReadAsync(cancellationToken);
+						using BlobData blobData = await ReadBlobDataAsync(cancellationToken);
 						refs.AddRange(blobData.Refs);
 					}
 
@@ -78,7 +78,7 @@ namespace EpicGames.Horde.Storage.Clients
 				=> _inner.OpenBodyAsync(offset, length, cancellationToken);
 
 			/// <inheritdoc/>
-			public async ValueTask<BlobData> ReadAsync(CancellationToken cancellationToken = default)
+			public async ValueTask<BlobData> ReadBlobDataAsync(CancellationToken cancellationToken = default)
 			{
 				IReadOnlyMemoryOwner<byte> data = await _inner.ReadBodyAsync(cancellationToken);
 
@@ -220,7 +220,7 @@ namespace EpicGames.Horde.Storage.Clients
 		public async ValueTask<BlobData> ReadBlobAsync(BlobLocator locator, CancellationToken cancellationToken = default)
 		{
 			IBlobHandle handle = CreateBlobHandle(locator);
-			return await handle.ReadAsync(cancellationToken);
+			return await handle.ReadBlobDataAsync(cancellationToken);
 		}
 
 		/// <inheritdoc/>
diff --git a/Engine/Source/Programs/Shared/EpicGames.Horde/Storage/Clients/HttpStorageClient.cs b/Engine/Source/Programs/Shared/EpicGames.Horde/Storage/Clients/HttpStorageClient.cs
--- a/Engine/Source/Programs/Shared/EpicGames.Horde/Storage/Clients/HttpStorageClient.cs
+++ b/Engine/Source/Programs/Shared/EpicGames.Horde/Storage/Clients/HttpStorageClient.cs
@@ -43,7 +43,7 @@ namespace EpicGames.Horde.Storage.Clients
 				=> _httpStorageClient._backend.OpenAsync(_locator.ToString(), offset, length, cancellationToken);
 
 			/// <inheritdoc/>
-			public ValueTask<BlobData> ReadAsync(CancellationToken cancellationToken = default)
+			public ValueTask<BlobData> ReadBlobDataAsync(CancellationToken cancellationToken = default)
 				=> _httpStorageClient.ReadBlobAsync(_locator, cancellationToken);
 
 			/// <inheritdoc/>
diff --git a/Engine/Source/Programs/Shared/EpicGames.Horde/Storage/Clients/KeyValueStorageClient.cs b/Engine/Source/Programs/Shared/EpicGames.Horde/Storage/Clients/KeyValueStorageClient.cs
--- a/Engine/Source/Programs/Shared/EpicGames.Horde/Storage/Clients/KeyValueStorageClient.cs
+++ b/Engine/Source/Programs/Shared/EpicGames.Horde/Storage/Clients/KeyValueStorageClient.cs
@@ -35,7 +35,7 @@ namespace EpicGames.Horde.Storage.Clients
 			public ValueTask FlushAsync(CancellationToken cancellationToken = default) => default;
 
 			/// <inheritdoc/>
-			public ValueTask<BlobData> ReadAsync(CancellationToken cancellationToken = default) => _keyValueStorageClient.ReadBlobAsync(_locator, cancellationToken);
+			public ValueTask<BlobData> ReadBlobDataAsync(CancellationToken cancellationToken = default) => _keyValueStorageClient.ReadBlobAsync(_locator, cancellationToken);
 
 			/// <inheritdoc/>
 			public bool TryAppendIdentifier(Utf8StringBuilder builder)
diff --git a/Engine/Source/Programs/Shared/EpicGames.Horde/Storage/IBlobHandle.cs b/Engine/Source/Programs/Shared/EpicGames.Horde/Storage/IBlobHandle.cs
--- a/Engine/Source/Programs/Shared/EpicGames.Horde/Storage/IBlobHandle.cs
+++ b/Engine/Source/Programs/Shared/EpicGames.Horde/Storage/IBlobHandle.cs
@@ -11,7 +11,7 @@ using EpicGames.Core;
 namespace EpicGames.Horde.Storage
 {
 	/// <summary>
-	/// Handle to a node. Can be used to reference nodes that have not been flushed yet. Handles should have value equality semantics, and must override <see cref="Object.Equals(Object?)"/> and <see cref="Object.GetHashCode()"/>.
+	/// Handle to a node. Can be used to reference nodes that have not been flushed yet.
 	/// </summary>
 	public interface IBlobHandle
 	{
@@ -64,7 +64,7 @@ namespace EpicGames.Horde.Storage
 		IBlobHandle? Outer { get; }
 
 		/// <summary>
-		/// Flush the referenced not to underlying storage
+		/// Flush the referenced data to underlying storage
 		/// </summary>
 		ValueTask FlushAsync(CancellationToken cancellationToken = default);
 
@@ -76,7 +76,7 @@ namespace EpicGames.Horde.Storage
 		/// <param name="cancellationToken">Cancellation token for the operation</param>
 		async Task<Stream> OpenBodyAsync(int offset = 0, int? length = null, CancellationToken cancellationToken = default)
 		{
-			BlobData blobData = await ReadAsync(cancellationToken);
+			BlobData blobData = await ReadBlobDataAsync(cancellationToken);
 			int maxLength = blobData.Data.Length - offset;
 			ReadOnlyMemory<byte> memory = blobData.Data.Slice(offset, length.HasValue ? Math.Min(length.Value, maxLength) : maxLength);
 			return new BlobDataStream(blobData, memory);
@@ -86,27 +86,7 @@ namespace EpicGames.Horde.Storage
 		/// Reads the blob's data
 		/// </summary>
 		/// <param name="cancellationToken">Cancellation token for the operation</param>
-		ValueTask<BlobData> ReadAsync(CancellationToken cancellationToken = default);
-
-		/// <summary>
-		/// Gets the type of this blob
-		/// </summary>
-		/// <param name="cancellationToken">Cancellation token for the operation</param>
-		async ValueTask<BlobType> ReadTypeAsync(CancellationToken cancellationToken = default)
-		{
-			using BlobData data = await ReadAsync(cancellationToken);
-			return data.Type;
-		}
-
-		/// <summary>
-		/// Gets the outward references from this blob
-		/// </summary>
-		/// <param name="cancellationToken">Cancellation token for the operation</param>
-		async ValueTask<IReadOnlyList<IBlobHandle>> ReadImportsAsync(CancellationToken cancellationToken = default)
-		{
-			using BlobData data = await ReadAsync(cancellationToken);
-			return data.Refs;
-		}
+		ValueTask<BlobData> ReadBlobDataAsync(CancellationToken cancellationToken = default);
 
 		/// <summary>
 		/// Reads part of the blob, and returns a handle that can be used to access the data.
@@ -127,7 +107,7 @@ namespace EpicGames.Horde.Storage
 		/// <returns></returns>
 		async ValueTask<IReadOnlyMemoryOwner<byte>> ReadBodyAsync(int offset, int? length, CancellationToken cancellationToken = default)
 		{
-			BlobData data = await ReadAsync(cancellationToken);
+			BlobData data = await ReadBlobDataAsync(cancellationToken);
 			return new BlobDataFragment(data, offset, length);
 		}
 
@@ -146,6 +126,18 @@ namespace EpicGames.Horde.Storage
 			=> throw new InvalidOperationException("Not supported for this handle type.");
 	}
 
+	/// <summary>
+	/// Typed interface to a particular blob handle
+	/// </summary>
+	/// <typeparam name="T">Type of the deserialized blob</typeparam>
+	public interface IBlobHandle<out T> : IBlobHandle
+	{
+		/// <summary>
+		/// Hash of the target node
+		/// </summary>
+		IoHash Hash { get; }
+	}
+
 	/// <summary>
 	/// Instance of <see cref="IBlobHandle"/> which wraps an inner handle and a fragment
 	/// </summary>
@@ -181,7 +173,7 @@ namespace EpicGames.Horde.Storage
 		}
 
 		/// <inheritdoc/>
-		public ValueTask<BlobData> ReadAsync(CancellationToken cancellationToken = default)
+		public ValueTask<BlobData> ReadBlobDataAsync(CancellationToken cancellationToken = default)
 		{
 			throw new NotSupportedException("Blob fragment handles cannot be read directly, and should be deconstructed into more specific types.");
 		}
@@ -198,6 +190,51 @@ namespace EpicGames.Horde.Storage
 	/// </summary>
 	public static class BlobHandleExtensions
 	{
+		interface IWrappedBlobHandle
+		{
+			public IBlobHandle Inner { get; }
+		}
+
+		class TypedBlobHandle<T> : IWrappedBlobHandle, IBlobHandle<T>
+		{
+			readonly IBlobHandle _inner;
+			readonly IoHash _hash;
+
+			public TypedBlobHandle(IBlobHandle inner, IoHash hash)
+			{
+				_inner = inner;
+				_hash = hash;
+			}
+
+			public IBlobHandle Inner => _inner;
+			public IoHash Hash => _hash;
+
+			public IBlobHandle? Outer => _inner.Outer;
+
+			public ValueTask FlushAsync(CancellationToken cancellationToken = default) => _inner.FlushAsync(cancellationToken);
+			public ValueTask<BlobData> ReadBlobDataAsync(CancellationToken cancellationToken = default) => _inner.ReadBlobDataAsync(cancellationToken);
+			public bool TryAppendIdentifier(Utf8StringBuilder builder) => _inner.TryAppendIdentifier(builder);
+		}
+
+		/// <summary>
+		/// Creates a typed blob handle
+		/// </summary>
+		public static IBlobHandle<T> ForType<T>(this IBlobHandle handle, IoHash hash) => new TypedBlobHandle<T>(handle, hash);
+
+		/// <summary>
+		/// 
+		/// </summary>
+		/// <param name="handle"></param>
+		/// <returns></returns>
+		public static IBlobHandle Unwrap(this IBlobHandle handle)
+		{
+			while (handle is IWrappedBlobHandle wrappedHandle)
+			{
+				handle = wrappedHandle.Inner;
+			}
+			return handle;
+		}
+
 		/// <summary>
 		/// Gets a path to this blob that can be used to describe blob references over the wire.
 		/// </summary>
diff --git a/Engine/Source/Programs/Shared/EpicGames.Horde/Storage/IBlobReader.cs b/Engine/Source/Programs/Shared/EpicGames.Horde/Storage/IBlobReader.cs
--- a/Engine/Source/Programs/Shared/EpicGames.Horde/Storage/IBlobReader.cs
+++ b/Engine/Source/Programs/Shared/EpicGames.Horde/Storage/IBlobReader.cs
@@ -17,19 +17,19 @@ namespace EpicGames.Horde.Storage
 		BlobType Type { get; }
 
 		/// <summary>
-		/// Version of the current node, as specified via <see cref="BlobTypeAttribute"/>
+		/// Version of the current node, as specified via <see cref="BlobType"/>
 		/// </summary>
 		int Version { get; }
 
 		/// <summary>
-		/// Locations of all referenced nodes.
+		/// Locations of all referenced nodes. These handles do not have valid hashes.
 		/// </summary>
 		IReadOnlyList<IBlobHandle> References { get; }
 
 		/// <summary>
 		/// Gets the next serialized blob handle
 		/// </summary>
-		IBlobHandle ReadBlobReference();
+		IBlobHandle<T> ReadBlobHandle<T>();
 	}
 
 	/// <summary>
@@ -43,7 +43,7 @@ namespace EpicGames.Horde.Storage
 		public BlobType Type => _blobData.Type;
 
 		/// <summary>
-		/// Version of the current node, as specified via <see cref="BlobTypeAttribute"/>
+		/// Version of the current node, as specified via <see cref="BlobType"/>
 		/// </summary>
 		public int Version => Type.Version;
 
@@ -82,6 +82,11 @@ namespace EpicGames.Horde.Storage
 		/// <summary>
 		/// Gets the next serialized blob handle
 		/// </summary>
-		public IBlobHandle ReadBlobReference() => _blobData.Refs[_refIdx++];
+		public IBlobHandle<T> ReadBlobHandle<T>()
+		{
+			IBlobHandle handle = References[_refIdx++];
+			IoHash hash = this.ReadIoHash();
+			return handle.ForType<T>(hash);
+		}
 	}
 }
diff --git a/Engine/Source/Programs/Shared/EpicGames.Horde/Storage/IBlobWriter.cs b/Engine/Source/Programs/Shared/EpicGames.Horde/Storage/IBlobWriter.cs
--- a/Engine/Source/Programs/Shared/EpicGames.Horde/Storage/IBlobWriter.cs
+++ b/Engine/Source/Programs/Shared/EpicGames.Horde/Storage/IBlobWriter.cs
@@ -12,10 +12,10 @@ namespace EpicGames.Horde.Storage
 	public interface IBlobWriter : IMemoryWriter
 	{
 		/// <summary>
-		/// Adds a reference to another blob. This reference is stored out of band, and will not result in any bytes written to the output.
+		/// Writes a reference to another blob. The blob's hash is serialized to the output stream.
 		/// </summary>
-		/// <param name="reference">Referenced blob</param>
-		void WriteBlobReference(IBlobHandle reference);
+		/// <param name="handle">Referenced blob</param>
+		void WriteBlobHandle<T>(IBlobHandle<T> handle);
 	}
 
 	/// <summary>
@@ -37,6 +37,11 @@ namespace EpicGames.Horde.Storage
 		/// <inheritdoc/>
 		public int Length => _length;
 
+		/// <summary>
+		/// Memory that has been written
+		/// </summary>
+		public ReadOnlyMemory<byte> WrittenMemory => _memory.Slice(0, _length);
+
 		/// <summary>
 		/// Constructor
 		/// </summary>
@@ -47,12 +52,6 @@ namespace EpicGames.Horde.Storage
 			_memory = treeWriter.GetOutputBuffer(0, 2 * 1024);
 		}
 
-		/// <summary>
-		/// Adds a reference to another blob. This reference is stored out of band, and will not result in any bytes written to the output.
-		/// </summary>
-		/// <param name="reference">Referenced blob</param>
-		public void WriteBlobReference(IBlobHandle reference) => _refs.Add(reference);
-
 		/// <summary>
 		/// Computes the hash of the written data
 		/// </summary>
@@ -61,10 +60,10 @@ namespace EpicGames.Horde.Storage
 		/// <summary>
 		/// Writes a handle to another node
 		/// </summary>
-		public void WriteHashedBlobHandle(IoHash hash, IBlobHandle target)
+		public void WriteBlobHandle<T>(IBlobHandle<T> target)
 		{
-			this.WriteIoHash(hash);
-			WriteBlobReference(target);
+			this.WriteIoHash(target.Hash);
+			_refs.Add(target);
 		}
 
 		/// <inheritdoc/>
diff --git a/Engine/Source/Programs/Shared/EpicGames.Horde/Storage/IStorageClient.cs b/Engine/Source/Programs/Shared/EpicGames.Horde/Storage/IStorageClient.cs
--- a/Engine/Source/Programs/Shared/EpicGames.Horde/Storage/IStorageClient.cs
+++ b/Engine/Source/Programs/Shared/EpicGames.Horde/Storage/IStorageClient.cs
@@ -104,7 +104,7 @@ namespace EpicGames.Horde.Storage
 		IStorageWriter CreateWriter(string? basePath = null);
 
 		/// <summary>
-		/// Read a blob from the underlying storage system. Calling <see cref="IBlobHandle.ReadAsync(CancellationToken)"/> is more efficient than calling this method repeatedly for small blobs.
+		/// Read a blob from the underlying storage system. Calling <see cref="IBlobHandle.ReadBlobDataAsync(CancellationToken)"/> is more efficient than calling this method repeatedly for small blobs.
 		/// </summary>
 		/// <param name="locator">Locator for the blob</param>
 		/// <param name="cancellationToken">Cancellation token for the operation</param>
diff --git a/Engine/Source/Programs/Shared/EpicGames.Horde/Storage/IStorageWriter.cs b/Engine/Source/Programs/Shared/EpicGames.Horde/Storage/IStorageWriter.cs
--- a/Engine/Source/Programs/Shared/EpicGames.Horde/Storage/IStorageWriter.cs
+++ b/Engine/Source/Programs/Shared/EpicGames.Horde/Storage/IStorageWriter.cs
@@ -185,7 +185,7 @@ namespace EpicGames.Horde.Storage
 			}
 
 			/// <inheritdoc/>
-			public ValueTask<BlobData> ReadAsync(CancellationToken cancellationToken = default)
+			public ValueTask<BlobData> ReadBlobDataAsync(CancellationToken cancellationToken = default)
 			{
 				if (_inner == null)
 				{
@@ -193,7 +193,7 @@ namespace EpicGames.Horde.Storage
 				}
 				else
 				{
-					return _inner.ReadAsync(cancellationToken);
+					return _inner.ReadBlobDataAsync(cancellationToken);
 				}
 			}
 
@@ -261,7 +261,7 @@ namespace EpicGames.Horde.Storage
 				_cache.Add(key, wrappedHandle);
 			}
 
-			wrappedHandle._inner = await _inner.WriteBlobAsync(type, size, references.ConvertAll(x => ((WrappedHandle)x)._inner!), aliases, cancellationToken);
+			wrappedHandle._inner = await _inner.WriteBlobAsync(type, size, references.ConvertAll(x => ((WrappedHandle)x.Unwrap())._inner!), aliases, cancellationToken);
 			return wrappedHandle;
 		}
 	}
@@ -285,25 +285,11 @@ namespace EpicGames.Horde.Storage
 		/// <param name="references">References to other nodes</param>
 		/// <param name="cancellationToken">Cancellation token for the operation</param>
 		/// <returns>Handle to the written node</returns>
-		public static ValueTask<IBlobHandle> WriteBlobAsync(this IStorageWriter writer, BlobType type, int size, IReadOnlyList<IBlobHandle> references, CancellationToken cancellationToken = default)
-		{
-			return writer.WriteBlobAsync(type, size, references, Array.Empty<AliasInfo>(), cancellationToken);
-		}
-
-		/// <summary>
-		/// Finish writing a node.
-		/// </summary>
-		/// <param name="writer">Writer instance to manipulate</param>
-		/// <param name="type">Type of the node that was written</param>
-		/// <param name="size">Used size of the buffer</param>
-		/// <param name="references">References to other nodes</param>
-		/// <param name="cancellationToken">Cancellation token for the operation</param>
-		/// <returns>Handle to the written node</returns>
-		public static async ValueTask<HashedNodeRef<T>> WriteHashedNodeRefAsync<T>(this IStorageWriter writer, BlobType type, int size, IReadOnlyList<IBlobHandle> references, CancellationToken cancellationToken = default) where T : Node
+		public static async ValueTask<IBlobHandle<T>> WriteBlobAsync<T>(this IStorageWriter writer, BlobType type, int size, IReadOnlyList<IBlobHandle> references, CancellationToken cancellationToken = default)
 		{
 			IoHash hash = IoHash.Compute(writer.GetOutputBuffer(size, size).Span.Slice(0, size));
-			IBlobHandle blobHandle = await WriteBlobAsync(writer, type, size, references, cancellationToken);
-			return new HashedNodeRef<T>(hash, blobHandle);
+			IBlobHandle blobHandle = await writer.WriteBlobAsync(type, size, references, Array.Empty<AliasInfo>(), cancellationToken);
+			return blobHandle.ForType<T>(hash);
 		}
 	}
 }
diff --git a/Engine/Source/Programs/Shared/EpicGames.Horde/Storage/Node.cs b/Engine/Source/Programs/Shared/EpicGames.Horde/Storage/Node.cs
deleted file mode 100644
diff --git a/Engine/Source/Programs/Shared/EpicGames.Horde/Storage/NodeRef.cs b/Engine/Source/Programs/Shared/EpicGames.Horde/Storage/NodeRef.cs
deleted file mode 100644
diff --git a/Engine/Source/Programs/Shared/EpicGames.Horde/Storage/Nodes/CbNode.cs b/Engine/Source/Programs/Shared/EpicGames.Horde/Storage/Nodes/CbNode.cs
--- a/Engine/Source/Programs/Shared/EpicGames.Horde/Storage/Nodes/CbNode.cs
+++ b/Engine/Source/Programs/Shared/EpicGames.Horde/Storage/Nodes/CbNode.cs
@@ -9,15 +9,40 @@ namespace EpicGames.Horde.Storage.Nodes
 	/// <summary>
 	/// A node containing arbitrary compact binary data
 	/// </summary>
-	[BlobType("{34A0793F-42F4-8364-A798-32862932841C}", 1)]
-	public class CbNode : Node
+	[BlobConverter(typeof(CbNodeConverter))]
+	public class CbNode
+	{
+		/// <summary>
+		/// The compact binary object
+		/// </summary>
+		public CbObject Object { get; set; }
+
+		/// <summary>
+		/// Imported nodes
+		/// </summary>
+		public IReadOnlyList<IBlobHandle<object>> References { get; }
+
+		/// <summary>
+		/// Constructor
+		/// </summary>
+		/// <param name="obj">The compact binary object</param>
+		/// <param name="references">List of references to attachments</param>
+		public CbNode(CbObject obj, IReadOnlyList<IBlobHandle<object>> references)
+		{
+			Object = obj;
+			References = references;
+		}
+	}
+
+	class CbNodeConverter : BlobConverter<CbNode>
 	{
 		class HandleMapper
 		{
 			readonly IBlobReader _reader;
-			readonly List<HashedNodeRef> _refs;
+			readonly List<IBlobHandle<object>> _refs;
+			int _refIdx;
 
-			public HandleMapper(IBlobReader reader, List<HashedNodeRef> refs)
+			public HandleMapper(IBlobReader reader, List<IBlobHandle<object>> refs)
 			{
 				_reader = reader;
 				_refs = refs;
@@ -27,8 +52,8 @@ namespace EpicGames.Horde.Storage.Nodes
 			{
 				if (field.IsAttachment())
 				{
-					IBlobHandle handle = _reader.ReadBlobReference();
-					_refs.Add(new HashedNodeRef(field.AsAttachment(), handle));
+					IBlobHandle handle = _reader.References[_refIdx++];
+					_refs.Add(handle.ForType<object>(field.AsAttachment()));
 				}
 				else if (field.IsArray())
 				{
@@ -43,45 +68,24 @@ namespace EpicGames.Horde.Storage.Nodes
 			}
 		}
 
-		/// <summary>
-		/// The compact binary object
-		/// </summary>
-		public CbObject Object { get; set; }
+		public static BlobType BlobType { get; } = new BlobType("{34A0793F-42F4-8364-A798-32862932841C}", 1);
 
-		/// <summary>
-		/// Imported nodes
-		/// </summary>
-		public IReadOnlyList<HashedNodeRef> References { get; }
-
-		/// <summary>
-		/// Constructor
-		/// </summary>
-		/// <param name="obj">The compact binary object</param>
-		/// <param name="references">List of references to attachments</param>
-		public CbNode(CbObject obj, IReadOnlyList<HashedNodeRef> references)
-		{
-			Object = obj;
-			References = references;
-		}
-
-		/// <summary>
-		/// Deserialization constructor
-		/// </summary>
-		/// <param name="reader">Reader to deserialize from</param>
-		public CbNode(IBlobReader reader)
+		/// <inheritdoc/>
+		public override CbNode Read(IBlobReader reader, BlobSerializerOptions options)
 		{
-			Object = new CbObject(reader.GetMemory().ToArray());
+			CbObject Object = new CbObject(reader.GetMemory().ToArray());
 
-			List<HashedNodeRef> references = new List<HashedNodeRef>();
+			List<IBlobHandle<object>> references = new List<IBlobHandle<object>>();
 			Object.IterateAttachments(new HandleMapper(reader, references).IterateField);
 
-			References = references;
+			return new CbNode(Object, references);
 		}
 
 		/// <inheritdoc/>
-		public override void Serialize(IBlobWriter writer)
+		public override BlobType Write(IBlobWriter writer, CbNode value, BlobSerializerOptions options)
 		{
-			writer.WriteFixedLengthBytes(Object.GetView().Span);
+			writer.WriteFixedLengthBytes(value.Object.GetView().Span);
+			return BlobType;
 		}
 	}
 }
diff --git a/Engine/Source/Programs/Shared/EpicGames.Horde/Storage/Nodes/ChunkedDataNode.cs b/Engine/Source/Programs/Shared/EpicGames.Horde/Storage/Nodes/ChunkedDataNode.cs
--- a/Engine/Source/Programs/Shared/EpicGames.Horde/Storage/Nodes/ChunkedDataNode.cs
+++ b/Engine/Source/Programs/Shared/EpicGames.Horde/Storage/Nodes/ChunkedDataNode.cs
@@ -18,17 +18,16 @@ namespace EpicGames.Horde.Storage.Nodes
 	/// Chunks are pushed into a tree hierarchy as data is appended to the root, with nodes of the tree also split along content-aware boundaries with <see cref="IoHash.NumBytes"/> granularity.
 	/// Once a chunk has been written to storage, it is treated as immutable.
 	/// </summary>
-	public abstract class ChunkedDataNode : Node
+	[BlobConverter(typeof(ChunkedDataNodeConverter))]
+	public abstract class ChunkedDataNode
 	{
 		/// <summary>
 		/// Copies the contents of this node and its children to the given output stream
 		/// </summary>
 		/// <param name="outputStream">The output stream to receive the data</param>
+		/// <param name="options">Options for controling serialization</param>
 		/// <param name="cancellationToken">Cancellation token for the operation</param>
-		public abstract Task CopyToStreamAsync(Stream outputStream, CancellationToken cancellationToken);
-
-		static readonly Guid s_leafNodeGuid = Node.GetNodeType<LeafChunkedDataNode>().Guid;
-		static readonly Guid s_interiorNodeGuid = Node.GetNodeType<InteriorChunkedDataNode>().Guid;
+		public abstract Task CopyToStreamAsync(Stream outputStream, BlobSerializerOptions? options = null, CancellationToken cancellationToken = default);
 
 		/// <summary>
 		/// Copy the contents of the node to the output stream without creating the intermediate FileNodes
@@ -38,12 +37,12 @@ namespace EpicGames.Horde.Storage.Nodes
 		/// <param name="cancellationToken">Cancellation token for the operation</param>
 		public static async Task CopyToStreamAsync(IBlobHandle handle, Stream outputStream, CancellationToken cancellationToken)
 		{
-			using BlobData blobData = await handle.ReadAsync(cancellationToken);
-			if (blobData.Type.Guid == s_leafNodeGuid)
+			using BlobData blobData = await handle.ReadBlobDataAsync(cancellationToken);
+			if (blobData.Type.Guid == LeafChunkedDataNodeConverter.BlobType.Guid)
 			{
 				await LeafChunkedDataNode.CopyToStreamAsync(blobData, outputStream, cancellationToken);
 			}
-			else if (blobData.Type.Guid == s_interiorNodeGuid)
+			else if (blobData.Type.Guid == InteriorChunkedDataNodeConverter.BlobType.Guid)
 			{
 				await InteriorChunkedDataNode.CopyToStreamAsync(blobData, outputStream, cancellationToken);
 			}
@@ -75,18 +74,56 @@ namespace EpicGames.Horde.Storage.Nodes
 		/// <summary>
 		/// Serialize this node and its children into a byte array
 		/// </summary>
+		/// <param name="options">Options to control serialization</param>
 		/// <param name="cancellationToken"></param>
 		/// <returns>Array of data stored by the tree</returns>
-		public async Task<byte[]> ToByteArrayAsync(CancellationToken cancellationToken)
+		public async Task<byte[]> ToByteArrayAsync(BlobSerializerOptions? options = null, CancellationToken cancellationToken = default)
 		{
 			using (MemoryStream stream = new MemoryStream())
 			{
-				await CopyToStreamAsync(stream, cancellationToken);
+				await CopyToStreamAsync(stream, options, cancellationToken);
 				return stream.ToArray();
 			}
 		}
 	}
 
+	class ChunkedDataNodeConverter : BlobConverter<ChunkedDataNode>
+	{
+		/// <inheritdoc/>
+		public override ChunkedDataNode Read(IBlobReader reader, BlobSerializerOptions options)
+		{
+			if (reader.Type.Guid == LeafChunkedDataNodeConverter.BlobType.Guid)
+			{
+				return options.GetConverter<LeafChunkedDataNode>().Read(reader, options);
+			}
+			else if (reader.Type.Guid == InteriorChunkedDataNodeConverter.BlobType.Guid)
+			{
+				return options.GetConverter<InteriorChunkedDataNode>().Read(reader, options);
+			}
+			else
+			{
+				throw new NotSupportedException();
+			}
+		}
+
+		/// <inheritdoc/>
+		public override BlobType Write(IBlobWriter writer, ChunkedDataNode value, BlobSerializerOptions options)
+		{
+			if (value is LeafChunkedDataNode leafNode)
+			{
+				return options.GetConverter<LeafChunkedDataNode>().Write(writer, leafNode, options);
+			}
+			else if (value is InteriorChunkedDataNode interiorNode)
+			{
+				return options.GetConverter<InteriorChunkedDataNode>().Write(writer, interiorNode, options);
+			}
+			else
+			{
+				throw new NotSupportedException();
+			}
+		}
+	}
+
 	/// <summary>
 	/// Type of a chunked data node
 	/// </summary>
@@ -112,16 +149,44 @@ namespace EpicGames.Horde.Storage.Nodes
 	/// Reference to a chunked data node
 	/// </summary>
 	/// <param name="Type">Type of the referenced node</param>
-	/// <param name="Hash">Hash of the target node</param>
 	/// <param name="Length">Length of the data stream within this node</param>
 	/// <param name="Handle">Handle to the target node</param>
-	public record class ChunkedDataNodeRef(ChunkedDataNodeType Type, IoHash Hash, long Length, IBlobHandle Handle)
+	public record class ChunkedDataNodeRef(ChunkedDataNodeType Type, long Length, IBlobHandle<ChunkedDataNode> Handle)
 	{
+		/// <summary>
+		/// Leaf node constructor
+		/// </summary>
+		public ChunkedDataNodeRef(long length, IBlobHandle<LeafChunkedDataNode> handle)
+			: this(ChunkedDataNodeType.Leaf, length, handle)
+		{
+		}
+
+		/// <summary>
+		/// Interior node constructor
+		/// </summary>
+		public ChunkedDataNodeRef(long length, IBlobHandle<InteriorChunkedDataNode> handle)
+			: this(ChunkedDataNodeType.Interior, length, handle)
+		{
+		}
+
 		/// <summary>
 		/// Read the node which is the target of this ref
 		/// </summary>
-		public ValueTask<ChunkedDataNode> ExpandAsync(CancellationToken cancellationToken = default)
-			=> Handle.ReadNodeAsync<ChunkedDataNode>(cancellationToken);
+		public async ValueTask<ChunkedDataNode> ReadBlobAsync(BlobSerializerOptions? options = null, CancellationToken cancellationToken = default)
+		{
+			if (Type == ChunkedDataNodeType.Leaf)
+			{
+				return await Handle.ReadBlobAsync<LeafChunkedDataNode>(options, cancellationToken);
+			}
+			else if (Type == ChunkedDataNodeType.Interior)
+			{
+				return await Handle.ReadBlobAsync<InteriorChunkedDataNode>(options, cancellationToken);
+			}
+			else
+			{
+				throw new InvalidDataException();
+			}
+		}
 	}
 
 	/// <summary>
@@ -134,14 +199,9 @@ namespace EpicGames.Horde.Storage.Nodes
 	/// <summary>
 	/// File node that contains a chunk of data
 	/// </summary>
-	[BlobType("{B27AFB68-4A4B-9E20-8A78-D8A439D49840}", 1)]
+	[BlobConverter(typeof(LeafChunkedDataNodeConverter))]
 	public sealed class LeafChunkedDataNode : ChunkedDataNode
 	{
-		/// <summary>
-		/// Static accessor for the blob type
-		/// </summary>
-		public static BlobType BlobType { get; } = Node.GetNodeType<LeafChunkedDataNode>();
-
 		/// <summary>
 		/// Data for this node
 		/// </summary>
@@ -150,27 +210,13 @@ namespace EpicGames.Horde.Storage.Nodes
 		/// <summary>
 		/// Create an empty leaf node
 		/// </summary>
-		public LeafChunkedDataNode()
+		public LeafChunkedDataNode(ReadOnlyMemory<byte> data)
 		{
-		}
-
-		/// <summary>
-		/// Create a leaf node from the given serialized data
-		/// </summary>
-		public LeafChunkedDataNode(IBlobReader reader)
-		{
-			// Keep this code in sync with CopyToStreamAsync
-			Data = reader.GetMemory().ToArray();
+			Data = data;
 		}
 
 		/// <inheritdoc/>
-		public override void Serialize(IBlobWriter writer)
-		{
-			writer.WriteFixedLengthBytes(Data.Span);
-		}
-
-		/// <inheritdoc/>
-		public override async Task CopyToStreamAsync(Stream outputStream, CancellationToken cancellationToken)
+		public override async Task CopyToStreamAsync(Stream outputStream, BlobSerializerOptions? options, CancellationToken cancellationToken)
 		{
 			await outputStream.WriteAsync(Data, cancellationToken);
 		}
@@ -260,7 +306,7 @@ namespace EpicGames.Horde.Storage.Nodes
 
 				int nextLength = GetChunkLength(readBuffer.Memory.Span.Slice(0, size), options);
 
-				Memory<byte> outputBuffer = writer.GetOutputBuffer(0, nextLength);
+				Memory<byte> outputBuffer = writer.GetOutputBuffer(0, nextLength).Slice(0, nextLength);
 				readBuffer.Memory.Slice(0, nextLength).CopyTo(outputBuffer);
 				hasher.Update(outputBuffer.Span);
 
@@ -271,8 +317,8 @@ namespace EpicGames.Horde.Storage.Nodes
 					sizeSinceProgressUpdate = 0;
 				}
 
-				HashedNodeRef<ChunkedDataNode> nodeRef = await writer.WriteHashedNodeRefAsync<ChunkedDataNode>(GetNodeType<LeafChunkedDataNode>(), nextLength, Array.Empty<IBlobHandle>(), cancellationToken);
-				leafNodeRefs.Add(new ChunkedDataNodeRef(ChunkedDataNodeType.Leaf, nodeRef.Hash, nextLength, nodeRef.Handle));
+				IBlobHandle<LeafChunkedDataNode> blobHandle = await writer.WriteBlobAsync<LeafChunkedDataNode>(LeafChunkedDataNodeConverter.BlobType, nextLength, Array.Empty<IBlobHandle>(), cancellationToken);
+				leafNodeRefs.Add(new ChunkedDataNodeRef(nextLength, blobHandle));
 
 				readBuffer.Memory.Slice(nextLength, size - nextLength).CopyTo(readBuffer.Memory);
 				size -= nextLength;
@@ -336,6 +382,26 @@ namespace EpicGames.Horde.Storage.Nodes
 		}
 	}
 
+	class LeafChunkedDataNodeConverter : BlobConverter<LeafChunkedDataNode>
+	{
+		/// <summary>
+		/// Static accessor for the blob type
+		/// </summary>
+		public static BlobType BlobType { get; } = new BlobType("{B27AFB68-4A4B-9E20-8A78-D8A439D49840}", 1);
+
+		public override LeafChunkedDataNode Read(IBlobReader reader, BlobSerializerOptions options)
+		{
+			byte[] Data = reader.GetMemory().ToArray();
+			return new LeafChunkedDataNode(Data);
+		}
+
+		public override BlobType Write(IBlobWriter writer, LeafChunkedDataNode value, BlobSerializerOptions options)
+		{
+			writer.WriteFixedLengthBytes(value.Data.Span);
+			return BlobType;
+		}
+	}
+
 	/// <summary>
 	/// Options for creating interior nodes
 	/// </summary>
@@ -357,14 +423,9 @@ namespace EpicGames.Horde.Storage.Nodes
 	/// <summary>
 	/// An interior file node
 	/// </summary>
-	[BlobType("{F4DEDDBC-4C7A-70CB-11F0-4783B9CDCCAF}", 2)] // Pending V3
+	[BlobConverter(typeof(InteriorChunkedDataNodeConverter))]
 	public class InteriorChunkedDataNode : ChunkedDataNode
 	{
-		/// <summary>
-		/// Static accessor for the blob type
-		/// </summary>
-		public static BlobType BlobType { get; } = Node.GetNodeType<InteriorChunkedDataNode>();
-
 		/// <summary>
 		/// Child nodes
 		/// </summary>
@@ -379,48 +440,6 @@ namespace EpicGames.Horde.Storage.Nodes
 			Children = children;
 		}
 
-		/// <summary>
-		/// Constructor
-		/// </summary>
-		public InteriorChunkedDataNode(IBlobReader reader)
-		{
-			// Keep this code in sync with CopyToStreamAsync
-			List<ChunkedDataNodeRef> children = new List<ChunkedDataNodeRef>();
-			while (reader.GetMemory().Length > 0)
-			{
-				IoHash hash = reader.ReadIoHash();
-
-				ChunkedDataNodeType type = ChunkedDataNodeType.Unknown;
-				if (reader.Version >= 2)
-				{
-					type = (ChunkedDataNodeType)reader.ReadUnsignedVarInt();
-				}
-
-				long length = 0;
-				if (reader.Version >= 3)
-				{
-					length = (long)reader.ReadUnsignedVarInt();
-				}
-
-				IBlobHandle handle = reader.ReadBlobReference();
-
-				children.Add(new ChunkedDataNodeRef(type, hash, length, handle));
-			}
-			Children = children;
-		}
-
-		/// <inheritdoc/>
-		public override void Serialize(IBlobWriter writer)
-		{
-			foreach (ChunkedDataNodeRef child in Children)
-			{
-				writer.WriteIoHash(child.Hash);
-				writer.WriteUnsignedVarInt((int)child.Type);
-				// Pending V3: writer.WriteUnsignedVarInt((ulong)child.Length);
-				writer.WriteBlobReference(child.Handle);
-			}
-		}
-
 		/// <summary>
 		/// Create a tree of nodes from the given list of handles, splitting nodes in each layer based on the hash of the last node.
 		/// </summary>
@@ -456,8 +475,8 @@ namespace EpicGames.Horde.Storage.Nodes
 				handleBuffer.Clear();
 				foreach ((InteriorChunkedDataNode interiorNode, long interiorLength) in interiorNodes)
 				{
-					HashedNodeRef<ChunkedDataNode> newNodeRef = await writer.WriteHashedNodeAsync<ChunkedDataNode>(interiorNode, cancellationToken);
-					handleBuffer.Add(new ChunkedDataNodeRef(ChunkedDataNodeType.Interior, newNodeRef.Hash, interiorLength, newNodeRef.Handle));
+					IBlobHandle<InteriorChunkedDataNode> interiorHandle = await writer.WriteBlobAsync(interiorNode, null, cancellationToken);
+					handleBuffer.Add(new ChunkedDataNodeRef(ChunkedDataNodeType.Interior, interiorLength, interiorHandle));
 				}
 
 				nodeRefs = handleBuffer;
@@ -481,7 +500,7 @@ namespace EpicGames.Horde.Storage.Nodes
 				index = Math.Min(index + options.MinChildCount, nodeRefs.Count);
 				for (; index < maxIndex; index++)
 				{
-					nodeRefs[index].Hash.CopyTo(buffer);
+					nodeRefs[index].Handle.Hash.CopyTo(buffer);
 
 					uint value = BinaryPrimitives.ReadUInt32LittleEndian(buffer);
 					if (value < options.SliceThreshold)
@@ -502,12 +521,12 @@ namespace EpicGames.Horde.Storage.Nodes
 		}
 
 		/// <inheritdoc/>
-		public override async Task CopyToStreamAsync(Stream outputStream, CancellationToken cancellationToken)
+		public override async Task CopyToStreamAsync(Stream outputStream, BlobSerializerOptions? options, CancellationToken cancellationToken)
 		{
 			foreach (ChunkedDataNodeRef childNodeRef in Children)
 			{
-				ChunkedDataNode childNode = await childNodeRef.ExpandAsync(cancellationToken);
-				await childNode.CopyToStreamAsync(outputStream, cancellationToken);
+				ChunkedDataNode childNode = await childNodeRef.ReadBlobAsync(options, cancellationToken);
+				await childNode.CopyToStreamAsync(outputStream, options, cancellationToken);
 			}
 		}
 
@@ -522,8 +541,7 @@ namespace EpicGames.Horde.Storage.Nodes
 			BlobReader nodeReader = new BlobReader(nodeData);
 			while (nodeReader.GetMemory(0).Length > 0)
 			{
-				_ = nodeReader.ReadIoHash();
-				IBlobHandle handle = nodeReader.ReadBlobReference();
+				IBlobHandle<ChunkedDataNode> handle = nodeReader.ReadBlobHandle<ChunkedDataNode>();
 				if (nodeReader.Version >= 2)
 				{
 					_ = nodeReader.ReadUnsignedVarInt(); // Type
@@ -536,4 +554,48 @@ namespace EpicGames.Horde.Storage.Nodes
 			}
 		}
 	}
+
+	class InteriorChunkedDataNodeConverter : BlobConverter<InteriorChunkedDataNode>
+	{
+		/// <summary>
+		/// Static accessor for the blob type
+		/// </summary>
+		public static BlobType BlobType { get; } = new BlobType("{F4DEDDBC-4C7A-70CB-11F0-4783B9CDCCAF}", 2); // Pending V3
+
+		public override InteriorChunkedDataNode Read(IBlobReader reader, BlobSerializerOptions options)
+		{
+			// Keep this code in sync with CopyToStreamAsync
+			List<ChunkedDataNodeRef> children = new List<ChunkedDataNodeRef>();
+			while (reader.GetMemory().Length > 0)
+			{
+				IBlobHandle<ChunkedDataNode> handle = reader.ReadBlobHandle<ChunkedDataNode>();
+
+				ChunkedDataNodeType type = ChunkedDataNodeType.Unknown;
+				if (reader.Version >= 2)
+				{
+					type = (ChunkedDataNodeType)reader.ReadUnsignedVarInt();
+				}
+
+				long length = 0;
+				if (reader.Version >= 3)
+				{
+					length = (long)reader.ReadUnsignedVarInt();
+				}
+
+				children.Add(new ChunkedDataNodeRef(type, length, handle));
+			}
+			return new InteriorChunkedDataNode(children);
+		}
+
+		public override BlobType Write(IBlobWriter writer, InteriorChunkedDataNode value, BlobSerializerOptions options)
+		{
+			foreach (ChunkedDataNodeRef child in value.Children)
+			{
+				writer.WriteBlobHandle(child.Handle);
+				writer.WriteUnsignedVarInt((int)child.Type);
+				// Pending V3: writer.WriteUnsignedVarInt((ulong)child.Length);
+			}
+			return InteriorChunkedDataNodeConverter.BlobType;
+		}
+	}
 }
diff --git a/Engine/Source/Programs/Shared/EpicGames.Horde/Storage/Nodes/ChunkedDataWriter.cs b/Engine/Source/Programs/Shared/EpicGames.Horde/Storage/Nodes/ChunkedDataWriter.cs
--- a/Engine/Source/Programs/Shared/EpicGames.Horde/Storage/Nodes/ChunkedDataWriter.cs
+++ b/Engine/Source/Programs/Shared/EpicGames.Horde/Storage/Nodes/ChunkedDataWriter.cs
@@ -70,8 +70,6 @@ namespace EpicGames.Horde.Storage.Nodes
 		/// </summary>
 		public const int DefaultBufferLength = 32 * 1024;
 
-		static readonly BlobType s_leafNodeType = Node.GetNodeType<LeafChunkedDataNode>();
-
 		readonly IStorageWriter _writer;
 		readonly ChunkingOptions _options;
 		readonly Blake3.Hasher _hasher;
@@ -345,8 +343,8 @@ namespace EpicGames.Horde.Storage.Nodes
 		/// <returns>Handle to the written leaf node</returns>
 		async ValueTask FlushLeafNodeAsync(CancellationToken cancellationToken)
 		{
-			HashedNodeRef<ChunkedDataNode> nodeRef  = await _writer.WriteHashedNodeRefAsync<ChunkedDataNode>(s_leafNodeType, _leafLength, Array.Empty<IBlobHandle>(), cancellationToken);
-			_leafHandles.Add(new ChunkedDataNodeRef(ChunkedDataNodeType.Leaf, nodeRef.Hash, _leafLength, nodeRef.Handle));
+			IBlobHandle<LeafChunkedDataNode> leafHandle = await _writer.WriteBlobAsync<LeafChunkedDataNode>(LeafChunkedDataNodeConverter.BlobType, _leafLength, Array.Empty<IBlobHandle>(), cancellationToken);
+			_leafHandles.Add(new ChunkedDataNodeRef(_leafLength, leafHandle));
 			ResetLeafState();
 		}
 	}
diff --git a/Engine/Source/Programs/Shared/EpicGames.Horde/Storage/Nodes/CommitNode.cs b/Engine/Source/Programs/Shared/EpicGames.Horde/Storage/Nodes/CommitNode.cs
--- a/Engine/Source/Programs/Shared/EpicGames.Horde/Storage/Nodes/CommitNode.cs
+++ b/Engine/Source/Programs/Shared/EpicGames.Horde/Storage/Nodes/CommitNode.cs
@@ -9,8 +9,7 @@ namespace EpicGames.Horde.Storage.Nodes
 	/// <summary>
 	/// A node representing commit metadata
 	/// </summary>
-	[BlobType("{64D50724-41C0-6B22-1CB5-90A8171824D6}", 1)]
-	public class CommitNode : Node
+	public class CommitNode
 	{
 		/// <summary>
 		/// The commit number
@@ -20,7 +19,7 @@ namespace EpicGames.Horde.Storage.Nodes
 		/// <summary>
 		/// Reference to the parent commit
 		/// </summary>
-		public NodeRef<CommitNode>? Parent { get; set; }
+		public IBlobHandle<CommitNode>? Parent { get; set; }
 
 		/// <summary>
 		/// Human readable name of the author of this change
@@ -60,66 +59,82 @@ namespace EpicGames.Horde.Storage.Nodes
 		/// <summary>
 		/// Metadata for this commit, keyed by arbitrary GUID
 		/// </summary>
-		public Dictionary<Guid, HashedNodeRef> Metadata { get; } = new Dictionary<Guid, HashedNodeRef>();
+		public Dictionary<Guid, IBlobHandle<object>> Metadata { get; } = new Dictionary<Guid, IBlobHandle<object>>();
 
 		/// <summary>
 		/// Constructor
 		/// </summary>
-		/// <param name="number">Commit number</param>
-		/// <param name="parent">The parent commit</param>
-		/// <param name="author">Author of this commit</param>
-		/// <param name="message">Message for the commit</param>
-		/// <param name="time">The commit time</param>
-		/// <param name="contents">Contents of the tree at this commit</param>
-		public CommitNode(int number, NodeRef<CommitNode>? parent, string author, string message, DateTime time, DirectoryNodeRef contents)
+		public CommitNode(int number, IBlobHandle<CommitNode>? parent, string author, string? authorId, string? committer, string? commiterId, string message, DateTime time, DirectoryNodeRef contents, Dictionary<Guid, IBlobHandle<object>> metadata)
 		{
 			Number = number;
 			Parent = parent;
 			Author = author;
+			AuthorId = authorId;
+			Committer = committer;
+			CommitterId = commiterId;
 			Message = message;
 			Time = time;
 			Contents = contents;
+			Metadata = metadata;
 		}
+	}
 
-		/// <summary>
-		/// Deserializing constructor
-		/// </summary>
-		/// <param name="reader"></param>
-		public CommitNode(IBlobReader reader)
+	class CommitNodeConverter : BlobConverter<CommitNode>
+	{
+		static readonly BlobType s_blobType = new BlobType("{64D50724-41C0-6B22-1CB5-90A8171824D6}", 1);
+
+		public override CommitNode Read(IBlobReader reader, BlobSerializerOptions options)
 		{
-			Number = (int)reader.ReadUnsignedVarInt();
-			Parent = reader.ReadOptionalNodeRef<CommitNode>();
-			Author = reader.ReadString();
-			AuthorId = reader.ReadOptionalString();
-			Committer = reader.ReadOptionalString();
-			CommitterId = reader.ReadOptionalString();
-			Message = reader.ReadString();
-			Time = reader.ReadDateTime();
-
-			IoHash hash = reader.ReadIoHash();
+			int number = (int)reader.ReadUnsignedVarInt();
+
+			IBlobHandle<CommitNode>? parent;
+			if (reader.ReadBoolean())
+			{
+				parent = reader.ReadBlobHandle<CommitNode>();
+			}
+			else
+			{
+				parent = null; 
+			}
+
+			string author = reader.ReadString();
+			string? authorId = reader.ReadOptionalString();
+			string? committer = reader.ReadOptionalString();
+			string? committerId = reader.ReadOptionalString();
+			string message = reader.ReadString();
+			DateTime time = reader.ReadDateTime();
+
+			IBlobHandle<DirectoryNode> contentsNode = reader.ReadBlobHandle<DirectoryNode>();
 			long length = (long)reader.ReadUnsignedVarInt();
-			Contents = new DirectoryNodeRef(hash, length, reader.ReadBlobReference());
+			DirectoryNodeRef contents = new DirectoryNodeRef(length, contentsNode);
+
+			Dictionary<Guid, IBlobHandle<object>> metadata = reader.ReadDictionary(() => reader.ReadGuidUnrealOrder(), () => reader.ReadBlobHandle<object>());
 
-			Metadata = reader.ReadDictionary(() => reader.ReadGuidUnrealOrder(), () => reader.ReadHashedNodeRef());
+			return new CommitNode(number, parent, author, authorId, committer, committerId, message, time, contents, metadata);
 		}
 
 		/// <inheritdoc/>
-		public override void Serialize(IBlobWriter writer)
+		public override BlobType Write(IBlobWriter writer, CommitNode value, BlobSerializerOptions options)
 		{
-			writer.WriteUnsignedVarInt(Number);
-			writer.WriteOptionalNodeRef(Parent);
-			writer.WriteString(Author);
-			writer.WriteOptionalString(AuthorId);
-			writer.WriteOptionalString(Committer);
-			writer.WriteOptionalString(CommitterId);
-			writer.WriteString(Message);
-			writer.WriteDateTime(Time);
-
-			writer.WriteIoHash(Contents.Hash);
-			writer.WriteUnsignedVarInt((ulong)Contents.Length);
-			writer.WriteBlobReference(Contents.Handle);
-
-			writer.WriteDictionary(Metadata, key => writer.WriteGuidUnrealOrder(key), value => writer.WriteHashedNodeRef(value));
+			writer.WriteUnsignedVarInt(value.Number);
+			writer.WriteBoolean(value.Parent != null);
+			if (value.Parent != null)
+			{
+				writer.WriteBlobHandle(value.Parent);
+			}
+			writer.WriteString(value.Author);
+			writer.WriteOptionalString(value.AuthorId);
+			writer.WriteOptionalString(value.Committer);
+			writer.WriteOptionalString(value.CommitterId);
+			writer.WriteString(value.Message);
+			writer.WriteDateTime(value.Time);
+
+			writer.WriteBlobHandle(value.Contents.Handle);
+			writer.WriteUnsignedVarInt((ulong)value.Contents.Length);
+
+			writer.WriteDictionary(value.Metadata, key => writer.WriteGuidUnrealOrder(key), value => writer.WriteBlobHandle(value));
+
+			return s_blobType;
 		}
 	}
 }
diff --git a/Engine/Source/Programs/Shared/EpicGames.Horde/Storage/Nodes/DdcRefNode.cs b/Engine/Source/Programs/Shared/EpicGames.Horde/Storage/Nodes/DdcRefNode.cs
--- a/Engine/Source/Programs/Shared/EpicGames.Horde/Storage/Nodes/DdcRefNode.cs
+++ b/Engine/Source/Programs/Shared/EpicGames.Horde/Storage/Nodes/DdcRefNode.cs
@@ -1,6 +1,7 @@
 // Copyright Epic Games, Inc. All Rights Reserved.
 
 using System.Collections.Generic;
+using System.Linq;
 using EpicGames.Core;
 using EpicGames.Horde.Storage;
 
@@ -9,8 +10,8 @@ namespace Horde.Server.Ddc
 	/// <summary>
 	/// A node containing ref data
 	/// </summary>
-	[BlobType("{0C7E5F25-4B55-454B-63F4-4A9B74D00651}", 1)]
-	public class DdcRefNode : Node
+	[BlobConverter(typeof(DdcRefNodeConverter))]
+	public class DdcRefNode
 	{
 		/// <summary>
 		/// Hash of the root node
@@ -30,27 +31,27 @@ namespace Horde.Server.Ddc
 			RootHash = rootHash;
 			References = new List<(IoHash, IBlobHandle)>();
 		}
+	}
 
-		/// <summary>
-		/// Constructor
-		/// </summary>
-		public DdcRefNode(IBlobReader reader)
-		{
-			RootHash = reader.ReadIoHash();
-			References = reader.ReadList(x => (reader.ReadIoHash(), reader.ReadBlobReference()));
-		}
+	class DdcRefNodeConverter : BlobConverter<DdcRefNode>
+	{
+		static BlobType s_blobType = new BlobType("{0C7E5F25-4B55-454B-63F4-4A9B74D00651}", 1);
 
-		/// <inheritdoc/>
-		public override void Serialize(IBlobWriter writer)
+		public override DdcRefNode Read(IBlobReader reader, BlobSerializerOptions options)
 		{
-			writer.WriteIoHash(RootHash);
-			writer.WriteList(References, x => WriteReference(writer, x.Hash, x.Handle));
+			IoHash rootHash = reader.ReadIoHash();
+			List<IBlobHandle<object>> references = reader.ReadList(x => reader.ReadBlobHandle<object>());
+
+			DdcRefNode refNode = new DdcRefNode(rootHash);
+			refNode.References.AddRange(references.Select(x => (x.Hash, (IBlobHandle)x)));
+			return refNode;
 		}
 
-		static void WriteReference(IBlobWriter writer, IoHash hash, IBlobHandle handle)
+		public override BlobType Write(IBlobWriter writer, DdcRefNode value, BlobSerializerOptions options)
 		{
-			writer.WriteIoHash(hash);
-			writer.WriteBlobReference(handle);
+			writer.WriteIoHash(value.RootHash);
+			writer.WriteList(value.References, x => writer.WriteBlobHandle(x.Handle.ForType<object>(x.Hash)));
+			return s_blobType;
 		}
 	}
 }
diff --git a/Engine/Source/Programs/Shared/EpicGames.Horde/Storage/Nodes/DirectoryEntry.cs b/Engine/Source/Programs/Shared/EpicGames.Horde/Storage/Nodes/DirectoryEntry.cs
--- a/Engine/Source/Programs/Shared/EpicGames.Horde/Storage/Nodes/DirectoryEntry.cs
+++ b/Engine/Source/Programs/Shared/EpicGames.Horde/Storage/Nodes/DirectoryEntry.cs
@@ -9,14 +9,7 @@ namespace EpicGames.Horde.Storage.Nodes
 	/// Entry for a directory within a directory node
 	/// </summary>
 	[DebuggerDisplay("{Name}")]
-	public record class DirectoryEntry(string Name, IoHash Hash, long Length, IBlobHandle Handle) : DirectoryNodeRef(Hash, Length, Handle)
+	public record class DirectoryEntry(string Name, long Length, IBlobHandle<DirectoryNode> Handle) : DirectoryNodeRef(Length, Handle)
 	{
-		/// <summary>
-		/// Constructor
-		/// </summary>
-		public DirectoryEntry(string name, long length, HashedNodeRef<DirectoryNode> target)
-			: this(name, target.Hash, length, target.Handle)
-		{
-		}
 	}
 }
diff --git a/Engine/Source/Programs/Shared/EpicGames.Horde/Storage/Nodes/DirectoryNode.cs b/Engine/Source/Programs/Shared/EpicGames.Horde/Storage/Nodes/DirectoryNode.cs
--- a/Engine/Source/Programs/Shared/EpicGames.Horde/Storage/Nodes/DirectoryNode.cs
+++ b/Engine/Source/Programs/Shared/EpicGames.Horde/Storage/Nodes/DirectoryNode.cs
@@ -1,6 +1,7 @@
 // Copyright Epic Games, Inc. All Rights Reserved.
 
 using EpicGames.Core;
+using EpicGames.Serialization;
 using Microsoft.Extensions.Logging;
 using System;
 using System.Buffers;
@@ -140,9 +141,14 @@ namespace EpicGames.Horde.Storage.Nodes
 	/// <summary>
 	/// A directory node
 	/// </summary>
-	[BlobType("{0714EC11-4D07-291A-8AE7-7F86799980D6}", 2)]
-	public class DirectoryNode : Node
+	[BlobConverter(typeof(DirectoryNodeConverter))]
+	public class DirectoryNode
 	{
+		/// <summary>
+		/// Type of serialized directory node blobs
+		/// </summary>
+		public static BlobType BlobType { get; } = new BlobType("{0714EC11-4D07-291A-8AE7-7F86799980D6}", 2);
+
 		readonly SortedDictionary<string, FileEntry> _nameToFileEntry = new SortedDictionary<string, FileEntry>(StringComparer.Ordinal);
 		readonly SortedDictionary<string, DirectoryEntry> _nameToDirectoryEntry = new SortedDictionary<string, DirectoryEntry>(StringComparer.Ordinal);
 
@@ -185,87 +191,6 @@ namespace EpicGames.Horde.Storage.Nodes
 			Flags = flags;
 		}
 
-		/// <summary>
-		/// Deserialization constructor
-		/// </summary>
-		/// <param name="reader">Reader to deserialize from</param>
-		public DirectoryNode(IBlobReader reader)
-		{
-			Flags = (DirectoryFlags)reader.ReadUnsignedVarInt();
-
-			int fileCount = (int)reader.ReadUnsignedVarInt();
-			for (int idx = 0; idx < fileCount; idx++)
-			{
-				IoHash targetHash = reader.ReadIoHash();
-
-				ChunkedDataNodeType targetType = ChunkedDataNodeType.Unknown;
-				if (reader.Version >= 2)
-				{
-					targetType = (ChunkedDataNodeType)reader.ReadUnsignedVarInt();
-				}
-
-				string name = reader.ReadString();
-				FileEntryFlags flags = (FileEntryFlags)reader.ReadUnsignedVarInt();
-				long length = (long)reader.ReadUnsignedVarInt();
-				IoHash streamHash = reader.ReadIoHash();
-				ChunkedDataNodeRef target = new ChunkedDataNodeRef(targetType, targetHash, length, reader.ReadBlobReference());
-
-				ReadOnlyMemory<byte> customData = default;
-				if ((flags & FileEntryFlags.HasCustomData) != 0)
-				{
-					customData = reader.ReadVariableLengthBytes();
-					flags &= ~FileEntryFlags.HasCustomData;
-				}
-
-				_nameToFileEntry[name] = new FileEntry(name, flags, length, streamHash, target, customData);
-			}
-
-			int directoryCount = (int)reader.ReadUnsignedVarInt();
-			for (int idx = 0; idx < directoryCount; idx++)
-			{
-				IoHash hash = reader.ReadIoHash();
-				long length = (long)reader.ReadUnsignedVarInt();
-				string name = reader.ReadString();
-
-				_nameToDirectoryEntry[name] = new DirectoryEntry(name, hash, length, reader.ReadBlobReference());
-			}
-		}
-
-		/// <inheritdoc/>
-		public override void Serialize(IBlobWriter writer)
-		{
-			writer.WriteUnsignedVarInt((ulong)Flags);
-
-			writer.WriteUnsignedVarInt(Files.Count);
-			foreach (FileEntry fileEntry in _nameToFileEntry.Values)
-			{
-				writer.WriteIoHash(fileEntry.Target.Hash);
-				writer.WriteUnsignedVarInt((int)fileEntry.Target.Type);
-				writer.WriteBlobReference(fileEntry.Target.Handle);
-
-				FileEntryFlags flags = (fileEntry.CustomData.Length > 0) ? (fileEntry.Flags | FileEntryFlags.HasCustomData) : (fileEntry.Flags & ~FileEntryFlags.HasCustomData);
-
-				writer.WriteString(fileEntry.Name);
-				writer.WriteUnsignedVarInt((ulong)flags);
-				writer.WriteUnsignedVarInt((ulong)fileEntry.Length);
-				writer.WriteIoHash(fileEntry.StreamHash);
-
-				if ((flags & FileEntryFlags.HasCustomData) != 0)
-				{
-					writer.WriteVariableLengthBytes(fileEntry.CustomData.Span);
-				}
-			}
-
-			writer.WriteUnsignedVarInt(Directories.Count);
-			foreach (DirectoryEntry directoryEntry in _nameToDirectoryEntry.Values)
-			{
-				writer.WriteIoHash(directoryEntry.Hash);
-				writer.WriteUnsignedVarInt((ulong)directoryEntry.Length);
-				writer.WriteString(directoryEntry.Name);
-				writer.WriteBlobReference(directoryEntry.Handle);
-			}
-		}
-
 		/// <summary>
 		/// Clear the contents of this directory
 		/// </summary>
@@ -359,9 +284,10 @@ namespace EpicGames.Horde.Storage.Nodes
 		/// Attempts to get a file entry from a path
 		/// </summary>
 		/// <param name="path">Path to the directory</param>
+		/// <param name="options">Options controling serialiation</param>
 		/// <param name="cancellationToken">Cancellation token</param>
 		/// <returns>The directory with the given path, or null if it was not found</returns>
-		public async ValueTask<FileEntry?> GetFileEntryByPathAsync(string path, CancellationToken cancellationToken = default)
+		public async ValueTask<FileEntry?> GetFileEntryByPathAsync(string path, BlobSerializerOptions? options = null, CancellationToken cancellationToken = default)
 		{
 			FileEntry? fileEntry;
 
@@ -375,7 +301,7 @@ namespace EpicGames.Horde.Storage.Nodes
 			}
 			else
 			{
-				DirectoryNode? directoryNode = await GetDirectoryByPathAsync(path.Substring(0, slashIdx), cancellationToken);
+				DirectoryNode? directoryNode = await GetDirectoryByPathAsync(path.Substring(0, slashIdx), options, cancellationToken);
 				if (directoryNode == null)
 				{
 					return null;
@@ -393,11 +319,12 @@ namespace EpicGames.Horde.Storage.Nodes
 		/// Attempts to get a directory entry from a path
 		/// </summary>
 		/// <param name="path">Path to the directory</param>
+		/// <param name="options">Options controling serialiation</param>
 		/// <param name="cancellationToken">Cancellation token</param>
 		/// <returns>The directory with the given path, or null if it was not found</returns>
-		public ValueTask<DirectoryNode?> GetDirectoryByPathAsync(string path, CancellationToken cancellationToken = default) => GetDirectoryByPathAsync(this, path, cancellationToken);
+		public ValueTask<DirectoryNode?> GetDirectoryByPathAsync(string path, BlobSerializerOptions? options = null, CancellationToken cancellationToken = default) => GetDirectoryByPathAsync(this, path, options, cancellationToken);
 
-		static async ValueTask<DirectoryNode?> GetDirectoryByPathAsync(DirectoryNode directoryNode, string path, CancellationToken cancellationToken = default)
+		static async ValueTask<DirectoryNode?> GetDirectoryByPathAsync(DirectoryNode directoryNode, string path, BlobSerializerOptions? options = null, CancellationToken cancellationToken = default)
 		{
 			while (path.Length > 0)
 			{
@@ -421,7 +348,7 @@ namespace EpicGames.Horde.Storage.Nodes
 					return null;
 				}
 
-				directoryNode = await directoryEntry.ExpandAsync(cancellationToken);
+				directoryNode = await directoryEntry.Handle.ReadBlobAsync(options, cancellationToken);
 			}
 			return directoryNode;
 		}
@@ -430,9 +357,10 @@ namespace EpicGames.Horde.Storage.Nodes
 		/// Deletes a file with the given path
 		/// </summary>
 		/// <param name="path"></param>
+		/// <param name="options">Options controlling serialization</param>
 		/// <param name="cancellationToken"></param>
 		/// <returns></returns>
-		public async ValueTask<bool> DeleteFileByPathAsync(string path, CancellationToken cancellationToken = default)
+		public async ValueTask<bool> DeleteFileByPathAsync(string path, BlobSerializerOptions? options = null, CancellationToken cancellationToken = default)
 		{
 			string remainingPath = path;
 			for (DirectoryNode? directory = this; directory != null;)
@@ -444,7 +372,7 @@ namespace EpicGames.Horde.Storage.Nodes
 				}
 				if (length > 0)
 				{
-					directory = await directory.TryOpenDirectoryAsync(remainingPath.Substring(0, length), cancellationToken);
+					directory = await directory.TryOpenDirectoryAsync(remainingPath.Substring(0, length), options, cancellationToken);
 				}
 				remainingPath = remainingPath.Substring(length + 1);
 			}
@@ -488,11 +416,12 @@ namespace EpicGames.Horde.Storage.Nodes
 		/// Tries to get a directory with the given name
 		/// </summary>
 		/// <param name="name">Name of the new directory</param>
+		/// <param name="options">Options controling serialiation</param>
 		/// <param name="cancellationToken">Cancellation token for the operation</param>
 		/// <returns>The new directory object</returns>
-		public async ValueTask<DirectoryNode> OpenDirectoryAsync(string name, CancellationToken cancellationToken = default)
+		public async ValueTask<DirectoryNode> OpenDirectoryAsync(string name, BlobSerializerOptions? options = null, CancellationToken cancellationToken = default)
 		{
-			DirectoryNode? directoryNode = await TryOpenDirectoryAsync(name, cancellationToken);
+			DirectoryNode? directoryNode = await TryOpenDirectoryAsync(name, options, cancellationToken);
 			if (directoryNode == null)
 			{
 				throw new DirectoryNotFoundException();
@@ -504,13 +433,14 @@ namespace EpicGames.Horde.Storage.Nodes
 		/// Tries to get a directory with the given name
 		/// </summary>
 		/// <param name="name">Name of the new directory</param>
+		/// <param name="options">Options controling serialiation</param>
 		/// <param name="cancellationToken">Cancellation token for the operation</param>
 		/// <returns>The new directory object</returns>
-		public async ValueTask<DirectoryNode?> TryOpenDirectoryAsync(string name, CancellationToken cancellationToken = default)
+		public async ValueTask<DirectoryNode?> TryOpenDirectoryAsync(string name, BlobSerializerOptions? options = null, CancellationToken cancellationToken = default)
 		{
 			if (TryGetDirectoryEntry(name, out DirectoryEntry? entry))
 			{
-				return await entry.ExpandAsync(cancellationToken);
+				return await entry.Handle.ReadBlobAsync(options, cancellationToken);
 			}
 			else
 			{
@@ -527,22 +457,22 @@ namespace EpicGames.Horde.Storage.Nodes
 
 		#endregion
 
-		/// <inheritdoc cref="AddFilesAsync(DirectoryReference, IEnumerable{FileInfo}, IStorageWriter, ChunkingOptions?, IProgress{ICopyStats}?, CancellationToken)"/>
-		public async Task AddFilesAsync(DirectoryInfo directoryInfo, IStorageWriter writer, ChunkingOptions? options = null, IProgress<ICopyStats>? progress = null, CancellationToken cancellationToken = default)
+		/// <inheritdoc cref="AddFilesAsync(DirectoryReference, IEnumerable{FileInfo}, IStorageWriter, ChunkingOptions?, IProgress{ICopyStats}?, BlobSerializerOptions?, CancellationToken)"/>
+		public async Task AddFilesAsync(DirectoryInfo directoryInfo, IStorageWriter writer, ChunkingOptions? options = null, IProgress<ICopyStats>? progress = null, BlobSerializerOptions? serializerOptions = null, CancellationToken cancellationToken = default)
 		{
-			await AddFilesAsync(new DirectoryReference(directoryInfo), directoryInfo.EnumerateFiles("*", SearchOption.AllDirectories).ToList(), writer, options, progress, cancellationToken);
+			await AddFilesAsync(new DirectoryReference(directoryInfo), directoryInfo.EnumerateFiles("*", SearchOption.AllDirectories).ToList(), writer, options, progress, serializerOptions, cancellationToken);
 		}
 
-		/// <inheritdoc cref="AddFilesAsync(DirectoryReference, IEnumerable{FileInfo}, IStorageWriter, ChunkingOptions?, IProgress{ICopyStats}?, CancellationToken)"/>
-		public Task AddFilesAsync(DirectoryReference baseDir, IEnumerable<FileReference> files, IStorageWriter writer, ChunkingOptions? options = null, IProgress<ICopyStats>? progress = null, CancellationToken cancellationToken = default)
+		/// <inheritdoc cref="AddFilesAsync(DirectoryReference, IEnumerable{FileInfo}, IStorageWriter, ChunkingOptions?, IProgress{ICopyStats}?, BlobSerializerOptions?, CancellationToken)"/>
+		public Task AddFilesAsync(DirectoryReference baseDir, IEnumerable<FileReference> files, IStorageWriter writer, ChunkingOptions? options = null, IProgress<ICopyStats>? progress = null, BlobSerializerOptions? serializerOptions = null, CancellationToken cancellationToken = default)
 		{
-			return AddFilesAsync(baseDir, files.Select(x => x.ToFileInfo()).ToList(), writer, options, progress, cancellationToken);
+			return AddFilesAsync(baseDir, files.Select(x => x.ToFileInfo()).ToList(), writer, options, progress, serializerOptions, cancellationToken);
 		}
 
-		/// <inheritdoc cref="AddFilesAsync(DirectoryReference, IEnumerable{FileInfo}, IStorageWriter, ChunkingOptions?, IProgress{ICopyStats}?, CancellationToken)"/>
-		public Task AddFilesAsync(DirectoryInfo baseDir, IEnumerable<FileInfo> files, IStorageWriter writer, ChunkingOptions? options = null, IProgress<ICopyStats>? progress = null, CancellationToken cancellationToken = default)
+		/// <inheritdoc cref="AddFilesAsync(DirectoryReference, IEnumerable{FileInfo}, IStorageWriter, ChunkingOptions?, IProgress{ICopyStats}?, BlobSerializerOptions?, CancellationToken)"/>
+		public Task AddFilesAsync(DirectoryInfo baseDir, IEnumerable<FileInfo> files, IStorageWriter writer, ChunkingOptions? options = null, IProgress<ICopyStats>? progress = null, BlobSerializerOptions? serializerOptions = null, CancellationToken cancellationToken = default)
 		{
-			return AddFilesAsync(new DirectoryReference(baseDir), files.ToList(), writer, options, progress, cancellationToken);
+			return AddFilesAsync(new DirectoryReference(baseDir), files.ToList(), writer, options, progress, serializerOptions, cancellationToken);
 		}
 
 		/// <summary>
@@ -553,8 +483,9 @@ namespace EpicGames.Horde.Storage.Nodes
 		/// <param name="options">Options for chunking file content</param>
 		/// <param name="writer">Writer for new node data</param>
 		/// <param name="progress">Feedback interface for progress updates</param>
+		/// <param name="serializerOptions">Options controlling serialization</param>
 		/// <param name="cancellationToken">Cancellation token for the operation</param>
-		public async Task AddFilesAsync(DirectoryReference baseDir, IEnumerable<FileInfo> files, IStorageWriter writer, ChunkingOptions? options = null, IProgress<ICopyStats>? progress = null, CancellationToken cancellationToken = default)
+		public async Task AddFilesAsync(DirectoryReference baseDir, IEnumerable<FileInfo> files, IStorageWriter writer, ChunkingOptions? options = null, IProgress<ICopyStats>? progress = null, BlobSerializerOptions? serializerOptions = null, CancellationToken cancellationToken = default)
 		{
 			options ??= new ChunkingOptions();
 
@@ -623,7 +554,7 @@ namespace EpicGames.Horde.Storage.Nodes
 			}
 
 			// Add all the new entries to the tree
-			await UpdateAsync(update, writer, cancellationToken);
+			await UpdateAsync(update, writer, serializerOptions, cancellationToken);
 		}
 
 		static List<(int Index, int Count)> ComputePartitions(IReadOnlyList<FileInfo> files, long totalSize)
@@ -691,12 +622,13 @@ namespace EpicGames.Horde.Storage.Nodes
 		/// </summary>
 		/// <param name="updates">Files to add</param>
 		/// <param name="writer">Writer for new node data</param>
+		/// <param name="options">Options controling serialiation</param>
 		/// <param name="cancellationToken">Cancellation token for the operation</param>
-		public Task UpdateAsync(IEnumerable<FileUpdate> updates, IStorageWriter writer, CancellationToken cancellationToken = default)
+		public Task UpdateAsync(IEnumerable<FileUpdate> updates, IStorageWriter writer, BlobSerializerOptions? options = null, CancellationToken cancellationToken = default)
 		{
 			DirectoryUpdate update = new DirectoryUpdate();
 			update.AddFiles(updates);
-			return UpdateAsync(update, writer, cancellationToken);
+			return UpdateAsync(update, writer, options, cancellationToken);
 		}
 
 		/// <summary>
@@ -704,8 +636,9 @@ namespace EpicGames.Horde.Storage.Nodes
 		/// </summary>
 		/// <param name="update">Files to add</param>
 		/// <param name="writer">Writer for new node data</param>
+		/// <param name="options">Options controling serialiation</param>
 		/// <param name="cancellationToken">Cancellation token for the operation</param>
-		public async Task UpdateAsync(DirectoryUpdate update, IStorageWriter writer, CancellationToken cancellationToken = default)
+		public async Task UpdateAsync(DirectoryUpdate update, IStorageWriter writer, BlobSerializerOptions? options = null, CancellationToken cancellationToken = default)
 		{
 			foreach ((string name, DirectoryUpdate? directory) in update.Directories)
 			{
@@ -715,11 +648,11 @@ namespace EpicGames.Horde.Storage.Nodes
 				}
 				else
 				{
-					DirectoryNode? childNode = await TryOpenDirectoryAsync(name, cancellationToken);
+					DirectoryNode? childNode = await TryOpenDirectoryAsync(name, options, cancellationToken);
 					childNode ??= new DirectoryNode();
-					await childNode.UpdateAsync(directory, writer, cancellationToken);
-					HashedNodeRef<DirectoryNode> nodeRef = await writer.WriteHashedNodeAsync(childNode, cancellationToken);
-					_nameToDirectoryEntry[name] = new DirectoryEntry(name, nodeRef.Hash, childNode.Length, nodeRef.Handle);
+					await childNode.UpdateAsync(directory, writer, options, cancellationToken);
+					IBlobHandle<DirectoryNode> handle = await writer.WriteBlobAsync<DirectoryNode>(childNode, options, cancellationToken);
+					_nameToDirectoryEntry[name] = new DirectoryEntry(name, childNode.Length, handle);
 				}
 			}
 			foreach ((string name, FileEntry? file) in update.Files)
@@ -741,8 +674,9 @@ namespace EpicGames.Horde.Storage.Nodes
 		/// <param name="stream">Input stream</param>
 		/// <param name="writer">Writer for new nodes</param>
 		/// <param name="options"></param>
+		/// <param name="serializerOptions">Options controling serialiation</param>
 		/// <param name="cancellationToken">Cancellation token for the operation</param>
-		public async Task CopyFromZipStreamAsync(Stream stream, IStorageWriter writer, ChunkingOptions options, CancellationToken cancellationToken = default)
+		public async Task CopyFromZipStreamAsync(Stream stream, IStorageWriter writer, ChunkingOptions options, BlobSerializerOptions? serializerOptions = null, CancellationToken cancellationToken = default)
 		{
 			// Create all the leaf nodes
 			List<(ZipArchiveEntry, LeafChunkedData)> entries = new List<(ZipArchiveEntry, LeafChunkedData)>();
@@ -776,16 +710,17 @@ namespace EpicGames.Horde.Storage.Nodes
 			}
 
 			// Update the tree
-			await UpdateAsync(updates, writer, cancellationToken);
+			await UpdateAsync(updates, writer, serializerOptions, cancellationToken);
 		}
 
 		/// <summary>
 		/// Utility function to allow extracting a packed directory to disk
 		/// </summary>
 		/// <param name="directoryInfo"></param>
+		/// <param name="options"></param>
 		/// <param name="logger"></param>
 		/// <param name="cancellationToken"></param>
-		public Task CopyToDirectoryAsync(DirectoryInfo directoryInfo, ILogger logger, CancellationToken cancellationToken) => CopyToDirectoryAsync(directoryInfo, null, logger, cancellationToken);
+		public Task CopyToDirectoryAsync(DirectoryInfo directoryInfo, BlobSerializerOptions? options, ILogger logger, CancellationToken cancellationToken) => CopyToDirectoryAsync(directoryInfo, null, options, logger, cancellationToken);
 
 		/// <summary>
 		/// Utility function to allow extracting a packed directory to disk
@@ -793,8 +728,9 @@ namespace EpicGames.Horde.Storage.Nodes
 		/// <param name="directoryInfo">Direcotry to write to</param>
 		/// <param name="progress">Sink for progress updates</param>
 		/// <param name="logger">Logger for output</param>
+		/// <param name="options">Options controlling serialization</param>
 		/// <param name="cancellationToken">Cancellation token for the operation</param>
-		public async Task CopyToDirectoryAsync(DirectoryInfo directoryInfo, IProgress<ICopyStats>? progress, ILogger logger, CancellationToken cancellationToken)
+		public async Task CopyToDirectoryAsync(DirectoryInfo directoryInfo, IProgress<ICopyStats>? progress, BlobSerializerOptions? options, ILogger logger, CancellationToken cancellationToken)
 		{
 			int numTasks = Math.Min(1 + (int)(Length / (16 * 1024 * 1024)), 128);
 			logger.LogInformation("Splitting read into {NumThreads} threads", numTasks);
@@ -820,7 +756,7 @@ namespace EpicGames.Horde.Storage.Nodes
 						maxOffset++;
 					}
 
-					tasks.Add(Task.Run(() => CopyToDirectoryInternalAsync(directoryInfo, minOffset, maxOffset - minOffset, copyStats, logger, cancellationToken), cancellationToken));
+					tasks.Add(Task.Run(() => CopyToDirectoryInternalAsync(directoryInfo, minOffset, maxOffset - minOffset, copyStats, options, logger, cancellationToken), cancellationToken));
 					offset = maxOffset;
 				}
 			}
@@ -831,7 +767,7 @@ namespace EpicGames.Horde.Storage.Nodes
 			}
 		}
 
-		async Task CopyToDirectoryInternalAsync(DirectoryInfo directoryInfo, long windowOffset, long windowLength, CopyStats? copyStats, ILogger logger, CancellationToken cancellationToken)
+		async Task CopyToDirectoryInternalAsync(DirectoryInfo directoryInfo, long windowOffset, long windowLength, CopyStats? copyStats, BlobSerializerOptions? options, ILogger logger, CancellationToken cancellationToken)
 		{
 			directoryInfo.Create();
 
@@ -853,8 +789,8 @@ namespace EpicGames.Horde.Storage.Nodes
 				if (windowOffset < directoryEntry.Length && windowOffset + windowLength >= 0)
 				{
 					DirectoryInfo subDirectoryInfo = directoryInfo.CreateSubdirectory(directoryEntry.Name.ToString());
-					DirectoryNode subDirectoryNode = await directoryEntry.ExpandAsync(cancellationToken);
-					await subDirectoryNode.CopyToDirectoryInternalAsync(subDirectoryInfo, windowOffset, windowLength, copyStats, logger, cancellationToken);
+					DirectoryNode subDirectoryNode = await directoryEntry.Handle.ReadBlobAsync(options, cancellationToken);
+					await subDirectoryNode.CopyToDirectoryInternalAsync(subDirectoryInfo, windowOffset, windowLength, copyStats, options, logger, cancellationToken);
 				}
 				windowOffset -= directoryEntry.Length;
 			}
@@ -864,9 +800,93 @@ namespace EpicGames.Horde.Storage.Nodes
 		/// Returns a stream containing the zipped contents of this directory
 		/// </summary>
 		/// <param name="filter">Filter for files to include in the zip</param>
+		/// <param name="options">Options controlling serialization</param>
 		/// <param name="logger">Logger for diagnostic output</param>
 		/// <returns>Stream containing zipped archive data</returns>
-		public Stream AsZipStream(FileFilter? filter = null, ILogger? logger = null) => new DirectoryNodeZipStream(this, filter, logger);
+		public Stream AsZipStream(FileFilter? filter = null, BlobSerializerOptions? options = null, ILogger? logger = null) => new DirectoryNodeZipStream(this, filter, options, logger);
+	}
+
+	class DirectoryNodeConverter : BlobConverter<DirectoryNode>
+	{
+		/// <inheritdoc/>
+		public override DirectoryNode Read(IBlobReader reader, BlobSerializerOptions options)
+		{
+			DirectoryNode directoryNode = new DirectoryNode((DirectoryFlags)reader.ReadUnsignedVarInt());
+
+			int fileCount = (int)reader.ReadUnsignedVarInt();
+			for (int idx = 0; idx < fileCount; idx++)
+			{
+				IBlobHandle<ChunkedDataNode> targetHandle = reader.ReadBlobHandle<ChunkedDataNode>();
+
+				ChunkedDataNodeType targetType = ChunkedDataNodeType.Unknown;
+				if (reader.Version >= 2)
+				{
+					targetType = (ChunkedDataNodeType)reader.ReadUnsignedVarInt();
+				}
+
+				string name = reader.ReadString();
+				FileEntryFlags flags = (FileEntryFlags)reader.ReadUnsignedVarInt();
+				long length = (long)reader.ReadUnsignedVarInt();
+				IoHash streamHash = reader.ReadIoHash();
+				ChunkedDataNodeRef target = new ChunkedDataNodeRef(targetType, length, targetHandle);
+
+				ReadOnlyMemory<byte> customData = default;
+				if ((flags & FileEntryFlags.HasCustomData) != 0)
+				{
+					customData = reader.ReadVariableLengthBytes();
+					flags &= ~FileEntryFlags.HasCustomData;
+				}
+
+				directoryNode.AddFile(new FileEntry(name, flags, length, streamHash, target, customData));
+			}
+
+			int directoryCount = (int)reader.ReadUnsignedVarInt();
+			for (int idx = 0; idx < directoryCount; idx++)
+			{
+				IBlobHandle<DirectoryNode> directoryHandle = reader.ReadBlobHandle<DirectoryNode>();
+				long length = (long)reader.ReadUnsignedVarInt();
+				string name = reader.ReadString();
+
+				directoryNode.AddDirectory(new DirectoryEntry(name, length, directoryHandle));
+			}
+
+			return directoryNode;
+		}
+
+		/// <inheritdoc/>
+		public override BlobType Write(IBlobWriter writer, DirectoryNode value, BlobSerializerOptions options)
+		{
+			writer.WriteUnsignedVarInt((ulong)value.Flags);
+
+			writer.WriteUnsignedVarInt(value.Files.Count);
+			foreach (FileEntry fileEntry in value.Files)
+			{
+				writer.WriteBlobHandle(fileEntry.Target.Handle);
+				writer.WriteUnsignedVarInt((int)fileEntry.Target.Type);
+
+				FileEntryFlags flags = (fileEntry.CustomData.Length > 0) ? (fileEntry.Flags | FileEntryFlags.HasCustomData) : (fileEntry.Flags & ~FileEntryFlags.HasCustomData);
+
+				writer.WriteString(fileEntry.Name);
+				writer.WriteUnsignedVarInt((ulong)flags);
+				writer.WriteUnsignedVarInt((ulong)fileEntry.Length);
+				writer.WriteIoHash(fileEntry.StreamHash);
+
+				if ((flags & FileEntryFlags.HasCustomData) != 0)
+				{
+					writer.WriteVariableLengthBytes(fileEntry.CustomData.Span);
+				}
+			}
+
+			writer.WriteUnsignedVarInt(value.Directories.Count);
+			foreach (DirectoryEntry directoryEntry in value.Directories)
+			{
+				writer.WriteBlobHandle(directoryEntry.Handle);
+				writer.WriteUnsignedVarInt((ulong)directoryEntry.Length);
+				writer.WriteString(directoryEntry.Name);
+			}
+
+			return DirectoryNode.BlobType;
+		}
 	}
 
 	/// <summary>
@@ -1042,11 +1062,12 @@ namespace EpicGames.Horde.Storage.Nodes
 		/// </summary>
 		/// <param name="node">Root node to copy from</param>
 		/// <param name="filter">Filter for files to include in the zip</param>
+		/// <param name="options">Options controling serialiation</param>
 		/// <param name="logger">Optional logger for debug tracing</param>
-		public DirectoryNodeZipStream(DirectoryNode node, FileFilter? filter, ILogger? logger)
+		public DirectoryNodeZipStream(DirectoryNode node, FileFilter? filter, BlobSerializerOptions? options, ILogger? logger)
 		{
 			_pipe = new Pipe();
-			_backgroundTask = BackgroundTask.StartNew(ctx => CopyToPipeAsync(node, filter, _pipe.Writer, logger, ctx));
+			_backgroundTask = BackgroundTask.StartNew(ctx => CopyToPipeAsync(node, filter, _pipe.Writer, options, logger, ctx));
 			_logger = logger;
 		}
 
@@ -1105,14 +1126,14 @@ namespace EpicGames.Horde.Storage.Nodes
 			return length;
 		}
 
-		static async Task CopyToPipeAsync(DirectoryNode node, FileFilter? filter, PipeWriter writer, ILogger? logger, CancellationToken cancellationToken)
+		static async Task CopyToPipeAsync(DirectoryNode node, FileFilter? filter, PipeWriter writer, BlobSerializerOptions? options, ILogger? logger, CancellationToken cancellationToken)
 		{
 			using Stream outputStream = writer.AsStream();
 			using ZipArchive archive = new ZipArchive(outputStream, ZipArchiveMode.Create);
-			await CopyFilesAsync(node, "", filter, archive, logger, cancellationToken);
+			await CopyFilesAsync(node, "", filter, archive, options, logger, cancellationToken);
 		}
 
-		static async Task CopyFilesAsync(DirectoryNode directory, string prefix, FileFilter? filter, ZipArchive archive, ILogger? logger, CancellationToken cancellationToken)
+		static async Task CopyFilesAsync(DirectoryNode directory, string prefix, FileFilter? filter, ZipArchive archive, BlobSerializerOptions? options, ILogger? logger, CancellationToken cancellationToken)
 		{
 			int numDirs = directory.Directories.Count;
 			int numFiles = directory.Files.Count;
@@ -1123,8 +1144,8 @@ namespace EpicGames.Horde.Storage.Nodes
 				string directoryPath = $"{prefix}{directoryEntry.Name}/";
 				if (filter == null || filter.PossiblyMatches(directoryPath))
 				{
-					DirectoryNode node = await directoryEntry.ExpandAsync(cancellationToken);
-					await CopyFilesAsync(node, directoryPath, filter, archive, logger, cancellationToken);
+					DirectoryNode node = await directoryEntry.Handle.ReadBlobAsync(options, cancellationToken);
+					await CopyFilesAsync(node, directoryPath, filter, archive, options, logger, cancellationToken);
 				}
 				numCopiedDirs++;
 			}
@@ -1181,31 +1202,31 @@ namespace EpicGames.Horde.Storage.Nodes
 		/// <summary>
 		/// Writes a tree of files to a storage writer
 		/// </summary>
-		public static async Task<HashedNodeRef<DirectoryNode>> WriteFilesAsync(this IStorageWriter writer, DirectoryReference baseDir, ChunkingOptions? options = null, IProgress<ICopyStats>? progress = null, CancellationToken cancellationToken = default)
+		public static async Task<IBlobHandle<DirectoryNode>> WriteFilesAsync(this IStorageWriter writer, DirectoryReference baseDir, ChunkingOptions? options = null, IProgress<ICopyStats>? progress = null, CancellationToken cancellationToken = default)
 		{
 			DirectoryNode outputNode = new DirectoryNode();
-			await outputNode.AddFilesAsync(baseDir, DirectoryReference.EnumerateFiles(baseDir, "*", SearchOption.AllDirectories), writer, options, progress, cancellationToken);
-			return await writer.WriteHashedNodeAsync(outputNode, cancellationToken);
+			await outputNode.AddFilesAsync(baseDir, DirectoryReference.EnumerateFiles(baseDir, "*", SearchOption.AllDirectories), writer, options, progress, null, cancellationToken);
+			return await writer.WriteBlobAsync(outputNode, cancellationToken: cancellationToken);
 		}
 
 		/// <summary>
 		/// Writes a tree of files to a storage writer
 		/// </summary>
-		public static async Task<HashedNodeRef<DirectoryNode>> WriteFilesAsync(this IStorageWriter writer, DirectoryInfo baseDir, IReadOnlyList<FileInfo> files, ChunkingOptions? options = null, IProgress<ICopyStats>? progress = null, CancellationToken cancellationToken = default)
+		public static async Task<IBlobHandle<DirectoryNode>> WriteFilesAsync(this IStorageWriter writer, DirectoryInfo baseDir, IReadOnlyList<FileInfo> files, ChunkingOptions? options = null, IProgress<ICopyStats>? progress = null, CancellationToken cancellationToken = default)
 		{
 			DirectoryNode outputNode = new DirectoryNode();
-			await outputNode.AddFilesAsync(baseDir, files, writer, options, progress, cancellationToken);
-			return await writer.WriteHashedNodeAsync(outputNode, cancellationToken);
+			await outputNode.AddFilesAsync(baseDir, files, writer, options, progress, null, cancellationToken);
+			return await writer.WriteBlobAsync(outputNode, cancellationToken: cancellationToken);
 		}
 
 		/// <summary>
 		/// Writes a tree of files to a storage writer
 		/// </summary>
-		public static async Task<HashedNodeRef<DirectoryNode>> WriteFilesAsync(this IStorageWriter writer, DirectoryReference baseDir, IReadOnlyList<FileReference> files, ChunkingOptions? options = null, IProgress<ICopyStats>? progress = null, CancellationToken cancellationToken = default)
+		public static async Task<IBlobHandle<DirectoryNode>> WriteFilesAsync(this IStorageWriter writer, DirectoryReference baseDir, IReadOnlyList<FileReference> files, ChunkingOptions? options = null, IProgress<ICopyStats>? progress = null, CancellationToken cancellationToken = default)
 		{
 			DirectoryNode outputNode = new DirectoryNode();
-			await outputNode.AddFilesAsync(baseDir, files, writer, options, progress, cancellationToken);
-			return await writer.WriteHashedNodeAsync(outputNode, cancellationToken);
+			await outputNode.AddFilesAsync(baseDir, files, writer, options, progress, null, cancellationToken);
+			return await writer.WriteBlobAsync(outputNode, cancellationToken: cancellationToken);
 		}
 	}
 }
diff --git a/Engine/Source/Programs/Shared/EpicGames.Horde/Storage/Nodes/DirectoryNodeRef.cs b/Engine/Source/Programs/Shared/EpicGames.Horde/Storage/Nodes/DirectoryNodeRef.cs
--- a/Engine/Source/Programs/Shared/EpicGames.Horde/Storage/Nodes/DirectoryNodeRef.cs
+++ b/Engine/Source/Programs/Shared/EpicGames.Horde/Storage/Nodes/DirectoryNodeRef.cs
@@ -9,23 +9,7 @@ namespace EpicGames.Horde.Storage.Nodes
 	/// <summary>
 	/// Reference to a directory node, including the target hash and length
 	/// </summary>
-	/// <param name="Hash">Hash of the target node</param>
 	/// <param name="Length">Sum total of all the file lengths in this directory tree</param>
 	/// <param name="Handle">Handle to the target node</param>
-	public record class DirectoryNodeRef(IoHash Hash, long Length, IBlobHandle Handle)
-	{
-		/// <summary>
-		/// Constructor
-		/// </summary>
-		public DirectoryNodeRef(long length, HashedNodeRef<DirectoryNode> target)
-			: this(target.Hash, length, target.Handle)
-		{
-		}
-
-		/// <summary>
-		/// Get the target directory node
-		/// </summary>
-		public ValueTask<DirectoryNode> ExpandAsync(CancellationToken cancellationToken = default)
-			=> Handle.ReadNodeAsync<DirectoryNode>(cancellationToken);
-	}
+	public record class DirectoryNodeRef(long Length, IBlobHandle<DirectoryNode> Handle);
 }
diff --git a/Engine/Source/Programs/Shared/EpicGames.Horde/Storage/Nodes/RedirectNode.cs b/Engine/Source/Programs/Shared/EpicGames.Horde/Storage/Nodes/RedirectNode.cs
--- a/Engine/Source/Programs/Shared/EpicGames.Horde/Storage/Nodes/RedirectNode.cs
+++ b/Engine/Source/Programs/Shared/EpicGames.Horde/Storage/Nodes/RedirectNode.cs
@@ -5,43 +5,37 @@ namespace EpicGames.Horde.Storage.Nodes
 	/// <summary>
 	/// A node containing arbitrary compact binary data
 	/// </summary>
-	[BlobType("{BE09E54F-47CA-7A6B-2A97-AFBC183B1538}", 1)]
-	public class RedirectNode : Node 
+	[BlobConverter(typeof(RedirectNodeConverter<>))]
+	public class RedirectNode<T>
 	{
 		/// <summary>
 		/// The target handle
 		/// </summary>
-		public IBlobHandle Handle { get; }
-
-		/// <summary>
-		/// Constructor
-		/// </summary>
-		/// <param name="target">Target node for the redirect</param>
-		public RedirectNode(NodeRef target) => Handle = target.Handle;
+		public IBlobHandle<T> Target { get; }
 
 		/// <summary>
 		/// Constructor
 		/// </summary>
 		/// <param name="handle">Target node for the redirect</param>
-		public RedirectNode(IBlobHandle handle) => Handle = handle;
-
-		/// <summary>
-		/// Deserialization constructor
-		/// </summary>
-		/// <param name="reader">Reader to deserialize from</param>
-		public RedirectNode(IBlobReader reader) => Handle = reader.ReadBlobReference();
+		public RedirectNode(IBlobHandle<T> handle) => Target = handle;
+	}
 
-		/// <summary>
-		/// Gets a typed reference to the target node
-		/// </summary>
-		public NodeRef ToNodeRef() => new NodeRef(Handle);
+	class RedirectNodeConverter<T> : BlobConverter<RedirectNode<T>>
+	{
+		BlobType s_blobType = new BlobType("{BE09E54F-47CA-7A6B-2A97-AFBC183B1538}", 1);
 
-		/// <summary>
-		/// Gets a typed reference to the target node
-		/// </summary>
-		public NodeRef<TTarget> ToNodeRef<TTarget>() where TTarget : Node => new NodeRef<TTarget>(Handle);
+		/// <inheritdoc/>
+		public override RedirectNode<T> Read(IBlobReader reader, BlobSerializerOptions options)
+		{
+			IBlobHandle<T> handle = reader.ReadBlobHandle<T>();
+			return new RedirectNode<T>(handle);
+		}
 
 		/// <inheritdoc/>
-		public override void Serialize(IBlobWriter writer) => writer.WriteBlobReference(Handle);
+		public override BlobType Write(IBlobWriter writer, RedirectNode<T> value, BlobSerializerOptions options)
+		{
+			writer.WriteBlobHandle<T>(value.Target);
+			return s_blobType;
+		}
 	}
 }
diff --git a/Engine/Source/Programs/Shared/EpicGames.Horde/Storage/Workspace.cs b/Engine/Source/Programs/Shared/EpicGames.Horde/Storage/Workspace.cs
--- a/Engine/Source/Programs/Shared/EpicGames.Horde/Storage/Workspace.cs
+++ b/Engine/Source/Programs/Shared/EpicGames.Horde/Storage/Workspace.cs
@@ -650,7 +650,7 @@ namespace EpicGames.Horde.Storage
 					DirectoryReference subDirPath = DirectoryReference.Combine(dirRef, subDirEntry.Name.ToString());
 					DirectoryState subDirState = dirState.FindOrAddDirectory(subDirEntry.Name);
 
-					DirectoryNode subDirNode = await subDirEntry.ExpandAsync(cancellationToken);
+					DirectoryNode subDirNode = await subDirEntry.Handle.ReadBlobAsync(cancellationToken: cancellationToken);
 					await SyncDirectoryAsync(subDirPath, subDirState, subDirNode, flag, cacheDirState, cancellationToken);
 
 					dirState.LayerFlags |= flag;
@@ -723,14 +723,13 @@ namespace EpicGames.Horde.Storage
 
 		async Task ExtractDataAsync(ChunkedDataNodeRef nodeRef, Stream outputStream, CancellationToken cancellationToken)
 		{
-			BlobType blobType = await nodeRef.Handle.ReadTypeAsync(cancellationToken);
-			if (blobType.Guid == LeafChunkedDataNode.BlobType.Guid)
+			if (nodeRef.Type == ChunkedDataNodeType.Leaf)
 			{
 				await ExtractLeafDataAsync(nodeRef, outputStream, cancellationToken);
 			}
-			else if (blobType.Guid == InteriorChunkedDataNode.BlobType.Guid)
+			else if (nodeRef.Type == ChunkedDataNodeType.Interior)
 			{
-				InteriorChunkedDataNode interiorNode = await nodeRef.Handle.ReadNodeAsync<InteriorChunkedDataNode>(cancellationToken);
+				InteriorChunkedDataNode interiorNode = await nodeRef.Handle.ReadBlobAsync<InteriorChunkedDataNode>(cancellationToken: cancellationToken);
 				foreach (ChunkedDataNodeRef childNodeRef in interiorNode.Children)
 				{
 					await ExtractDataAsync(childNodeRef, outputStream, cancellationToken);
@@ -738,7 +737,7 @@ namespace EpicGames.Horde.Storage
 			}
 			else
 			{
-				throw new InvalidDataException($"Invalid node type {blobType}");
+				throw new InvalidDataException($"Invalid node type {nodeRef.Type}");
 			}
 		}
 
@@ -746,7 +745,7 @@ namespace EpicGames.Horde.Storage
 		{
 			// Try to copy cached data for this node
 			HashInfo? hashInfo;
-			if (_hashes.TryGetValue(nodeRef.Hash, out hashInfo))
+			if (_hashes.TryGetValue(nodeRef.Handle.Hash, out hashInfo))
 			{
 				if (await TryCopyCachedDataAsync(hashInfo, 0, hashInfo.Length, outputStream, cancellationToken))
 				{
@@ -755,7 +754,7 @@ namespace EpicGames.Horde.Storage
 			}
 
 			// Otherwise 
-			using BlobData blobData = await nodeRef.Handle.ReadAsync(cancellationToken);
+			using BlobData blobData = await nodeRef.Handle.ReadBlobDataAsync(cancellationToken);
 			await LeafChunkedDataNode.CopyToStreamAsync(blobData, outputStream, cancellationToken);
 		}
 
diff --git a/Engine/Source/Programs/UnrealBuildTool/Artifacts/HordeStorageArtifactCache.cs b/Engine/Source/Programs/UnrealBuildTool/Artifacts/HordeStorageArtifactCache.cs
--- a/Engine/Source/Programs/UnrealBuildTool/Artifacts/HordeStorageArtifactCache.cs
+++ b/Engine/Source/Programs/UnrealBuildTool/Artifacts/HordeStorageArtifactCache.cs
@@ -30,7 +30,7 @@ namespace UnrealBuildTool.Artifacts
 		/// Collection of output file references.  There should be exactly the same number
 		/// of file references as outputs in the action
 		/// </summary>
-		public readonly HashedNodeRef<ChunkedDataNode>[] OutputRefs;
+		public readonly IBlobHandle<ChunkedDataNode>[] OutputRefs;
 
 		/// <summary>
 		/// Construct a new horde artifact number
@@ -40,7 +40,7 @@ namespace UnrealBuildTool.Artifacts
 		public HordeArtifactAction(ArtifactAction artifactAction)
 		{
 			ArtifactAction = artifactAction;
-			OutputRefs = new HashedNodeRef<ChunkedDataNode>[ArtifactAction.Outputs.Length];
+			OutputRefs = new IBlobHandle<ChunkedDataNode>[ArtifactAction.Outputs.Length];
 		}
 
 		/// <summary>
@@ -50,7 +50,7 @@ namespace UnrealBuildTool.Artifacts
 		public HordeArtifactAction(IBlobReader reader)
 		{
 			ArtifactAction = reader.ReadArtifactAction();
-			OutputRefs = reader.ReadVariableLengthArray(() => new HashedNodeRef<ChunkedDataNode>(reader));
+			OutputRefs = reader.ReadVariableLengthArray(() => reader.ReadBlobHandle<ChunkedDataNode>());
 		}
 
 		/// <summary>
@@ -60,15 +60,9 @@ namespace UnrealBuildTool.Artifacts
 		public void Serialize(IBlobWriter writer)
 		{
 			writer.WriteArtifactAction(ArtifactAction);
-			writer.WriteVariableLengthArray(OutputRefs, x => x.Serialize(writer));
+			writer.WriteVariableLengthArray(OutputRefs, x => writer.WriteBlobHandle(x));
 		}
 
-		/// <summary>
-		/// Return an enumeration of the references
-		/// </summary>
-		/// <returns></returns>
-		public IEnumerable<NodeRef> EnumerateRefs() => OutputRefs;
-
 		/// <summary>
 		/// Write all the files to disk
 		/// </summary>
@@ -88,7 +82,7 @@ namespace UnrealBuildTool.Artifacts
 				string outputName = artifact.GetFullPath(ArtifactAction.DirectoryMapping);
 				using FileStream stream = new(outputName, FileMode.Open, FileAccess.Read, FileShare.Read);
 				ChunkedData chunkedData = await fileWriter.CreateAsync(stream, leafOptions.TargetSize, cancellationToken);
-				OutputRefs[index++] = new HashedNodeRef<ChunkedDataNode>(chunkedData.Root.Hash, chunkedData.Root.Handle);
+				OutputRefs[index++] = chunkedData.Root.Handle;
 			}
 		}
 	}
@@ -123,10 +117,9 @@ namespace UnrealBuildTool.Artifacts
 	/// <summary>
 	/// Horde node that represents a collection of action nodes 
 	/// </summary>
-	[BlobType("{E8DBCD77-4CAE-861D-0758-7FB733256ED2}")]
-	class ArtifactActionCollectionNode : Node
+	[BlobConverter(typeof(ArtifactActionCollectionNodeConverter))]
+	class ArtifactActionCollectionNode
 	{
-
 		/// <summary>
 		/// Collection of actions
 		/// </summary>
@@ -138,20 +131,23 @@ namespace UnrealBuildTool.Artifacts
 		public ArtifactActionCollectionNode()
 		{
 		}
+	}
 
-		/// <summary>
-		/// Construct a new artifact action collection from the source reader
-		/// </summary>
-		/// <param name="reader">Source reader</param>
-		public ArtifactActionCollectionNode(IBlobReader reader)
+	class ArtifactActionCollectionNodeConverter : BlobConverter<ArtifactActionCollectionNode>
+	{
+		static readonly BlobType s_blobType = new BlobType("{E8DBCD77-4CAE-861D-0758-7FB733256ED2}", 1);
+
+		public override ArtifactActionCollectionNode Read(IBlobReader reader, BlobSerializerOptions options)
 		{
-			ArtifactActions = reader.ReadDictionary<IoHash, HordeArtifactAction>(() => reader.ReadIoHash(), () => reader.ReadHordeArtifactAction());
+			ArtifactActionCollectionNode node = new ArtifactActionCollectionNode();
+			node.ArtifactActions = reader.ReadDictionary<IoHash, HordeArtifactAction>(() => reader.ReadIoHash(), () => reader.ReadHordeArtifactAction());
+			return node;
 		}
 
-		/// <inheritdoc/>
-		public override void Serialize(IBlobWriter writer)
+		public override BlobType Write(IBlobWriter writer, ArtifactActionCollectionNode value, BlobSerializerOptions options)
 		{
-			writer.WriteDictionary<IoHash, HordeArtifactAction>(ArtifactActions, (x) => writer.WriteIoHash(x), (x) => writer.WriteHordeArtifactAction(x));
+			writer.WriteDictionary<IoHash, HordeArtifactAction>(value.ArtifactActions, (x) => writer.WriteIoHash(x), (x) => writer.WriteHordeArtifactAction(x));
+			return s_blobType;
 		}
 	}
 
@@ -236,11 +232,6 @@ namespace UnrealBuildTool.Artifacts
 			return cache;
 		}
 
-		static HordeStorageArtifactCache()
-		{
-			Node.RegisterType<ArtifactActionCollectionNode>();
-		}
-
 		/// <summary>
 		/// Constructor
 		/// </summary>
@@ -277,7 +268,7 @@ namespace UnrealBuildTool.Artifacts
 					{
 						artifactActions.AddRange(_pendingWrites.Where(x => x.Key == key));
 					}
-					ArtifactActionCollectionNode? node = await _store.TryReadRefAsync<ArtifactActionCollectionNode>(GetRefName(key), default, cancellationToken);
+					ArtifactActionCollectionNode? node = await _store.TryReadRefAsync<ArtifactActionCollectionNode>(GetRefName(key), default, cancellationToken: cancellationToken);
 					if (node != null)
 					{
 						foreach (HordeArtifactAction artifactAction in node.ArtifactActions.Values)
@@ -309,7 +300,7 @@ namespace UnrealBuildTool.Artifacts
 			{
 				output[index] = false;
 				ArtifactAction artifactAction = artifactActions[index];
-				ArtifactActionCollectionNode? node = await _store.TryReadRefAsync<ArtifactActionCollectionNode>(GetRefName(artifactAction.Key), default, cancellationToken);
+				ArtifactActionCollectionNode? node = await _store.TryReadRefAsync<ArtifactActionCollectionNode>(GetRefName(artifactAction.Key), default, cancellationToken: cancellationToken);
 				if (node != null)
 				{
 					if (node.ArtifactActions.TryGetValue(artifactAction.ActionKey, out HordeArtifactAction hordeArtifactAction))
@@ -317,9 +308,9 @@ namespace UnrealBuildTool.Artifacts
 						output[index] = true;
 
 						int refIndex = 0;
-						foreach (HashedNodeRef<ChunkedDataNode> artifactRef in hordeArtifactAction.OutputRefs)
+						foreach (IBlobHandle<ChunkedDataNode> artifactRef in hordeArtifactAction.OutputRefs)
 						{
-							if (artifactRef.Handle == null)
+							if (artifactRef == null)
 							{
 								output[index] = false;
 								break;
@@ -328,7 +319,7 @@ namespace UnrealBuildTool.Artifacts
 							{
 								string outputName = hordeArtifactAction.ArtifactAction.Outputs[refIndex++].GetFullPath(artifactAction.DirectoryMapping);
 								using FileStream stream = new(outputName, FileMode.Create, FileAccess.Write, FileShare.ReadWrite);
-								await ChunkedDataNode.CopyToStreamAsync(artifactRef.Handle, stream, cancellationToken);
+								await ChunkedDataNode.CopyToStreamAsync(artifactRef, stream, cancellationToken);
 							}
 							catch (Exception)
 							{
@@ -476,7 +467,7 @@ namespace UnrealBuildTool.Artifacts
 					RefName refName = GetRefName(artifactAction.Key);
 
 					// Locate the destination collection for this key
-					ArtifactActionCollectionNode? node = _store!.TryReadRefAsync<ArtifactActionCollectionNode>(refName, default, cancellationToken).Result;
+					ArtifactActionCollectionNode? node = _store!.TryReadRefAsync<ArtifactActionCollectionNode>(refName, default, cancellationToken: cancellationToken).Result;
 					node ??= new ArtifactActionCollectionNode();
 
 					// Update the artifact action collection
@@ -486,11 +477,11 @@ namespace UnrealBuildTool.Artifacts
 					// Save the artifact action file
 					await using IStorageWriter writer = _store!.CreateWriter();
 					await hordeArtifactAction.WriteFilesAsync(writer, cancellationToken);
-					NodeRef<ArtifactActionCollectionNode> nodeRef = await writer.WriteNodeAsync(node);
+					IBlobHandle<ArtifactActionCollectionNode> nodeRef = await writer.WriteBlobAsync(node);
 					await writer.FlushAsync();
 
 					// Save the collection
-					await _store.WriteRefTargetAsync(refName, nodeRef.Handle, cancellationToken: cancellationToken);
+					await _store.WriteRefTargetAsync(refName, nodeRef, cancellationToken: cancellationToken);
 				}, cancellationToken));
 			}
 			return tasks;
diff --git a/Engine/Source/Programs/UnrealBuildTool/Executors/Experimental/UBAHordeSession.cs b/Engine/Source/Programs/UnrealBuildTool/Executors/Experimental/UBAHordeSession.cs
--- a/Engine/Source/Programs/UnrealBuildTool/Executors/Experimental/UBAHordeSession.cs
+++ b/Engine/Source/Programs/UnrealBuildTool/Executors/Experimental/UBAHordeSession.cs
@@ -304,7 +304,8 @@ namespace UnrealBuildTool
 			await using IStorageWriter writer = _storage.CreateWriter();
 			DirectoryNode sandbox = new();
 			await sandbox.AddFilesAsync(baseDir, files, writer, cancellationToken: cancellationToken);
-			IBlobHandle handle = await writer.FlushAsync(sandbox, cancellationToken);
+			IBlobHandle<DirectoryNode> handle = await writer.WriteBlobAsync(sandbox);
+			await writer.FlushAsync(cancellationToken);
 			return handle.GetLocator();
 		}
 
