Horde: Fix some static analysis warnings.

diff --git a/Engine/Source/Programs/Horde/Horde.Server/Storage/StorageController.cs b/Engine/Source/Programs/Horde/Horde.Server/Storage/StorageController.cs
--- a/Engine/Source/Programs/Horde/Horde.Server/Storage/StorageController.cs
+++ b/Engine/Source/Programs/Horde/Horde.Server/Storage/StorageController.cs
@@ -1,7 +1,6 @@
 // Copyright Epic Games, Inc. All Rights Reserved.
 
 using System;
-using System.Collections.Generic;
 using System.IO;
 using System.Linq;
 using System.Net.Http.Headers;
@@ -13,7 +12,6 @@ using Amazon.EC2.Model;
 using EpicGames.Core;
 using EpicGames.Horde.Storage;
 using EpicGames.Horde.Storage.Bundles;
-using EpicGames.Horde.Storage.Nodes;
 using Horde.Server.Acls;
 using Horde.Server.Server;
 using Horde.Server.Utilities;
@@ -424,9 +422,10 @@ namespace Horde.Server.Storage
 
 			return new { type = blobData.Type.Guid, @class = node.GetType().Name, content = content };
 		}
-		*/
+
 		static string GetNodeLink(string linkBase, IBlobHandle handle) => GetNodeLink(linkBase, handle.GetLocator());
 		
 		static string GetNodeLink(string linkBase, BlobLocator locator) => $"{linkBase}/nodes/{locator}";
+		*/
 	}
 }
diff --git a/Engine/Source/Programs/Horde/Horde.Server/Storage/StorageService.cs b/Engine/Source/Programs/Horde/Horde.Server/Storage/StorageService.cs
--- a/Engine/Source/Programs/Horde/Horde.Server/Storage/StorageService.cs
+++ b/Engine/Source/Programs/Horde/Horde.Server/Storage/StorageService.cs
@@ -68,9 +68,6 @@ namespace Horde.Server.Storage
 			/// <inheritdoc/>
 			public ValueTask FlushAsync(CancellationToken cancellationToken = default) => default;
 
-			/// <inheritdoc/>
-			public ValueTask<BlobType> ReadTypeAsync(CancellationToken cancellationToken = default) => new ValueTask<BlobType>(BlobType.Leaf);
-
 			/// <inheritdoc/>
 			public async Task<Stream> OpenBodyAsync(int offset = 0, int? length = null, CancellationToken cancellationToken = default)
 			{
@@ -87,9 +84,6 @@ namespace Horde.Server.Storage
 				return await _backend.OpenAsync(_path, offset, length, cancellationToken);
 			}
 
-			/// <inheritdoc/>
-			public ValueTask<IReadOnlyList<IBlobHandle>> ReadImportsAsync(CancellationToken cancellationToken = default) => new ValueTask<IReadOnlyList<IBlobHandle>>(Array.Empty<IBlobHandle>());
-
 			/// <inheritdoc/>
 			public async ValueTask<BlobData> ReadBlobDataAsync(CancellationToken cancellationToken = default)
 			{
diff --git a/Engine/Source/Programs/Horde/Horde/Commands/Bundles/BundleDump.cs b/Engine/Source/Programs/Horde/Horde/Commands/Bundles/BundleDump.cs
--- a/Engine/Source/Programs/Horde/Horde/Commands/Bundles/BundleDump.cs
+++ b/Engine/Source/Programs/Horde/Horde/Commands/Bundles/BundleDump.cs
@@ -2,7 +2,6 @@
 
 using System.ComponentModel;
 using EpicGames.Core;
-using EpicGames.Horde.Storage;
 using EpicGames.Horde.Storage.Bundles.V1;
 using Microsoft.Extensions.Logging;
 
diff --git a/Engine/Source/Programs/Shared/EpicGames.Horde.Tests/BundleStoreTests.cs b/Engine/Source/Programs/Shared/EpicGames.Horde.Tests/BundleStoreTests.cs
--- a/Engine/Source/Programs/Shared/EpicGames.Horde.Tests/BundleStoreTests.cs
+++ b/Engine/Source/Programs/Shared/EpicGames.Horde.Tests/BundleStoreTests.cs
@@ -15,7 +15,6 @@ using System.IO;
 using EpicGames.Core;
 using Microsoft.Extensions.Caching.Memory;
 using Microsoft.Extensions.Logging.Abstractions;
-using System.Reflection;
 using EpicGames.Horde.Storage.Bundles;
 using EpicGames.Horde.Storage.Bundles.V1;
 using EpicGames.Serialization;
diff --git a/Engine/Source/Programs/Shared/EpicGames.Horde.Tests/StorageClientTests.cs b/Engine/Source/Programs/Shared/EpicGames.Horde.Tests/StorageClientTests.cs
--- a/Engine/Source/Programs/Shared/EpicGames.Horde.Tests/StorageClientTests.cs
+++ b/Engine/Source/Programs/Shared/EpicGames.Horde.Tests/StorageClientTests.cs
@@ -1,8 +1,6 @@
 // Copyright Epic Games, Inc. All Rights Reserved.
 
 using System;
-using System.Collections.Generic;
-using System.Threading;
 using System.Threading.Tasks;
 using EpicGames.Core;
 using EpicGames.Horde.Storage;
diff --git a/Engine/Source/Programs/Shared/EpicGames.Horde/Compute/AgentMessage.cs b/Engine/Source/Programs/Shared/EpicGames.Horde/Compute/AgentMessage.cs
--- a/Engine/Source/Programs/Shared/EpicGames.Horde/Compute/AgentMessage.cs
+++ b/Engine/Source/Programs/Shared/EpicGames.Horde/Compute/AgentMessage.cs
@@ -9,7 +9,6 @@ using System.Threading;
 using System.Threading.Tasks;
 using EpicGames.Core;
 using EpicGames.Horde.Storage;
-using EpicGames.Horde.Storage.Nodes;
 
 namespace EpicGames.Horde.Compute
 {
diff --git a/Engine/Source/Programs/Shared/EpicGames.Horde/Storage/BlobConverter.cs b/Engine/Source/Programs/Shared/EpicGames.Horde/Storage/BlobConverter.cs
--- a/Engine/Source/Programs/Shared/EpicGames.Horde/Storage/BlobConverter.cs
+++ b/Engine/Source/Programs/Shared/EpicGames.Horde/Storage/BlobConverter.cs
@@ -1,8 +1,6 @@
 // Copyright Epic Games, Inc. All Rights Reserved.
 
 using System;
-using System.Threading;
-using System.Threading.Tasks;
 
 namespace EpicGames.Horde.Storage
 {
diff --git a/Engine/Source/Programs/Shared/EpicGames.Horde/Storage/BlobSerializer.cs b/Engine/Source/Programs/Shared/EpicGames.Horde/Storage/BlobSerializer.cs
--- a/Engine/Source/Programs/Shared/EpicGames.Horde/Storage/BlobSerializer.cs
+++ b/Engine/Source/Programs/Shared/EpicGames.Horde/Storage/BlobSerializer.cs
@@ -3,13 +3,9 @@
 using System;
 using System.Collections.Concurrent;
 using System.Collections.Generic;
-using System.Linq;
 using System.Reflection;
-using System.Text;
 using System.Threading;
 using System.Threading.Tasks;
-using EpicGames.Core;
-using EpicGames.Serialization;
 
 namespace EpicGames.Horde.Storage
 {
@@ -102,8 +98,8 @@ namespace EpicGames.Horde.Storage
 		public IList<BlobConverter> Converters { get; } = new List<BlobConverter>();
 
 		readonly ConcurrentDictionary<Type, BlobConverter> _cachedConverters = new ConcurrentDictionary<Type, BlobConverter>();
-		readonly static ConcurrentDictionary<Type, BlobConverter> s_cachedDefaultConverters = new ConcurrentDictionary<Type, BlobConverter>();
-		readonly static Func<Type, BlobConverter> s_createDefaultConverter = CreateDefaultConverter;
+		static readonly ConcurrentDictionary<Type, BlobConverter> s_cachedDefaultConverters = new ConcurrentDictionary<Type, BlobConverter>();
+		static readonly Func<Type, BlobConverter> s_createDefaultConverter = CreateDefaultConverter;
 
 		/// <summary>
 		/// Gets a converter for the given type
diff --git a/Engine/Source/Programs/Shared/EpicGames.Horde/Storage/Clients/BundleStorageClient.cs b/Engine/Source/Programs/Shared/EpicGames.Horde/Storage/Clients/BundleStorageClient.cs
--- a/Engine/Source/Programs/Shared/EpicGames.Horde/Storage/Clients/BundleStorageClient.cs
+++ b/Engine/Source/Programs/Shared/EpicGames.Horde/Storage/Clients/BundleStorageClient.cs
@@ -27,7 +27,6 @@ namespace EpicGames.Horde.Storage.Clients
 		{
 			readonly BundleStorageClient _storageClient;
 			readonly IBlobHandle _inner;
-			List<IBlobHandle>? _refs;
 
 			/// <inheritdoc/>
 			public IBlobHandle? Outer => _inner.Outer;
@@ -44,35 +43,6 @@ namespace EpicGames.Horde.Storage.Clients
 			/// <inheritdoc/>
 			public ValueTask FlushAsync(CancellationToken cancellationToken = default) => _inner.FlushAsync(cancellationToken);
 
-			/// <inheritdoc/>
-			public ValueTask<BlobType> ReadTypeAsync(CancellationToken cancellationToken = default) => new ValueTask<BlobType>(Bundle.BlobType);
-
-			/// <inheritdoc/>
-			public async ValueTask<IReadOnlyList<IBlobHandle>> ReadImportsAsync(CancellationToken cancellationToken = default)
-			{
-				if (_refs == null)
-				{
-					List<IBlobHandle> refs = new List<IBlobHandle>();
-
-					Bundles.V1.BundleHeader header = await _storageClient.ReadHeaderAsync(_inner.GetLocator(), cancellationToken);
-					if (header.Exports.Count > 0)
-					{
-						foreach (BlobLocator import in header.Imports)
-						{
-							refs.Add(_storageClient.CreateBlobHandle(new BlobLocator(import.Path)));
-						}
-					}
-					else
-					{
-						using BlobData blobData = await ReadBlobDataAsync(cancellationToken);
-						refs.AddRange(blobData.Refs);
-					}
-
-					_refs = refs;
-				}
-				return _refs;
-			}
-
 			/// <inheritdoc/>
 			public Task<Stream> OpenBodyAsync(int offset = 0, int? length = null, CancellationToken cancellationToken = default)
 				=> _inner.OpenBodyAsync(offset, length, cancellationToken);
diff --git a/Engine/Source/Programs/Shared/EpicGames.Horde/Storage/IBlobHandle.cs b/Engine/Source/Programs/Shared/EpicGames.Horde/Storage/IBlobHandle.cs
--- a/Engine/Source/Programs/Shared/EpicGames.Horde/Storage/IBlobHandle.cs
+++ b/Engine/Source/Programs/Shared/EpicGames.Horde/Storage/IBlobHandle.cs
@@ -1,7 +1,6 @@
 // Copyright Epic Games, Inc. All Rights Reserved.
 
 using System;
-using System.Collections.Generic;
 using System.Diagnostics.CodeAnalysis;
 using System.IO;
 using System.Threading;
diff --git a/Engine/Source/Programs/Shared/EpicGames.Horde/Storage/Nodes/CbNode.cs b/Engine/Source/Programs/Shared/EpicGames.Horde/Storage/Nodes/CbNode.cs
--- a/Engine/Source/Programs/Shared/EpicGames.Horde/Storage/Nodes/CbNode.cs
+++ b/Engine/Source/Programs/Shared/EpicGames.Horde/Storage/Nodes/CbNode.cs
@@ -73,12 +73,12 @@ namespace EpicGames.Horde.Storage.Nodes
 		/// <inheritdoc/>
 		public override CbNode Read(IBlobReader reader, BlobSerializerOptions options)
 		{
-			CbObject Object = new CbObject(reader.GetMemory().ToArray());
+			CbObject obj = new CbObject(reader.GetMemory().ToArray());
 
 			List<IBlobHandle<object>> references = new List<IBlobHandle<object>>();
-			Object.IterateAttachments(new HandleMapper(reader, references).IterateField);
+			obj.IterateAttachments(new HandleMapper(reader, references).IterateField);
 
-			return new CbNode(Object, references);
+			return new CbNode(obj, references);
 		}
 
 		/// <inheritdoc/>
diff --git a/Engine/Source/Programs/Shared/EpicGames.Horde/Storage/Nodes/ChunkedDataNode.cs b/Engine/Source/Programs/Shared/EpicGames.Horde/Storage/Nodes/ChunkedDataNode.cs
--- a/Engine/Source/Programs/Shared/EpicGames.Horde/Storage/Nodes/ChunkedDataNode.cs
+++ b/Engine/Source/Programs/Shared/EpicGames.Horde/Storage/Nodes/ChunkedDataNode.cs
@@ -391,8 +391,8 @@ namespace EpicGames.Horde.Storage.Nodes
 
 		public override LeafChunkedDataNode Read(IBlobReader reader, BlobSerializerOptions options)
 		{
-			byte[] Data = reader.GetMemory().ToArray();
-			return new LeafChunkedDataNode(Data);
+			byte[] data = reader.GetMemory().ToArray();
+			return new LeafChunkedDataNode(data);
 		}
 
 		public override BlobType Write(IBlobWriter writer, LeafChunkedDataNode value, BlobSerializerOptions options)
diff --git a/Engine/Source/Programs/Shared/EpicGames.Horde/Storage/Nodes/DirectoryEntry.cs b/Engine/Source/Programs/Shared/EpicGames.Horde/Storage/Nodes/DirectoryEntry.cs
--- a/Engine/Source/Programs/Shared/EpicGames.Horde/Storage/Nodes/DirectoryEntry.cs
+++ b/Engine/Source/Programs/Shared/EpicGames.Horde/Storage/Nodes/DirectoryEntry.cs
@@ -1,7 +1,6 @@
 // Copyright Epic Games, Inc. All Rights Reserved.
 
 using System.Diagnostics;
-using EpicGames.Core;
 
 namespace EpicGames.Horde.Storage.Nodes
 {
diff --git a/Engine/Source/Programs/Shared/EpicGames.Horde/Storage/Nodes/DirectoryNodeRef.cs b/Engine/Source/Programs/Shared/EpicGames.Horde/Storage/Nodes/DirectoryNodeRef.cs
--- a/Engine/Source/Programs/Shared/EpicGames.Horde/Storage/Nodes/DirectoryNodeRef.cs
+++ b/Engine/Source/Programs/Shared/EpicGames.Horde/Storage/Nodes/DirectoryNodeRef.cs
@@ -1,9 +1,5 @@
 // Copyright Epic Games, Inc. All Rights Reserved.
 
-using System.Threading;
-using System.Threading.Tasks;
-using EpicGames.Core;
-
 namespace EpicGames.Horde.Storage.Nodes
 {
 	/// <summary>
diff --git a/Engine/Source/Programs/Shared/EpicGames.Horde/Storage/Nodes/RedirectNode.cs b/Engine/Source/Programs/Shared/EpicGames.Horde/Storage/Nodes/RedirectNode.cs
--- a/Engine/Source/Programs/Shared/EpicGames.Horde/Storage/Nodes/RedirectNode.cs
+++ b/Engine/Source/Programs/Shared/EpicGames.Horde/Storage/Nodes/RedirectNode.cs
@@ -22,7 +22,7 @@ namespace EpicGames.Horde.Storage.Nodes
 
 	class RedirectNodeConverter<T> : BlobConverter<RedirectNode<T>>
 	{
-		BlobType s_blobType = new BlobType("{BE09E54F-47CA-7A6B-2A97-AFBC183B1538}", 1);
+		static readonly BlobType s_blobType = new BlobType("{BE09E54F-47CA-7A6B-2A97-AFBC183B1538}", 1);
 
 		/// <inheritdoc/>
 		public override RedirectNode<T> Read(IBlobReader reader, BlobSerializerOptions options)
