HordeTrayApp: Ignore processes that are descendants of HordeAgent when checking for critical processes

#rnx

diff --git a/Engine/Source/Programs/Horde/Horde.Agent.TrayApp/Program.cs b/Engine/Source/Programs/Horde/Horde.Agent.TrayApp/Program.cs
--- a/Engine/Source/Programs/Horde/Horde.Agent.TrayApp/Program.cs
+++ b/Engine/Source/Programs/Horde/Horde.Agent.TrayApp/Program.cs
@@ -512,12 +512,29 @@ namespace Horde.Agent.TrayApp
 
 			for (; ; )
 			{
-				_idleCriticalProcessCount = _settings.Idle.CriticalProcesses
-					.Select(x => Path.GetFileNameWithoutExtension(x).ToUpperInvariant())
-					.Distinct()
-					.SelectMany(x => Process.GetProcessesByName(x))
-					.Count();
-
+				try
+				{
+					if (RuntimeInformation.IsOSPlatform(OSPlatform.Windows) && _settings.Idle.CriticalProcesses.Any())
+					{
+						IEnumerable<int> hordeProcessIds = Process.GetProcessesByName("HordeAgent").Select(x => x.Id);
+						IEnumerable<Process> criticalProcesses = _settings.Idle.CriticalProcesses
+							.Select(x => Path.GetFileNameWithoutExtension(x).ToUpperInvariant())
+							.Distinct()
+							.SelectMany(x => Process.GetProcessesByName(x));
+
+						// Ignore processes that are descendants of HordeAgent
+						if (hordeProcessIds.Any())
+						{
+							criticalProcesses = criticalProcesses.Where(x => !ProcessUtils.GetAncestorProcesses(x).Select(x => x.Id).Intersect(hordeProcessIds).Any());
+						}
+
+						_idleCriticalProcessCount = criticalProcesses.Count();
+					}
+				}
+				catch (InvalidOperationException)
+				{
+					// If a process stops running Process.Id will throw an exception
+				}
 				await Task.Delay(sampleInterval, cancellationToken);
 			}
 		}
diff --git a/Engine/Source/Programs/Shared/EpicGames.Core/ProcessUtils.cs b/Engine/Source/Programs/Shared/EpicGames.Core/ProcessUtils.cs
--- a/Engine/Source/Programs/Shared/EpicGames.Core/ProcessUtils.cs
+++ b/Engine/Source/Programs/Shared/EpicGames.Core/ProcessUtils.cs
@@ -5,6 +5,7 @@ using System.Collections.Generic;
 using System.Diagnostics;
 using System.Linq;
 using System.Runtime.InteropServices;
+using System.Runtime.Versioning;
 using System.Text;
 using System.Threading;
 using Microsoft.Extensions.Logging;
@@ -40,6 +41,21 @@ namespace EpicGames.Core
 		[DllImport("kernel32.dll")]
 		static extern uint WaitForMultipleObjects(int nCount, IntPtr[] lpHandles, bool bWaitAll, uint dwMilliseconds);
 
+		[StructLayout(LayoutKind.Sequential)]
+		struct ProcessBasicInformation
+		{
+			// These members must match PROCESS_BASIC_INFORMATION
+			public IntPtr _exitStatus;
+			public IntPtr _pebBaseAddress;
+			public IntPtr _affinityMask;
+			public IntPtr _basePriority;
+			public IntPtr _uniqueProcessId;
+			public IntPtr _inheritedFromUniqueProcessId;
+		}
+
+		[DllImport("ntdll.dll")]
+		private static extern int NtQueryInformationProcess(IntPtr processHandle, int processInformationClass, ref ProcessBasicInformation processInformation, int processInformationLength, out int returnLength);
+
 		/// <summary>
 		/// Attempts to terminate all processes matching a predicate
 		/// </summary>
@@ -66,6 +82,7 @@ namespace EpicGames.Core
 		/// <param name="logger">Logging device</param>
 		/// <param name="cancellationToken">Cancellation token to abort the search</param>
 		/// <returns>True if the process succeeded</returns>
+		[SupportedOSPlatform("windows")]
 		static bool TerminateProcessesWin32(Predicate<FileReference> predicate, ILogger logger, CancellationToken cancellationToken)
 		{
 			Dictionary<int, int> processToCount = new Dictionary<int, int>();
@@ -238,5 +255,67 @@ namespace EpicGames.Core
 				}
 			}
 		}
+
+		/// <summary>
+		/// Gets the parent process of a specified process
+		/// </summary>
+		/// <param name="process">The process</param>
+		/// <returns>The parent process if running, otherwise null</returns>
+		public static Process? GetParentProcess(Process process)
+		{
+			if (RuntimeInformation.IsOSPlatform(OSPlatform.Windows))
+			{
+				return GetParentProcessWin32(process);
+			}
+			else
+			{
+				throw new PlatformNotSupportedException();
+			}
+		}
+
+		[SupportedOSPlatform("windows")]
+		static Process? GetParentProcessWin32(Process process)
+		{
+			try
+			{
+				ProcessBasicInformation pbi = new ProcessBasicInformation();
+				if (NtQueryInformationProcess(process.Handle, 0, ref pbi, Marshal.SizeOf(pbi), out _) == 0)
+				{
+					return Process.GetProcessById(pbi._inheritedFromUniqueProcessId.ToInt32());
+				}
+			}
+			catch (Exception)
+			{
+			}
+			return null;
+		}
+
+		/// <summary>
+		/// Gets all ancestor processes of a specified windows process
+		/// </summary>
+		/// <param name="process">The process</param>
+		/// <returns>IEnumerable containing the parent processes</returns>
+		public static IEnumerable<Process> GetAncestorProcesses(Process process)
+		{
+			if (RuntimeInformation.IsOSPlatform(OSPlatform.Windows))
+			{
+				return GetAncestorProcessesWin32(process);
+			}
+			else
+			{
+				throw new PlatformNotSupportedException();
+			}
+		}
+
+		[SupportedOSPlatform("windows")]
+		static IEnumerable<Process> GetAncestorProcessesWin32(Process process)
+		{
+			Process? parentProcess = GetParentProcess(process);
+			while (parentProcess != null)
+			{
+				yield return parentProcess;
+				parentProcess = GetParentProcess(parentProcess);
+			}
+		}
 	}
 }
