Horde: Update the version number for new bundle interior nodes, while leaving the compute system using the older version.

diff --git a/Engine/Source/Programs/Shared/EpicGames.Horde/Storage/Nodes/ChunkedDataNode.cs b/Engine/Source/Programs/Shared/EpicGames.Horde/Storage/Nodes/ChunkedDataNode.cs
--- a/Engine/Source/Programs/Shared/EpicGames.Horde/Storage/Nodes/ChunkedDataNode.cs
+++ b/Engine/Source/Programs/Shared/EpicGames.Horde/Storage/Nodes/ChunkedDataNode.cs
@@ -42,7 +42,7 @@ namespace EpicGames.Horde.Storage.Nodes
 			{
 				await LeafChunkedDataNode.CopyToStreamAsync(blobData, outputStream, cancellationToken);
 			}
-			else if (blobData.Type.Guid == InteriorChunkedDataNodeConverter.BlobType.Guid)
+			else if (blobData.Type.Guid == InteriorChunkedDataNodeConverter.BlobTypeGuid)
 			{
 				await InteriorChunkedDataNode.CopyToStreamAsync(blobData, outputStream, cancellationToken);
 			}
@@ -96,7 +96,7 @@ namespace EpicGames.Horde.Storage.Nodes
 			{
 				return options.GetConverter<LeafChunkedDataNode>().Read(reader, options);
 			}
-			else if (reader.Type.Guid == InteriorChunkedDataNodeConverter.BlobType.Guid)
+			else if (reader.Type.Guid == InteriorChunkedDataNodeConverter.BlobTypeGuid)
 			{
 				return options.GetConverter<InteriorChunkedDataNode>().Read(reader, options);
 			}
@@ -445,12 +445,13 @@ namespace EpicGames.Horde.Storage.Nodes
 		/// </summary>
 		/// <param name="leafChunkedData">List of leaf handles</param>
 		/// <param name="options">Options for splitting the tree</param>
+		/// <param name="serializerOptions">Options for serialization</param>
 		/// <param name="writer">Output writer for new interior nodes</param>
 		/// <param name="cancellationToken">Cancellation token for the operation</param>
 		/// <returns>Handle to the root node of the tree</returns>
-		public static async Task<ChunkedData> CreateTreeAsync(LeafChunkedData leafChunkedData, InteriorChunkedDataNodeOptions options, IStorageWriter writer, CancellationToken cancellationToken)
+		public static async Task<ChunkedData> CreateTreeAsync(LeafChunkedData leafChunkedData, InteriorChunkedDataNodeOptions options, IStorageWriter writer, BlobSerializerOptions? serializerOptions, CancellationToken cancellationToken)
 		{
-			ChunkedDataNodeRef rootRef = await CreateTreeAsync(leafChunkedData.LeafHandles, options, writer, cancellationToken);
+			ChunkedDataNodeRef rootRef = await CreateTreeAsync(leafChunkedData.LeafHandles, options, writer, serializerOptions, cancellationToken);
 			return new ChunkedData(leafChunkedData.Hash, rootRef);
 		}
 
@@ -458,11 +459,12 @@ namespace EpicGames.Horde.Storage.Nodes
 		/// Create a tree of nodes from the given list of handles, splitting nodes in each layer based on the hash of the last node.
 		/// </summary>
 		/// <param name="nodeRefs">List of leaf nodes</param>
-		/// <param name="options">Options for splitting the tree</param>
+		/// <param name="chunkingOptions">Options for splitting the tree</param>
 		/// <param name="writer">Output writer for new interior nodes</param>
+		/// <param name="serializerOptions">Options for serializing blobs</param>
 		/// <param name="cancellationToken">Cancellation token for the operation</param>
 		/// <returns>Handle to the root node of the tree</returns>
-		public static async Task<ChunkedDataNodeRef> CreateTreeAsync(List<ChunkedDataNodeRef> nodeRefs, InteriorChunkedDataNodeOptions options, IStorageWriter writer, CancellationToken cancellationToken)
+		public static async Task<ChunkedDataNodeRef> CreateTreeAsync(List<ChunkedDataNodeRef> nodeRefs, InteriorChunkedDataNodeOptions chunkingOptions, IStorageWriter writer, BlobSerializerOptions? serializerOptions, CancellationToken cancellationToken)
 		{
 			List<ChunkedDataNodeRef> handleBuffer = new List<ChunkedDataNodeRef>();
 
@@ -470,12 +472,12 @@ namespace EpicGames.Horde.Storage.Nodes
 			while (nodeRefs.Count > 1)
 			{
 				interiorNodes.Clear();
-				CreateTreeLayer(nodeRefs, options, interiorNodes);
+				CreateTreeLayer(nodeRefs, chunkingOptions, interiorNodes);
 
 				handleBuffer.Clear();
 				foreach ((InteriorChunkedDataNode interiorNode, long interiorLength) in interiorNodes)
 				{
-					IBlobHandle<InteriorChunkedDataNode> interiorHandle = await writer.WriteBlobAsync(interiorNode, null, cancellationToken);
+					IBlobHandle<InteriorChunkedDataNode> interiorHandle = await writer.WriteBlobAsync(interiorNode, serializerOptions, cancellationToken);
 					handleBuffer.Add(new ChunkedDataNodeRef(ChunkedDataNodeType.Interior, interiorLength, interiorHandle));
 				}
 
@@ -555,13 +557,28 @@ namespace EpicGames.Horde.Storage.Nodes
 		}
 	}
 
-	class InteriorChunkedDataNodeConverter : BlobConverter<InteriorChunkedDataNode>
+	/// <summary>
+	/// Converter for interior node types
+	/// </summary>
+	public class InteriorChunkedDataNodeConverter : BlobConverter<InteriorChunkedDataNode>
 	{
 		/// <summary>
-		/// Static accessor for the blob type
+		/// Static accessor for the blob type guid
 		/// </summary>
-		public static BlobType BlobType { get; } = new BlobType("{F4DEDDBC-4C7A-70CB-11F0-4783B9CDCCAF}", 2); // Pending V3
+		public static Guid BlobTypeGuid { get; } = Guid.Parse("{F4DEDDBC-4C7A-70CB-11F0-4783B9CDCCAF}");
+
+		readonly int _writeVersion;
 
+		/// <summary>
+		/// Constructor
+		/// </summary>
+		/// <param name="writeVersion">Version number for serialized data</param>
+		public InteriorChunkedDataNodeConverter(int writeVersion = 3)
+		{
+			_writeVersion = writeVersion;
+		}
+
+		/// <inheritdoc/>
 		public override InteriorChunkedDataNode Read(IBlobReader reader, BlobSerializerOptions options)
 		{
 			// Keep this code in sync with CopyToStreamAsync
@@ -587,15 +604,22 @@ namespace EpicGames.Horde.Storage.Nodes
 			return new InteriorChunkedDataNode(children);
 		}
 
+		/// <inheritdoc/>
 		public override BlobType Write(IBlobWriter writer, InteriorChunkedDataNode value, BlobSerializerOptions options)
 		{
 			foreach (ChunkedDataNodeRef child in value.Children)
 			{
 				writer.WriteBlobHandle(child.Handle);
-				writer.WriteUnsignedVarInt((int)child.Type);
-				// Pending V3: writer.WriteUnsignedVarInt((ulong)child.Length);
+				if (_writeVersion >= 2)
+				{
+					writer.WriteUnsignedVarInt((int)child.Type);
+				}
+				if (_writeVersion >= 3)
+				{
+					writer.WriteUnsignedVarInt((ulong)child.Length);
+				}
 			}
-			return InteriorChunkedDataNodeConverter.BlobType;
+			return new BlobType(BlobTypeGuid, _writeVersion);
 		}
 	}
 }
diff --git a/Engine/Source/Programs/Shared/EpicGames.Horde/Storage/Nodes/ChunkedDataWriter.cs b/Engine/Source/Programs/Shared/EpicGames.Horde/Storage/Nodes/ChunkedDataWriter.cs
--- a/Engine/Source/Programs/Shared/EpicGames.Horde/Storage/Nodes/ChunkedDataWriter.cs
+++ b/Engine/Source/Programs/Shared/EpicGames.Horde/Storage/Nodes/ChunkedDataWriter.cs
@@ -71,7 +71,8 @@ namespace EpicGames.Horde.Storage.Nodes
 		public const int DefaultBufferLength = 32 * 1024;
 
 		readonly IStorageWriter _writer;
-		readonly ChunkingOptions _options;
+		readonly ChunkingOptions _chunkingOptions;
+		readonly BlobSerializerOptions? _serializerOptions;
 		readonly Blake3.Hasher _hasher;
 
 		// Tree state
@@ -91,11 +92,13 @@ namespace EpicGames.Horde.Storage.Nodes
 		/// Constructor
 		/// </summary>
 		/// <param name="writer">Writer for new nodes</param>
-		/// <param name="options">Chunking options</param>
-		public ChunkedDataWriter(IStorageWriter writer, ChunkingOptions options)
+		/// <param name="chunkingOptions">Chunking options</param>
+		/// <param name="serializerOptions">Options for serialization</param>
+		public ChunkedDataWriter(IStorageWriter writer, ChunkingOptions chunkingOptions, BlobSerializerOptions? serializerOptions)
 		{
 			_writer = writer;
-			_options = options;
+			_chunkingOptions = chunkingOptions;
+			_serializerOptions = serializerOptions;
 			_hasher = Blake3.Hasher.New();
 		}
 
@@ -215,7 +218,7 @@ namespace EpicGames.Horde.Storage.Nodes
 			for (; ; )
 			{
 				Memory<byte> buffer = _writer.GetOutputBuffer(_leafLength, _leafLength);
-				int appendLength = AppendToLeafNode(buffer.Span.Slice(0, _leafLength), data.Span, ref _leafHash, _options.LeafOptions);
+				int appendLength = AppendToLeafNode(buffer.Span.Slice(0, _leafLength), data.Span, ref _leafHash, _chunkingOptions.LeafOptions);
 
 				buffer = _writer.GetOutputBuffer(_leafLength, _leafLength + appendLength);
 				data.Slice(0, appendLength).CopyTo(buffer.Slice(_leafLength));
@@ -320,7 +323,7 @@ namespace EpicGames.Horde.Storage.Nodes
 		public async Task<ChunkedData> CompleteAsync(CancellationToken cancellationToken)
 		{
 			await FlushLeafNodeAsync(cancellationToken);
-			ChunkedDataNodeRef rootHandle = await InteriorChunkedDataNode.CreateTreeAsync(_leafHandles, _options.InteriorOptions, _writer, cancellationToken);
+			ChunkedDataNodeRef rootHandle = await InteriorChunkedDataNode.CreateTreeAsync(_leafHandles, _chunkingOptions.InteriorOptions, _writer, _serializerOptions, cancellationToken);
 			return new ChunkedData(IoHash.FromBlake3(_hasher), rootHandle);
 		}
 
diff --git a/Engine/Source/Programs/Shared/EpicGames.Horde/Storage/Nodes/DirectoryNode.cs b/Engine/Source/Programs/Shared/EpicGames.Horde/Storage/Nodes/DirectoryNode.cs
--- a/Engine/Source/Programs/Shared/EpicGames.Horde/Storage/Nodes/DirectoryNode.cs
+++ b/Engine/Source/Programs/Shared/EpicGames.Horde/Storage/Nodes/DirectoryNode.cs
@@ -522,7 +522,7 @@ namespace EpicGames.Horde.Storage.Nodes
 					await Parallel.ForEachAsync(partitions, cancellationToken, (filePartition, ctx) => CreateLeafChunkNodesAsync(writer, batch, leafChunkedFiles, filePartition.Start, filePartition.Count, copyStats, options, cancellationToken));
 
 					// Create interior nodes for all the leaf chunks
-					ChunkedData[] chunkedFiles = await CreateInteriorChunkNodesAsync(leafChunkedFiles, options.InteriorOptions, writer, cancellationToken);
+					ChunkedData[] chunkedFiles = await CreateInteriorChunkNodesAsync(leafChunkedFiles, options.InteriorOptions, writer, serializerOptions, cancellationToken);
 
 					// Write all the interior nodes and generate the directory update
 					for (int idx = 0; idx < batch.Count; idx++)
@@ -607,12 +607,12 @@ namespace EpicGames.Horde.Storage.Nodes
 			await writerFork.FlushAsync(cancellationToken);
 		}
 
-		static async Task<ChunkedData[]> CreateInteriorChunkNodesAsync(LeafChunkedData[] leafChunkedFiles, InteriorChunkedDataNodeOptions options, IStorageWriter writer, CancellationToken cancellationToken)
+		static async Task<ChunkedData[]> CreateInteriorChunkNodesAsync(LeafChunkedData[] leafChunkedFiles, InteriorChunkedDataNodeOptions chunkingOptions, IStorageWriter writer, BlobSerializerOptions? serializerOptions, CancellationToken cancellationToken)
 		{
 			ChunkedData[] chunkedFiles = new ChunkedData[leafChunkedFiles.Length];
 			for (int idx = 0; idx < leafChunkedFiles.Length; idx++)
 			{
-				chunkedFiles[idx] = await InteriorChunkedDataNode.CreateTreeAsync(leafChunkedFiles[idx], options, writer, cancellationToken);
+				chunkedFiles[idx] = await InteriorChunkedDataNode.CreateTreeAsync(leafChunkedFiles[idx], chunkingOptions, writer, serializerOptions, cancellationToken);
 			}
 			return chunkedFiles;
 		}
@@ -705,7 +705,7 @@ namespace EpicGames.Horde.Storage.Nodes
 					flags |= FileEntryFlags.Executable;
 				}
 
-				ChunkedData chunkedFile = await InteriorChunkedDataNode.CreateTreeAsync(leafChunkedFile, options.InteriorOptions, writer, cancellationToken);
+				ChunkedData chunkedFile = await InteriorChunkedDataNode.CreateTreeAsync(leafChunkedFile, options.InteriorOptions, writer, serializerOptions, cancellationToken);
 				updates.Add(new FileUpdate(entry.FullName, flags, entry.Length, chunkedFile));
 			}
 
diff --git a/Engine/Source/Programs/UnrealBuildTool/Artifacts/HordeStorageArtifactCache.cs b/Engine/Source/Programs/UnrealBuildTool/Artifacts/HordeStorageArtifactCache.cs
--- a/Engine/Source/Programs/UnrealBuildTool/Artifacts/HordeStorageArtifactCache.cs
+++ b/Engine/Source/Programs/UnrealBuildTool/Artifacts/HordeStorageArtifactCache.cs
@@ -75,7 +75,7 @@ namespace UnrealBuildTool.Artifacts
 			InteriorChunkedDataNodeOptions interiorOptions = new(1, 10, 20);
 			ChunkingOptions options = new() { LeafOptions = leafOptions, InteriorOptions = interiorOptions };
 
-			using ChunkedDataWriter fileWriter = new(writer, options);
+			using ChunkedDataWriter fileWriter = new(writer, options, null);
 			int index = 0;
 			foreach (ArtifactFile artifact in ArtifactAction.Outputs)
 			{
diff --git a/Engine/Source/Programs/UnrealBuildTool/Executors/Experimental/UBAHordeSession.cs b/Engine/Source/Programs/UnrealBuildTool/Executors/Experimental/UBAHordeSession.cs
--- a/Engine/Source/Programs/UnrealBuildTool/Executors/Experimental/UBAHordeSession.cs
+++ b/Engine/Source/Programs/UnrealBuildTool/Executors/Experimental/UBAHordeSession.cs
@@ -301,9 +301,12 @@ namespace UnrealBuildTool
 
 		async Task<BlobLocator> CreateToolAsync(DirectoryReference baseDir, IEnumerable<FileReference> files, CancellationToken cancellationToken)
 		{
+			BlobSerializerOptions serializerOptions = new BlobSerializerOptions();
+			serializerOptions.Converters.Add(new InteriorChunkedDataNodeConverter(2)); // Lock to v2 for now. Could change based on protocol version.
+
 			await using IStorageWriter writer = _storage.CreateWriter();
 			DirectoryNode sandbox = new();
-			await sandbox.AddFilesAsync(baseDir, files, writer, cancellationToken: cancellationToken);
+			await sandbox.AddFilesAsync(baseDir, files, writer, serializerOptions: serializerOptions, cancellationToken: cancellationToken);
 			IBlobHandle<DirectoryNode> handle = await writer.WriteBlobAsync(sandbox);
 			await writer.FlushAsync(cancellationToken);
 			return handle.GetLocator();
