Horde: Fix compile errors.

diff --git a/Engine/Source/Programs/Horde/Horde.Agent/Utility/TempStorage.cs b/Engine/Source/Programs/Horde/Horde.Agent/Utility/TempStorage.cs
--- a/Engine/Source/Programs/Horde/Horde.Agent/Utility/TempStorage.cs
+++ b/Engine/Source/Programs/Horde/Horde.Agent/Utility/TempStorage.cs
@@ -647,7 +647,7 @@ namespace Horde.Storage.Utility
 			FileReference localFileListLocation = GetTagManifestLocation(manifestDir, nodeName, tagName);
 			fileList.Save(localFileListLocation);
 
-			using ChunkedDataWriter fileNodeWriter = new ChunkedDataWriter(writer, new ChunkingOptions());
+			using ChunkedDataWriter fileNodeWriter = new ChunkedDataWriter(writer, new ChunkingOptions(), BlobSerializerOptions.Default);
 			ChunkedData fileNodeData = await fileNodeWriter.CreateAsync(localFileListLocation.ToFileInfo(), cancellationToken);
 
 			return new FileEntry(localFileListLocation.GetFileName(), FileEntryFlags.None, fileNodeWriter.Length, fileNodeData);
diff --git a/Engine/Source/Programs/Horde/Horde.Server/Perforce/PerforceReplicator.cs b/Engine/Source/Programs/Horde/Horde.Server/Perforce/PerforceReplicator.cs
--- a/Engine/Source/Programs/Horde/Horde.Server/Perforce/PerforceReplicator.cs
+++ b/Engine/Source/Programs/Horde/Horde.Server/Perforce/PerforceReplicator.cs
@@ -130,7 +130,7 @@ namespace Horde.Server.Perforce
 				public Handle(IStorageWriter writer, ChunkingOptions options)
 				{
 					Hash = IncrementalHash.CreateHash(HashAlgorithmName.MD5);
-					FileWriter = new ChunkedDataWriter(writer, options);
+					FileWriter = new ChunkedDataWriter(writer, options, BlobSerializerOptions.Default);
 				}
 
 				public void Dispose()
diff --git a/Engine/Source/Programs/Shared/EpicGames.Horde.Tests/BundleStoreTests.cs b/Engine/Source/Programs/Shared/EpicGames.Horde.Tests/BundleStoreTests.cs
--- a/Engine/Source/Programs/Shared/EpicGames.Horde.Tests/BundleStoreTests.cs
+++ b/Engine/Source/Programs/Shared/EpicGames.Horde.Tests/BundleStoreTests.cs
@@ -293,7 +293,7 @@ namespace EpicGames.Horde.Tests
 			{
 				await using IStorageWriter writer = store.CreateWriter();
 
-				using ChunkedDataWriter fileWriter = new ChunkedDataWriter(writer, new ChunkingOptions());
+				using ChunkedDataWriter fileWriter = new ChunkedDataWriter(writer, new ChunkingOptions(), BlobSerializerOptions.Default);
 				ChunkedData fileHandle = await fileWriter.CreateAsync(Encoding.UTF8.GetBytes("world"), CancellationToken.None);
 
 				List<FileUpdate> fileUpdates = new List<FileUpdate>();
@@ -350,7 +350,7 @@ namespace EpicGames.Horde.Tests
 				ChunkingOptions options = new ChunkingOptions();
 				options.LeafOptions = new LeafChunkedDataNodeOptions(128, 256, 64 * 1024);
 
-				using ChunkedDataWriter fileWriter = new ChunkedDataWriter(writer, options);
+				using ChunkedDataWriter fileWriter = new ChunkedDataWriter(writer, options, BlobSerializerOptions.Default);
 				for (int idx = 0; idx < chunk.Length / 16; idx++)
 				{
 					await fileWriter.AppendAsync(chunk.AsMemory(idx * 16, 16), CancellationToken.None);
@@ -398,7 +398,7 @@ namespace EpicGames.Horde.Tests
 				ChunkingOptions options = new ChunkingOptions();
 				options.LeafOptions = new LeafChunkedDataNodeOptions(128, 256, 64 * 1024);
 
-				using ChunkedDataWriter fileWriter = new ChunkedDataWriter(writer, options);
+				using ChunkedDataWriter fileWriter = new ChunkedDataWriter(writer, options, BlobSerializerOptions.Default);
 				ChunkedData chunkedData = await fileWriter.CreateAsync(data, CancellationToken.None);
 
 				root = new DirectoryNode(DirectoryFlags.None);
diff --git a/Engine/Source/Programs/Shared/EpicGames.Horde.Tests/BundleTests.cs b/Engine/Source/Programs/Shared/EpicGames.Horde.Tests/BundleTests.cs
--- a/Engine/Source/Programs/Shared/EpicGames.Horde.Tests/BundleTests.cs
+++ b/Engine/Source/Programs/Shared/EpicGames.Horde.Tests/BundleTests.cs
@@ -77,7 +77,7 @@ namespace EpicGames.Horde.Tests
 			ChunkingOptions options = new ChunkingOptions();
 			options.LeafOptions = new LeafChunkedDataNodeOptions(8, 8, 8);
 
-			using ChunkedDataWriter fileNodeWriter = new ChunkedDataWriter(writer, options);
+			using ChunkedDataWriter fileNodeWriter = new ChunkedDataWriter(writer, options, BlobSerializerOptions.Default);
 
 			ChunkedDataNode node;
 			ChunkedDataNodeRef nodeRef;
diff --git a/Engine/Source/Programs/Shared/EpicGames.Horde.Tests/ChunkingTests.cs b/Engine/Source/Programs/Shared/EpicGames.Horde.Tests/ChunkingTests.cs
--- a/Engine/Source/Programs/Shared/EpicGames.Horde.Tests/ChunkingTests.cs
+++ b/Engine/Source/Programs/Shared/EpicGames.Horde.Tests/ChunkingTests.cs
@@ -60,7 +60,7 @@ namespace EpicGames.Horde.Tests
 
 				using MemoryStream emptyStream = new MemoryStream();
 				LeafChunkedData leafChunkedData = await LeafChunkedDataNode.CreateFromStreamAsync(writer, emptyStream, new LeafChunkedDataNodeOptions(64, 64, 64), CancellationToken.None);
-				ChunkedData chunkedData = await InteriorChunkedDataNode.CreateTreeAsync(leafChunkedData, new InteriorChunkedDataNodeOptions(4, 4, 4), writer, CancellationToken.None); 
+				ChunkedData chunkedData = await InteriorChunkedDataNode.CreateTreeAsync(leafChunkedData, new InteriorChunkedDataNodeOptions(4, 4, 4), writer, BlobSerializerOptions.Default, CancellationToken.None); 
 
 				DirectoryNode directory = new DirectoryNode();
 				directory.AddFile("test.foo", FileEntryFlags.None, 0, chunkedData);
@@ -110,7 +110,7 @@ namespace EpicGames.Horde.Tests
 
 			const int NumIterations = 100;
 			{
-				using ChunkedDataWriter fileWriter = new ChunkedDataWriter(writer, options);
+				using ChunkedDataWriter fileWriter = new ChunkedDataWriter(writer, options, BlobSerializerOptions.Default);
 
 				for (int idx = 0; idx < NumIterations; idx++)
 				{
diff --git a/Engine/Source/Programs/Shared/EpicGames.Horde.Tests/ComputeSocketTests.cs b/Engine/Source/Programs/Shared/EpicGames.Horde.Tests/ComputeSocketTests.cs
--- a/Engine/Source/Programs/Shared/EpicGames.Horde.Tests/ComputeSocketTests.cs
+++ b/Engine/Source/Programs/Shared/EpicGames.Horde.Tests/ComputeSocketTests.cs
@@ -186,7 +186,7 @@ namespace EpicGames.Horde.Tests
 
 						byte[] data = Encoding.UTF8.GetBytes("Hello world");
 
-						using ChunkedDataWriter writer = new ChunkedDataWriter(treeWriter, new ChunkingOptions());
+						using ChunkedDataWriter writer = new ChunkedDataWriter(treeWriter, new ChunkingOptions(), BlobSerializerOptions.Default);
 						ChunkedData chunkedData = await writer.CreateAsync(data, cancellationToken);
 
 						DirectoryNode directory = new DirectoryNode();
