Horde: Add a parameterless constructor to InteriorChunkedDataNodeConverter to fix tests.

diff --git a/Engine/Source/Programs/Shared/EpicGames.Horde/Storage/Nodes/ChunkedDataNode.cs b/Engine/Source/Programs/Shared/EpicGames.Horde/Storage/Nodes/ChunkedDataNode.cs
--- a/Engine/Source/Programs/Shared/EpicGames.Horde/Storage/Nodes/ChunkedDataNode.cs
+++ b/Engine/Source/Programs/Shared/EpicGames.Horde/Storage/Nodes/ChunkedDataNode.cs
@@ -569,11 +569,18 @@ namespace EpicGames.Horde.Storage.Nodes
 
 		readonly int _writeVersion;
 
+		/// <summary>
+		/// Constructor
+		/// </summary>
+		public InteriorChunkedDataNodeConverter() : this(3)
+		{
+		}
+
 		/// <summary>
 		/// Constructor
 		/// </summary>
 		/// <param name="writeVersion">Version number for serialized data</param>
-		public InteriorChunkedDataNodeConverter(int writeVersion = 3)
+		public InteriorChunkedDataNodeConverter(int writeVersion)
 		{
 			_writeVersion = writeVersion;
 		}
