Fix static analysis warnings.

diff --git a/Engine/Source/Programs/Shared/EpicGames.Horde/Replicators/ReplicatorMessages.cs b/Engine/Source/Programs/Shared/EpicGames.Horde/Replicators/ReplicatorMessages.cs
--- a/Engine/Source/Programs/Shared/EpicGames.Horde/Replicators/ReplicatorMessages.cs
+++ b/Engine/Source/Programs/Shared/EpicGames.Horde/Replicators/ReplicatorMessages.cs
@@ -1,8 +1,6 @@
 ﻿// Copyright Epic Games, Inc. All Rights Reserved.
 
 using System;
-using System.Text;
-using System.Threading.Tasks;
 
 namespace EpicGames.Horde.Replicators
 {
diff --git a/Engine/Source/Programs/Shared/EpicGames.OIDC/WindowsTokenStore.cs b/Engine/Source/Programs/Shared/EpicGames.OIDC/WindowsTokenStore.cs
--- a/Engine/Source/Programs/Shared/EpicGames.OIDC/WindowsTokenStore.cs
+++ b/Engine/Source/Programs/Shared/EpicGames.OIDC/WindowsTokenStore.cs
@@ -170,6 +170,7 @@ namespace EpicGames.OIDC
 			if (disposing)
 			{
 				SaveStoreToDisk();
+				_mutex.Dispose();
 			}
 		}
 
