Horde: Fix formatting of replicator ids in HTTP responses.

diff --git a/Engine/Source/Programs/Horde/Horde.Server/Replicators/ReplicatorsController.cs b/Engine/Source/Programs/Horde/Horde.Server/Replicators/ReplicatorsController.cs
--- a/Engine/Source/Programs/Horde/Horde.Server/Replicators/ReplicatorsController.cs
+++ b/Engine/Source/Programs/Horde/Horde.Server/Replicators/ReplicatorsController.cs
@@ -146,6 +146,8 @@ namespace Horde.Server.Replicators
 		{
 			GetReplicatorResponse response = new GetReplicatorResponse();
 			response.Id = replicatorId;
+			response.StreamId = replicatorId.StreamId;
+			response.StreamReplicatorId = replicatorId.StreamReplicatorId;
 
 			if (replicator != null)
 			{
diff --git a/Engine/Source/Programs/Shared/EpicGames.Horde/Replicators/ReplicatorId.cs b/Engine/Source/Programs/Shared/EpicGames.Horde/Replicators/ReplicatorId.cs
--- a/Engine/Source/Programs/Shared/EpicGames.Horde/Replicators/ReplicatorId.cs
+++ b/Engine/Source/Programs/Shared/EpicGames.Horde/Replicators/ReplicatorId.cs
@@ -3,6 +3,8 @@
 using System;
 using System.ComponentModel;
 using System.Globalization;
+using System.Text.Json;
+using System.Text.Json.Serialization;
 using EpicGames.Core;
 using EpicGames.Horde.Streams;
 using EpicGames.Serialization;
@@ -15,8 +17,22 @@ namespace EpicGames.Horde.Replicators
 	[LogValueType]
 	[JsonSchemaString]
 	[TypeConverter(typeof(ReplicatorIdTypeConverter))]
+	[JsonConverter(typeof(ReplicatorIdJsonConverter))]
 	public record struct ReplicatorId(StreamId StreamId, StreamReplicatorId StreamReplicatorId)
 	{
+		/// <summary>
+		/// Parse a replicator id
+		/// </summary>
+		public static ReplicatorId Parse(string text)
+		{
+			ReplicatorId replicatorId;
+			if (!TryParse(text, out replicatorId))
+			{
+				throw new FormatException($"Unable to parse '{text}' as a replicator id");
+			}
+			return replicatorId;
+		}
+
 		/// <summary>
 		/// Parse a replicator id
 		/// </summary>
@@ -39,10 +55,18 @@ namespace EpicGames.Horde.Replicators
 		public override string ToString() => $"{StreamId}:{StreamReplicatorId}";
 	}
 
-	/// <summary>
-	/// Class which serializes types with a <see cref="StringIdConverter{T}"/> to Json
-	/// </summary>
-	public sealed class ReplicatorIdTypeConverter : TypeConverter
+	sealed class ReplicatorIdJsonConverter : JsonConverter<ReplicatorId>
+	{
+		/// <inheritdoc/>
+		public override ReplicatorId Read(ref Utf8JsonReader reader, Type typeToConvert, JsonSerializerOptions options)
+			=> ReplicatorId.Parse(reader.GetString() ?? String.Empty);
+
+		/// <inheritdoc/>
+		public override void Write(Utf8JsonWriter writer, ReplicatorId value, JsonSerializerOptions options)
+			=> writer.WriteStringValue(value.ToString());
+	}
+
+	sealed class ReplicatorIdTypeConverter : TypeConverter
 	{
 		/// <inheritdoc/>
 		public override bool CanConvertFrom(ITypeDescriptorContext? context, Type sourceType) => sourceType == typeof(string);
diff --git a/Engine/Source/Programs/Shared/EpicGames.Horde/Replicators/ReplicatorMessages.cs b/Engine/Source/Programs/Shared/EpicGames.Horde/Replicators/ReplicatorMessages.cs
--- a/Engine/Source/Programs/Shared/EpicGames.Horde/Replicators/ReplicatorMessages.cs
+++ b/Engine/Source/Programs/Shared/EpicGames.Horde/Replicators/ReplicatorMessages.cs
@@ -1,6 +1,7 @@
 ﻿// Copyright Epic Games, Inc. All Rights Reserved.
 
 using System;
+using EpicGames.Horde.Streams;
 
 namespace EpicGames.Horde.Replicators
 {
@@ -14,6 +15,16 @@ namespace EpicGames.Horde.Replicators
 		/// </summary>
 		public ReplicatorId Id { get; set; }
 
+		/// <summary>
+		/// Identifier for the stream
+		/// </summary>
+		public StreamId StreamId { get; set; }
+
+		/// <summary>
+		/// Identifier for this replicator within the stream
+		/// </summary>
+		public StreamReplicatorId StreamReplicatorId { get; set; }
+
 		/// <summary>
 		/// Whether to pause replication
 		/// </summary>
