Horde: Move BundleStorageClient into the EpicGames.Horde.Bundles namespace.

diff --git a/Engine/Source/Programs/Horde/Horde.Server/Commands/Generate/ToolDataCommand.cs b/Engine/Source/Programs/Horde/Horde.Server/Commands/Generate/ToolDataCommand.cs
--- a/Engine/Source/Programs/Horde/Horde.Server/Commands/Generate/ToolDataCommand.cs
+++ b/Engine/Source/Programs/Horde/Horde.Server/Commands/Generate/ToolDataCommand.cs
@@ -6,6 +6,7 @@ using System.Text.Json.Nodes;
 using System.Threading.Tasks;
 using EpicGames.Core;
 using EpicGames.Horde.Storage;
+using EpicGames.Horde.Storage.Bundles;
 using EpicGames.Horde.Storage.Clients;
 using EpicGames.Horde.Storage.Nodes;
 using Horde.Server.Utilities;
diff --git a/Engine/Source/Programs/Horde/Horde.Server/Storage/StorageService.cs b/Engine/Source/Programs/Horde/Horde.Server/Storage/StorageService.cs
--- a/Engine/Source/Programs/Horde/Horde.Server/Storage/StorageService.cs
+++ b/Engine/Source/Programs/Horde/Horde.Server/Storage/StorageService.cs
@@ -11,6 +11,7 @@ using System.Threading.Tasks;
 using EpicGames.Core;
 using EpicGames.Horde.Storage;
 using EpicGames.Horde.Storage.Backends;
+using EpicGames.Horde.Storage.Bundles;
 using EpicGames.Horde.Storage.Clients;
 using EpicGames.Redis;
 using EpicGames.Redis.Utility;
diff --git a/Engine/Source/Programs/Horde/Horde.Server/Tools/ToolCollection.cs b/Engine/Source/Programs/Horde/Horde.Server/Tools/ToolCollection.cs
--- a/Engine/Source/Programs/Horde/Horde.Server/Tools/ToolCollection.cs
+++ b/Engine/Source/Programs/Horde/Horde.Server/Tools/ToolCollection.cs
@@ -2,6 +2,7 @@
 
 using EpicGames.Core;
 using EpicGames.Horde.Storage;
+using EpicGames.Horde.Storage.Bundles;
 using EpicGames.Horde.Storage.Clients;
 using EpicGames.Horde.Storage.Nodes;
 using EpicGames.Horde.Tools;
diff --git a/Engine/Source/Programs/Horde/Horde/Commands/Bundles/BundleCreate.cs b/Engine/Source/Programs/Horde/Horde/Commands/Bundles/BundleCreate.cs
--- a/Engine/Source/Programs/Horde/Horde/Commands/Bundles/BundleCreate.cs
+++ b/Engine/Source/Programs/Horde/Horde/Commands/Bundles/BundleCreate.cs
@@ -4,6 +4,7 @@ using System.ComponentModel;
 using System.Diagnostics;
 using EpicGames.Core;
 using EpicGames.Horde.Storage;
+using EpicGames.Horde.Storage.Bundles;
 using EpicGames.Horde.Storage.Clients;
 using EpicGames.Horde.Storage.Nodes;
 using Microsoft.Extensions.Logging;
diff --git a/Engine/Source/Programs/Horde/Horde/Commands/Bundles/BundleExtract.cs b/Engine/Source/Programs/Horde/Horde/Commands/Bundles/BundleExtract.cs
--- a/Engine/Source/Programs/Horde/Horde/Commands/Bundles/BundleExtract.cs
+++ b/Engine/Source/Programs/Horde/Horde/Commands/Bundles/BundleExtract.cs
@@ -4,6 +4,7 @@ using System.ComponentModel;
 using System.Diagnostics;
 using EpicGames.Core;
 using EpicGames.Horde.Storage;
+using EpicGames.Horde.Storage.Bundles;
 using EpicGames.Horde.Storage.Clients;
 using EpicGames.Horde.Storage.Nodes;
 using Microsoft.Extensions.Logging;
diff --git a/Engine/Source/Programs/Horde/Horde/Commands/Compute/ComputeRun.cs b/Engine/Source/Programs/Horde/Horde/Commands/Compute/ComputeRun.cs
--- a/Engine/Source/Programs/Horde/Horde/Commands/Compute/ComputeRun.cs
+++ b/Engine/Source/Programs/Horde/Horde/Commands/Compute/ComputeRun.cs
@@ -8,6 +8,7 @@ using EpicGames.Horde.Common;
 using EpicGames.Horde.Compute;
 using EpicGames.Horde.Compute.Clients;
 using EpicGames.Horde.Storage;
+using EpicGames.Horde.Storage.Bundles;
 using EpicGames.Horde.Storage.Clients;
 using EpicGames.Horde.Storage.Nodes;
 using Microsoft.Extensions.Logging;
diff --git a/Engine/Source/Programs/Horde/Horde/Commands/Workspace/WorkspaceSync.cs b/Engine/Source/Programs/Horde/Horde/Commands/Workspace/WorkspaceSync.cs
--- a/Engine/Source/Programs/Horde/Horde/Commands/Workspace/WorkspaceSync.cs
+++ b/Engine/Source/Programs/Horde/Horde/Commands/Workspace/WorkspaceSync.cs
@@ -4,6 +4,7 @@ using System.ComponentModel;
 using System.Diagnostics;
 using EpicGames.Core;
 using EpicGames.Horde.Storage;
+using EpicGames.Horde.Storage.Bundles;
 using EpicGames.Horde.Storage.Clients;
 using EpicGames.Horde.Storage.Nodes;
 using Microsoft.Extensions.Logging;
diff --git a/Engine/Source/Programs/Shared/EpicGames.Horde.Tests/ComputeSocketTests.cs b/Engine/Source/Programs/Shared/EpicGames.Horde.Tests/ComputeSocketTests.cs
--- a/Engine/Source/Programs/Shared/EpicGames.Horde.Tests/ComputeSocketTests.cs
+++ b/Engine/Source/Programs/Shared/EpicGames.Horde.Tests/ComputeSocketTests.cs
@@ -13,6 +13,7 @@ using EpicGames.Core;
 using EpicGames.Horde.Compute;
 using EpicGames.Horde.Compute.Transports;
 using EpicGames.Horde.Storage;
+using EpicGames.Horde.Storage.Bundles;
 using EpicGames.Horde.Storage.Clients;
 using EpicGames.Horde.Storage.Nodes;
 using Microsoft.Extensions.Logging;
diff --git a/Engine/Source/Programs/Shared/EpicGames.Horde/Compute/AgentMessageHandler.cs b/Engine/Source/Programs/Shared/EpicGames.Horde/Compute/AgentMessageHandler.cs
--- a/Engine/Source/Programs/Shared/EpicGames.Horde/Compute/AgentMessageHandler.cs
+++ b/Engine/Source/Programs/Shared/EpicGames.Horde/Compute/AgentMessageHandler.cs
@@ -16,6 +16,7 @@ using System.Linq;
 using EpicGames.Horde.Storage.Clients;
 using System.Runtime.ExceptionServices;
 using System.Text;
+using EpicGames.Horde.Storage.Bundles;
 
 namespace EpicGames.Horde.Compute
 {
diff --git a/Engine/Source/Programs/Shared/EpicGames.Horde/Storage/Bundles/BundleStorageClient.cs b/Engine/Source/Programs/Shared/EpicGames.Horde/Storage/Bundles/BundleStorageClient.cs
new file mode 100644
--- /dev/null
+++ b/Engine/Source/Programs/Shared/EpicGames.Horde/Storage/Bundles/BundleStorageClient.cs
@@ -0,0 +1,359 @@
+// Copyright Epic Games, Inc. All Rights Reserved.
+
+using System;
+using System.Buffers.Text;
+using System.Collections.Generic;
+using System.IO;
+using System.Linq;
+using System.Threading;
+using System.Threading.Tasks;
+using EpicGames.Core;
+using EpicGames.Horde.Storage.Bundles;
+using EpicGames.Horde.Storage.Bundles.V1;
+using EpicGames.Horde.Storage.Bundles.V2;
+using EpicGames.Horde.Storage.Clients;
+using Microsoft.Extensions.Logging;
+
+namespace EpicGames.Horde.Storage.Bundles
+{
+	/// <summary>
+	/// Base class for an implementation of <see cref="IStorageClient"/>, providing implementations for some common functionality using bundles.
+	/// </summary>
+	public sealed class BundleStorageClient : IStorageClient
+	{
+		/// <summary>
+		/// Handle to a bundle object
+		/// </summary>
+		class BundleHandle : IBlobHandle
+		{
+			readonly BundleStorageClient _storageClient;
+			readonly IBlobHandle _inner;
+
+			/// <inheritdoc/>
+			public IBlobHandle? Outer => _inner.Outer;
+
+			/// <summary>
+			/// Constructor
+			/// </summary>
+			public BundleHandle(BundleStorageClient storageClient, IBlobHandle inner)
+			{
+				_storageClient = storageClient;
+				_inner = inner;
+			}
+
+			/// <inheritdoc/>
+			public ValueTask FlushAsync(CancellationToken cancellationToken = default) => _inner.FlushAsync(cancellationToken);
+
+			/// <inheritdoc/>
+			public Task<Stream> OpenBodyAsync(int offset = 0, int? length = null, CancellationToken cancellationToken = default)
+				=> _inner.OpenBodyAsync(offset, length, cancellationToken);
+
+			/// <inheritdoc/>
+			public async ValueTask<BlobData> ReadBlobDataAsync(CancellationToken cancellationToken = default)
+			{
+				IReadOnlyMemoryOwner<byte> data = await _inner.ReadBodyAsync(cancellationToken);
+
+				List<BlobLocator> importLocators = ReadImportsFromData(data.Memory).ToList();
+				List<IBlobHandle> importHandles = importLocators.ConvertAll(x => _storageClient.CreateBlobHandle(x));
+
+				return new BlobDataWithOwner(Bundle.BlobType, data.Memory, importHandles, data);
+			}
+
+			IEnumerable<BlobLocator> ReadImportsFromData(ReadOnlyMemory<byte> data)
+			{
+				BundleSignature signature = Bundle.ReadSignature(data.Span);
+				if (signature.Version <= BundleVersion.LatestV1)
+				{
+					return ReadImportsFromDataV1(data);
+				}
+				else if (signature.Version <= BundleVersion.LatestV2)
+				{
+					return ReadImportsFromDataV2(data);
+				}
+				else
+				{
+					throw new InvalidOperationException($"Unsupported bundle version {(int)signature.Version}");
+				}
+			}
+
+			static IEnumerable<BlobLocator> ReadImportsFromDataV1(ReadOnlyMemory<byte> data)
+			{
+				BundleHeader header = BundleHeader.Read(data);
+				return header.Imports.Select(x => x.BaseLocator);
+			}
+
+			IEnumerable<BlobLocator> ReadImportsFromDataV2(ReadOnlyMemory<byte> data)
+			{
+				HashSet<BlobLocator> locators = new HashSet<BlobLocator>();
+				while (data.Length > 0)
+				{
+					BundleSignature signature = Bundle.ReadSignature(data.Span);
+
+					using IRefCountedHandle<Bundles.V2.Packet> packet = Bundles.V2.Packet.Decode(data, _storageClient._cache.Allocator);
+					for (int idx = 0; idx < packet.Target.GetImportCount(); idx++)
+					{
+						PacketImport import = packet.Target.GetImport(idx);
+						if (import.BaseIdx == -1)
+						{
+							locators.Add(new BlobLocator(import.Fragment.Clone()));
+						}
+					}
+
+					data = data.Slice(signature.HeaderLength);
+				}
+				return locators;
+			}
+
+			/// <inheritdoc/>
+			public bool TryAppendIdentifier(Utf8StringBuilder builder)
+				=> _inner.TryAppendIdentifier(builder);
+
+			/// <inheritdoc/>
+			public IBlobHandle GetFragmentHandle(ReadOnlySpan<byte> fragment)
+			{
+				int exportIdx;
+				if (Utf8Parser.TryParse(fragment, out exportIdx, out int numBytesRead) && numBytesRead == fragment.Length)
+				{
+					return new Bundles.V1.FlushedNodeHandle(_storageClient._bundleReader, _inner.GetLocator(), this, exportIdx);
+				}
+
+				int ampIdx = fragment.IndexOf((byte)'&');
+				if (ampIdx == -1)
+				{
+					return new Bundles.V2.PacketHandle(_storageClient, this, fragment, _storageClient._cache);
+				}
+				else
+				{
+					return new Bundles.V2.PacketHandle(_storageClient, this, fragment.Slice(0, ampIdx), _storageClient._cache).GetFragmentHandle(fragment.Slice(ampIdx + 1));
+				}
+			}
+
+			/// <inheritdoc/>
+			public override bool Equals(object? obj) => obj is BundleHandle other && _inner.Equals(other._inner);
+
+			/// <inheritdoc/>
+			public override int GetHashCode() => _inner.GetHashCode();
+		}
+
+		readonly IStorageClient _inner;
+		readonly BundleCache _cache;
+		readonly Bundles.V1.BundleReader _bundleReader;
+
+		/// <summary>
+		/// Allocator which trims the cache to keep below a maximum size
+		/// </summary>
+		public IMemoryAllocator<byte> Allocator => _cache.Allocator;
+
+		/// <inheritdoc/>
+		public bool SupportsRedirects { get; } = false;
+
+		/// <summary>
+		/// Constructor
+		/// </summary>
+		public BundleStorageClient(IStorageClient inner, BundleCache cache, ILogger logger)
+		{
+			_inner = inner;
+			_cache = cache;
+			_bundleReader = new Bundles.V1.BundleReader(this, cache, logger);
+		}
+
+		/// <inheritdoc/>
+		public void Dispose()
+		{
+			_inner.Dispose();
+		}
+
+		/// <summary>
+		/// Creates a bundle storage client around a memory client backend
+		/// </summary>
+		public static BundleStorageClient CreateFromMemory(ILogger logger)
+		{
+			MemoryStorageClient blobStore = new MemoryStorageClient();
+			return new BundleStorageClient(blobStore, BundleCache.None, logger);
+		}
+
+		/// <summary>
+		/// Creates a bundle storage client around a directory on the filesystem
+		/// </summary>
+		public static BundleStorageClient CreateFromDirectory(DirectoryReference rootDir, BundleCache cache, ILogger logger)
+		{
+			FileStorageClient fileStorageClient = new FileStorageClient(rootDir, logger);
+			return new BundleStorageClient(fileStorageClient, cache, logger);
+		}
+
+		#region Blobs
+
+		/// <inheritdoc/>
+		public async Task<Stream> OpenBlobAsync(BlobLocator locator, int offset = 0, int? length = null, CancellationToken cancellationToken = default) 
+			=> await OpenBundleAsync(locator, offset, length, cancellationToken);
+
+		/// <inheritdoc/>
+		public async ValueTask<BlobData> ReadBlobAsync(BlobLocator locator, CancellationToken cancellationToken = default)
+		{
+			IBlobHandle handle = CreateBlobHandle(locator);
+			return await handle.ReadBlobDataAsync(cancellationToken);
+		}
+
+		/// <inheritdoc/>
+		public async ValueTask<IBlobHandle> WriteBlobAsync(BlobType type, Stream stream, IReadOnlyList<IBlobHandle> references, string? basePath = null, CancellationToken cancellationToken = default)
+		{
+			if (type == Bundle.BlobType)
+			{
+				return await WriteBundleAsync(stream, references, basePath, cancellationToken);
+			}
+			else
+			{
+				await using IStorageWriter writer = CreateWriter(basePath);
+
+				int length = 0;
+				for (; ; )
+				{
+					int readLength = await stream.ReadAsync(writer.GetOutputBuffer(length, length + 1), cancellationToken);
+					if (readLength == 0)
+					{
+						break;
+					}
+					length += readLength;
+				}
+
+				return await writer.WriteBlobAsync(type, length, references, Array.Empty<AliasInfo>(), cancellationToken);
+			}
+		}
+
+		/// <inheritdoc/>
+		public ValueTask<Uri?> TryGetReadRedirectAsync(BlobLocator locator, CancellationToken cancellationToken = default) => new ValueTask<Uri?>();
+
+		/// <inheritdoc/>
+		public ValueTask<(BlobLocator, Uri)?> TryGetWriteRedirectAsync(string? prefix = null, CancellationToken cancellationToken = default) => new ValueTask<(BlobLocator, Uri)?>();
+
+		#endregion
+
+		#region Bundles
+
+		/// <summary>
+		/// Read a bundle from the underlying storage
+		/// </summary>
+		async Task<Stream> OpenBundleAsync(BlobLocator locator, int offset, int? length, CancellationToken cancellationToken)
+			=> await _inner.CreateBlobHandle(locator).OpenBodyAsync(offset, length, cancellationToken);
+
+		/// <summary>
+		/// Write a bundle to the underlying storage
+		/// </summary>
+		ValueTask<IBlobHandle> WriteBundleAsync(Stream stream, IReadOnlyList<IBlobHandle> references, string? basePath = null, CancellationToken cancellationToken = default)
+			=> _inner.WriteBlobAsync(Bundle.BlobType, stream, references, basePath, cancellationToken);
+
+		/// <inheritdoc/>
+		public Task<Bundles.V1.BundleHeader> ReadHeaderAsync(BlobLocator locator, CancellationToken cancellationToken) => _bundleReader.ReadHeaderAsync(locator, cancellationToken);
+
+		#endregion
+
+		#region Nodes
+
+		/// <inheritdoc/>
+		public IBlobHandle CreateBlobHandle(BlobLocator locator)
+		{
+			if (locator.TryUnwrap(out BlobLocator baseLocator, out Utf8String fragment))
+			{
+				BundleHandle bundleHandle = new BundleHandle(this, CreateBlobHandle(baseLocator));
+
+				int exportIdx;
+				if (Utf8Parser.TryParse(fragment.Span, out exportIdx, out int numBytesRead) && numBytesRead == fragment.Length)
+				{
+					return new Bundles.V1.FlushedNodeHandle(_bundleReader, baseLocator, bundleHandle, exportIdx);
+				}
+
+				int ampIdx = fragment.IndexOf('&');
+				if (ampIdx == -1)
+				{
+					return new Bundles.V2.PacketHandle(this, bundleHandle, fragment.Span, _cache);
+				}
+				else
+				{
+					return new Bundles.V2.ExportHandle(new Bundles.V2.PacketHandle(this, bundleHandle, fragment.Slice(0, ampIdx).Span, _cache), fragment.Span.Slice(ampIdx + 1));
+				}
+			}
+			else
+			{
+				return new BundleHandle(this, _inner.CreateBlobHandle(baseLocator));
+			}
+		}
+
+		/// <inheritdoc/>
+		public IStorageWriter CreateWriter(string? basePath = null, BundleOptions? options = null)
+		{
+			options ??= BundleOptions.Default;
+
+			if (options.MaxVersion == BundleVersion.LatestV1)
+			{
+				return new Bundles.V1.BundleWriter(this, _bundleReader, basePath, options);
+			}
+			else if(options.MaxVersion == BundleVersion.LatestV2)
+			{
+				return new Bundles.V2.BundleWriter(this, basePath, _cache, options);
+			}
+			else
+			{
+				throw new InvalidOperationException($"Unsupported bundle version: {(int)options.MaxVersion}");
+			}
+		}
+
+		/// <inheritdoc/>
+		IStorageWriter IStorageClient.CreateWriter(string? basePath) => CreateWriter(basePath);
+
+		#endregion
+
+		#region Aliases
+
+		/// <inheritdoc/>
+		public Task AddAliasAsync(string name, IBlobHandle handle, int rank, ReadOnlyMemory<byte> data, CancellationToken cancellationToken)
+			=> _inner.AddAliasAsync(name, handle, rank, data, cancellationToken);
+
+		/// <inheritdoc/>
+		public Task RemoveAliasAsync(string name, IBlobHandle handle, CancellationToken cancellationToken)
+			=> _inner.RemoveAliasAsync(name, handle, cancellationToken);
+
+		/// <inheritdoc/>
+		public async Task<BlobAlias[]> FindAliasesAsync(string name, int? maxLength = null, CancellationToken cancellationToken = default)
+		{
+			BlobAlias[] aliases = await _inner.FindAliasesAsync(name, maxLength, cancellationToken);
+			for (int idx = 0; idx < aliases.Length; idx++)
+			{
+				BlobAlias alias = aliases[idx];
+				aliases[idx] = new BlobAlias(CreateBlobHandle(alias.Target.GetLocator()), alias.Rank, alias.Data);
+			}
+			return aliases;
+		}
+
+		#endregion
+
+		#region Refs
+
+		/// <inheritdoc/>
+		public Task<bool> DeleteRefAsync(RefName name, CancellationToken cancellationToken = default)
+			=> _inner.DeleteRefAsync(name, cancellationToken);
+
+		/// <inheritdoc/>
+		public async Task<IBlobHandle?> TryReadRefAsync(RefName name, RefCacheTime cacheTime = default, CancellationToken cancellationToken = default)
+		{
+			IBlobHandle? target = await _inner.TryReadRefAsync(name, cacheTime, cancellationToken);
+			if (target != null)
+			{
+				target = CreateBlobHandle(target.GetLocator());
+			}
+			return target;
+		}
+
+		/// <inheritdoc/>
+		public Task WriteRefAsync(RefName name, IBlobHandle target, RefOptions? options = null, CancellationToken cancellationToken = default)
+			=> _inner.WriteRefAsync(name, target, options, cancellationToken);
+
+		#endregion
+
+		/// <inheritdoc/>
+		public void GetStats(StorageStats stats)
+		{
+			_inner.GetStats(stats);
+			_bundleReader.GetStats(stats);
+		}
+	}
+}
diff --git a/Engine/Source/Programs/Shared/EpicGames.Horde/Storage/Clients/BundleStorageClient.cs b/Engine/Source/Programs/Shared/EpicGames.Horde/Storage/Clients/BundleStorageClient.cs
deleted file mode 100644
diff --git a/Engine/Source/Programs/Shared/EpicGames.Horde/Storage/Clients/HttpStorageClient.cs b/Engine/Source/Programs/Shared/EpicGames.Horde/Storage/Clients/HttpStorageClient.cs
--- a/Engine/Source/Programs/Shared/EpicGames.Horde/Storage/Clients/HttpStorageClient.cs
+++ b/Engine/Source/Programs/Shared/EpicGames.Horde/Storage/Clients/HttpStorageClient.cs
@@ -13,6 +13,7 @@ using Microsoft.Extensions.Logging;
 using EpicGames.Horde.Storage.Backends;
 using System.IO;
 using System.Collections.Generic;
+using EpicGames.Horde.Storage.Bundles;
 
 namespace EpicGames.Horde.Storage.Clients
 {
diff --git a/Engine/Source/Programs/UnrealBuildTool/Artifacts/HordeStorageArtifactCache.cs b/Engine/Source/Programs/UnrealBuildTool/Artifacts/HordeStorageArtifactCache.cs
--- a/Engine/Source/Programs/UnrealBuildTool/Artifacts/HordeStorageArtifactCache.cs
+++ b/Engine/Source/Programs/UnrealBuildTool/Artifacts/HordeStorageArtifactCache.cs
@@ -8,6 +8,7 @@ using System.Threading;
 using System.Threading.Tasks;
 using EpicGames.Core;
 using EpicGames.Horde.Storage;
+using EpicGames.Horde.Storage.Bundles;
 using EpicGames.Horde.Storage.Clients;
 using EpicGames.Horde.Storage.Nodes;
 using Microsoft.Extensions.Logging;
diff --git a/Engine/Source/Programs/UnrealBuildTool/Executors/Experimental/UBAHordeSession.cs b/Engine/Source/Programs/UnrealBuildTool/Executors/Experimental/UBAHordeSession.cs
--- a/Engine/Source/Programs/UnrealBuildTool/Executors/Experimental/UBAHordeSession.cs
+++ b/Engine/Source/Programs/UnrealBuildTool/Executors/Experimental/UBAHordeSession.cs
@@ -14,6 +14,7 @@ using EpicGames.Horde;
 using EpicGames.Horde.Compute;
 using EpicGames.Horde.Compute.Clients;
 using EpicGames.Horde.Storage;
+using EpicGames.Horde.Storage.Bundles;
 using EpicGames.Horde.Storage.Clients;
 using EpicGames.Horde.Storage.Nodes;
 using EpicGames.OIDC;
