Horde: Add shared base classes for export/packet/bundle handle types.

diff --git a/Engine/Source/Programs/Horde/Horde.Server/Commands/Generate/ToolDataCommand.cs b/Engine/Source/Programs/Horde/Horde.Server/Commands/Generate/ToolDataCommand.cs
--- a/Engine/Source/Programs/Horde/Horde.Server/Commands/Generate/ToolDataCommand.cs
+++ b/Engine/Source/Programs/Horde/Horde.Server/Commands/Generate/ToolDataCommand.cs
@@ -7,7 +7,6 @@ using System.Threading.Tasks;
 using EpicGames.Core;
 using EpicGames.Horde.Storage;
 using EpicGames.Horde.Storage.Bundles;
-using EpicGames.Horde.Storage.Clients;
 using EpicGames.Horde.Storage.Nodes;
 using Horde.Server.Utilities;
 using Microsoft.Extensions.Logging;
diff --git a/Engine/Source/Programs/Horde/Horde.Server/Commands/Test/ReplicationCommand.cs b/Engine/Source/Programs/Horde/Horde.Server/Commands/Test/ReplicationCommand.cs
--- a/Engine/Source/Programs/Horde/Horde.Server/Commands/Test/ReplicationCommand.cs
+++ b/Engine/Source/Programs/Horde/Horde.Server/Commands/Test/ReplicationCommand.cs
@@ -12,7 +12,6 @@ using EpicGames.Horde.Streams;
 using Horde.Server.Configuration;
 using EpicGames.Horde.Replicators;
 using Horde.Server.Replicators;
-using Microsoft.CodeAnalysis.Operations;
 
 namespace Horde.Server.Commands.Test
 {
diff --git a/Engine/Source/Programs/Horde/Horde.Server/Replicators/PerforceReplicator.cs b/Engine/Source/Programs/Horde/Horde.Server/Replicators/PerforceReplicator.cs
--- a/Engine/Source/Programs/Horde/Horde.Server/Replicators/PerforceReplicator.cs
+++ b/Engine/Source/Programs/Horde/Horde.Server/Replicators/PerforceReplicator.cs
@@ -5,7 +5,6 @@ using System.Buffers.Binary;
 using System.Collections.Generic;
 using System.Diagnostics;
 using System.Linq;
-using System.Runtime.InteropServices;
 using System.Security.Cryptography;
 using System.Text;
 using System.Threading;
@@ -14,7 +13,6 @@ using EpicGames.Core;
 using EpicGames.Horde.Replicators;
 using EpicGames.Horde.Storage;
 using EpicGames.Horde.Storage.Bundles;
-using EpicGames.Horde.Storage.Clients;
 using EpicGames.Horde.Storage.Nodes;
 using EpicGames.Horde.Streams;
 using EpicGames.Perforce;
@@ -39,9 +37,9 @@ namespace Horde.Server.Replicators
 	/// <summary>
 	/// Exception thrown to indicate that the state of a replicator has been modified externally
 	/// </summary>
-	class ReplicatorModifiedException : Exception
+	public class ReplicatorModifiedException : Exception
 	{
-		public ReplicatorModifiedException() : base("The replicator has been modified externally")
+		internal ReplicatorModifiedException() : base("The replicator has been modified externally")
 		{
 		}
 	}
@@ -298,7 +296,7 @@ namespace Horde.Server.Replicators
 		/// <summary>
 		/// Runs a replication loop for a stream
 		/// </summary>
-		public async Task RunAsync(ReplicatorId replicatorId, StreamConfig streamConfig, ReplicatorConfig replicatorConfig, PerforceReplicationOptions options, CancellationToken cancellationToken = default)
+		public async Task RunAsync(ReplicatorId replicatorId, StreamConfig streamConfig, PerforceReplicationOptions options, CancellationToken cancellationToken = default)
 		{
 			_logger.LogInformation("Starting replication background task for {ReplicatorId}", replicatorId);
 
@@ -318,7 +316,7 @@ namespace Horde.Server.Replicators
 		/// <summary>
 		/// Runs the replicator for a single change
 		/// </summary>
-		public async Task<IReplicator> RunOnceAsync(IReplicator replicator, StreamConfig streamConfig, PerforceReplicationOptions options, CancellationToken cancellationToken)
+		public async Task<IReplicator> RunOnceAsync(IReplicator replicator, StreamConfig streamConfig, PerforceReplicationOptions replicatorOptions, CancellationToken cancellationToken)
 		{
 			RefName refName = new RefName(streamConfig.Id.ToString());
 			RefName incRefName = GetIncrementalRefName(replicator.Id);
@@ -371,12 +369,10 @@ namespace Horde.Server.Replicators
 			int change = replicator.CurrentChange!.Value;
 			_logger.LogInformation("Replicating {ReplicatorId} change {Change}", replicator.Id, change);
 
-			PerforceReplicationOptions replicationOptions = new PerforceReplicationOptions();
-
 			BlobSerializerOptions blobOptions = new BlobSerializerOptions();
 			try
 			{
-				replicator = await WriteInternalAsync(replicator, change, streamConfig, replicationOptions, blobOptions, cancellationToken);
+				replicator = await WriteInternalAsync(replicator, change, streamConfig, replicatorOptions, blobOptions, cancellationToken);
 				return replicator;
 			}
 			catch (OperationCanceledException ex)
diff --git a/Engine/Source/Programs/Horde/Horde.Server/Replicators/ReplicationService.cs b/Engine/Source/Programs/Horde/Horde.Server/Replicators/ReplicationService.cs
--- a/Engine/Source/Programs/Horde/Horde.Server/Replicators/ReplicationService.cs
+++ b/Engine/Source/Programs/Horde/Horde.Server/Replicators/ReplicationService.cs
@@ -91,7 +91,7 @@ namespace Horde.Server.Replicators
 								if (replicatorConfig.Enabled)
 								{
 									ReplicatorId replicatorId = new ReplicatorId(streamConfig.Id, replicatorConfig.Id);
-									replicatorIdToTask.Add(replicatorId, BackgroundTask.StartNew(ctx => RunReplicationGuardedAsync(replicatorId, streamConfig, replicatorConfig, ctx)));
+									replicatorIdToTask.Add(replicatorId, BackgroundTask.StartNew(ctx => RunReplicationGuardedAsync(replicatorId, streamConfig, ctx)));
 									removeReplicators.Remove(replicatorId);
 								}
 							}
@@ -114,7 +114,7 @@ namespace Horde.Server.Replicators
 			}
 		}
 
-		async Task RunReplicationGuardedAsync(ReplicatorId replicatorId, StreamConfig streamConfig, ReplicatorConfig replicatorConfig, CancellationToken cancellationToken)
+		async Task RunReplicationGuardedAsync(ReplicatorId replicatorId, StreamConfig streamConfig, CancellationToken cancellationToken)
 		{
 			_logger.LogInformation("Started background task for replication of {ReplicatorId}", replicatorId);
 
@@ -129,7 +129,7 @@ namespace Horde.Server.Replicators
 				.Build();
 
 			PerforceReplicationOptions replicationOptions = new PerforceReplicationOptions();
-			await pipeline.ExecuteAsync(async ctx => await _replicator.RunAsync(replicatorId, streamConfig, replicatorConfig, replicationOptions, ctx), cancellationToken);
+			await pipeline.ExecuteAsync(async ctx => await _replicator.RunAsync(replicatorId, streamConfig, replicationOptions, ctx), cancellationToken);
 		}
 	}
 }
diff --git a/Engine/Source/Programs/Horde/Horde.Server/Replicators/ReplicatorCollection.cs b/Engine/Source/Programs/Horde/Horde.Server/Replicators/ReplicatorCollection.cs
--- a/Engine/Source/Programs/Horde/Horde.Server/Replicators/ReplicatorCollection.cs
+++ b/Engine/Source/Programs/Horde/Horde.Server/Replicators/ReplicatorCollection.cs
@@ -183,22 +183,22 @@ namespace Horde.Server.Replicators
 
 			if (options.Pause != null)
 			{
-				updates.Add(Builders<ReplicatorDoc>.Update.SetOrUnsetBool(x => x.Pause, options.Pause ?? false));
+				updates.Add(Builders<ReplicatorDoc>.Update.SetOrUnsetBool(x => x.Pause, options.Pause.Value));
 			}
 
 			if (options.Clean != null)
 			{
-				updates.Add(Builders<ReplicatorDoc>.Update.SetOrUnsetBool(x => x.Clean, options.Clean ?? false));
+				updates.Add(Builders<ReplicatorDoc>.Update.SetOrUnsetBool(x => x.Clean, options.Clean.Value));
 			}
 
 			if (options.Reset != null)
 			{
-				updates.Add(Builders<ReplicatorDoc>.Update.SetOrUnsetBool(x => x.Reset, options.Reset ?? false));
+				updates.Add(Builders<ReplicatorDoc>.Update.SetOrUnsetBool(x => x.Reset, options.Reset.Value));
 			}
 
 			if (options.SingleStep != null)
 			{
-				updates.Add(Builders<ReplicatorDoc>.Update.SetOrUnsetBool(x => x.SingleStep, options.SingleStep ?? false));
+				updates.Add(Builders<ReplicatorDoc>.Update.SetOrUnsetBool(x => x.SingleStep, options.SingleStep.Value));
 			}
 
 			if (options.LastChange != null)
diff --git a/Engine/Source/Programs/Horde/Horde.Server/Storage/StorageService.cs b/Engine/Source/Programs/Horde/Horde.Server/Storage/StorageService.cs
--- a/Engine/Source/Programs/Horde/Horde.Server/Storage/StorageService.cs
+++ b/Engine/Source/Programs/Horde/Horde.Server/Storage/StorageService.cs
@@ -51,13 +51,13 @@ namespace Horde.Server.Storage
 	/// </summary>
 	public sealed class StorageService : IHostedService, IStorageClientFactory, IAsyncDisposable
 	{
-		sealed class LeafBlobHandle : IBlobHandle
+		sealed class LeafBlobHandle : BlobHandle
 		{
 			readonly IStorageBackend _backend;
 			readonly string _path;
 			readonly Tracer _tracer;
 
-			public IBlobHandle? Outer => null;
+			public override IBlobHandle? Outer => null;
 
 			public LeafBlobHandle(IStorageBackend backend, string path, Tracer tracer)
 			{
@@ -67,10 +67,10 @@ namespace Horde.Server.Storage
 			}
 
 			/// <inheritdoc/>
-			public ValueTask FlushAsync(CancellationToken cancellationToken = default) => default;
+			public override ValueTask FlushAsync(CancellationToken cancellationToken = default) => default;
 
 			/// <inheritdoc/>
-			public async Task<Stream> OpenBodyAsync(int offset = 0, int? length = null, CancellationToken cancellationToken = default)
+			public override async Task<Stream> OpenBodyAsync(int offset = 0, int? length = null, CancellationToken cancellationToken = default)
 			{
 				using TelemetrySpan span = _tracer.StartActiveSpan($"{nameof(StorageService)}.{nameof(LeafBlobHandle)}.{nameof(OpenBodyAsync)}");
 				span.SetAttribute("path", _path);
@@ -86,14 +86,14 @@ namespace Horde.Server.Storage
 			}
 
 			/// <inheritdoc/>
-			public async ValueTask<BlobData> ReadBlobDataAsync(CancellationToken cancellationToken = default)
+			public override async ValueTask<BlobData> ReadBlobDataAsync(CancellationToken cancellationToken = default)
 			{
 				IReadOnlyMemoryOwner<byte> obj = await _backend.ReadAsync(_path, cancellationToken);
 				return new BlobDataWithOwner(BlobType.Leaf, obj.Memory, Array.Empty<IBlobHandle>(), obj);
 			}
 
 			/// <inheritdoc/>
-			public bool TryAppendIdentifier(Utf8StringBuilder builder)
+			public override bool TryAppendIdentifier(Utf8StringBuilder builder)
 			{
 				builder.Append(_path);
 				return true;
diff --git a/Engine/Source/Programs/Horde/Horde.Server/Tools/ToolCollection.cs b/Engine/Source/Programs/Horde/Horde.Server/Tools/ToolCollection.cs
--- a/Engine/Source/Programs/Horde/Horde.Server/Tools/ToolCollection.cs
+++ b/Engine/Source/Programs/Horde/Horde.Server/Tools/ToolCollection.cs
@@ -3,7 +3,6 @@
 using EpicGames.Core;
 using EpicGames.Horde.Storage;
 using EpicGames.Horde.Storage.Bundles;
-using EpicGames.Horde.Storage.Clients;
 using EpicGames.Horde.Storage.Nodes;
 using EpicGames.Horde.Tools;
 using EpicGames.Serialization;
diff --git a/Engine/Source/Programs/Shared/EpicGames.Horde.Tests/BlobDataTests.cs b/Engine/Source/Programs/Shared/EpicGames.Horde.Tests/BlobDataTests.cs
--- a/Engine/Source/Programs/Shared/EpicGames.Horde.Tests/BlobDataTests.cs
+++ b/Engine/Source/Programs/Shared/EpicGames.Horde.Tests/BlobDataTests.cs
@@ -13,17 +13,17 @@ namespace EpicGames.Horde.Tests
 	[TestClass]
 	public class BlobDataTests
 	{
-		class DummyHandle : IBlobHandle
+		class DummyHandle : BlobHandle
 		{
 			readonly BlobLocator _locator;
 
-			public IBlobHandle? Outer => null;
+			public override IBlobHandle? Outer => null;
 
 			public DummyHandle(string locator) => _locator = new BlobLocator(new Utf8String(locator));
-			public ValueTask FlushAsync(CancellationToken cancellationToken = default) => default;
-			public ValueTask<BlobData> ReadBlobDataAsync(CancellationToken cancellationToken = default) => throw new NotImplementedException();
+			public override ValueTask FlushAsync(CancellationToken cancellationToken = default) => default;
+			public override ValueTask<BlobData> ReadBlobDataAsync(CancellationToken cancellationToken = default) => throw new NotImplementedException();
 
-			public bool TryAppendIdentifier(Utf8StringBuilder builder)
+			public override bool TryAppendIdentifier(Utf8StringBuilder builder)
 			{
 				builder.Append(_locator.Path);
 				return true;
diff --git a/Engine/Source/Programs/Shared/EpicGames.Horde.Tests/BundlePacketTests.cs b/Engine/Source/Programs/Shared/EpicGames.Horde.Tests/BundlePacketTests.cs
--- a/Engine/Source/Programs/Shared/EpicGames.Horde.Tests/BundlePacketTests.cs
+++ b/Engine/Source/Programs/Shared/EpicGames.Horde.Tests/BundlePacketTests.cs
@@ -16,12 +16,12 @@ namespace EpicGames.Horde.Tests
 	[TestClass]
 	public class BundlePacketTests
 	{
-		class FakeHandle : IBlobHandle
+		class FakeHandle : BlobHandle
 		{
 			readonly IBlobHandle? _outer;
 			readonly string? _fragment;
 
-			public IBlobHandle? Outer => _outer;
+			public override IBlobHandle? Outer => _outer;
 
 			public FakeHandle(IBlobHandle? outer, string? fragment)
 			{
@@ -29,10 +29,10 @@ namespace EpicGames.Horde.Tests
 				_fragment = fragment;
 			}
 
-			public ValueTask<BlobData> ReadBlobDataAsync(CancellationToken cancellationToken = default) => throw new NotImplementedException();
-			public ValueTask FlushAsync(CancellationToken cancellationToken = default) => throw new NotImplementedException();
+			public override ValueTask<BlobData> ReadBlobDataAsync(CancellationToken cancellationToken = default) => throw new NotImplementedException();
+			public override ValueTask FlushAsync(CancellationToken cancellationToken = default) => throw new NotImplementedException();
 
-			public bool TryAppendIdentifier(Utf8StringBuilder builder)
+			public override bool TryAppendIdentifier(Utf8StringBuilder builder)
 			{
 				if (_fragment == null)
 				{
diff --git a/Engine/Source/Programs/Shared/EpicGames.Horde/Compute/AgentMessageHandler.cs b/Engine/Source/Programs/Shared/EpicGames.Horde/Compute/AgentMessageHandler.cs
--- a/Engine/Source/Programs/Shared/EpicGames.Horde/Compute/AgentMessageHandler.cs
+++ b/Engine/Source/Programs/Shared/EpicGames.Horde/Compute/AgentMessageHandler.cs
@@ -13,7 +13,6 @@ using Microsoft.Extensions.Logging;
 using EpicGames.Horde.Storage;
 using System.IO;
 using System.Linq;
-using EpicGames.Horde.Storage.Clients;
 using System.Runtime.ExceptionServices;
 using System.Text;
 using EpicGames.Horde.Storage.Bundles;
diff --git a/Engine/Source/Programs/Shared/EpicGames.Horde/Storage/Bundles/BundleHandle.cs b/Engine/Source/Programs/Shared/EpicGames.Horde/Storage/Bundles/BundleHandle.cs
new file mode 100644
--- /dev/null
+++ b/Engine/Source/Programs/Shared/EpicGames.Horde/Storage/Bundles/BundleHandle.cs
@@ -0,0 +1,136 @@
+﻿// Copyright Epic Games, Inc. All Rights Reserved.
+
+using System;
+using System.Buffers.Text;
+using System.Collections.Generic;
+using System.IO;
+using System.Linq;
+using System.Threading;
+using System.Threading.Tasks;
+using EpicGames.Core;
+using EpicGames.Horde.Storage.Bundles.V1;
+using EpicGames.Horde.Storage.Bundles.V2;
+
+namespace EpicGames.Horde.Storage.Bundles
+{
+	/// <summary>
+	/// Base class for packet handles
+	/// </summary>
+	public abstract class BundleHandle : BlobHandle
+	{
+	}
+
+	/// <summary>
+	/// Generic flushed bundle handle; can be either V1 or V2 format.
+	/// </summary>
+	class FlushedBundleHandle : BundleHandle
+	{
+		readonly BundleStorageClient _storageClient;
+		readonly IBlobHandle _inner;
+
+		/// <inheritdoc/>
+		public override IBlobHandle? Outer => _inner.Outer;
+
+		/// <summary>
+		/// Constructor
+		/// </summary>
+		public FlushedBundleHandle(BundleStorageClient storageClient, IBlobHandle inner)
+		{
+			_storageClient = storageClient;
+			_inner = inner;
+		}
+
+		/// <inheritdoc/>
+		public override ValueTask FlushAsync(CancellationToken cancellationToken = default) => _inner.FlushAsync(cancellationToken);
+
+		/// <inheritdoc/>
+		public override Task<Stream> OpenBodyAsync(int offset = 0, int? length = null, CancellationToken cancellationToken = default)
+			=> _inner.OpenBodyAsync(offset, length, cancellationToken);
+
+		/// <inheritdoc/>
+		public override async ValueTask<BlobData> ReadBlobDataAsync(CancellationToken cancellationToken = default)
+		{
+			IReadOnlyMemoryOwner<byte> data = await _inner.ReadBodyAsync(cancellationToken);
+
+			List<BlobLocator> importLocators = ReadImportsFromData(data.Memory).ToList();
+			List<IBlobHandle> importHandles = importLocators.ConvertAll(x => _storageClient.CreateBlobHandle(x));
+
+			return new BlobDataWithOwner(Bundle.BlobType, data.Memory, importHandles, data);
+		}
+
+		IEnumerable<BlobLocator> ReadImportsFromData(ReadOnlyMemory<byte> data)
+		{
+			BundleSignature signature = Bundle.ReadSignature(data.Span);
+			if (signature.Version <= BundleVersion.LatestV1)
+			{
+				return ReadImportsFromDataV1(data);
+			}
+			else if (signature.Version <= BundleVersion.LatestV2)
+			{
+				return ReadImportsFromDataV2(data);
+			}
+			else
+			{
+				throw new InvalidOperationException($"Unsupported bundle version {(int)signature.Version}");
+			}
+		}
+
+		static IEnumerable<BlobLocator> ReadImportsFromDataV1(ReadOnlyMemory<byte> data)
+		{
+			BundleHeader header = BundleHeader.Read(data);
+			return header.Imports.Select(x => x.BaseLocator);
+		}
+
+		IEnumerable<BlobLocator> ReadImportsFromDataV2(ReadOnlyMemory<byte> data)
+		{
+			HashSet<BlobLocator> locators = new HashSet<BlobLocator>();
+			while (data.Length > 0)
+			{
+				BundleSignature signature = Bundle.ReadSignature(data.Span);
+
+				using IRefCountedHandle<Bundles.V2.Packet> packet = Bundles.V2.Packet.Decode(data, _storageClient.Cache.Allocator);
+				for (int idx = 0; idx < packet.Target.GetImportCount(); idx++)
+				{
+					PacketImport import = packet.Target.GetImport(idx);
+					if (import.BaseIdx == -1)
+					{
+						locators.Add(new BlobLocator(import.Fragment.Clone()));
+					}
+				}
+
+				data = data.Slice(signature.HeaderLength);
+			}
+			return locators;
+		}
+
+		/// <inheritdoc/>
+		public override bool TryAppendIdentifier(Utf8StringBuilder builder)
+			=> _inner.TryAppendIdentifier(builder);
+
+		/// <inheritdoc/>
+		public override IBlobHandle GetFragmentHandle(ReadOnlySpan<byte> fragment)
+		{
+			int exportIdx;
+			if (Utf8Parser.TryParse(fragment, out exportIdx, out int numBytesRead) && numBytesRead == fragment.Length)
+			{
+				return new Bundles.V1.FlushedNodeHandle(_storageClient.BundleReader, _inner.GetLocator(), this, exportIdx);
+			}
+
+			int ampIdx = fragment.IndexOf((byte)'&');
+			if (ampIdx == -1)
+			{
+				return new Bundles.V2.FlushedPacketHandle(_storageClient, this, fragment, _storageClient.Cache);
+			}
+			else
+			{
+				return new Bundles.V2.FlushedPacketHandle(_storageClient, this, fragment.Slice(0, ampIdx), _storageClient.Cache).GetFragmentHandle(fragment.Slice(ampIdx + 1));
+			}
+		}
+
+		/// <inheritdoc/>
+		public override bool Equals(object? obj) => obj is FlushedBundleHandle other && _inner.Equals(other._inner);
+
+		/// <inheritdoc/>
+		public override int GetHashCode() => _inner.GetHashCode();
+	}
+}
diff --git a/Engine/Source/Programs/Shared/EpicGames.Horde/Storage/Bundles/BundleStorageClient.cs b/Engine/Source/Programs/Shared/EpicGames.Horde/Storage/Bundles/BundleStorageClient.cs
--- a/Engine/Source/Programs/Shared/EpicGames.Horde/Storage/Bundles/BundleStorageClient.cs
+++ b/Engine/Source/Programs/Shared/EpicGames.Horde/Storage/Bundles/BundleStorageClient.cs
@@ -4,13 +4,9 @@ using System;
 using System.Buffers.Text;
 using System.Collections.Generic;
 using System.IO;
-using System.Linq;
 using System.Threading;
 using System.Threading.Tasks;
 using EpicGames.Core;
-using EpicGames.Horde.Storage.Bundles;
-using EpicGames.Horde.Storage.Bundles.V1;
-using EpicGames.Horde.Storage.Bundles.V2;
 using EpicGames.Horde.Storage.Clients;
 using Microsoft.Extensions.Logging;
 
@@ -21,129 +17,22 @@ namespace EpicGames.Horde.Storage.Bundles
 	/// </summary>
 	public sealed class BundleStorageClient : IStorageClient
 	{
-		/// <summary>
-		/// Handle to a bundle object
-		/// </summary>
-		class BundleHandle : IBlobHandle
-		{
-			readonly BundleStorageClient _storageClient;
-			readonly IBlobHandle _inner;
-
-			/// <inheritdoc/>
-			public IBlobHandle? Outer => _inner.Outer;
-
-			/// <summary>
-			/// Constructor
-			/// </summary>
-			public BundleHandle(BundleStorageClient storageClient, IBlobHandle inner)
-			{
-				_storageClient = storageClient;
-				_inner = inner;
-			}
-
-			/// <inheritdoc/>
-			public ValueTask FlushAsync(CancellationToken cancellationToken = default) => _inner.FlushAsync(cancellationToken);
-
-			/// <inheritdoc/>
-			public Task<Stream> OpenBodyAsync(int offset = 0, int? length = null, CancellationToken cancellationToken = default)
-				=> _inner.OpenBodyAsync(offset, length, cancellationToken);
-
-			/// <inheritdoc/>
-			public async ValueTask<BlobData> ReadBlobDataAsync(CancellationToken cancellationToken = default)
-			{
-				IReadOnlyMemoryOwner<byte> data = await _inner.ReadBodyAsync(cancellationToken);
-
-				List<BlobLocator> importLocators = ReadImportsFromData(data.Memory).ToList();
-				List<IBlobHandle> importHandles = importLocators.ConvertAll(x => _storageClient.CreateBlobHandle(x));
-
-				return new BlobDataWithOwner(Bundle.BlobType, data.Memory, importHandles, data);
-			}
-
-			IEnumerable<BlobLocator> ReadImportsFromData(ReadOnlyMemory<byte> data)
-			{
-				BundleSignature signature = Bundle.ReadSignature(data.Span);
-				if (signature.Version <= BundleVersion.LatestV1)
-				{
-					return ReadImportsFromDataV1(data);
-				}
-				else if (signature.Version <= BundleVersion.LatestV2)
-				{
-					return ReadImportsFromDataV2(data);
-				}
-				else
-				{
-					throw new InvalidOperationException($"Unsupported bundle version {(int)signature.Version}");
-				}
-			}
-
-			static IEnumerable<BlobLocator> ReadImportsFromDataV1(ReadOnlyMemory<byte> data)
-			{
-				BundleHeader header = BundleHeader.Read(data);
-				return header.Imports.Select(x => x.BaseLocator);
-			}
-
-			IEnumerable<BlobLocator> ReadImportsFromDataV2(ReadOnlyMemory<byte> data)
-			{
-				HashSet<BlobLocator> locators = new HashSet<BlobLocator>();
-				while (data.Length > 0)
-				{
-					BundleSignature signature = Bundle.ReadSignature(data.Span);
-
-					using IRefCountedHandle<Bundles.V2.Packet> packet = Bundles.V2.Packet.Decode(data, _storageClient._cache.Allocator);
-					for (int idx = 0; idx < packet.Target.GetImportCount(); idx++)
-					{
-						PacketImport import = packet.Target.GetImport(idx);
-						if (import.BaseIdx == -1)
-						{
-							locators.Add(new BlobLocator(import.Fragment.Clone()));
-						}
-					}
-
-					data = data.Slice(signature.HeaderLength);
-				}
-				return locators;
-			}
-
-			/// <inheritdoc/>
-			public bool TryAppendIdentifier(Utf8StringBuilder builder)
-				=> _inner.TryAppendIdentifier(builder);
-
-			/// <inheritdoc/>
-			public IBlobHandle GetFragmentHandle(ReadOnlySpan<byte> fragment)
-			{
-				int exportIdx;
-				if (Utf8Parser.TryParse(fragment, out exportIdx, out int numBytesRead) && numBytesRead == fragment.Length)
-				{
-					return new Bundles.V1.FlushedNodeHandle(_storageClient._bundleReader, _inner.GetLocator(), this, exportIdx);
-				}
-
-				int ampIdx = fragment.IndexOf((byte)'&');
-				if (ampIdx == -1)
-				{
-					return new Bundles.V2.PacketHandle(_storageClient, this, fragment, _storageClient._cache);
-				}
-				else
-				{
-					return new Bundles.V2.PacketHandle(_storageClient, this, fragment.Slice(0, ampIdx), _storageClient._cache).GetFragmentHandle(fragment.Slice(ampIdx + 1));
-				}
-			}
-
-			/// <inheritdoc/>
-			public override bool Equals(object? obj) => obj is BundleHandle other && _inner.Equals(other._inner);
-
-			/// <inheritdoc/>
-			public override int GetHashCode() => _inner.GetHashCode();
-		}
-
 		readonly IStorageClient _inner;
 		readonly BundleCache _cache;
 		readonly Bundles.V1.BundleReader _bundleReader;
 
+		internal Bundles.V1.BundleReader BundleReader => _bundleReader;
+
 		/// <summary>
 		/// Allocator which trims the cache to keep below a maximum size
 		/// </summary>
 		public IMemoryAllocator<byte> Allocator => _cache.Allocator;
 
+		/// <summary>
+		/// Cache for bundle data
+		/// </summary>
+		public BundleCache Cache => _cache;
+
 		/// <inheritdoc/>
 		public bool SupportsRedirects { get; } = false;
 
@@ -254,7 +143,7 @@ namespace EpicGames.Horde.Storage.Bundles
 		{
 			if (locator.TryUnwrap(out BlobLocator baseLocator, out Utf8String fragment))
 			{
-				BundleHandle bundleHandle = new BundleHandle(this, CreateBlobHandle(baseLocator));
+				FlushedBundleHandle bundleHandle = new FlushedBundleHandle(this, CreateBlobHandle(baseLocator));
 
 				int exportIdx;
 				if (Utf8Parser.TryParse(fragment.Span, out exportIdx, out int numBytesRead) && numBytesRead == fragment.Length)
@@ -265,16 +154,16 @@ namespace EpicGames.Horde.Storage.Bundles
 				int ampIdx = fragment.IndexOf('&');
 				if (ampIdx == -1)
 				{
-					return new Bundles.V2.PacketHandle(this, bundleHandle, fragment.Span, _cache);
+					return new Bundles.V2.FlushedPacketHandle(this, bundleHandle, fragment.Span, _cache);
 				}
 				else
 				{
-					return new Bundles.V2.ExportHandle(new Bundles.V2.PacketHandle(this, bundleHandle, fragment.Slice(0, ampIdx).Span, _cache), fragment.Span.Slice(ampIdx + 1));
+					return new Bundles.V2.FlushedExportHandle(new Bundles.V2.FlushedPacketHandle(this, bundleHandle, fragment.Slice(0, ampIdx).Span, _cache), fragment.Span.Slice(ampIdx + 1));
 				}
 			}
 			else
 			{
-				return new BundleHandle(this, _inner.CreateBlobHandle(baseLocator));
+				return new FlushedBundleHandle(this, _inner.CreateBlobHandle(baseLocator));
 			}
 		}
 
diff --git a/Engine/Source/Programs/Shared/EpicGames.Horde/Storage/Bundles/V1/BundleWriter.cs b/Engine/Source/Programs/Shared/EpicGames.Horde/Storage/Bundles/V1/BundleWriter.cs
--- a/Engine/Source/Programs/Shared/EpicGames.Horde/Storage/Bundles/V1/BundleWriter.cs
+++ b/Engine/Source/Programs/Shared/EpicGames.Horde/Storage/Bundles/V1/BundleWriter.cs
@@ -18,7 +18,7 @@ namespace EpicGames.Horde.Storage.Bundles.V1
 	/// <summary>
 	/// Implementation of <see cref="IBlobHandle"/> for nodes which can be read from storage
 	/// </summary>
-	sealed class FlushedNodeHandle : IBlobHandle
+	sealed class FlushedNodeHandle : BlobHandle
 	{
 		readonly BundleReader _reader;
 
@@ -26,7 +26,7 @@ namespace EpicGames.Horde.Storage.Bundles.V1
 		public int ExportIdx { get; }
 
 		/// <inheritdoc/>
-		public IBlobHandle Outer { get; }
+		public override IBlobHandle Outer { get; }
 
 		/// <summary>
 		/// Constructor
@@ -58,17 +58,17 @@ namespace EpicGames.Horde.Storage.Bundles.V1
 		}
 
 		/// <inheritdoc/>
-		public bool TryAppendIdentifier(Utf8StringBuilder builder)
+		public override bool TryAppendIdentifier(Utf8StringBuilder builder)
 		{
 			builder.Append(ExportIdx);
 			return true;
 		}
 
 		/// <inheritdoc/>
-		public ValueTask<BlobData> ReadBlobDataAsync(CancellationToken cancellationToken = default) => _reader.ReadNodeDataAsync(BundleLocator, ExportIdx, cancellationToken);
+		public override ValueTask<BlobData> ReadBlobDataAsync(CancellationToken cancellationToken = default) => _reader.ReadNodeDataAsync(BundleLocator, ExportIdx, cancellationToken);
 
 		/// <inheritdoc/>
-		public ValueTask FlushAsync(CancellationToken cancellationToken = default) => new ValueTask();
+		public override ValueTask FlushAsync(CancellationToken cancellationToken = default) => new ValueTask();
 
 		/// <inheritdoc/>
 		public override bool Equals(object? obj) => obj is FlushedNodeHandle other && BundleLocator == other.BundleLocator && ExportIdx == other.ExportIdx;
@@ -84,7 +84,7 @@ namespace EpicGames.Horde.Storage.Bundles.V1
 	public sealed class BundleWriter : IStorageWriter
 	{
 		// Information about a unique output node. Note that multiple node refs may de-duplicate to the same output node.
-		internal class PendingNode : IBlobHandle
+		internal class PendingNode : BlobHandle
 		{
 			readonly BundleReader _reader;
 
@@ -103,7 +103,7 @@ namespace EpicGames.Horde.Storage.Bundles.V1
 			public PendingBundle? PendingBundle => _pendingBundle;
 			public FlushedNodeHandle? FlushedNodeHandle => _flushedHandle;
 
-			public IBlobHandle? Outer => (_flushedHandle != null) ? _flushedHandle.Outer : throw new NotSupportedException();
+			public override IBlobHandle? Outer => (_flushedHandle != null) ? _flushedHandle.Outer : throw new NotSupportedException();
 
 			public PendingNode(BundleReader reader, BlobType blobType, int packet, int offset, int length, IReadOnlyList<IBlobHandle> refs, IReadOnlyList<AliasInfo> aliases, PendingBundle pendingBundle)
 			{
@@ -120,7 +120,7 @@ namespace EpicGames.Horde.Storage.Bundles.V1
 			}
 
 			/// <inheritdoc/>
-			public bool TryAppendIdentifier(Utf8StringBuilder builder)
+			public override bool TryAppendIdentifier(Utf8StringBuilder builder)
 			{
 				return _flushedHandle?.TryAppendIdentifier(builder) ?? false;
 			}
@@ -136,7 +136,7 @@ namespace EpicGames.Horde.Storage.Bundles.V1
 			}
 
 			/// <inheritdoc/>
-			public async ValueTask<BlobData> ReadBlobDataAsync(CancellationToken cancellationToken = default)
+			public override async ValueTask<BlobData> ReadBlobDataAsync(CancellationToken cancellationToken = default)
 			{
 				if (_flushedHandle == null)
 				{
@@ -154,7 +154,7 @@ namespace EpicGames.Horde.Storage.Bundles.V1
 			}
 
 			/// <inheritdoc/>
-			public async ValueTask FlushAsync(CancellationToken cancellationToken = default)
+			public override async ValueTask FlushAsync(CancellationToken cancellationToken = default)
 			{
 				if (_flushedHandle != null)
 				{
diff --git a/Engine/Source/Programs/Shared/EpicGames.Horde/Storage/Bundles/V2/BundleWriter.cs b/Engine/Source/Programs/Shared/EpicGames.Horde/Storage/Bundles/V2/BundleWriter.cs
--- a/Engine/Source/Programs/Shared/EpicGames.Horde/Storage/Bundles/V2/BundleWriter.cs
+++ b/Engine/Source/Programs/Shared/EpicGames.Horde/Storage/Bundles/V2/BundleWriter.cs
@@ -19,12 +19,12 @@ namespace EpicGames.Horde.Storage.Bundles.V2
 	public sealed class BundleWriter : IStorageWriter
 	{
 		// An export that has been written and is waiting to be flushed to disk
-		internal sealed class PendingExportHandle : IBlobHandle
+		internal sealed class PendingExportHandle : ExportHandle
 		{
 			readonly PendingPacketHandle _packet;
 			readonly int _exportIdx;
 
-			public IBlobHandle? Outer => _packet;
+			public override PacketHandle? OuterPacket => _packet;
 
 			public PendingExportHandle(PendingPacketHandle packet, int exportIdx)
 			{
@@ -32,37 +32,37 @@ namespace EpicGames.Horde.Storage.Bundles.V2
 				_exportIdx = exportIdx;
 			}
 
-			public ValueTask FlushAsync(CancellationToken cancellationToken = default) 
+			public override ValueTask FlushAsync(CancellationToken cancellationToken = default) 
 				=> _packet.FlushAsync(cancellationToken);
 
-			public ValueTask<BlobData> ReadBlobDataAsync(CancellationToken cancellationToken = default)
+			public override ValueTask<BlobData> ReadBlobDataAsync(CancellationToken cancellationToken = default)
 				=> _packet.ReadExportAsync(_exportIdx, cancellationToken);
 
-			public bool TryAppendIdentifier(Utf8StringBuilder builder)
+			public override bool TryAppendIdentifier(Utf8StringBuilder builder)
 			{
-				ExportHandle.AppendIdentifier(builder, _exportIdx);
+				FlushedExportHandle.AppendIdentifier(builder, _exportIdx);
 				return true;
 			}
 		}
 
 		// Packet that is still being built, but may be redirected to a flushed packet
-		internal sealed class PendingPacketHandle : IBlobHandle
+		internal sealed class PendingPacketHandle : PacketHandle
 		{
 			readonly PendingBundleHandle _bundle;
-			PacketHandle? _flushedHandle;
+			FlushedPacketHandle? _flushedHandle;
 
-			public IBlobHandle? Outer => _flushedHandle?.Outer ?? _bundle;
+			public override BundleHandle? OuterBundle => _flushedHandle?.OuterBundle ?? _bundle;
 			public IBlobHandle? FlushedHandle => _flushedHandle;
 
 			public PendingPacketHandle(PendingBundleHandle bundle) => _bundle = bundle;
 
-			public ValueTask FlushAsync(CancellationToken cancellationToken = default) 
+			public override ValueTask FlushAsync(CancellationToken cancellationToken = default) 
 				=> _bundle.FlushAsync(cancellationToken);
 
-			public void CompletePacket(PacketHandle flushedHandle)
+			public void CompletePacket(FlushedPacketHandle flushedHandle)
 				=> _flushedHandle = flushedHandle;
 
-			public ValueTask<BlobData> ReadBlobDataAsync(CancellationToken cancellationToken = default)
+			public override ValueTask<BlobData> ReadBlobDataAsync(CancellationToken cancellationToken = default)
 			{
 				lock (_bundle.LockObject)
 				{
@@ -86,12 +86,12 @@ namespace EpicGames.Horde.Storage.Bundles.V2
 				return await _flushedHandle!.ReadExportAsync(exportIdx, cancellationToken);
 			}
 
-			public bool TryAppendIdentifier(Utf8StringBuilder builder)
+			public override bool TryAppendIdentifier(Utf8StringBuilder builder)
 				=> _flushedHandle?.TryAppendIdentifier(builder) ?? false;
 		}
 
 		// Fragment of a bundle that needs to be written to storage.
-		internal sealed class PendingBundleHandle : IBlobHandle, IDisposable
+		internal sealed class PendingBundleHandle : BundleHandle, IDisposable
 		{
 			readonly object _lockObject = new object();
 
@@ -113,7 +113,7 @@ namespace EpicGames.Horde.Storage.Bundles.V2
 			/// </summary>
 			public object LockObject => _lockObject;
 
-			public IBlobHandle? Outer => null;
+			public override IBlobHandle? Outer => null;
 			public IBlobHandle? FlushedHandle => _flushedHandle;
 
 			/// <summary>
@@ -216,7 +216,7 @@ namespace EpicGames.Horde.Storage.Bundles.V2
 					// Point the packet handle to the encoded data
 					lock (_lockObject)
 					{
-						PacketHandle flushedPacketHandle = new PacketHandle(_storageClient, this, packetOffset, packetLength, _cache);
+						FlushedPacketHandle flushedPacketHandle = new FlushedPacketHandle(_storageClient, this, packetOffset, packetLength, _cache);
 						_packetHandle.CompletePacket(flushedPacketHandle);
 					}
 
@@ -238,7 +238,7 @@ namespace EpicGames.Horde.Storage.Bundles.V2
 			}
 
 			// Write this bundle to storage
-			public async ValueTask FlushAsync(CancellationToken cancellationToken = default)
+			public override async ValueTask FlushAsync(CancellationToken cancellationToken = default)
 			{
 				// Check we haven't already flushed this bundle
 				if (_encodedPacketWriter == null)
@@ -286,11 +286,11 @@ namespace EpicGames.Horde.Storage.Bundles.V2
 			}
 
 			/// <inheritdoc/>
-			public ValueTask<BlobData> ReadBlobDataAsync(CancellationToken cancellationToken)
+			public override ValueTask<BlobData> ReadBlobDataAsync(CancellationToken cancellationToken)
 				=> GetFlushedHandle().ReadBlobDataAsync(cancellationToken);
 
 			/// <inheritdoc/>
-			public async Task<Stream> OpenBodyAsync(int offset, int? length, CancellationToken cancellationToken = default)
+			public override async Task<Stream> OpenBodyAsync(int offset, int? length, CancellationToken cancellationToken = default)
 			{
 				if (_flushedHandle == null)
 				{
@@ -302,7 +302,7 @@ namespace EpicGames.Horde.Storage.Bundles.V2
 			}
 
 			/// <inheritdoc/>
-			public async ValueTask<IReadOnlyMemoryOwner<byte>> ReadBodyAsync(int offset, int? length, CancellationToken cancellationToken = default)
+			public override async ValueTask<IReadOnlyMemoryOwner<byte>> ReadBodyAsync(int offset, int? length, CancellationToken cancellationToken = default)
 			{
 				if (_flushedHandle == null)
 				{
@@ -323,7 +323,7 @@ namespace EpicGames.Horde.Storage.Bundles.V2
 			}
 
 			/// <inheritdoc/>
-			public bool TryAppendIdentifier(Utf8StringBuilder builder)
+			public override bool TryAppendIdentifier(Utf8StringBuilder builder)
 				=> _flushedHandle?.TryAppendIdentifier(builder) ?? false;
 
 			IBlobHandle GetFlushedHandle()
diff --git a/Engine/Source/Programs/Shared/EpicGames.Horde/Storage/Bundles/V2/ExportHandle.cs b/Engine/Source/Programs/Shared/EpicGames.Horde/Storage/Bundles/V2/ExportHandle.cs
--- a/Engine/Source/Programs/Shared/EpicGames.Horde/Storage/Bundles/V2/ExportHandle.cs
+++ b/Engine/Source/Programs/Shared/EpicGames.Horde/Storage/Bundles/V2/ExportHandle.cs
@@ -9,23 +9,35 @@ using EpicGames.Core;
 
 namespace EpicGames.Horde.Storage.Bundles.V2
 {
+	/// <summary>
+	/// Base class for packet handles
+	/// </summary>
+	public abstract class ExportHandle : BlobHandle
+	{
+		/// <inheritdoc/>
+		public sealed override IBlobHandle? Outer => OuterPacket;
+
+		/// <inheritdoc cref="IBlobHandle.Outer"/>
+		public abstract PacketHandle? OuterPacket { get; }
+	}
+
 	/// <summary>
 	/// Handle to an export within a packet. Same implementation is used for flushed and pending exports.
 	/// </summary>
-	public class ExportHandle : IBlobHandle
+	class FlushedExportHandle : ExportHandle
 	{
 		static readonly Utf8String s_fragmentPrefix = new Utf8String("exp=");
 
-		readonly PacketHandle _packet;
+		readonly FlushedPacketHandle _packet;
 		readonly int _exportIdx;
 
 		/// <inheritdoc/>
-		public IBlobHandle? Outer => _packet;
+		public override PacketHandle? OuterPacket => _packet;
 
 		/// <summary>
 		/// Constructor
 		/// </summary>
-		public ExportHandle(PacketHandle packet, int exportIdx)
+		public FlushedExportHandle(FlushedPacketHandle packet, int exportIdx)
 		{
 			_packet = packet;
 			_exportIdx = exportIdx;
@@ -34,7 +46,7 @@ namespace EpicGames.Horde.Storage.Bundles.V2
 		/// <summary>
 		/// Constructor
 		/// </summary>
-		public ExportHandle(PacketHandle packet, ReadOnlySpan<byte> fragment)
+		public FlushedExportHandle(FlushedPacketHandle packet, ReadOnlySpan<byte> fragment)
 		{
 			_packet = packet;
 			if (!TryParse(fragment, out _exportIdx))
@@ -44,7 +56,7 @@ namespace EpicGames.Horde.Storage.Bundles.V2
 		}
 
 		/// <inheritdoc/>
-		public ValueTask FlushAsync(CancellationToken cancellationToken = default) => default;
+		public override ValueTask FlushAsync(CancellationToken cancellationToken = default) => default;
 
 		/// <summary>
 		/// Attempt to parse an export index from the given fragment
@@ -68,11 +80,11 @@ namespace EpicGames.Horde.Storage.Bundles.V2
 		}
 
 		/// <inheritdoc/>
-		public ValueTask<BlobData> ReadBlobDataAsync(CancellationToken cancellationToken = default)
+		public override ValueTask<BlobData> ReadBlobDataAsync(CancellationToken cancellationToken = default)
 			=> _packet.ReadExportAsync(_exportIdx, cancellationToken);
 
 		/// <inheritdoc/>
-		public bool TryAppendIdentifier(Utf8StringBuilder builder)
+		public override bool TryAppendIdentifier(Utf8StringBuilder builder)
 		{
 			AppendIdentifier(builder, _exportIdx);
 			return true;
@@ -89,7 +101,7 @@ namespace EpicGames.Horde.Storage.Bundles.V2
 
 		/// <inheritdoc/>
 		public override bool Equals(object? obj)
-			=> obj is ExportHandle other && _packet.Equals(other._packet) && _exportIdx == other._exportIdx;
+			=> obj is FlushedExportHandle other && _packet.Equals(other._packet) && _exportIdx == other._exportIdx;
 
 		/// <inheritdoc/>
 		public override int GetHashCode()
diff --git a/Engine/Source/Programs/Shared/EpicGames.Horde/Storage/Bundles/V2/PacketHandle.cs b/Engine/Source/Programs/Shared/EpicGames.Horde/Storage/Bundles/V2/PacketHandle.cs
--- a/Engine/Source/Programs/Shared/EpicGames.Horde/Storage/Bundles/V2/PacketHandle.cs
+++ b/Engine/Source/Programs/Shared/EpicGames.Horde/Storage/Bundles/V2/PacketHandle.cs
@@ -11,26 +11,38 @@ using EpicGames.Core;
 
 namespace EpicGames.Horde.Storage.Bundles.V2
 {
+	/// <summary>
+	/// Base class for packet handles
+	/// </summary>
+	public abstract class PacketHandle : BlobHandle
+	{
+		/// <inheritdoc/>
+		public sealed override IBlobHandle? Outer => OuterBundle;
+
+		/// <inheritdoc/>
+		public abstract BundleHandle? OuterBundle { get; }
+	}
+
 	/// <summary>
 	/// Handle to an packet within a bundle. 
 	/// </summary>
-	public class PacketHandle : IBlobHandle
+	class FlushedPacketHandle : PacketHandle
 	{
 		static readonly Utf8String s_fragmentPrefix = new Utf8String("pkt=");
 
 		readonly IStorageClient _storageClient;
-		readonly IBlobHandle _outer;
+		readonly BundleHandle _outer;
 		readonly int _packetOffset;
 		readonly int _packetLength;
 		readonly BundleCache _cache;
 
 		/// <inheritdoc/>
-		public IBlobHandle? Outer => _outer;
+		public override BundleHandle? OuterBundle => _outer;
 
 		/// <summary>
 		/// Constructor
 		/// </summary>
-		public PacketHandle(IStorageClient storageClient, IBlobHandle outer, int packetOffset, int packetLength, BundleCache cache)
+		public FlushedPacketHandle(IStorageClient storageClient, BundleHandle outer, int packetOffset, int packetLength, BundleCache cache)
 		{
 			_storageClient = storageClient;
 
@@ -43,7 +55,7 @@ namespace EpicGames.Horde.Storage.Bundles.V2
 		/// <summary>
 		/// Constructor
 		/// </summary>
-		public PacketHandle(IStorageClient storageClient, IBlobHandle outer, ReadOnlySpan<byte> fragment, BundleCache cache)
+		public FlushedPacketHandle(IStorageClient storageClient, BundleHandle outer, ReadOnlySpan<byte> fragment, BundleCache cache)
 		{
 			_storageClient = storageClient;
 
@@ -92,10 +104,10 @@ namespace EpicGames.Horde.Storage.Bundles.V2
 		}
 
 		/// <inheritdoc/>
-		public ValueTask FlushAsync(CancellationToken cancellationToken = default) => default;
+		public override ValueTask FlushAsync(CancellationToken cancellationToken = default) => default;
 
 		/// <inheritdoc/>
-		public async ValueTask<BlobData> ReadBlobDataAsync(CancellationToken cancellationToken = default)
+		public override async ValueTask<BlobData> ReadBlobDataAsync(CancellationToken cancellationToken = default)
 		{
 			try
 			{
@@ -148,7 +160,7 @@ namespace EpicGames.Horde.Storage.Bundles.V2
 		}
 
 		/// <inheritdoc/>
-		public bool TryAppendIdentifier(Utf8StringBuilder builder)
+		public override bool TryAppendIdentifier(Utf8StringBuilder builder)
 		{
 			AppendIdentifier(builder, _packetOffset, _packetLength);
 			return true;
@@ -166,12 +178,12 @@ namespace EpicGames.Horde.Storage.Bundles.V2
 		}
 
 		/// <inheritdoc/>
-		public IBlobHandle GetFragmentHandle(ReadOnlySpan<byte> fragment)
-			=> new ExportHandle(this, fragment);
+		public override IBlobHandle GetFragmentHandle(ReadOnlySpan<byte> fragment)
+			=> new FlushedExportHandle(this, fragment);
 
 		/// <inheritdoc/>
 		public override bool Equals(object? obj)
-			=> obj is PacketHandle other && _outer.Equals(other._outer) && _packetOffset == other._packetOffset;
+			=> obj is FlushedPacketHandle other && _outer.Equals(other._outer) && _packetOffset == other._packetOffset;
 
 		/// <inheritdoc/>
 		public override int GetHashCode()
diff --git a/Engine/Source/Programs/Shared/EpicGames.Horde/Storage/Bundles/V2/PacketReader.cs b/Engine/Source/Programs/Shared/EpicGames.Horde/Storage/Bundles/V2/PacketReader.cs
--- a/Engine/Source/Programs/Shared/EpicGames.Horde/Storage/Bundles/V2/PacketReader.cs
+++ b/Engine/Source/Programs/Shared/EpicGames.Horde/Storage/Bundles/V2/PacketReader.cs
@@ -8,12 +8,12 @@ namespace EpicGames.Horde.Storage.Bundles.V2
 	/// <summary>
 	/// Utility class for constructing BlobData objects from a packet, caching any computed handles to other blobs.
 	/// </summary>
-	public sealed class PacketReader : IDisposable
+	sealed class PacketReader : IDisposable
 	{
 		readonly IStorageClient _storageClient;
 		readonly BundleCache _cache;
-		readonly IBlobHandle _bundleHandle;
-		readonly PacketHandle _packetHandle;
+		readonly BundleHandle _bundleHandle;
+		readonly FlushedPacketHandle _packetHandle;
 
 		Packet _decodedPacket;
 		IBlobHandle?[] _cachedImportHandles;
@@ -33,7 +33,7 @@ namespace EpicGames.Horde.Storage.Bundles.V2
 		/// <param name="packetHandle"></param>
 		/// <param name="decodedPacket">Data for the packet</param>
 		/// <param name="memoryOwner">Owner for the packet data</param>
-		public PacketReader(IStorageClient storageClient, BundleCache cache, IBlobHandle bundleHandle, PacketHandle packetHandle, Packet decodedPacket, IRefCountedHandle memoryOwner)
+		public PacketReader(IStorageClient storageClient, BundleCache cache, BundleHandle bundleHandle, FlushedPacketHandle packetHandle, Packet decodedPacket, IRefCountedHandle memoryOwner)
 		{
 			_storageClient = storageClient;
 			_cache = cache;
@@ -124,10 +124,10 @@ namespace EpicGames.Horde.Storage.Bundles.V2
 						importHandle = _storageClient.CreateBlobHandle(new BlobLocator(fragment));
 						break;
 					case PacketImport.CurrentBundleBaseIdx:
-						importHandle = new PacketHandle(_storageClient, _bundleHandle, fragment, _cache);
+						importHandle = new FlushedPacketHandle(_storageClient, _bundleHandle, fragment, _cache);
 						break;
 					case PacketImport.CurrentPacketBaseIdx:
-						importHandle = new ExportHandle(_packetHandle, fragment);
+						importHandle = new FlushedExportHandle(_packetHandle, fragment);
 						break;
 					default:
 						importHandle = GetImportHandle(import.BaseIdx).GetFragmentHandle(fragment);
diff --git a/Engine/Source/Programs/Shared/EpicGames.Horde/Storage/Clients/HttpStorageClient.cs b/Engine/Source/Programs/Shared/EpicGames.Horde/Storage/Clients/HttpStorageClient.cs
--- a/Engine/Source/Programs/Shared/EpicGames.Horde/Storage/Clients/HttpStorageClient.cs
+++ b/Engine/Source/Programs/Shared/EpicGames.Horde/Storage/Clients/HttpStorageClient.cs
@@ -22,13 +22,13 @@ namespace EpicGames.Horde.Storage.Clients
 	/// </summary>
 	public sealed class HttpStorageClient : IStorageClient
 	{
-		class Handle : IBlobHandle
+		class Handle : BlobHandle
 		{
 			readonly HttpStorageClient _httpStorageClient;
 			readonly BlobLocator _locator;
 
 			/// <inheritdoc/>
-			public IBlobHandle? Outer => null;
+			public override IBlobHandle? Outer => null;
 
 			public Handle(HttpStorageClient httpStorageClient, BlobLocator locator)
 			{
@@ -37,18 +37,18 @@ namespace EpicGames.Horde.Storage.Clients
 			}
 
 			/// <inheritdoc/>
-			public ValueTask FlushAsync(CancellationToken cancellationToken = default) => default;
+			public override ValueTask FlushAsync(CancellationToken cancellationToken = default) => default;
 
 			/// <inheritdoc/>
-			public Task<Stream> OpenBodyAsync(int offset = 0, int? length = null, CancellationToken cancellationToken = default)
+			public override Task<Stream> OpenBodyAsync(int offset = 0, int? length = null, CancellationToken cancellationToken = default)
 				=> _httpStorageClient._backend.OpenAsync(_locator.ToString(), offset, length, cancellationToken);
 
 			/// <inheritdoc/>
-			public ValueTask<BlobData> ReadBlobDataAsync(CancellationToken cancellationToken = default)
+			public override ValueTask<BlobData> ReadBlobDataAsync(CancellationToken cancellationToken = default)
 				=> _httpStorageClient.ReadBlobAsync(_locator, cancellationToken);
 
 			/// <inheritdoc/>
-			public bool TryAppendIdentifier(Utf8StringBuilder builder)
+			public override bool TryAppendIdentifier(Utf8StringBuilder builder)
 			{
 				builder.Append(_locator.Path);
 				return true;
diff --git a/Engine/Source/Programs/Shared/EpicGames.Horde/Storage/Clients/KeyValueStorageClient.cs b/Engine/Source/Programs/Shared/EpicGames.Horde/Storage/Clients/KeyValueStorageClient.cs
--- a/Engine/Source/Programs/Shared/EpicGames.Horde/Storage/Clients/KeyValueStorageClient.cs
+++ b/Engine/Source/Programs/Shared/EpicGames.Horde/Storage/Clients/KeyValueStorageClient.cs
@@ -14,13 +14,13 @@ namespace EpicGames.Horde.Storage.Clients
 	/// </summary>
 	public abstract class KeyValueStorageClient : IStorageClient
 	{
-		class Handle : IBlobHandle
+		class Handle : BlobHandle
 		{
 			readonly KeyValueStorageClient _keyValueStorageClient;
 			readonly BlobLocator _locator;
 
 			/// <inheritdoc/>
-			public IBlobHandle? Outer => null;
+			public override IBlobHandle? Outer => null;
 
 			/// <summary>
 			/// Constructor
@@ -32,13 +32,13 @@ namespace EpicGames.Horde.Storage.Clients
 			}
 
 			/// <inheritdoc/>
-			public ValueTask FlushAsync(CancellationToken cancellationToken = default) => default;
+			public override ValueTask FlushAsync(CancellationToken cancellationToken = default) => default;
 
 			/// <inheritdoc/>
-			public ValueTask<BlobData> ReadBlobDataAsync(CancellationToken cancellationToken = default) => _keyValueStorageClient.ReadBlobAsync(_locator, cancellationToken);
+			public override ValueTask<BlobData> ReadBlobDataAsync(CancellationToken cancellationToken = default) => _keyValueStorageClient.ReadBlobAsync(_locator, cancellationToken);
 
 			/// <inheritdoc/>
-			public bool TryAppendIdentifier(Utf8StringBuilder builder)
+			public override bool TryAppendIdentifier(Utf8StringBuilder builder)
 			{
 				builder.Append(_locator.Path);
 				return true;
diff --git a/Engine/Source/Programs/Shared/EpicGames.Horde/Storage/IBlobHandle.cs b/Engine/Source/Programs/Shared/EpicGames.Horde/Storage/IBlobHandle.cs
--- a/Engine/Source/Programs/Shared/EpicGames.Horde/Storage/IBlobHandle.cs
+++ b/Engine/Source/Programs/Shared/EpicGames.Horde/Storage/IBlobHandle.cs
@@ -14,5 +14,76 @@
 	/// </summary>
 	public interface IBlobHandle
 	{
+		/// <summary>
+		/// For a blob nested within another blob, gets a handle to the containing blob (eg. For a bundle node, will return the packet. For a bundle packet, will return the bundle. For a bundle or other non-nested blob, returns null.)
+		/// </summary>
+		IBlobHandle? Outer { get; }
+
+		/// <summary>
+		/// Flush the referenced data to underlying storage
+		/// </summary>
+		ValueTask FlushAsync(CancellationToken cancellationToken = default);
+
+		/// <summary>
+		/// Open the blob's data stream
+		/// </summary>
+		/// <param name="offset">Start offset of the stream</param>
+		/// <param name="length">Length of the stream</param>
+		/// <param name="cancellationToken">Cancellation token for the operation</param>
+		Task<Stream> OpenBodyAsync(int offset = 0, int? length = null, CancellationToken cancellationToken = default);
+
+		/// <summary>
+		/// Reads the blob's data
+		/// </summary>
+		/// <param name="cancellationToken">Cancellation token for the operation</param>
+		ValueTask<BlobData> ReadBlobDataAsync(CancellationToken cancellationToken = default);
+
+		/// <summary>
+		/// Reads part of the blob, and returns a handle that can be used to access the data.
+		/// </summary>
+		/// <param name="cancellationToken">Cancellation token for the operation</param>
+		/// <returns></returns>
+		ValueTask<IReadOnlyMemoryOwner<byte>> ReadBodyAsync(CancellationToken cancellationToken = default);
+
+		/// <summary>
+		/// Reads part of the blob, and returns a handle that can be used to access the data.
+		/// </summary>
+		/// <param name="offset">Offset of the data within the blob</param>
+		/// <param name="length">Length of the data to read</param>
+		/// <param name="cancellationToken">Cancellation token for the operation</param>
+		/// <returns></returns>
+		ValueTask<IReadOnlyMemoryOwner<byte>> ReadBodyAsync(int offset, int? length, CancellationToken cancellationToken = default);
+
+		/// <summary>
+		/// Gets an identifier for this blob, relative to its outer
+		/// </summary>
+		/// <param name="builder">Builder for appending the identifier to</param>
+		/// <returns>True if an identifier was returned, false otherwise</returns>
+		bool TryAppendIdentifier(Utf8StringBuilder builder);
+
+		/// <summary>
+		/// Gets a handle to a nested blob object
+		/// </summary>
+		/// <param name="fragment">Name of the blob fragment</param>
+		public IBlobHandle GetFragmentHandle(ReadOnlySpan<byte> fragment);
+	}
+
+	/// <summary>
+	/// Typed interface to a particular blob handle
+	/// </summary>
+	/// <typeparam name="T">Type of the deserialized blob</typeparam>
+	public interface IBlobHandle<out T> : IBlobHandle
+	{
+		/// <summary>
+		/// Hash of the target node
+		/// </summary>
+		IoHash Hash { get; }
+	}
+
+	/// <summary>
+	/// Base class for blob handles
+	/// </summary>
+	public abstract class BlobHandle : IBlobHandle
+	{
 		private class BlobDataFragment : IReadOnlyMemoryOwner<byte>
 		{
@@ -57,23 +128,14 @@
 			}
 		}
 
-		/// <summary>
-		/// For a blob nested within another blob, gets a handle to the containing blob (eg. For a bundle node, will return the packet. For a bundle packet, will return the bundle. For a bundle or other non-nested blob, returns null.)
-		/// </summary>
-		IBlobHandle? Outer { get; }
+		/// <inheritdoc/>
+		public abstract IBlobHandle? Outer { get; }
 
-		/// <summary>
-		/// Flush the referenced data to underlying storage
-		/// </summary>
-		ValueTask FlushAsync(CancellationToken cancellationToken = default);
+		/// <inheritdoc/>
+		public abstract ValueTask FlushAsync(CancellationToken cancellationToken = default);
 
-		/// <summary>
-		/// Open the blob's data stream
-		/// </summary>
-		/// <param name="offset">Start offset of the stream</param>
-		/// <param name="length">Length of the stream</param>
-		/// <param name="cancellationToken">Cancellation token for the operation</param>
-		async Task<Stream> OpenBodyAsync(int offset = 0, int? length = null, CancellationToken cancellationToken = default)
+		/// <inheritdoc/>
+		public virtual async Task<Stream> OpenBodyAsync(int offset = 0, int? length = null, CancellationToken cancellationToken = default)
 		{
 			BlobData blobData = await ReadBlobDataAsync(cancellationToken);
 			int maxLength = blobData.Data.Length - offset;
@@ -81,71 +143,39 @@
 			return new BlobDataStream(blobData, memory);
 		}
 
-		/// <summary>
-		/// Reads the blob's data
-		/// </summary>
-		/// <param name="cancellationToken">Cancellation token for the operation</param>
-		ValueTask<BlobData> ReadBlobDataAsync(CancellationToken cancellationToken = default);
+		/// <inheritdoc/>
+		public abstract ValueTask<BlobData> ReadBlobDataAsync(CancellationToken cancellationToken = default);
 
-		/// <summary>
-		/// Reads part of the blob, and returns a handle that can be used to access the data.
-		/// </summary>
-		/// <param name="cancellationToken">Cancellation token for the operation</param>
-		/// <returns></returns>
-		async ValueTask<IReadOnlyMemoryOwner<byte>> ReadBodyAsync(CancellationToken cancellationToken = default)
+		/// <inheritdoc/>
+		public virtual async ValueTask<IReadOnlyMemoryOwner<byte>> ReadBodyAsync(CancellationToken cancellationToken = default)
 		{
 			return await ReadBodyAsync(0, null, cancellationToken);
 		}
 
-		/// <summary>
-		/// Reads part of the blob, and returns a handle that can be used to access the data.
-		/// </summary>
-		/// <param name="offset">Offset of the data within the blob</param>
-		/// <param name="length">Length of the data to read</param>
-		/// <param name="cancellationToken">Cancellation token for the operation</param>
-		/// <returns></returns>
-		async ValueTask<IReadOnlyMemoryOwner<byte>> ReadBodyAsync(int offset, int? length, CancellationToken cancellationToken = default)
+		/// <inheritdoc/>
+		public virtual async ValueTask<IReadOnlyMemoryOwner<byte>> ReadBodyAsync(int offset, int? length, CancellationToken cancellationToken = default)
 		{
 			BlobData data = await ReadBlobDataAsync(cancellationToken);
 			return new BlobDataFragment(data, offset, length);
 		}
 
-		/// <summary>
-		/// Gets an identifier for this blob, relative to its outer
-		/// </summary>
-		/// <param name="builder">Builder for appending the identifier to</param>
-		/// <returns>True if an identifier was returned, false otherwise</returns>
-		bool TryAppendIdentifier(Utf8StringBuilder builder);
+		/// <inheritdoc/>
+		public abstract bool TryAppendIdentifier(Utf8StringBuilder builder);
 
-		/// <summary>
-		/// Gets a handle to a nested blob object
-		/// </summary>
-		/// <param name="fragment">Name of the blob fragment</param>
-		public IBlobHandle GetFragmentHandle(ReadOnlySpan<byte> fragment)
+		/// <inheritdoc/>
+		public virtual IBlobHandle GetFragmentHandle(ReadOnlySpan<byte> fragment)
 			=> throw new InvalidOperationException("Not supported for this handle type.");
 	}
 
 	/// <summary>
-	/// Typed interface to a particular blob handle
-	/// </summary>
-	/// <typeparam name="T">Type of the deserialized blob</typeparam>
-	public interface IBlobHandle<out T> : IBlobHandle
-	{
-		/// <summary>
-		/// Hash of the target node
-		/// </summary>
-		IoHash Hash { get; }
-	}
-
-	/// <summary>
 	/// Instance of <see cref="IBlobHandle"/> which wraps an inner handle and a fragment
 	/// </summary>
-	public class BlobFragmentHandle : IBlobHandle
+	public class BlobFragmentHandle : BlobHandle
 	{
 		/// <summary>
 		/// Handle to the outer blob
 		/// </summary>
-		public IBlobHandle Outer { get; }
+		public override IBlobHandle Outer { get; }
 
 		/// <summary>
 		/// The fragment portion of the handle
@@ -162,17 +192,17 @@
 		}
 
 		/// <inheritdoc/>
-		public ValueTask FlushAsync(CancellationToken cancellationToken = default) => default;
+		public override ValueTask FlushAsync(CancellationToken cancellationToken = default) => default;
 
 		/// <inheritdoc/>
-		public bool TryAppendIdentifier(Utf8StringBuilder builder)
+		public override bool TryAppendIdentifier(Utf8StringBuilder builder)
 		{
 			builder.Append(Fragment);
 			return true;
 		}
 
 		/// <inheritdoc/>
-		public ValueTask<BlobData> ReadBlobDataAsync(CancellationToken cancellationToken = default)
+		public override ValueTask<BlobData> ReadBlobDataAsync(CancellationToken cancellationToken = default)
 		{
 			throw new NotSupportedException("Blob fragment handles cannot be read directly, and should be deconstructed into more specific types.");
 		}
@@ -207,11 +237,14 @@
 
 			public IBlobHandle Inner => _inner;
 			public IoHash Hash => _hash;
-
 			public IBlobHandle? Outer => _inner.Outer;
 
 			public ValueTask FlushAsync(CancellationToken cancellationToken = default) => _inner.FlushAsync(cancellationToken);
+			public IBlobHandle GetFragmentHandle(ReadOnlySpan<byte> fragment) => _inner.GetFragmentHandle(fragment);
+			public Task<Stream> OpenBodyAsync(int offset = 0, int? length = null, CancellationToken cancellationToken = default) => _inner.OpenBodyAsync(offset, length, cancellationToken);
 			public ValueTask<BlobData> ReadBlobDataAsync(CancellationToken cancellationToken = default) => _inner.ReadBlobDataAsync(cancellationToken);
+			public ValueTask<IReadOnlyMemoryOwner<byte>> ReadBodyAsync(CancellationToken cancellationToken = default) => _inner.ReadBodyAsync(cancellationToken);
+			public ValueTask<IReadOnlyMemoryOwner<byte>> ReadBodyAsync(int offset, int? length, CancellationToken cancellationToken = default) => _inner.ReadBodyAsync(offset, length, cancellationToken);
 			public bool TryAppendIdentifier(Utf8StringBuilder builder) => _inner.TryAppendIdentifier(builder);
 		}
 
diff --git a/Engine/Source/Programs/Shared/EpicGames.Horde/Storage/IStorageWriter.cs b/Engine/Source/Programs/Shared/EpicGames.Horde/Storage/IStorageWriter.cs
--- a/Engine/Source/Programs/Shared/EpicGames.Horde/Storage/IStorageWriter.cs
+++ b/Engine/Source/Programs/Shared/EpicGames.Horde/Storage/IStorageWriter.cs
@@ -157,22 +157,22 @@ namespace EpicGames.Horde.Storage
 			internal bool TryGetValue(BlobKey key, [NotNullWhen(true)] out IBlobHandle? handle) => _blobKeyToHandle.TryGetValue(key, out handle);
 		}
 
-		class WrappedHandle : IBlobHandle
+		class WrappedHandle : BlobHandle
 		{
 			public object _lockObject = new object();
 			public IBlobHandle? _inner;
 
 			/// <inheritdoc/>
-			public IBlobHandle? Outer => _inner?.Outer;
+			public override IBlobHandle? Outer => _inner?.Outer;
 
 			/// <inheritdoc/>
-			public bool TryAppendIdentifier(Utf8StringBuilder builder)
+			public override bool TryAppendIdentifier(Utf8StringBuilder builder)
 			{
 				return _inner?.TryAppendIdentifier(builder) ?? false;
 			}
 
 			/// <inheritdoc/>
-			public ValueTask FlushAsync(CancellationToken cancellationToken)
+			public override ValueTask FlushAsync(CancellationToken cancellationToken)
 			{
 				if (_inner == null)
 				{
@@ -185,7 +185,7 @@ namespace EpicGames.Horde.Storage
 			}
 
 			/// <inheritdoc/>
-			public ValueTask<BlobData> ReadBlobDataAsync(CancellationToken cancellationToken = default)
+			public override ValueTask<BlobData> ReadBlobDataAsync(CancellationToken cancellationToken = default)
 			{
 				if (_inner == null)
 				{
