Horde: Include IP and port in exception for failed compute client connect attempts

diff --git a/Engine/Source/Programs/Shared/EpicGames.Horde/Compute/Clients/ServerComputeClient.cs b/Engine/Source/Programs/Shared/EpicGames.Horde/Compute/Clients/ServerComputeClient.cs
--- a/Engine/Source/Programs/Shared/EpicGames.Horde/Compute/Clients/ServerComputeClient.cs
+++ b/Engine/Source/Programs/Shared/EpicGames.Horde/Compute/Clients/ServerComputeClient.cs
@@ -87,6 +87,22 @@ namespace EpicGames.Horde.Compute.Clients
 			return new TunnelHandshakeResponse(isSuccess, parts[3]);
 		}
 	}
+
+	/// <summary>
+	/// Exception for ServerComputeClient
+	/// </summary>
+	public class ServerComputeClientException : ComputeException
+	{
+		/// <inheritdoc/>
+		public ServerComputeClientException(string message) : base(message)
+		{
+		}
+
+		/// <inheritdoc/>
+		public ServerComputeClientException(string? message, Exception? innerException) : base(message, innerException)
+		{
+		}
+	}
 	
 	/// <summary>
 	/// Helper class to enlist remote resources to perform compute-intensive tasks.
@@ -283,25 +299,32 @@ namespace EpicGames.Horde.Compute.Clients
 			
 			workerLogger.LogDebug("Connecting to {AgentId} at {AgentAddress} ({ConnectionType} via {ConnectionAddress}) with nonce {Nonce} and encryption {Encryption}...",
 				response.AgentId, agentAddress, response.ConnectionMode, response.ConnectionAddress ?? "None", response.Nonce, response.Encryption);
-			switch (response.ConnectionMode)
+			try
 			{
-				case ConnectionMode.Direct:
-					await socket.ConnectAsync(IPAddress.Parse(response.Ip), response.Port, cancellationToken);
-					break;
-
-				case ConnectionMode.Tunnel when !String.IsNullOrEmpty(response.ConnectionAddress):
-					(string host, int port) = ParseHostPort(response.ConnectionAddress);
-					await socket.ConnectAsync(host, port, cancellationToken);
-					await TunnelHandshakeAsync(socket, response, cancellationToken);
-					break;
+				switch (response.ConnectionMode)
+				{
+					case ConnectionMode.Direct:
+						await socket.ConnectAsync(IPAddress.Parse(response.Ip), response.Port, cancellationToken);
+						break;
+
+					case ConnectionMode.Tunnel when !String.IsNullOrEmpty(response.ConnectionAddress):
+						(string host, int port) = ParseHostPort(response.ConnectionAddress);
+						await socket.ConnectAsync(host, port, cancellationToken);
+						await TunnelHandshakeAsync(socket, response, cancellationToken);
+						break;
 				
-				case ConnectionMode.Relay when !String.IsNullOrEmpty(response.ConnectionAddress):
-					response.Ip = response.ConnectionAddress;
-					await socket.ConnectAsync(IPAddress.Parse(response.ConnectionAddress), response.Ports[ConnectionMetadataPort.ComputeId].Port, cancellationToken);
-					break;
+					case ConnectionMode.Relay when !String.IsNullOrEmpty(response.ConnectionAddress):
+						response.Ip = response.ConnectionAddress;
+						await socket.ConnectAsync(IPAddress.Parse(response.ConnectionAddress), response.Ports[ConnectionMetadataPort.ComputeId].Port, cancellationToken);
+						break;
 				
-				default:
-					throw new Exception($"Unable to resolve connection mode ({response.ConnectionMode} via {response.ConnectionAddress ?? "none"})");
+					default:
+						throw new Exception($"Unable to resolve connection mode ({response.ConnectionMode} via {response.ConnectionAddress ?? "none"})");
+				}
+			}
+			catch (SocketException se)
+			{
+				throw new ServerComputeClientException($"Unable to connect to {agentAddress}", se);
 			}
 
 			// Send the nonce
diff --git a/Engine/Source/Programs/Shared/EpicGames.Horde/Compute/ComputeExceptions.cs b/Engine/Source/Programs/Shared/EpicGames.Horde/Compute/ComputeExceptions.cs
--- a/Engine/Source/Programs/Shared/EpicGames.Horde/Compute/ComputeExceptions.cs
+++ b/Engine/Source/Programs/Shared/EpicGames.Horde/Compute/ComputeExceptions.cs
@@ -9,11 +9,13 @@ namespace EpicGames.Horde.Compute
 	/// </summary>
 	public abstract class ComputeException : Exception
 	{
-		/// <summary>
-		/// Constructor
-		/// </summary>
-		protected ComputeException(string message)
-			: base(message)
+		/// <inheritdoc/>
+		protected ComputeException(string message) : base(message)
+		{
+		}
+
+		/// <inheritdoc/>
+		protected ComputeException(string? message, Exception? innerException) : base(message, innerException)
 		{
 		}
 	}
