Horde: Remove methods from IStorageBackend which are no longer required with the IObjectStore refactor.

* Storage backends now function in terms of BlobLocators rather than string paths.
* Object keys now typically have a ".blob" suffix for blobs, which is appended during conversion from BlobLocator to ObjectKey.

diff --git a/Engine/Source/Programs/Horde/Horde.Server.Tests/GcServiceTests.cs b/Engine/Source/Programs/Horde/Horde.Server.Tests/GcServiceTests.cs
--- a/Engine/Source/Programs/Horde/Horde.Server.Tests/GcServiceTests.cs
+++ b/Engine/Source/Programs/Horde/Horde.Server.Tests/GcServiceTests.cs
@@ -52,7 +52,7 @@ namespace Horde.Server.Tests
 			ObjectKey[] remaining = await backend.EnumerateAsync().ToArrayAsync();
 			Assert.AreEqual(nodes.Count, remaining.Length);
 
-			HashSet<ObjectKey> nodePaths = new HashSet<ObjectKey>(nodes.Select(x => new ObjectKey(x.BaseLocator.ToString())));
+			HashSet<ObjectKey> nodePaths = new HashSet<ObjectKey>(nodes.Select(x => StorageService.GetObjectKey(x.BaseLocator)));
 			Assert.IsTrue(remaining.All(x => nodePaths.Contains(x)));
 		}
 
diff --git a/Engine/Source/Programs/Horde/Horde.Server/Replicators/ReplicatorsController.cs b/Engine/Source/Programs/Horde/Horde.Server/Replicators/ReplicatorsController.cs
--- a/Engine/Source/Programs/Horde/Horde.Server/Replicators/ReplicatorsController.cs
+++ b/Engine/Source/Programs/Horde/Horde.Server/Replicators/ReplicatorsController.cs
@@ -193,7 +193,7 @@ namespace Horde.Server.Replicators
 			{
 				double copiedMb = (replicator.CurrentCopiedSize ?? 0) / (1024.0 * 1024.0);
 				double totalMb = replicator.CurrentSize.Value / (1024.0 * 1024.0);
-				message.Append($" ({copiedMb:n1}/{totalMb:n1}mb, {(copiedMb * 100.0) / totalMb:n1}%)");
+				message.Append($" ({copiedMb:n1}mb/{totalMb:n1}mb, {(copiedMb * 100.0) / totalMb:n1}%)");
 			}
 
 			return message.ToString();
diff --git a/Engine/Source/Programs/Horde/Horde.Server/Storage/ObjectStoreFactory.cs b/Engine/Source/Programs/Horde/Horde.Server/Storage/ObjectStoreFactory.cs
--- a/Engine/Source/Programs/Horde/Horde.Server/Storage/ObjectStoreFactory.cs
+++ b/Engine/Source/Programs/Horde/Horde.Server/Storage/ObjectStoreFactory.cs
@@ -21,88 +21,88 @@ namespace Horde.Server.Storage
 	/// </summary>
 	class ObjectStoreFactory : IObjectStoreFactory
 	{
-		class RefCountedBackend : IDisposable
+		class RefCountedObjectStore : IDisposable
 		{
 			public BackendId Id { get; }
 			public IoHash Hash { get; }
-			public IObjectStore Backend { get; }
+			public IObjectStore Inner { get; }
 
 			public int _refCount = 1;
 
-			public RefCountedBackend(BackendId id, IoHash hash, IObjectStore backend)
+			public RefCountedObjectStore(BackendId id, IoHash hash, IObjectStore inner)
 			{
 				Id = id;
 				Hash = hash;
-				Backend = backend;
+				Inner = inner;
 			}
 
 			public void Dispose()
 			{
-				Backend.Dispose();
+				Inner.Dispose();
 			}
 		}
 
-		class BackendWrapper : IObjectStore
+		class ObjectStoreWrapper : IObjectStore
 		{
 			readonly ObjectStoreFactory _owner;
-			RefCountedBackend _refCountedBackend;
-			IObjectStore _backend;
+			RefCountedObjectStore _refCountedObjectStore;
+			IObjectStore _inner;
 
-			public BackendWrapper(ObjectStoreFactory owner, RefCountedBackend refCountedBackend)
+			public ObjectStoreWrapper(ObjectStoreFactory owner, RefCountedObjectStore refCountedObjectStore)
 			{
 				_owner = owner;
-				_refCountedBackend = refCountedBackend;
-				_backend = _refCountedBackend.Backend;
+				_refCountedObjectStore = refCountedObjectStore;
+				_inner = _refCountedObjectStore.Inner;
 			}
 
 			public void Dispose()
 			{
-				if (_refCountedBackend != null)
+				if (_refCountedObjectStore != null)
 				{
-					_owner.ReleaseBackend(_refCountedBackend);
-					_refCountedBackend = null!;
+					_owner.ReleaseBackend(_refCountedObjectStore);
+					_refCountedObjectStore = null!;
 
-					_backend = null!;
+					_inner = null!;
 				}
 			}
 
 			#region IObjectStore Implementation
 
 			/// <inheritdoc/>
-			public bool SupportsRedirects => _backend.SupportsRedirects;
+			public bool SupportsRedirects => _inner.SupportsRedirects;
 
 			/// <inheritdoc/>
-			public Task<bool> ExistsAsync(ObjectKey key, CancellationToken cancellationToken = default) => _backend.ExistsAsync(key, cancellationToken);
+			public Task<bool> ExistsAsync(ObjectKey key, CancellationToken cancellationToken = default) => _inner.ExistsAsync(key, cancellationToken);
 
 			/// <inheritdoc/>
-			public Task DeleteAsync(ObjectKey key, CancellationToken cancellationToken = default) => _backend.DeleteAsync(key, cancellationToken);
+			public Task DeleteAsync(ObjectKey key, CancellationToken cancellationToken = default) => _inner.DeleteAsync(key, cancellationToken);
 
 			/// <inheritdoc/>
-			public IAsyncEnumerable<ObjectKey> EnumerateAsync(CancellationToken cancellationToken = default) => _backend.EnumerateAsync(cancellationToken);
+			public IAsyncEnumerable<ObjectKey> EnumerateAsync(CancellationToken cancellationToken = default) => _inner.EnumerateAsync(cancellationToken);
 
 			/// <inheritdoc/>
 			public Task<Stream> OpenAsync(ObjectKey key, int offset, int? length, CancellationToken cancellationToken = default)
 			{
-				return _backend?.OpenAsync(key, offset, length, cancellationToken) ?? throw new InvalidOperationException("Backend has already been disposed");
+				return _inner?.OpenAsync(key, offset, length, cancellationToken) ?? throw new InvalidOperationException("Backend has already been disposed");
 			}
 
 			/// <inheritdoc/>
 			public Task<IReadOnlyMemoryOwner<byte>> ReadAsync(ObjectKey key, int offset, int? length, CancellationToken cancellationToken = default)
 			{
-				return _backend?.ReadAsync(key, offset, length, cancellationToken) ?? throw new InvalidOperationException("Backend has already been disposed");
+				return _inner?.ReadAsync(key, offset, length, cancellationToken) ?? throw new InvalidOperationException("Backend has already been disposed");
 			}
 
 			/// <inheritdoc/>
-			public Task WriteAsync(ObjectKey key, Stream stream, CancellationToken cancellationToken = default) => _backend.WriteAsync(key, stream, cancellationToken);
+			public Task WriteAsync(ObjectKey key, Stream stream, CancellationToken cancellationToken = default) => _inner.WriteAsync(key, stream, cancellationToken);
 
 			/// <inheritdoc/>
-			public ValueTask<Uri?> TryGetReadRedirectAsync(ObjectKey key, CancellationToken cancellationToken = default) => _backend.TryGetReadRedirectAsync(key, cancellationToken);
+			public ValueTask<Uri?> TryGetReadRedirectAsync(ObjectKey key, CancellationToken cancellationToken = default) => _inner.TryGetReadRedirectAsync(key, cancellationToken);
 
 			/// <inheritdoc/>
-			public ValueTask<Uri?> TryGetWriteRedirectAsync(ObjectKey key, CancellationToken cancellationToken = default) => _backend.TryGetWriteRedirectAsync(key, cancellationToken);
+			public ValueTask<Uri?> TryGetWriteRedirectAsync(ObjectKey key, CancellationToken cancellationToken = default) => _inner.TryGetWriteRedirectAsync(key, cancellationToken);
 
 			/// <inheritdoc/>
-			public void GetStats(StorageStats stats) => _backend.GetStats(stats);
+			public void GetStats(StorageStats stats) => _inner.GetStats(stats);
 
 			#endregion
 		}
@@ -110,7 +110,7 @@ namespace Horde.Server.Storage
 		readonly IServiceProvider _serviceProvider;
 		readonly StorageBackendCache _storageBackendCache;
 		readonly object _lockObject = new object();
-		readonly Dictionary<IoHash, RefCountedBackend> _backends = new Dictionary<IoHash, RefCountedBackend>();
+		readonly Dictionary<IoHash, RefCountedObjectStore> _objectStores = new Dictionary<IoHash, RefCountedObjectStore>();
 		readonly ILogger _logger;
 
 		/// <summary>
@@ -139,10 +139,10 @@ namespace Horde.Server.Storage
 			}
 
 			// See if we've got an existing backend we can use
-			RefCountedBackend? refCountedBackend;
+			RefCountedObjectStore? refCountedBackend;
 			lock (_lockObject)
 			{
-				if (_backends.TryGetValue(hash, out refCountedBackend))
+				if (_objectStores.TryGetValue(hash, out refCountedBackend))
 				{
 					refCountedBackend._refCount++;
 					_logger.LogDebug("Adding reference to storage backend {Id}@{Hash}", refCountedBackend.Id, hash);
@@ -150,16 +150,16 @@ namespace Horde.Server.Storage
 				else
 				{
 					IObjectStore newBackend = CreateStorageBackend(config);
-					refCountedBackend = new RefCountedBackend(config.Id, hash, newBackend);
-					_backends.Add(hash, refCountedBackend);
+					refCountedBackend = new RefCountedObjectStore(config.Id, hash, newBackend);
+					_objectStores.Add(hash, refCountedBackend);
 					_logger.LogInformation("Created storage backend {Id}@{Hash}", refCountedBackend.Id, hash);
 				}
 			}
 
-			return new BackendWrapper(this, refCountedBackend);
+			return new ObjectStoreWrapper(this, refCountedBackend);
 		}
 
-		void ReleaseBackend(RefCountedBackend backend)
+		void ReleaseBackend(RefCountedObjectStore backend)
 		{
 			lock (_lockObject)
 			{
@@ -167,7 +167,7 @@ namespace Horde.Server.Storage
 
 				if (--backend._refCount == 0)
 				{
-					_backends.Remove(backend.Hash);
+					_objectStores.Remove(backend.Hash);
 					backend.Dispose();
 					_logger.LogInformation("Disposed storage backend {Id}@{Hash}", backend.Id, backend.Hash);
 				}
diff --git a/Engine/Source/Programs/Horde/Horde.Server/Storage/ObjectStores/AwsObjectStore.cs b/Engine/Source/Programs/Horde/Horde.Server/Storage/ObjectStores/AwsObjectStore.cs
--- a/Engine/Source/Programs/Horde/Horde.Server/Storage/ObjectStores/AwsObjectStore.cs
+++ b/Engine/Source/Programs/Horde/Horde.Server/Storage/ObjectStores/AwsObjectStore.cs
@@ -260,7 +260,7 @@ namespace Horde.Server.Storage.ObjectStores
 			}
 		}
 
-		string GetFullPath(ObjectKey key) => $"{_pathPrefix}{key}.blob";
+		string GetFullPath(ObjectKey key) => $"{_pathPrefix}{key}";
 
 		/// <inheritdoc/>
 		public Task<Stream> OpenAsync(ObjectKey key, int offset, int? length, CancellationToken cancellationToken)
diff --git a/Engine/Source/Programs/Horde/Horde.Server/Storage/StorageService.cs b/Engine/Source/Programs/Horde/Horde.Server/Storage/StorageService.cs
--- a/Engine/Source/Programs/Horde/Horde.Server/Storage/StorageService.cs
+++ b/Engine/Source/Programs/Horde/Horde.Server/Storage/StorageService.cs
@@ -53,15 +53,17 @@ namespace Horde.Server.Storage
 		sealed class LeafBlobHandle : BlobHandle
 		{
 			readonly IObjectStore _store;
+			readonly BlobLocator _locator;
 			readonly ObjectKey _key;
 			readonly Tracer _tracer;
 
 			public override IBlobHandle? Outer => null;
 
-			public LeafBlobHandle(IObjectStore store, ObjectKey key, Tracer tracer)
+			public LeafBlobHandle(IObjectStore store, BlobLocator locator, Tracer tracer)
 			{
 				_store = store;
-				_key = key;
+				_locator = locator;
+				_key = GetObjectKey(locator);
 				_tracer = tracer;
 			}
 
@@ -94,7 +96,7 @@ namespace Horde.Server.Storage
 			/// <inheritdoc/>
 			public override bool TryAppendIdentifier(Utf8StringBuilder builder)
 			{
-				builder.Append(_key.Path);
+				builder.Append(_locator.Path);
 				return true;
 			}
 
@@ -152,11 +154,11 @@ namespace Horde.Server.Storage
 			{
 				if (locator.TryUnwrap(out BlobLocator baseLocator, out Utf8String fragment))
 				{
-					return new BlobFragmentHandle(new LeafBlobHandle(_store, new ObjectKey(baseLocator.ToString()), _tracer), fragment);
+					return new BlobFragmentHandle(new LeafBlobHandle(_store, baseLocator, _tracer), fragment);
 				}
 				else
 				{
-					return new LeafBlobHandle(_store, new ObjectKey(locator.ToString()), _tracer);
+					return new LeafBlobHandle(_store, locator, _tracer);
 				}
 			}
 
@@ -176,19 +178,18 @@ namespace Horde.Server.Storage
 			/// <inheritdoc/>
 			public async ValueTask<IBlobHandle> WriteBlobAsync(BlobType type, Stream stream, IReadOnlyList<IBlobHandle> references, string? basePath = null, CancellationToken cancellationToken = default)
 			{
-				string path = StorageHelpers.CreateUniqueName(basePath);
-				await _store.WriteAsync(new ObjectKey(path), stream, cancellationToken);
+				BlobLocator locator = StorageHelpers.CreateUniqueLocator(basePath);
 
-				BlobLocator locator = new BlobLocator(path);
+				await _store.WriteAsync(GetObjectKey(locator), stream, cancellationToken);
 				await _outer.AddBlobAsync(NamespaceId, locator, null, cancellationToken);
 
-				return CreateBlobHandle(new BlobLocator(path));
+				return CreateBlobHandle(locator);
 			}
 
 			/// <inheritdoc/>
 			public ValueTask<Uri?> TryGetReadRedirectAsync(BlobLocator locator, CancellationToken cancellationToken = default)
 			{
-				return _store.TryGetReadRedirectAsync(new ObjectKey(locator.Path.ToString()), cancellationToken);
+				return _store.TryGetReadRedirectAsync(GetObjectKey(locator), cancellationToken);
 			}
 
 			/// <inheritdoc/>
@@ -199,15 +200,14 @@ namespace Horde.Server.Storage
 					return null;
 				}
 
-				string path = StorageHelpers.CreateUniqueName(prefix);
+				BlobLocator locator = StorageHelpers.CreateUniqueLocator(prefix);
 
-				Uri? url = await _store.TryGetWriteRedirectAsync(new ObjectKey(path), cancellationToken);
+				Uri? url = await _store.TryGetWriteRedirectAsync(GetObjectKey(locator), cancellationToken);
 				if (url == null)
 				{
 					return null;
 				}
 
-				BlobLocator locator = new BlobLocator(path);
 				await _outer.AddBlobAsync(NamespaceId, locator, null, cancellationToken);
 
 				return (locator, url);
@@ -669,6 +669,8 @@ namespace Horde.Server.Storage
 			await _gcTicker.DisposeAsync();
 		}
 
+		internal static ObjectKey GetObjectKey(BlobLocator locator) => new ObjectKey($"{locator.Path}.blob");
+
 		/// <summary>
 		/// Creates a new storage client factory using the current global config value
 		/// </summary>
@@ -1185,7 +1187,7 @@ namespace Horde.Server.Storage
 							_ = _redisService.GetDatabase().SortedSetAddAsync(checkSet, entries, flags: CommandFlags.FireAndForget);
 							score = Math.BitIncrement(score);
 						}
-						await namespaceInfo.Store.DeleteAsync(new ObjectKey(info.Path), cancellationToken);
+						await namespaceInfo.Store.DeleteAsync(GetObjectKey(new BlobLocator(info.Path)), cancellationToken);
 					}
 				}
 				_ = _redisService.GetDatabase().SortedSetRemoveAsync(checkSet, values[0], CommandFlags.FireAndForget);
diff --git a/Engine/Source/Programs/Shared/EpicGames.Horde.Tests/BackendTests.cs b/Engine/Source/Programs/Shared/EpicGames.Horde.Tests/BackendTests.cs
--- a/Engine/Source/Programs/Shared/EpicGames.Horde.Tests/BackendTests.cs
+++ b/Engine/Source/Programs/Shared/EpicGames.Horde.Tests/BackendTests.cs
@@ -59,17 +59,17 @@ namespace EpicGames.Horde.Tests
 			using (TempDir tempDir = new TempDir("Cache"))
 			{
 				using FileStorageBackend backend = new FileStorageBackend(tempDir.Location);
-				string path = await backend.WriteBytesAsync(Encoding.UTF8.GetBytes("hello world"));
+				BlobLocator path = await backend.WriteBytesAsync(Encoding.UTF8.GetBytes("hello world"));
 
-				using (IReadOnlyMemoryOwner<byte> handle = await backend.ReadAsync(path, 0, null))
+				using (IReadOnlyMemoryOwner<byte> handle = await backend.ReadBlobAsync(path, 0, null))
 				{
 					Assert.AreEqual("hello world", Encoding.UTF8.GetString(handle.Memory.Span));
 				}
-				using (IReadOnlyMemoryOwner<byte> handle = await backend.ReadAsync(path, 0, 5))
+				using (IReadOnlyMemoryOwner<byte> handle = await backend.ReadBlobAsync(path, 0, 5))
 				{
 					Assert.AreEqual("hello", Encoding.UTF8.GetString(handle.Memory.Span));
 				}
-				using (IReadOnlyMemoryOwner<byte> handle = await backend.ReadAsync(path, 4, 3))
+				using (IReadOnlyMemoryOwner<byte> handle = await backend.ReadBlobAsync(path, 4, 3))
 				{
 					Assert.AreEqual("o w", Encoding.UTF8.GetString(handle.Memory.Span));
 				}
@@ -89,32 +89,32 @@ namespace EpicGames.Horde.Tests
 				await TestBackendAsync(cacheBackend);
 
 				Assert.AreEqual(1, cache.Items.Count());
-				byte[] value = await cacheBackend.ReadBytesAsync(cache.Items.First());
+				byte[] value = await cacheBackend.ReadBytesAsync(cache.GetLocators().First());
 				Assert.IsTrue(value.SequenceEqual(Encoding.UTF8.GetBytes("item 2")));
 
 				byte[] data3 = Encoding.UTF8.GetBytes("3");
-				string path3 = await cacheBackend.WriteBytesAsync(data3);
+				BlobLocator path3 = await cacheBackend.WriteBytesAsync(data3);
 				await cacheBackend.ReadBytesAsync(path3);
 
 				byte[] data4 = Encoding.UTF8.GetBytes("4");
-				string path4 = await cacheBackend.WriteBytesAsync(data4);
+				BlobLocator path4 = await cacheBackend.WriteBytesAsync(data4);
 				await cacheBackend.ReadBytesAsync(path4);
 
 				byte[] data5 = Encoding.UTF8.GetBytes("5");
-				string path5 = await cacheBackend.WriteBytesAsync(data5);
+				BlobLocator path5 = await cacheBackend.WriteBytesAsync(data5);
 				await cacheBackend.ReadBytesAsync(path5);
 
-				HashSet<string> paths = new HashSet<string>(cache.Items);
+				HashSet<BlobLocator> paths = new HashSet<BlobLocator>(cache.GetLocators());
 				Assert.AreEqual(4, paths.Count);
 				Assert.IsTrue(paths.Contains(path3));
 				Assert.IsTrue(paths.Contains(path4));
 				Assert.IsTrue(paths.Contains(path5));
 
 				byte[] data6 = Encoding.UTF8.GetBytes("12345678901");
-				string path6 = await cacheBackend.WriteBytesAsync(data6);
+				BlobLocator path6 = await cacheBackend.WriteBytesAsync(data6);
 				await cacheBackend.ReadBytesAsync(path6);
 
-				paths = new HashSet<string>(cache.Items);
+				paths = new HashSet<BlobLocator>(cache.GetLocators());
 				Assert.AreEqual(2, paths.Count);
 				Assert.IsTrue(paths.Contains(path5));
 				Assert.IsTrue(paths.Contains(path6));
@@ -122,7 +122,7 @@ namespace EpicGames.Horde.Tests
 				await cacheBackend.ReadBytesAsync(path3);
 				await cacheBackend.ReadBytesAsync(path4);
 
-				paths = new HashSet<string>(cache.Items);
+				paths = new HashSet<BlobLocator>(cache.GetLocators());
 				Assert.AreEqual(2, paths.Count);
 				Assert.IsTrue(paths.Contains(path3));
 				Assert.IsTrue(paths.Contains(path4));
@@ -134,10 +134,10 @@ namespace EpicGames.Horde.Tests
 			byte[] data1 = Encoding.UTF8.GetBytes("hello world");
 			byte[] data2 = Encoding.UTF8.GetBytes("item 2");
 
-			string path1 = await backend.WriteBytesAsync(data1);
+			BlobLocator path1 = await backend.WriteBytesAsync(data1);
 			byte[] outputData1 = await backend.ReadBytesAsync(path1);
 
-			string path2 = await backend.WriteBytesAsync(data2);
+			BlobLocator path2 = await backend.WriteBytesAsync(data2);
 			byte[] outputData2 = await backend.ReadBytesAsync(path2);
 
 			Assert.IsTrue(data1.SequenceEqual(outputData1));
diff --git a/Engine/Source/Programs/Shared/EpicGames.Horde/Storage/Backends/FileStorageBackend.cs b/Engine/Source/Programs/Shared/EpicGames.Horde/Storage/Backends/FileStorageBackend.cs
--- a/Engine/Source/Programs/Shared/EpicGames.Horde/Storage/Backends/FileStorageBackend.cs
+++ b/Engine/Source/Programs/Shared/EpicGames.Horde/Storage/Backends/FileStorageBackend.cs
@@ -1,9 +1,7 @@
 // Copyright Epic Games, Inc. All Rights Reserved.
 
 using System;
-using System.Collections.Generic;
 using System.IO;
-using System.Runtime.CompilerServices;
 using System.Threading;
 using System.Threading.Tasks;
 using EpicGames.Core;
@@ -39,67 +37,39 @@ namespace EpicGames.Horde.Storage.Backends
 		/// <summary>
 		/// Gets the path for storing a file on disk
 		/// </summary>
-		static ObjectKey GetBlobFile(string path) => new ObjectKey($"{path}.blob");
+		static ObjectKey GetBlobFile(BlobLocator locator) => new ObjectKey(locator.Path);
 
 		/// <inheritdoc/>
-		public Task<Stream> OpenAsync(string path, int offset, int? length, CancellationToken cancellationToken)
-			=> _objectStore.OpenAsync(GetBlobFile(path), offset, length, cancellationToken);
+		public Task<Stream> OpenBlobAsync(BlobLocator locator, int offset, int? length, CancellationToken cancellationToken)
+			=> _objectStore.OpenAsync(GetBlobFile(locator), offset, length, cancellationToken);
 
 		/// <summary>
 		/// Maps a file into memory for reading, and returns a handle to it
 		/// </summary>
-		/// <param name="path">Path to the file</param>
+		/// <param name="locator">Path to the file</param>
 		/// <param name="offset">Offset of the data to retrieve</param>
 		/// <param name="length">Length of the data</param>
 		/// <returns>Handle to the data. Must be disposed by the caller.</returns>
-		public IReadOnlyMemoryOwner<byte> Read(string path, int offset, int? length)
-			=> _objectStore.Read(GetBlobFile(path), offset, length);
+		public IReadOnlyMemoryOwner<byte> Read(BlobLocator locator, int offset, int? length)
+			=> _objectStore.Read(GetBlobFile(locator), offset, length);
 
 		/// <inheritdoc/>
-		public Task<IReadOnlyMemoryOwner<byte>> ReadAsync(string path, int offset, int? length, CancellationToken cancellationToken = default)
-			=> _objectStore.ReadAsync(GetBlobFile(path), offset, length, cancellationToken);
+		public Task<IReadOnlyMemoryOwner<byte>> ReadBlobAsync(BlobLocator locator, int offset, int? length, CancellationToken cancellationToken = default)
+			=> _objectStore.ReadAsync(GetBlobFile(locator), offset, length, cancellationToken);
 
 		/// <inheritdoc/>
-		public async Task<string> WriteAsync(Stream stream, string? prefix = null, CancellationToken cancellationToken = default)
+		public async Task<BlobLocator> WriteBlobAsync(Stream stream, string? prefix = null, CancellationToken cancellationToken = default)
 		{
-			string path = StorageHelpers.CreateUniqueName(prefix);
-			await WriteExplicitPathAsync(path, stream, cancellationToken);
-			return path;
+			BlobLocator locator = StorageHelpers.CreateUniqueLocator(prefix);
+			await _objectStore.WriteAsync(GetBlobFile(locator), stream, cancellationToken);
+			return locator;
 		}
 
 		/// <inheritdoc/>
-		public Task WriteExplicitPathAsync(string path, Stream stream, CancellationToken cancellationToken = default)
-			=> _objectStore.WriteAsync(GetBlobFile(path), stream, cancellationToken);
+		public ValueTask<Uri?> TryGetBlobReadRedirectAsync(BlobLocator path, CancellationToken cancellationToken = default) => default;
 
 		/// <inheritdoc/>
-		public Task<bool> ExistsAsync(string path, CancellationToken cancellationToken)
-			=> _objectStore.ExistsAsync(GetBlobFile(path), cancellationToken);
-
-		/// <summary>
-		/// Delete a file from the store
-		/// </summary>
-		/// <param name="path"></param>
-		public void Delete(string path)
-			=> _objectStore.Delete(GetBlobFile(path));
-
-		/// <inheritdoc/>
-		public Task DeleteAsync(string path, CancellationToken cancellationToken)
-			=> _objectStore.DeleteAsync(GetBlobFile(path), cancellationToken);
-
-		/// <inheritdoc/>
-		public async IAsyncEnumerable<string> EnumerateAsync([EnumeratorCancellation] CancellationToken cancellationToken = default)
-		{
-			await foreach (ObjectKey locator in _objectStore.EnumerateAsync(cancellationToken))
-			{
-				yield return locator.Path.ToString();
-			}
-		}
-
-		/// <inheritdoc/>
-		public ValueTask<Uri?> TryGetReadRedirectAsync(string path, CancellationToken cancellationToken = default) => default;
-
-		/// <inheritdoc/>
-		public ValueTask<(string, Uri)?> TryGetWriteRedirectAsync(string? prefix = null, CancellationToken cancellationToken = default) => default;
+		public ValueTask<(BlobLocator, Uri)?> TryGetBlobWriteRedirectAsync(string? prefix = null, CancellationToken cancellationToken = default) => default;
 
 		/// <inheritdoc/>
 		public void GetStats(StorageStats stats)
diff --git a/Engine/Source/Programs/Shared/EpicGames.Horde/Storage/Backends/HttpStorageBackend.cs b/Engine/Source/Programs/Shared/EpicGames.Horde/Storage/Backends/HttpStorageBackend.cs
--- a/Engine/Source/Programs/Shared/EpicGames.Horde/Storage/Backends/HttpStorageBackend.cs
+++ b/Engine/Source/Programs/Shared/EpicGames.Horde/Storage/Backends/HttpStorageBackend.cs
@@ -1,7 +1,6 @@
 // Copyright Epic Games, Inc. All Rights Reserved.
 
 using System;
-using System.Collections.Generic;
 using System.IO;
 using System.Net.Http;
 using System.Net.Http.Headers;
@@ -51,19 +50,19 @@ namespace EpicGames.Horde.Storage.Backends
 		#region Blobs
 
 		/// <inheritdoc/>
-		public async Task<Stream> OpenAsync(string path, int offset, int? length, CancellationToken cancellationToken = default)
+		public async Task<Stream> OpenBlobAsync(BlobLocator locator, int offset, int? length, CancellationToken cancellationToken = default)
 		{
 			if (offset == 0 && length == null)
 			{
-				_logger.LogDebug("Reading {Locator}", path);
+				_logger.LogDebug("Reading {Locator}", locator);
 			}
 			else if (length == null)
 			{
-				_logger.LogDebug("Reading {Locator} ({Offset}..)", path, offset);
+				_logger.LogDebug("Reading {Locator} ({Offset}..)", locator, offset);
 			}
 			else
 			{
-				_logger.LogDebug("Reading {Locator} ({Offset}+{Length})", path, offset, length);
+				_logger.LogDebug("Reading {Locator} ({Offset}+{Length})", locator, offset, length);
 			}
 
 			if (length.HasValue && length.Value == 0)
@@ -73,7 +72,7 @@ namespace EpicGames.Horde.Storage.Backends
 
 			using (HttpClient httpClient = _createClient())
 			{
-				using (HttpRequestMessage request = new HttpRequestMessage(HttpMethod.Get, $"{_basePath}/blobs/{path}"))
+				using (HttpRequestMessage request = new HttpRequestMessage(HttpMethod.Get, $"{_basePath}/blobs/{locator}"))
 				{
 					if (offset != 0 || length != null)
 					{
@@ -88,9 +87,9 @@ namespace EpicGames.Horde.Storage.Backends
 		}
 
 		/// <inheritdoc/>
-		public async Task<IReadOnlyMemoryOwner<byte>> ReadAsync(string path, int offset, int? length, CancellationToken cancellationToken = default)
+		public async Task<IReadOnlyMemoryOwner<byte>> ReadBlobAsync(BlobLocator locator, int offset, int? length, CancellationToken cancellationToken = default)
 		{
-			using (Stream stream = await OpenAsync(path, offset, length, cancellationToken))
+			using (Stream stream = await OpenBlobAsync(locator, offset, length, cancellationToken))
 			{
 				byte[] data = await stream.ReadAllBytesAsync(cancellationToken);
 				return ReadOnlyMemoryOwner.Create(data);
@@ -98,7 +97,7 @@ namespace EpicGames.Horde.Storage.Backends
 		}
 
 		/// <inheritdoc/>
-		public async Task<string> WriteAsync(Stream stream, string? prefix = null, CancellationToken cancellationToken = default)
+		public async Task<BlobLocator> WriteBlobAsync(Stream stream, string? prefix = null, CancellationToken cancellationToken = default)
 		{
 			using StreamContent streamContent = new StreamContent(stream);
 
@@ -121,7 +120,7 @@ namespace EpicGames.Horde.Storage.Backends
 							}
 						}
 						_logger.LogDebug("Written {Locator} (using redirect)", redirectResponse.Blob);
-						return redirectResponse.Blob;
+						return new BlobLocator(redirectResponse.Blob);
 					}
 				}
 			}
@@ -129,12 +128,9 @@ namespace EpicGames.Horde.Storage.Backends
 			WriteBlobResponse response = await SendWriteRequestAsync(streamContent, prefix, cancellationToken);
 			_supportsUploadRedirects = response.SupportsRedirects ?? false;
 			_logger.LogDebug("Written {Locator} (direct)", response.Blob);
-			return response.Blob;
+			return new BlobLocator(response.Blob);
 		}
 
-		/// <inheritdoc/>
-		public Task WriteExplicitPathAsync(string path, Stream stream, CancellationToken cancellationToken = default) => throw new NotSupportedException();
-
 		async Task<WriteBlobResponse> SendWriteRequestAsync(StreamContent? streamContent, string? prefix = null, CancellationToken cancellationToken = default)
 		{
 			using (HttpClient httpClient = _createClient())
@@ -162,30 +158,21 @@ namespace EpicGames.Horde.Storage.Backends
 		}
 
 		/// <inheritdoc/>
-		public Task<bool> ExistsAsync(string path, CancellationToken cancellationToken = default) => throw new NotSupportedException();
-
-		/// <inheritdoc/>
-		public Task DeleteAsync(string path, CancellationToken cancellationToken = default) => throw new NotSupportedException();
-
-		/// <inheritdoc/>
-		public IAsyncEnumerable<string> EnumerateAsync(CancellationToken cancellationToken = default) => throw new NotSupportedException();
-
-		/// <inheritdoc/>
-		public ValueTask<Uri?> TryGetReadRedirectAsync(string path, CancellationToken cancellationToken = default)
+		public ValueTask<Uri?> TryGetBlobReadRedirectAsync(BlobLocator locator, CancellationToken cancellationToken = default)
 		{
 			// We don't currently have any need for getting a read redirect explicitly, though ReadAsync() calls may redirect us automatically.
 			return default;
 		}
 
 		/// <inheritdoc/>
-		public async ValueTask<(string, Uri)?> TryGetWriteRedirectAsync(string? prefix = null, CancellationToken cancellationToken = default)
+		public async ValueTask<(BlobLocator, Uri)?> TryGetBlobWriteRedirectAsync(string? prefix = null, CancellationToken cancellationToken = default)
 		{
 			if (_supportsUploadRedirects)
 			{
 				WriteBlobResponse redirectResponse = await SendWriteRequestAsync(null, prefix, cancellationToken);
 				if (redirectResponse.UploadUrl != null)
 				{
-					return (redirectResponse.Blob, redirectResponse.UploadUrl);
+					return (new BlobLocator(redirectResponse.Blob), redirectResponse.UploadUrl);
 				}
 			}
 			return null;
diff --git a/Engine/Source/Programs/Shared/EpicGames.Horde/Storage/Backends/MemoryStorageBackend.cs b/Engine/Source/Programs/Shared/EpicGames.Horde/Storage/Backends/MemoryStorageBackend.cs
--- a/Engine/Source/Programs/Shared/EpicGames.Horde/Storage/Backends/MemoryStorageBackend.cs
+++ b/Engine/Source/Programs/Shared/EpicGames.Horde/Storage/Backends/MemoryStorageBackend.cs
@@ -4,7 +4,6 @@ using System;
 using System.Collections.Concurrent;
 using System.Collections.Generic;
 using System.IO;
-using System.Runtime.CompilerServices;
 using System.Threading;
 using System.Threading.Tasks;
 using EpicGames.Core;
@@ -19,12 +18,12 @@ namespace EpicGames.Horde.Storage.Backends
 		/// <summary>
 		/// Data storage
 		/// </summary>
-		readonly ConcurrentDictionary<string, byte[]> _pathToData = new ConcurrentDictionary<string, byte[]>();
+		readonly ConcurrentDictionary<BlobLocator, byte[]> _locatorToData = new ConcurrentDictionary<BlobLocator, byte[]>();
 
 		/// <summary>
 		/// Read only access to the stored blobs
 		/// </summary>
-		public IReadOnlyDictionary<string, byte[]> Blobs => _pathToData; 
+		public IReadOnlyDictionary<BlobLocator, byte[]> Blobs => _locatorToData; 
 
 		/// <inheritdoc/>
 		public bool SupportsRedirects => false;
@@ -35,20 +34,20 @@ namespace EpicGames.Horde.Storage.Backends
 		}
 
 		/// <inheritdoc/>
-		public Task<Stream> OpenAsync(string path, int offset, int? length, CancellationToken cancellationToken)
+		public Task<Stream> OpenBlobAsync(BlobLocator locator, int offset, int? length, CancellationToken cancellationToken)
 		{
-			return Task.FromResult<Stream>(new ReadOnlyMemoryStream(GetData(path, offset, length)));
+			return Task.FromResult<Stream>(new ReadOnlyMemoryStream(GetData(locator, offset, length)));
 		}
 
 		/// <inheritdoc/>
-		public Task<IReadOnlyMemoryOwner<byte>> ReadAsync(string path, int offset, int? length, CancellationToken cancellationToken)
+		public Task<IReadOnlyMemoryOwner<byte>> ReadBlobAsync(BlobLocator locator, int offset, int? length, CancellationToken cancellationToken)
 		{
-			return Task.FromResult(ReadOnlyMemoryOwner.Create(GetData(path, offset, length)));
+			return Task.FromResult(ReadOnlyMemoryOwner.Create(GetData(locator, offset, length)));
 		}
 
-		ReadOnlyMemory<byte> GetData(string path, int offset, int? length)
+		ReadOnlyMemory<byte> GetData(BlobLocator locator, int offset, int? length)
 		{
-			ReadOnlyMemory<byte> data = _pathToData[path].AsMemory(offset);
+			ReadOnlyMemory<byte> data = _locatorToData[locator].AsMemory(offset);
 			if (length != null && length.Value < data.Length)
 			{
 				data = data.Slice(0, length.Value);
@@ -57,52 +56,22 @@ namespace EpicGames.Horde.Storage.Backends
 		}
 
 		/// <inheritdoc/>
-		public async Task<string> WriteAsync(Stream stream, string? prefix = null, CancellationToken cancellationToken = default)
-		{
-			string path = StorageHelpers.CreateUniqueName(prefix);
-			await WriteExplicitPathAsync(path, stream, cancellationToken);
-			return path;
-		}
-
-		/// <inheritdoc/>
-		public async Task WriteExplicitPathAsync(string path, Stream stream, CancellationToken cancellationToken = default)
+		public async Task<BlobLocator> WriteBlobAsync(Stream stream, string? prefix = null, CancellationToken cancellationToken = default)
 		{
+			BlobLocator locator = StorageHelpers.CreateUniqueLocator(prefix);
 			using (MemoryStream buffer = new MemoryStream())
 			{
 				await stream.CopyToAsync(buffer, cancellationToken);
-				_pathToData[path] = buffer.ToArray();
-			}
-		}
-
-		/// <inheritdoc/>
-		public Task<bool> ExistsAsync(string path, CancellationToken cancellationToken)
-		{
-			return Task.FromResult(_pathToData.ContainsKey(path));
-		}
-
-		/// <inheritdoc/>
-		public Task DeleteAsync(string path, CancellationToken cancellationToken)
-		{
-			_pathToData.TryRemove(path, out _);
-			return Task.CompletedTask;
-		}
-
-		/// <inheritdoc/>
-		public async IAsyncEnumerable<string> EnumerateAsync([EnumeratorCancellation] CancellationToken cancellationToken = default)
-		{
-			foreach (string path in _pathToData.Keys)
-			{
-				yield return path;
-				cancellationToken.ThrowIfCancellationRequested();
-				await Task.Yield();
+				_locatorToData[locator] = buffer.ToArray();
 			}
+			return locator;
 		}
 
 		/// <inheritdoc/>
-		public ValueTask<Uri?> TryGetReadRedirectAsync(string path, CancellationToken cancellationToken = default) => default;
+		public ValueTask<Uri?> TryGetBlobReadRedirectAsync(BlobLocator locator, CancellationToken cancellationToken = default) => default;
 
 		/// <inheritdoc/>
-		public ValueTask<(string, Uri)?> TryGetWriteRedirectAsync(string? prefix = null, CancellationToken cancellationToken = default) => default;
+		public ValueTask<(BlobLocator, Uri)?> TryGetBlobWriteRedirectAsync(string? prefix = null, CancellationToken cancellationToken = default) => default;
 
 		/// <inheritdoc/>
 		public void GetStats(StorageStats stats) { }
diff --git a/Engine/Source/Programs/Shared/EpicGames.Horde/Storage/Backends/PrefixStorageBackend.cs b/Engine/Source/Programs/Shared/EpicGames.Horde/Storage/Backends/PrefixStorageBackend.cs
deleted file mode 100644
diff --git a/Engine/Source/Programs/Shared/EpicGames.Horde/Storage/Clients/FileStorageClient.cs b/Engine/Source/Programs/Shared/EpicGames.Horde/Storage/Clients/FileStorageClient.cs
--- a/Engine/Source/Programs/Shared/EpicGames.Horde/Storage/Clients/FileStorageClient.cs
+++ b/Engine/Source/Programs/Shared/EpicGames.Horde/Storage/Clients/FileStorageClient.cs
@@ -64,15 +64,15 @@ namespace EpicGames.Horde.Storage.Clients
 		/// <inheritdoc/>
 		public override async ValueTask<BlobData> ReadBlobAsync(BlobLocator locator, CancellationToken cancellationToken = default)
 		{
-			IReadOnlyMemoryOwner<byte> owner = await _backend.ReadAsync(locator.ToString(), cancellationToken);
+			IReadOnlyMemoryOwner<byte> owner = await _backend.ReadAsync(locator, cancellationToken);
 			return new BlobDataWithOwner(BlobType.Leaf, owner.Memory, Array.Empty<IBlobHandle>(), owner);
 		}
 
 		/// <inheritdoc/>
 		public override async ValueTask<IBlobHandle> WriteBlobAsync(BlobType type, Stream stream, IReadOnlyList<IBlobHandle> references, string? basePath = null, CancellationToken cancellationToken = default)
 		{
-			string path = await _backend.WriteAsync(stream, basePath, cancellationToken);
-			return CreateBlobHandle(new BlobLocator(path));
+			BlobLocator locator = await _backend.WriteBlobAsync(stream, basePath, cancellationToken);
+			return CreateBlobHandle(locator);
 		}
 
 		/// <inheritdoc/>
diff --git a/Engine/Source/Programs/Shared/EpicGames.Horde/Storage/Clients/HttpStorageClient.cs b/Engine/Source/Programs/Shared/EpicGames.Horde/Storage/Clients/HttpStorageClient.cs
--- a/Engine/Source/Programs/Shared/EpicGames.Horde/Storage/Clients/HttpStorageClient.cs
+++ b/Engine/Source/Programs/Shared/EpicGames.Horde/Storage/Clients/HttpStorageClient.cs
@@ -41,7 +41,7 @@ namespace EpicGames.Horde.Storage.Clients
 
 			/// <inheritdoc/>
 			public override Task<Stream> OpenBodyAsync(int offset = 0, int? length = null, CancellationToken cancellationToken = default)
-				=> _httpStorageClient._backend.OpenAsync(_locator.ToString(), offset, length, cancellationToken);
+				=> _httpStorageClient._backend.OpenBlobAsync(_locator, offset, length, cancellationToken);
 
 			/// <inheritdoc/>
 			public override ValueTask<BlobData> ReadBlobDataAsync(CancellationToken cancellationToken = default)
@@ -108,31 +108,24 @@ namespace EpicGames.Horde.Storage.Clients
 		/// <inheritdoc/>
 		public async ValueTask<BlobData> ReadBlobAsync(BlobLocator locator, CancellationToken cancellationToken = default)
 		{
-			IReadOnlyMemoryOwner<byte> owner = await _backend.ReadAsync(locator.ToString(), cancellationToken);
+			IReadOnlyMemoryOwner<byte> owner = await _backend.ReadAsync(locator, cancellationToken);
 			return new BlobDataWithOwner(BlobType.Leaf, owner.Memory, Array.Empty<IBlobHandle>(), owner);
 		}
 
 		/// <inheritdoc/>
 		public async ValueTask<IBlobHandle> WriteBlobAsync(BlobType type, Stream stream, IReadOnlyList<IBlobHandle> references, string? basePath = null, CancellationToken cancellationToken = default)
 		{
-			string path = await _backend.WriteAsync(stream, basePath, cancellationToken);
-			return new Handle(this, new BlobLocator(path));
+			BlobLocator locator = await _backend.WriteBlobAsync(stream, basePath, cancellationToken);
+			return new Handle(this, locator);
 		}
 
 		/// <inheritdoc/>
 		public ValueTask<Uri?> TryGetReadRedirectAsync(BlobLocator locator, CancellationToken cancellationToken = default)
-			=> _backend.TryGetReadRedirectAsync(locator.ToString(), cancellationToken);
+			=> _backend.TryGetBlobReadRedirectAsync(locator, cancellationToken);
 
 		/// <inheritdoc/>
-		public async ValueTask<(BlobLocator, Uri)?> TryGetWriteRedirectAsync(string? prefix = null, CancellationToken cancellationToken = default)
-		{
-			(string Path, Uri Url)? result = await _backend.TryGetWriteRedirectAsync(prefix, cancellationToken);
-			if (result == null)
-			{
-				return null;
-			}
-			return (new BlobLocator(result.Value.Path), result.Value.Url);
-		}
+		public ValueTask<(BlobLocator, Uri)?> TryGetWriteRedirectAsync(string? prefix = null, CancellationToken cancellationToken = default)
+			=> _backend.TryGetBlobWriteRedirectAsync(prefix, cancellationToken);
 
 		#endregion
 
diff --git a/Engine/Source/Programs/Shared/EpicGames.Horde/Storage/IStorageBackend.cs b/Engine/Source/Programs/Shared/EpicGames.Horde/Storage/IStorageBackend.cs
--- a/Engine/Source/Programs/Shared/EpicGames.Horde/Storage/IStorageBackend.cs
+++ b/Engine/Source/Programs/Shared/EpicGames.Horde/Storage/IStorageBackend.cs
@@ -1,7 +1,6 @@
 // Copyright Epic Games, Inc. All Rights Reserved.
 
 using System;
-using System.Collections.Generic;
 using System.IO;
 using System.Text;
 using System.Threading;
@@ -23,22 +22,22 @@ namespace EpicGames.Horde.Storage
 		/// <summary>
 		/// Attempts to open a read stream for the given path.
 		/// </summary>
-		/// <param name="path">Relative path within the bucket</param>
+		/// <param name="locator">Relative path within the bucket</param>
 		/// <param name="offset">Offset to start reading from</param>
 		/// <param name="length">Length of data to read</param>
 		/// <param name="cancellationToken">Cancellation token for the operation</param>
 		/// <returns></returns>
-		Task<Stream> OpenAsync(string path, int offset, int? length, CancellationToken cancellationToken = default);
+		Task<Stream> OpenBlobAsync(BlobLocator locator, int offset, int? length, CancellationToken cancellationToken = default);
 
 		/// <summary>
 		/// Reads an object into memory and returns a handle to it.
 		/// </summary>
-		/// <param name="path">Path to the file</param>
+		/// <param name="locator">Path to the file</param>
 		/// <param name="offset">Offset of the data to retrieve</param>
 		/// <param name="length">Length of the data</param>
 		/// <param name="cancellationToken">Cancellation token for the operation</param>
 		/// <returns>Handle to the data. Must be disposed by the caller.</returns>
-		Task<IReadOnlyMemoryOwner<byte>> ReadAsync(string path, int offset, int? length, CancellationToken cancellationToken = default);
+		Task<IReadOnlyMemoryOwner<byte>> ReadBlobAsync(BlobLocator locator, int offset, int? length, CancellationToken cancellationToken = default);
 
 		/// <summary>
 		/// Writes a stream to the storage backend. If the stream throws an exception during read, the write will be aborted.
@@ -47,48 +46,15 @@ namespace EpicGames.Horde.Storage
 		/// <param name="prefix">Path prefix for the uploaded data</param>
 		/// <param name="cancellationToken">Cancellation token for the operation</param>
 		/// <returns>Path to the uploaded object</returns>
-		Task<string> WriteAsync(Stream stream, string? prefix = null, CancellationToken cancellationToken = default);
-
-		/// <summary>
-		/// Writes a stream to the storage backend. This overload is deprecated; prefer passing a prefix to allow the server to determine a unique path.
-		/// </summary>
-		/// <param name="path"></param>
-		/// <param name="stream">Stream to write</param>
-		/// <param name="cancellationToken">Cancellation token for the operation</param>
-		/// <returns>Path to the uploaded object</returns>
-		[Obsolete("Use WriteAsync() instead. Ability to specify an explicit path is deprecated and will be removed in a future release.")]
-		Task WriteExplicitPathAsync(string path, Stream stream, CancellationToken cancellationToken = default);
-
-		/// <summary>
-		/// Tests whether the given path exists
-		/// </summary>
-		/// <param name="path">Relative path within the bucket</param>
-		/// <param name="cancellationToken">Cancellation token for the operation</param>
-		/// <returns></returns>
-		Task<bool> ExistsAsync(string path, CancellationToken cancellationToken = default);
-
-		/// <summary>
-		/// Deletes a file with the given path
-		/// </summary>
-		/// <param name="path">Relative path within the bucket</param>
-		/// <param name="cancellationToken">Cancellation token for the operation</param>
-		/// <returns>Async task</returns>
-		Task DeleteAsync(string path, CancellationToken cancellationToken = default);
-
-		/// <summary>
-		/// Enumerates all the objects in the store
-		/// </summary>
-		/// <param name="cancellationToken">Cancellation token for the operation</param>
-		/// <returns>Sequence of object paths</returns>
-		IAsyncEnumerable<string> EnumerateAsync(CancellationToken cancellationToken = default);
+		Task<BlobLocator> WriteBlobAsync(Stream stream, string? prefix = null, CancellationToken cancellationToken = default);
 
 		/// <summary>
 		/// Gets a HTTP redirect for a read request
 		/// </summary>
-		/// <param name="path">Path to read from</param>
+		/// <param name="locator">Path to read from</param>
 		/// <param name="cancellationToken">Cancellation token for the operation</param>
 		/// <returns>Path to upload the data to</returns>
-		ValueTask<Uri?> TryGetReadRedirectAsync(string path, CancellationToken cancellationToken = default);
+		ValueTask<Uri?> TryGetBlobReadRedirectAsync(BlobLocator locator, CancellationToken cancellationToken = default);
 
 		/// <summary>
 		/// Gets a HTTP redirect for a write request
@@ -96,7 +62,7 @@ namespace EpicGames.Horde.Storage
 		/// <param name="prefix">Prefix for the uploaded data</param>
 		/// <param name="cancellationToken">Cancellation token for the operation</param>
 		/// <returns>Path for retrieval, and URI to upload the data to</returns>
-		ValueTask<(string, Uri)?> TryGetWriteRedirectAsync(string? prefix = null, CancellationToken cancellationToken = default);
+		ValueTask<(BlobLocator, Uri)?> TryGetBlobWriteRedirectAsync(string? prefix = null, CancellationToken cancellationToken = default);
 
 		/// <summary>
 		/// Gets stats for this storage backend
@@ -124,7 +90,7 @@ namespace EpicGames.Horde.Storage
 		/// </summary>
 		/// <param name="prefix">The prefix to use</param>
 		/// <returns>Unique name generated with the given prefix</returns>
-		public static string CreateUniqueName(string? prefix)
+		public static BlobLocator CreateUniqueLocator(string? prefix)
 		{
 			StringBuilder builder = new StringBuilder(prefix);
 			if (builder.Length > 0 && builder[^1] != '/')
@@ -133,7 +99,7 @@ namespace EpicGames.Horde.Storage
 			}
 			builder.Append(s_sessionPrefix);
 			builder.Append(Interlocked.Increment(ref s_increment));
-			return builder.ToString();
+			return new BlobLocator(builder.ToString());
 		}
 	}
 
@@ -146,30 +112,30 @@ namespace EpicGames.Horde.Storage
 		/// Attempts to open a read stream for the given path.
 		/// </summary>
 		/// <param name="storageBackend">Backend to read from</param>
-		/// <param name="path">Object name within the store</param>
+		/// <param name="locator">Object name within the store</param>
 		/// <param name="cancellationToken">Cancellation token for the operation</param>
 		/// <returns>Stream for the object</returns>
-		public static Task<Stream> OpenAsync(this IStorageBackend storageBackend, string path, CancellationToken cancellationToken = default) => storageBackend.OpenAsync(path, 0, null, cancellationToken);
+		public static Task<Stream> OpenAsync(this IStorageBackend storageBackend, BlobLocator locator, CancellationToken cancellationToken = default) => storageBackend.OpenBlobAsync(locator, 0, null, cancellationToken);
 
 		/// <summary>
 		/// Attempts to open a read stream for the given path.
 		/// </summary>
 		/// <param name="storageBackend">Backend to read from</param>
-		/// <param name="path">Object name within the store</param>
+		/// <param name="locator">Object name within the store</param>
 		/// <param name="cancellationToken">Cancellation token for the operation</param>
 		/// <returns>Stream for the object</returns>
-		public static Task<IReadOnlyMemoryOwner<byte>> ReadAsync(this IStorageBackend storageBackend, string path, CancellationToken cancellationToken = default) => storageBackend.ReadAsync(path, 0, null, cancellationToken);
+		public static Task<IReadOnlyMemoryOwner<byte>> ReadAsync(this IStorageBackend storageBackend, BlobLocator locator, CancellationToken cancellationToken = default) => storageBackend.ReadBlobAsync(locator, 0, null, cancellationToken);
 
 		/// <summary>
 		/// Reads an object as an array of bytes
 		/// </summary>
 		/// <param name="storageBackend">Backend to read from</param>
-		/// <param name="path">Object name within the store</param>
+		/// <param name="locator">Object name within the store</param>
 		/// <param name="cancellationToken">Cancellation token for the operation</param>
 		/// <returns>Contents of the object</returns>
-		public static async Task<byte[]> ReadBytesAsync(this IStorageBackend storageBackend, string path, CancellationToken cancellationToken = default)
+		public static async Task<byte[]> ReadBytesAsync(this IStorageBackend storageBackend, BlobLocator locator, CancellationToken cancellationToken = default)
 		{
-			using IReadOnlyMemoryOwner<byte> storageObject = await storageBackend.ReadAsync(path, cancellationToken);
+			using IReadOnlyMemoryOwner<byte> storageObject = await storageBackend.ReadAsync(locator, cancellationToken);
 			return storageObject.Memory.ToArray();
 		}
 
@@ -180,11 +146,11 @@ namespace EpicGames.Horde.Storage
 		/// <param name="data">Data to be written</param>
 		/// <param name="prefix">Prefix for the uploaded data</param>
 		/// <param name="cancellationToken">Cancellation token for the operation</param>
-		public static async Task<string> WriteBytesAsync(this IStorageBackend storageBackend, ReadOnlyMemory<byte> data, string? prefix = null, CancellationToken cancellationToken = default)
+		public static async Task<BlobLocator> WriteBytesAsync(this IStorageBackend storageBackend, ReadOnlyMemory<byte> data, string? prefix = null, CancellationToken cancellationToken = default)
 		{
 			using (ReadOnlyMemoryStream stream = new ReadOnlyMemoryStream(data))
 			{
-				return await storageBackend.WriteAsync(stream, prefix, cancellationToken);
+				return await storageBackend.WriteBlobAsync(stream, prefix, cancellationToken);
 			}
 		}
 	}
diff --git a/Engine/Source/Programs/Shared/EpicGames.Horde/Storage/ObjectStores/MemoryObjectStore.cs b/Engine/Source/Programs/Shared/EpicGames.Horde/Storage/ObjectStores/MemoryObjectStore.cs
--- a/Engine/Source/Programs/Shared/EpicGames.Horde/Storage/ObjectStores/MemoryObjectStore.cs
+++ b/Engine/Source/Programs/Shared/EpicGames.Horde/Storage/ObjectStores/MemoryObjectStore.cs
@@ -16,12 +16,12 @@ namespace EpicGames.Horde.Storage.ObjectStores
 	/// </summary>
 	public sealed class MemoryObjectStore : IObjectStore
 	{
-		readonly ConcurrentDictionary<ObjectKey, byte[]> _pathToData = new ConcurrentDictionary<ObjectKey, byte[]>();
+		readonly ConcurrentDictionary<ObjectKey, byte[]> _keyToData = new ConcurrentDictionary<ObjectKey, byte[]>();
 
 		/// <summary>
 		/// Read only access to the stored blobs
 		/// </summary>
-		public IReadOnlyDictionary<ObjectKey, byte[]> Blobs => _pathToData; 
+		public IReadOnlyDictionary<ObjectKey, byte[]> Blobs => _keyToData; 
 
 		/// <inheritdoc/>
 		public bool SupportsRedirects => false;
@@ -32,20 +32,20 @@ namespace EpicGames.Horde.Storage.ObjectStores
 		}
 
 		/// <inheritdoc/>
-		public Task<Stream> OpenAsync(ObjectKey locator, int offset, int? length, CancellationToken cancellationToken)
+		public Task<Stream> OpenAsync(ObjectKey key, int offset, int? length, CancellationToken cancellationToken)
 		{
-			return Task.FromResult<Stream>(new ReadOnlyMemoryStream(GetData(locator, offset, length)));
+			return Task.FromResult<Stream>(new ReadOnlyMemoryStream(GetData(key, offset, length)));
 		}
 
 		/// <inheritdoc/>
-		public Task<IReadOnlyMemoryOwner<byte>> ReadAsync(ObjectKey locator, int offset, int? length, CancellationToken cancellationToken)
+		public Task<IReadOnlyMemoryOwner<byte>> ReadAsync(ObjectKey key, int offset, int? length, CancellationToken cancellationToken)
 		{
-			return Task.FromResult(ReadOnlyMemoryOwner.Create(GetData(locator, offset, length)));
+			return Task.FromResult(ReadOnlyMemoryOwner.Create(GetData(key, offset, length)));
 		}
 
-		ReadOnlyMemory<byte> GetData(ObjectKey locator, int offset, int? length)
+		ReadOnlyMemory<byte> GetData(ObjectKey key, int offset, int? length)
 		{
-			ReadOnlyMemory<byte> data = _pathToData[locator].AsMemory(offset);
+			ReadOnlyMemory<byte> data = _keyToData[key].AsMemory(offset);
 			if (length != null && length.Value < data.Length)
 			{
 				data = data.Slice(0, length.Value);
@@ -54,44 +54,44 @@ namespace EpicGames.Horde.Storage.ObjectStores
 		}
 
 		/// <inheritdoc/>
-		public async Task WriteAsync(ObjectKey locator, Stream stream, CancellationToken cancellationToken = default)
+		public async Task WriteAsync(ObjectKey key, Stream stream, CancellationToken cancellationToken = default)
 		{
 			using (MemoryStream buffer = new MemoryStream())
 			{
 				await stream.CopyToAsync(buffer, cancellationToken);
-				_pathToData[locator] = buffer.ToArray();
+				_keyToData[key] = buffer.ToArray();
 			}
 		}
 
 		/// <inheritdoc/>
-		public Task<bool> ExistsAsync(ObjectKey locator, CancellationToken cancellationToken)
+		public Task<bool> ExistsAsync(ObjectKey key, CancellationToken cancellationToken)
 		{
-			return Task.FromResult(_pathToData.ContainsKey(locator));
+			return Task.FromResult(_keyToData.ContainsKey(key));
 		}
 
 		/// <inheritdoc/>
-		public Task DeleteAsync(ObjectKey locator, CancellationToken cancellationToken)
+		public Task DeleteAsync(ObjectKey key, CancellationToken cancellationToken)
 		{
-			_pathToData.TryRemove(locator, out _);
+			_keyToData.TryRemove(key, out _);
 			return Task.CompletedTask;
 		}
 
 		/// <inheritdoc/>
 		public async IAsyncEnumerable<ObjectKey> EnumerateAsync([EnumeratorCancellation] CancellationToken cancellationToken = default)
 		{
-			foreach (ObjectKey locator in _pathToData.Keys)
+			foreach (ObjectKey key in _keyToData.Keys)
 			{
-				yield return locator;
+				yield return key;
 				cancellationToken.ThrowIfCancellationRequested();
 				await Task.Yield();
 			}
 		}
 
 		/// <inheritdoc/>
-		public ValueTask<Uri?> TryGetReadRedirectAsync(ObjectKey path, CancellationToken cancellationToken = default) => default;
+		public ValueTask<Uri?> TryGetReadRedirectAsync(ObjectKey key, CancellationToken cancellationToken = default) => default;
 
 		/// <inheritdoc/>
-		public ValueTask<Uri?> TryGetWriteRedirectAsync(ObjectKey path, CancellationToken cancellationToken = default) => default;
+		public ValueTask<Uri?> TryGetWriteRedirectAsync(ObjectKey key, CancellationToken cancellationToken = default) => default;
 
 		/// <inheritdoc/>
 		public void GetStats(StorageStats stats) { }
diff --git a/Engine/Source/Programs/Shared/EpicGames.Horde/Storage/StorageBackendCache.cs b/Engine/Source/Programs/Shared/EpicGames.Horde/Storage/StorageBackendCache.cs
--- a/Engine/Source/Programs/Shared/EpicGames.Horde/Storage/StorageBackendCache.cs
+++ b/Engine/Source/Programs/Shared/EpicGames.Horde/Storage/StorageBackendCache.cs
@@ -1,7 +1,7 @@
 // Copyright Epic Games, Inc. All Rights Reserved.
 
 using EpicGames.Core;
-using EpicGames.Horde.Storage.Backends;
+using EpicGames.Horde.Storage.ObjectStores;
 using Microsoft.Extensions.Logging;
 using Microsoft.Extensions.Logging.Abstractions;
 using System;
@@ -21,15 +21,13 @@ namespace EpicGames.Horde.Storage
 	{
 		class Item
 		{
-			public string Key { get; }
-			public string Path { get; }
+			public ObjectKey Key { get; }
 			public long Length { get; }
 			public LinkedListNode<Item> ListNode { get; }
 
-			public Item(string key, string path, long length)
+			public Item(ObjectKey key, long length)
 			{
 				Key = key;
-				Path = path;
 				Length = length;
 				ListNode = new LinkedListNode<Item>(this);
 			}
@@ -80,7 +78,7 @@ namespace EpicGames.Horde.Storage
 			public async Task<IReadOnlyMemoryOwner<byte>> ReadAsync(ObjectKey key, int offset, int? length, CancellationToken cancellationToken = default)
 			{
 #pragma warning disable CA2000 // Dispose objects before losing scope
-				IReadOnlyMemoryOwner<byte> storageObject = await _cacheStorage.ReadAsync($"{_keyPrefix}{key}", ctx => _inner.OpenAsync(key, ctx), cancellationToken);
+				IReadOnlyMemoryOwner<byte> storageObject = await _cacheStorage.ReadAsync(new ObjectKey($"{_keyPrefix}{key}"), ctx => _inner.OpenAsync(key, ctx), cancellationToken);
 #pragma warning restore CA2000 // Dispose objects before losing scope
 				return storageObject.Slice(offset, length);
 			}
@@ -101,6 +99,8 @@ namespace EpicGames.Horde.Storage
 
 		sealed class BackendWrapper : IStorageBackend
 		{
+			public const string BlobExtension = ".blob";
+
 			readonly string _keyPrefix;
 			readonly StorageBackendCache _cacheStorage;
 			readonly IStorageBackend _inner;
@@ -116,31 +116,23 @@ namespace EpicGames.Horde.Storage
 
 			public void Dispose() => _inner.Dispose();
 
-			public async Task<Stream> OpenAsync(string path, int offset, int? length, CancellationToken cancellationToken = default)
+			public async Task<Stream> OpenBlobAsync(BlobLocator locator, int offset, int? length, CancellationToken cancellationToken = default)
 			{
-				IReadOnlyMemoryOwner<byte> storageObject = await ReadAsync(path, offset, length, cancellationToken);
+				IReadOnlyMemoryOwner<byte> storageObject = await ReadBlobAsync(locator, offset, length, cancellationToken);
 				return storageObject.AsStream();
 			}
 
-			public async Task<IReadOnlyMemoryOwner<byte>> ReadAsync(string path, int offset, int? length, CancellationToken cancellationToken = default)
+			public async Task<IReadOnlyMemoryOwner<byte>> ReadBlobAsync(BlobLocator locator, int offset, int? length, CancellationToken cancellationToken = default)
 			{
 #pragma warning disable CA2000 // Dispose objects before losing scope
-				IReadOnlyMemoryOwner<byte> storageObject = await _cacheStorage.ReadAsync($"{_keyPrefix}{path}", ctx => _inner.OpenAsync(path, ctx), cancellationToken);
+				IReadOnlyMemoryOwner<byte> storageObject = await _cacheStorage.ReadAsync(new ObjectKey($"{_keyPrefix}{locator}{BlobExtension}"), ctx => _inner.OpenAsync(locator, ctx), cancellationToken);
 #pragma warning restore CA2000 // Dispose objects before losing scope
 				return storageObject.Slice(offset, length);
 			}
 
-			public Task<string> WriteAsync(Stream stream, string? prefix = null, CancellationToken cancellationToken = default) => _inner.WriteAsync(stream, prefix, cancellationToken);
-
-#pragma warning disable CS0618 // Type or member is obsolete
-			public Task WriteExplicitPathAsync(string path, Stream stream, CancellationToken cancellationToken = default) => _inner.WriteExplicitPathAsync(path, stream, cancellationToken);
-#pragma warning restore CS0618 // Type or member is obsolete
-
-			public Task DeleteAsync(string path, CancellationToken cancellationToken = default) => _inner.DeleteAsync(path, cancellationToken);
-			public IAsyncEnumerable<string> EnumerateAsync(CancellationToken cancellationToken = default) => _inner.EnumerateAsync(cancellationToken);
-			public Task<bool> ExistsAsync(string path, CancellationToken cancellationToken = default) => _inner.ExistsAsync(path, cancellationToken);
-			public ValueTask<Uri?> TryGetReadRedirectAsync(string path, CancellationToken cancellationToken = default) => _inner.TryGetReadRedirectAsync(path, cancellationToken);
-			public ValueTask<(string, Uri)?> TryGetWriteRedirectAsync(string? prefix = null, CancellationToken cancellationToken = default) => _inner.TryGetWriteRedirectAsync(prefix, cancellationToken);
+			public Task<BlobLocator> WriteBlobAsync(Stream stream, string? prefix = null, CancellationToken cancellationToken = default) => _inner.WriteBlobAsync(stream, prefix, cancellationToken);
+			public ValueTask<Uri?> TryGetBlobReadRedirectAsync(BlobLocator locator, CancellationToken cancellationToken = default) => _inner.TryGetBlobReadRedirectAsync(locator, cancellationToken);
+			public ValueTask<(BlobLocator, Uri)?> TryGetBlobWriteRedirectAsync(string? prefix = null, CancellationToken cancellationToken = default) => _inner.TryGetBlobWriteRedirectAsync(prefix, cancellationToken);
 
 			public void GetStats(StorageStats stats)
 			{
@@ -151,13 +143,13 @@ namespace EpicGames.Horde.Storage
 
 		object LockObject => _items;
 
-		readonly FileStorageBackend _backend;
+		readonly FileObjectStore _objectStore;
 		readonly long _maxSize;
 		readonly ILogger _logger;
 
 		readonly LinkedList<Item> _items = new LinkedList<Item>();
-		readonly Dictionary<string, Item> _pathToItem = new Dictionary<string, Item>(StringComparer.Ordinal);
-		readonly Dictionary<string, PendingItem> _pathToPendingItem = new Dictionary<string, PendingItem>(StringComparer.Ordinal);
+		readonly Dictionary<ObjectKey, Item> _keyToItem = new Dictionary<ObjectKey, Item>();
+		readonly Dictionary<ObjectKey, PendingItem> _keyToPendingItem = new Dictionary<ObjectKey, PendingItem>();
 
 		long _size;
 		long _cleanCount;
@@ -166,7 +158,19 @@ namespace EpicGames.Horde.Storage
 		long _writeTimeTicks;
 		long _fetchBytes;
 
-		internal IEnumerable<string> Items => _items.Select(x => x.Key);
+		internal IEnumerable<ObjectKey> Items => _items.Select(x => x.Key);
+
+		internal IEnumerable<BlobLocator> GetLocators()
+		{
+			foreach (Item item in _items)
+			{
+				ObjectKey key = item.Key;
+				if (key.Path.ToString().EndsWith(BackendWrapper.BlobExtension))
+				{
+					yield return new BlobLocator(key.Path.Substring(0, key.Path.Length - BackendWrapper.BlobExtension.Length));
+				}
+			}
+		}
 
 		/// <summary>
 		/// Constructor
@@ -195,7 +199,7 @@ namespace EpicGames.Horde.Storage
 			cacheDir ??= new DirectoryReference(Path.Combine(Path.GetTempPath(), $"horde-{Guid.NewGuid().ToString("n")}"));
 			FileUtils.ForceDeleteDirectoryContents(cacheDir);
 
-			_backend = new FileStorageBackend(cacheDir);
+			_objectStore = new FileObjectStore(cacheDir);
 			_maxSize = maxSize ?? (50 * 1024 * 1024);
 			_logger = logger;
 		}
@@ -205,9 +209,9 @@ namespace EpicGames.Horde.Storage
 		{
 			foreach (Item item in _items)
 			{
-				_backend.Delete(item.Path);
+				_objectStore.Delete(item.Key);
 			}
-			_backend.Dispose();
+			_objectStore.Dispose();
 		}
 
 		/// <summary>
@@ -242,7 +246,7 @@ namespace EpicGames.Horde.Storage
 		}
 
 		/// <inheritdoc/>
-		public async Task<IReadOnlyMemoryOwner<byte>> ReadAsync(string key, Func<CancellationToken, Task<Stream>> createStreamAsync, CancellationToken cancellationToken = default)
+		public async Task<IReadOnlyMemoryOwner<byte>> ReadAsync(ObjectKey key, Func<CancellationToken, Task<Stream>> createStreamAsync, CancellationToken cancellationToken = default)
 		{
 			for (; ; )
 			{
@@ -250,17 +254,17 @@ namespace EpicGames.Horde.Storage
 				lock (LockObject)
 				{
 					Item? item;
-					if (_pathToItem.TryGetValue(key, out item))
+					if (_keyToItem.TryGetValue(key, out item))
 					{
 						_items.Remove(item.ListNode);
 						_items.AddFirst(item.ListNode);
-						return _backend.Read(item.Path, 0, null);
+						return _objectStore.Read(item.Key, 0, null);
 					}
 
-					if (!_pathToPendingItem.TryGetValue(key, out pendingItem))
+					if (!_keyToPendingItem.TryGetValue(key, out pendingItem))
 					{
 						pendingItem = new PendingItem(BackgroundTask.StartNew(x => ReadIntoCacheAsync(key, createStreamAsync, x)));
-						_pathToPendingItem.Add(key, pendingItem);
+						_keyToPendingItem.Add(key, pendingItem);
 					}
 
 					pendingItem.AddRef();
@@ -277,7 +281,7 @@ namespace EpicGames.Horde.Storage
 			}
 		}
 
-		async Task ReadIntoCacheAsync(string key, Func<CancellationToken, Task<Stream>> createStreamAsync, CancellationToken cancellationToken)
+		async Task ReadIntoCacheAsync(ObjectKey key, Func<CancellationToken, Task<Stream>> createStreamAsync, CancellationToken cancellationToken)
 		{
 			long openStartTicks = Stopwatch.GetTimestamp();
 			using Stream stream = await createStreamAsync(cancellationToken);
@@ -299,9 +303,9 @@ namespace EpicGames.Horde.Storage
 					Item item = node.Value;
 					try
 					{
-						_backend.Delete(item.Path);
+						_objectStore.Delete(item.Key);
 
-						_pathToItem.Remove(item.Key);
+						_keyToItem.Remove(item.Key);
 						_items.Remove(node);
 
 						_size -= item.Length;
@@ -309,7 +313,7 @@ namespace EpicGames.Horde.Storage
 					}
 					catch (Exception ex)
 					{
-						_logger.LogDebug(ex, "Unable to delete cache item {Path}: {Message}", item.Path, ex.Message);
+						_logger.LogDebug(ex, "Unable to delete cache item {Path}: {Message}", item.Key, ex.Message);
 					}
 				}
 			}
@@ -317,14 +321,14 @@ namespace EpicGames.Horde.Storage
 			Interlocked.Add(ref _cleanTimeTicks, cleanFinishTicks - cleanStartTicks);
 
 			long writeStartTicks = cleanFinishTicks;
-			string path = await _backend.WriteAsync(stream, cancellationToken: cancellationToken);
+			await _objectStore.WriteAsync(key, stream, cancellationToken: cancellationToken);
 			lock (LockObject)
 			{
-				Item item = new Item(key, path, totalLength);
+				Item item = new Item(key, totalLength);
 				_items.AddFirst(item.ListNode);
-				_pathToItem.Add(key, item);
+				_keyToItem.Add(key, item);
 
-				_pathToPendingItem.Remove(key);
+				_keyToPendingItem.Remove(key);
 			}
 			Interlocked.Add(ref _fetchBytes, totalLength);
 			long writeFinishTicks = Stopwatch.GetTimestamp();
