Horde: Fix static analysis warning.

diff --git a/Engine/Source/Programs/Shared/EpicGames.Horde/Storage/StorageBackendCache.cs b/Engine/Source/Programs/Shared/EpicGames.Horde/Storage/StorageBackendCache.cs
--- a/Engine/Source/Programs/Shared/EpicGames.Horde/Storage/StorageBackendCache.cs
+++ b/Engine/Source/Programs/Shared/EpicGames.Horde/Storage/StorageBackendCache.cs
@@ -165,7 +165,7 @@ namespace EpicGames.Horde.Storage
 			foreach (Item item in _items)
 			{
 				ObjectKey key = item.Key;
-				if (key.Path.ToString().EndsWith(BackendWrapper.BlobExtension))
+				if (key.Path.ToString().EndsWith(BackendWrapper.BlobExtension, StringComparison.OrdinalIgnoreCase))
 				{
 					yield return new BlobLocator(key.Path.Substring(0, key.Path.Length - BackendWrapper.BlobExtension.Length));
 				}
