OidcToken - Allow multiple instances of oidc token to read the windows token store at the same time, and prevent updating the token store file if there are no changes applied to it. Reduced the scope the mutex to only cover the updating of the token store.
This resolves race issues when running multiple oidc tokens at the same time.

diff --git a/Engine/Source/Programs/Shared/EpicGames.OIDC/WindowsTokenStore.cs b/Engine/Source/Programs/Shared/EpicGames.OIDC/WindowsTokenStore.cs
--- a/Engine/Source/Programs/Shared/EpicGames.OIDC/WindowsTokenStore.cs
+++ b/Engine/Source/Programs/Shared/EpicGames.OIDC/WindowsTokenStore.cs
@@ -36,21 +36,10 @@ namespace EpicGames.OIDC
 		}
 
 		private readonly Dictionary<string, byte[]> _providerToRefreshToken = new Dictionary<string, byte[]>();
-
-		private Mutex _mutex;
+		private bool _isDirty;
 
 		public WindowsTokenStore()
 		{
-			_mutex = new Mutex(false, "oidcTokenStoreDat");
-
-			try
-			{
-				_mutex.WaitOne();
-			} catch (AbandonedMutexException)
-			{
-
-			}
-
 			ReadStoreFromDisk();
 		}
 
@@ -68,7 +57,7 @@ namespace EpicGames.OIDC
 				return;
 			}
 
-			using FileStream fs = fi.OpenRead();
+			using FileStream fs = fi.Open(FileMode.Open, FileAccess.Read, FileShare.Read);
 			using TextReader tr = new StreamReader(fs);
 
 			TokenStoreState? state;
@@ -95,6 +84,11 @@ namespace EpicGames.OIDC
 
 		private void SaveStoreToDisk()
 		{
+			if (!_isDirty)
+			{
+				return;
+			}
+
 			FileInfo fi = GetStorePath();
 
 			if (!fi.Directory?.Exists ?? false)
@@ -109,7 +103,22 @@ namespace EpicGames.OIDC
 				JsonSerializer.Serialize<TokenStoreState>(writer, new TokenStoreState(_providerToRefreshToken));
 			}
 
-			File.Move(tempFile, fi.FullName, true);
+			{
+				using Mutex mutex = new Mutex(false, "oidcTokenStoreDat");
+
+				try
+				{
+					mutex.WaitOne();
+				}
+				catch (AbandonedMutexException)
+				{
+
+				}
+				File.Move(tempFile, fi.FullName, true);
+
+				mutex.ReleaseMutex();
+			}
+			_isDirty = false;
 		}
 
 		public bool TryGetRefreshToken(string oidcProvider, out string refreshToken)
@@ -158,6 +167,8 @@ namespace EpicGames.OIDC
 			byte[] encryptedToken = CryptProtectDataHelper.DoCryptProtectData(bytes, $"OidcToken-{providerIdentifier}", GetEntropy(providerIdentifier));
 
 			_providerToRefreshToken[providerIdentifier] = encryptedToken;
+
+			_isDirty = true;
 		}
 
 		public void Save()
@@ -170,7 +181,6 @@ namespace EpicGames.OIDC
 			if (disposing)
 			{
 				SaveStoreToDisk();
-				_mutex.Dispose();
 			}
 		}
 
