EpicGames.Core: Fix incorrect data returned by ReadOnlyMemoryArray enumerator.

diff --git a/Engine/Source/Programs/Shared/EpicGames.Core/ReadOnlyMemoryArray.cs b/Engine/Source/Programs/Shared/EpicGames.Core/ReadOnlyMemoryArray.cs
--- a/Engine/Source/Programs/Shared/EpicGames.Core/ReadOnlyMemoryArray.cs
+++ b/Engine/Source/Programs/Shared/EpicGames.Core/ReadOnlyMemoryArray.cs
@@ -128,7 +128,7 @@ namespace EpicGames.Core
 		/// <inheritdoc/>
 		public IEnumerator<ReadOnlyMemory<byte>> GetEnumerator()
 		{
-			int offset = 0;
+			int offset = GetOffset(0);
 			for (int index = 0; index < Count; index++)
 			{
 				int nextOffset = GetOffset(index + 1);
