Horde: Allow periods in ref names, as long as they follow the same requirements as filesystem paths (not appearing at the start or end of a path fragment).

diff --git a/Engine/Source/Programs/Shared/EpicGames.Horde/Storage/RefName.cs b/Engine/Source/Programs/Shared/EpicGames.Horde/Storage/RefName.cs
--- a/Engine/Source/Programs/Shared/EpicGames.Horde/Storage/RefName.cs
+++ b/Engine/Source/Programs/Shared/EpicGames.Horde/Storage/RefName.cs
@@ -80,6 +80,17 @@ namespace EpicGames.Horde.Storage
 						lastSlashIdx = idx;
 					}
 				}
+				else if (text[idx] == '.')
+				{
+					if (idx == 0 || text[idx - 1] == '/')
+					{
+						throw new ArgumentException($"{Encoding.UTF8.GetString(text)} is not a valid ref name (path fragment cannot start with a period)", argumentName);
+					}
+					if (idx + 1 == text.Length || text[idx + 1] == '/')
+					{
+						throw new ArgumentException($"{Encoding.UTF8.GetString(text)} is not a valid ref name (path fragment cannot end with a period)", argumentName);
+					}
+				}
 				else
 				{
 					if (!IsValidChar(text[idx]))
@@ -94,7 +105,7 @@ namespace EpicGames.Horde.Storage
 
 		static uint[] CreateValidCharsArray()
 		{
-			const string ValidChars = "0123456789abcdefghijklmnopqrstuvwxyz_/-+";
+			const string ValidChars = "0123456789abcdefghijklmnopqrstuvwxyz_/-+.";
 
 			uint[] validChars = new uint[256 / 8];
 			for (int idx = 0; idx < ValidChars.Length; idx++)
diff --git a/Engine/Source/Programs/Shared/EpicGames.Horde/StringId.cs b/Engine/Source/Programs/Shared/EpicGames.Horde/StringId.cs
--- a/Engine/Source/Programs/Shared/EpicGames.Horde/StringId.cs
+++ b/Engine/Source/Programs/Shared/EpicGames.Horde/StringId.cs
@@ -128,6 +128,11 @@ namespace EpicGames.Horde
 				throw new ArgumentException($"String id may not be longer than {MaxLength} characters", paramName);
 			}
 
+			if (text.Length > 0 && (text[0] == '.' || text[^1] == '.'))
+			{
+				throw new ArgumentException($"'{text}' is not a valid string id (cannot start or end with a period)");
+			}
+
 			for (int idx = 0; idx < text.Length; idx++)
 			{
 				char character = (char)text[idx];
@@ -139,7 +144,7 @@ namespace EpicGames.Horde
 					}
 					else
 					{
-						throw new ArgumentException($"{text} is not a valid string id", paramName);
+						throw new ArgumentException($"'{text}' is not a valid string id (character '{character}' is not allowed)", paramName);
 					}
 				}
 			}
