Horde: Avoid sending cancellation exceptions during compute execution

Shows up in client's log even though it's benign.

diff --git a/Engine/Source/Programs/Shared/EpicGames.Horde/Compute/AgentMessageHandler.cs b/Engine/Source/Programs/Shared/EpicGames.Horde/Compute/AgentMessageHandler.cs
--- a/Engine/Source/Programs/Shared/EpicGames.Horde/Compute/AgentMessageHandler.cs
+++ b/Engine/Source/Programs/Shared/EpicGames.Horde/Compute/AgentMessageHandler.cs
@@ -217,6 +217,11 @@ namespace EpicGames.Horde.Compute
 					await ExecuteProcessInternalAsync(channel, executable, arguments, workingDir, envVars, flags, cancellationToken);
 				}
 			}
+			catch (OperationCanceledException) when (cancellationToken.IsCancellationRequested)
+			{
+				// Skip sending back over compute socket as invocation is already being cancelled
+				_logger.LogInformation("Compute process execution cancelled");
+			}
 			catch (Exception ex)
 			{
 				await channel.SendExceptionAsync(ex, cancellationToken);
