Horde: Fix some bundle V2 decode errors introduced in CL 30648855

diff --git a/Engine/Source/Programs/Shared/EpicGames.Horde/Storage/Bundles/BundleStorageClient.cs b/Engine/Source/Programs/Shared/EpicGames.Horde/Storage/Bundles/BundleStorageClient.cs
--- a/Engine/Source/Programs/Shared/EpicGames.Horde/Storage/Bundles/BundleStorageClient.cs
+++ b/Engine/Source/Programs/Shared/EpicGames.Horde/Storage/Bundles/BundleStorageClient.cs
@@ -150,12 +150,13 @@ namespace EpicGames.Horde.Storage.Bundles
 			}
 
 			int ampIdx = fragment.IndexOf('&');
-			if (ampIdx != -1)
+			if (ampIdx == -1)
 			{
 				throw new Exception($"{locator} is not valid");
 			}
 
-			return new Bundles.V2.ExportHandle(new Bundles.V2.FlushedPacketHandle(this, bundleHandle, fragment.Slice(0, ampIdx).Span, _cache), fragment.Slice(ampIdx + 1));
+			Bundles.V2.PacketHandle packetHandle = new Bundles.V2.FlushedPacketHandle(this, bundleHandle, fragment.Slice(0, ampIdx).Span, _cache);
+			return new Bundles.V2.ExportHandle(packetHandle, fragment.Slice(ampIdx + 1));
 		}
 
 		/// <inheritdoc/>
diff --git a/Engine/Source/Programs/Shared/EpicGames.Horde/Storage/Bundles/V2/PacketHandle.cs b/Engine/Source/Programs/Shared/EpicGames.Horde/Storage/Bundles/V2/PacketHandle.cs
--- a/Engine/Source/Programs/Shared/EpicGames.Horde/Storage/Bundles/V2/PacketHandle.cs
+++ b/Engine/Source/Programs/Shared/EpicGames.Horde/Storage/Bundles/V2/PacketHandle.cs
@@ -43,7 +43,7 @@ namespace EpicGames.Horde.Storage.Bundles.V2
 			if (Bundle.TryGetLocator(out locator))
 			{
 				builder.Append(locator.Path);
-				builder.Append('?');
+				builder.Append('#');
 				return TryAppendIdentifier(builder);
 			}
 			return false;
@@ -207,7 +207,7 @@ namespace EpicGames.Horde.Storage.Bundles.V2
 		{
 			using IRefCountedHandle<IReadOnlyMemoryOwner<byte>> encodedData = await ReadEncodedPacketAsync(cancellationToken);
 			IRefCountedHandle<Packet> packet = Packet.Decode(encodedData.Target.Memory, _cache.Allocator);
-			return new PacketReader(_storageClient, _cache, packet.Target, packet);
+			return new PacketReader(_storageClient, _cache, _outer, this, packet.Target, packet);
 		}
 
 		async ValueTask<IRefCountedHandle<IReadOnlyMemoryOwner<byte>>> ReadEncodedPacketAsync(CancellationToken cancellationToken)
diff --git a/Engine/Source/Programs/Shared/EpicGames.Horde/Storage/Bundles/V2/PacketReader.cs b/Engine/Source/Programs/Shared/EpicGames.Horde/Storage/Bundles/V2/PacketReader.cs
--- a/Engine/Source/Programs/Shared/EpicGames.Horde/Storage/Bundles/V2/PacketReader.cs
+++ b/Engine/Source/Programs/Shared/EpicGames.Horde/Storage/Bundles/V2/PacketReader.cs
@@ -13,6 +13,8 @@ namespace EpicGames.Horde.Storage.Bundles.V2
 	{
 		readonly BundleStorageClient _storageClient;
 		readonly BundleCache _cache;
+		readonly BundleHandle _bundleHandle;
+		readonly FlushedPacketHandle _packetHandle;
 
 		Packet _decodedPacket;
 		object?[] _cachedImportHandles;
@@ -28,12 +30,16 @@ namespace EpicGames.Horde.Storage.Bundles.V2
 		/// </summary>
 		/// <param name="storageClient"></param>
 		/// <param name="cache"></param>
+		/// <param name="bundleHandle"></param>
+		/// <param name="packetHandle"></param>
 		/// <param name="decodedPacket">Data for the packet</param>
 		/// <param name="memoryOwner">Owner for the packet data</param>
-		public PacketReader(BundleStorageClient storageClient, BundleCache cache, Packet decodedPacket, IRefCountedHandle memoryOwner)
+		public PacketReader(BundleStorageClient storageClient, BundleCache cache, BundleHandle bundleHandle, FlushedPacketHandle packetHandle, Packet decodedPacket, IRefCountedHandle memoryOwner)
 		{
 			_storageClient = storageClient;
 			_cache = cache;
+			_bundleHandle = bundleHandle;
+			_packetHandle = packetHandle;
 			_decodedPacket = decodedPacket;
 			_memoryOwner = memoryOwner;
 			_cachedImportHandles = new object?[_decodedPacket.GetImportCount()];
@@ -124,6 +130,11 @@ namespace EpicGames.Horde.Storage.Bundles.V2
 
 		PacketHandle GetImportedPacketHandle(int packetIdx)
 		{
+			if (packetIdx == PacketImport.CurrentPacketBaseIdx)
+			{
+				return _packetHandle;
+			}
+
 			PacketHandle? packetHandle = _cachedImportHandles[packetIdx] as PacketHandle;
 			if (packetHandle is null)
 			{
@@ -141,6 +152,11 @@ namespace EpicGames.Horde.Storage.Bundles.V2
 
 		BundleHandle GetImportedBundleHandle(int bundleIdx)
 		{
+			if (bundleIdx == PacketImport.CurrentBundleBaseIdx)
+			{
+				return _bundleHandle;
+			}
+
 			BundleHandle? bundleHandle = _cachedImportHandles[bundleIdx] as BundleHandle;
 			if (bundleHandle is null)
 			{
