Horde: Improve performance of extracting data from bundles.

* Enumerating chunks to extract is now done asynchronously to extraction, and chunks are enumerated in order. This massively increases cache coherency over having separate writers run over different partitions of the data.
* Data is now written to disk using memory mapped files rather than buffered stream IO. Significantly faster in testing (~2x).

diff --git a/Engine/Source/Programs/Horde/Horde/Commands/Bundles/BundleExtract.cs b/Engine/Source/Programs/Horde/Horde/Commands/Bundles/BundleExtract.cs
--- a/Engine/Source/Programs/Horde/Horde/Commands/Bundles/BundleExtract.cs
+++ b/Engine/Source/Programs/Horde/Horde/Commands/Bundles/BundleExtract.cs
@@ -36,6 +36,10 @@ namespace Horde.Commands.Bundles
 		[Description("Directory to write extracted files.")]
 		public DirectoryReference OutputDir { get; set; } = null!;
 
+		[CommandLine("-CleanOutput")]
+		[Description("If set, deletes the contents of the output directory before extraction.")]
+		public bool CleanOutput { get; set; }
+
 		public BundleExtract(HttpStorageClientFactory storageClientFactory, BundleCache bundleCache, IOptions<CmdConfig> config)
 			: base(storageClientFactory, bundleCache, config)
 		{
@@ -43,6 +47,12 @@ namespace Horde.Commands.Bundles
 
 		public override async Task<int> ExecuteAsync(ILogger logger)
 		{
+			if(CleanOutput)
+			{
+				logger.LogInformation("Deleting contents of {OutputDir}...", OutputDir);
+				FileUtils.ForceDeleteDirectoryContents(OutputDir);
+			}
+
 			if (File != null)
 			{
 				using IStorageClient store = BundleStorageClient.CreateFromDirectory(File.Directory, BundleCache, logger);
diff --git a/Engine/Source/Programs/Shared/EpicGames.Horde/Storage/Bundles/BundleCache.cs b/Engine/Source/Programs/Shared/EpicGames.Horde/Storage/Bundles/BundleCache.cs
--- a/Engine/Source/Programs/Shared/EpicGames.Horde/Storage/Bundles/BundleCache.cs
+++ b/Engine/Source/Programs/Shared/EpicGames.Horde/Storage/Bundles/BundleCache.cs
@@ -178,6 +178,7 @@ namespace EpicGames.Horde.Storage
 		readonly CancellationTokenSource _cancellationSource = new CancellationTokenSource();
 		readonly MemoryAllocator _allocator;
 
+		int _numCacheHits;
 		long _currentSize;
 
 		readonly MemoryCache? _headerCache;
@@ -300,6 +301,7 @@ namespace EpicGames.Horde.Storage
 				CacheValue? item;
 				if (_itemLookup.TryGetValue(key, out item) && item.InitTask.TryGetResult(out IDisposable result))
 				{
+					Interlocked.Increment(ref _numCacheHits);
 					item.AddRef();
 					return new CacheValueHandle<TValue>(item, (TValue)result);
 				}
@@ -349,6 +351,7 @@ namespace EpicGames.Horde.Storage
 					CacheValue? untypedItem;
 					if (_itemLookup.TryGetValue(key, out untypedItem))
 					{
+						Interlocked.Increment(ref _numCacheHits);
 						item = untypedItem;
 					}
 					else
@@ -395,6 +398,15 @@ namespace EpicGames.Horde.Storage
 			Interlocked.Add(ref _currentSize, -size);
 		}
 
+		/// <summary>
+		/// Gets stats for the cache
+		/// </summary>
+		public void GetStats(StorageStats stats)
+		{
+			stats.Add("Num cache hits", _numCacheHits);
+			stats.Add("Cache size", _currentSize);
+		}
+
 		#region V1
 
 		static void AddCachedValue(IMemoryCache? cache, string cacheKey, object value, int size)
diff --git a/Engine/Source/Programs/Shared/EpicGames.Horde/Storage/Bundles/BundleStorageClient.cs b/Engine/Source/Programs/Shared/EpicGames.Horde/Storage/Bundles/BundleStorageClient.cs
--- a/Engine/Source/Programs/Shared/EpicGames.Horde/Storage/Bundles/BundleStorageClient.cs
+++ b/Engine/Source/Programs/Shared/EpicGames.Horde/Storage/Bundles/BundleStorageClient.cs
@@ -21,6 +21,7 @@ namespace EpicGames.Horde.Storage.Bundles
 	{
 		readonly IStorageBackend _backend;
 		readonly BundleCache _cache;
+		readonly PacketReaderStats _packetReaderStats = new PacketReaderStats();
 		readonly Bundles.V1.BundleReader _bundleReader;
 
 		internal Bundles.V1.BundleReader BundleReader => _bundleReader;
@@ -40,6 +41,11 @@ namespace EpicGames.Horde.Storage.Bundles
 		/// </summary>
 		public BundleCache Cache => _cache;
 
+		/// <summary>
+		/// Stats for reading bundles
+		/// </summary>
+		internal PacketReaderStats PacketReaderStats => _packetReaderStats;
+
 		/// <summary>
 		/// Constructor
 		/// </summary>
@@ -244,8 +250,10 @@ namespace EpicGames.Horde.Storage.Bundles
 		/// <inheritdoc/>
 		public void GetStats(StorageStats stats)
 		{
+			_cache.GetStats(stats);
 			_backend.GetStats(stats);
 			_bundleReader.GetStats(stats);
+			_packetReaderStats.GetStats(stats);
 		}
 	}
 }
diff --git a/Engine/Source/Programs/Shared/EpicGames.Horde/Storage/Bundles/V1/BundleReader.cs b/Engine/Source/Programs/Shared/EpicGames.Horde/Storage/Bundles/V1/BundleReader.cs
--- a/Engine/Source/Programs/Shared/EpicGames.Horde/Storage/Bundles/V1/BundleReader.cs
+++ b/Engine/Source/Programs/Shared/EpicGames.Horde/Storage/Bundles/V1/BundleReader.cs
@@ -586,10 +586,22 @@ namespace EpicGames.Horde.Storage.Bundles.V1
 		/// </summary>
 		public void GetStats(StorageStats stats)
 		{
-			stats.Add("Num bytes read", _numBytesRead);
-			stats.Add("Num header reads", _numHeaderReads);
-			stats.Add("Num packet reads", _numPacketReads);
-			stats.Add("Decode time (ms)", (_decodeTimeTicks * 1000) / Stopwatch.Frequency);
+			if (_numBytesRead > 0)
+			{
+				stats.Add("Num bytes read (v1)", _numBytesRead);
+			}
+			if (_numHeaderReads > 0)
+			{
+				stats.Add("Num header reads (v1)", _numHeaderReads);
+			}
+			if (_numPacketReads > 0)
+			{
+				stats.Add("Num packet reads (v1)", _numPacketReads);
+			}
+			if (_decodeTimeTicks > 0)
+			{
+				stats.Add("Decode time (ms) (v1)", (_decodeTimeTicks * 1000) / Stopwatch.Frequency);
+			}
 		}
 	}
 }
diff --git a/Engine/Source/Programs/Shared/EpicGames.Horde/Storage/Bundles/V2/PacketHandle.cs b/Engine/Source/Programs/Shared/EpicGames.Horde/Storage/Bundles/V2/PacketHandle.cs
--- a/Engine/Source/Programs/Shared/EpicGames.Horde/Storage/Bundles/V2/PacketHandle.cs
+++ b/Engine/Source/Programs/Shared/EpicGames.Horde/Storage/Bundles/V2/PacketHandle.cs
@@ -231,6 +231,11 @@ namespace EpicGames.Horde.Storage.Bundles.V2
 				Memory<byte> memory = leadingPacket.Memory.Slice(0, _packetLength);
 				await stream.ReadFixedLengthBytesAsync(memory, cancellationToken);
 
+				PacketReaderStats packetReaderStats = _storageClient.PacketReaderStats;
+				Interlocked.Increment(ref packetReaderStats._numReads);
+				Interlocked.Add(ref packetReaderStats._numBytesRead, _packetLength);
+				Interlocked.Increment(ref packetReaderStats._numPacketsRead);
+
 				// Read any other packets in the same stream
 				byte[] header = new byte[BundleSignature.NumBytes];
 				for (int readOffset = _packetLength; readOffset + BundleSignature.NumBytes < readLength;)
@@ -238,12 +243,14 @@ namespace EpicGames.Horde.Storage.Bundles.V2
 					int readBytes = await stream.ReadGreedyAsync(header, cancellationToken);
 					if (readBytes < header.Length)
 					{
+						Interlocked.Add(ref packetReaderStats._numWastedBytesRead, readBytes);
 						break;
 					}
 
 					BundleSignature signature = BundleSignature.Read(header);
 					if (readOffset + signature.HeaderLength >= readLength)
 					{
+						Interlocked.Add(ref packetReaderStats._numWastedBytesRead, readLength - readOffset);
 						break;
 					}
 
@@ -264,6 +271,9 @@ namespace EpicGames.Horde.Storage.Bundles.V2
 #pragma warning restore CA2000
 
 					readOffset += signature.HeaderLength;
+
+					Interlocked.Add(ref packetReaderStats._numBytesRead, signature.HeaderLength);
+					Interlocked.Increment(ref packetReaderStats._numPacketsRead);
 				}
 			}
 			catch
diff --git a/Engine/Source/Programs/Shared/EpicGames.Horde/Storage/Bundles/V2/PacketReader.cs b/Engine/Source/Programs/Shared/EpicGames.Horde/Storage/Bundles/V2/PacketReader.cs
--- a/Engine/Source/Programs/Shared/EpicGames.Horde/Storage/Bundles/V2/PacketReader.cs
+++ b/Engine/Source/Programs/Shared/EpicGames.Horde/Storage/Bundles/V2/PacketReader.cs
@@ -6,6 +6,25 @@ using EpicGames.Core;
 
 namespace EpicGames.Horde.Storage.Bundles.V2
 {
+	/// <summary>
+	/// Counters used for tracking operations performed by a BundleStorageClient
+	/// </summary>
+	class PacketReaderStats
+	{
+		public int _numReads;
+		public int _numPacketsRead;
+		public long _numBytesRead;
+		public long _numWastedBytesRead;
+
+		public void GetStats(StorageStats stats)
+		{
+			stats.Add("Num reads (v2)", _numReads);
+			stats.Add("Num packets read (v2)", _numPacketsRead);
+			stats.Add("Num bytes read (v2)", _numBytesRead);
+			stats.Add("Num wasted bytes read (v2)", _numWastedBytesRead);
+		}
+	}
+
 	/// <summary>
 	/// Utility class for constructing BlobData objects from a packet, caching any computed handles to other blobs.
 	/// </summary>
diff --git a/Engine/Source/Programs/Shared/EpicGames.Horde/Storage/Nodes/DirectoryNode.cs b/Engine/Source/Programs/Shared/EpicGames.Horde/Storage/Nodes/DirectoryNode.cs
--- a/Engine/Source/Programs/Shared/EpicGames.Horde/Storage/Nodes/DirectoryNode.cs
+++ b/Engine/Source/Programs/Shared/EpicGames.Horde/Storage/Nodes/DirectoryNode.cs
@@ -10,10 +10,12 @@ using System.Diagnostics;
 using System.Diagnostics.CodeAnalysis;
 using System.IO;
 using System.IO.Compression;
+using System.IO.MemoryMappedFiles;
 using System.IO.Pipelines;
 using System.Linq;
 using System.Runtime.InteropServices;
 using System.Threading;
+using System.Threading.Channels;
 using System.Threading.Tasks;
 
 namespace EpicGames.Horde.Storage.Nodes
@@ -722,6 +724,10 @@ namespace EpicGames.Horde.Storage.Nodes
 		/// <param name="cancellationToken"></param>
 		public Task CopyToDirectoryAsync(DirectoryInfo directoryInfo, BlobSerializerOptions? options, ILogger logger, CancellationToken cancellationToken) => CopyToDirectoryAsync(directoryInfo, null, options, logger, cancellationToken);
 
+		record class OutputDir(string Path, DirectoryNode Node);
+		record class OutputFile(OutputDir Directory, FileEntry FileEntry);
+		record class OutputChunk(OutputFile File, long Offset, long Length, IBlobHandle Handle);
+
 		/// <summary>
 		/// Utility function to allow extracting a packed directory to disk
 		/// </summary>
@@ -732,7 +738,7 @@ namespace EpicGames.Horde.Storage.Nodes
 		/// <param name="cancellationToken">Cancellation token for the operation</param>
 		public async Task CopyToDirectoryAsync(DirectoryInfo directoryInfo, IProgress<ICopyStats>? progress, BlobSerializerOptions? options, ILogger logger, CancellationToken cancellationToken)
 		{
-			int numTasks = Math.Min(1 + (int)(Length / (16 * 1024 * 1024)), 128);
+			int numTasks = Math.Min(1 + (int)(Length / (16 * 1024 * 1024)), 4);
 			logger.LogInformation("Splitting read into {NumThreads} threads", numTasks);
 
 			CopyStats? copyStats = null;
@@ -741,58 +747,158 @@ namespace EpicGames.Horde.Storage.Nodes
 				copyStats = new CopyStats(progress);
 			}
 
-			List<Task> tasks = new List<Task>();
-			try
+			Channel<OutputChunk> chunks = Channel.CreateBounded<OutputChunk>(new BoundedChannelOptions(128) { FullMode = BoundedChannelFullMode.Wait });
+			using (CancellationTokenSource cancellationSource = CancellationTokenSource.CreateLinkedTokenSource(cancellationToken))
+			{
+				// Helper method to run a background task and set a cancellation source on error
+				async Task RunBackgroundTask(Func<CancellationToken, Task> taskFunc)
+				{
+					try
+					{
+						await taskFunc(cancellationSource.Token);
+					}
+					catch (OperationCanceledException)
+					{
+						// Ignore
+					}
+					catch (Exception ex)
+					{
+						logger.LogError(ex, "Error while extracting data: {Message}", ex.Message);
+						cancellationSource.Cancel();
+					}
+				}
+
+				List<Task> tasks = new List<Task>();
+				tasks.Add(RunBackgroundTask(ctx => FindOutputChunksRootAsync(chunks.Writer, options, ctx)));
+				for (int idx = 0; idx < numTasks; idx++)
+				{
+					tasks.Add(RunBackgroundTask(ctx => ExtractAsync(chunks.Reader, new DirectoryReference(directoryInfo), copyStats, ctx)));
+				}
+
+				await Task.WhenAll(tasks);
+			}
+		}
+
+		static async Task ExtractAsync(ChannelReader<OutputChunk> chunkReader, DirectoryReference baseDir, CopyStats? copyStats, CancellationToken cancellationToken)
+		{
+			OutputChunk? chunk = await ReadNextChunkAsync(chunkReader, cancellationToken);
+			while (chunk != null)
 			{
-				long offset = 0;
-				for (int taskIdx = 0; taskIdx < numTasks; taskIdx++)
+				// Open the file for the current chunk
+				OutputFile file = chunk.File;
+				FileReference locator = FileReference.Combine(baseDir, file.Directory.Path, file.FileEntry.Name);
+				DirectoryReference.CreateDirectory(locator.Directory);
+
+				await using FileStream stream = FileReference.Open(locator, FileMode.OpenOrCreate, FileAccess.ReadWrite, FileShare.ReadWrite);
+				stream.SetLength(file.FileEntry.Length);
+
+				// If this file is empty, don't write anything and just move to the next chunk
+				if (file.FileEntry.Length == 0)
+				{
+					chunk = await ReadNextChunkAsync(chunkReader, cancellationToken);
+					continue;
+				}
+
+				// Process as many chunks as we can for this file
+				using MemoryMappedFile memoryMappedFile = MemoryMappedFile.CreateFromFile(stream, null, file.FileEntry.Length, MemoryMappedFileAccess.ReadWrite, HandleInheritability.None, false);
+				using MemoryMappedViewAccessor memoryMappedViewAccessor = memoryMappedFile.CreateViewAccessor(0, file.FileEntry.Length);
+				using MemoryMappedView memoryMappedView = new MemoryMappedView(memoryMappedViewAccessor);
+
+				while (chunk != null && chunk.File == file)
 				{
-					long minOffset = offset;
-					long maxOffset = (Length * (taskIdx + 1)) / numTasks;
+					// Write this chunk
+					using (BlobData data = await chunk.Handle.ReadBlobDataAsync(cancellationToken))
+					{
+						data.Data.CopyTo(memoryMappedView!.GetMemory(chunk.Offset, data.Data.Length));
+					}
 
-					// Entries may be zero-length, so need to make sure the last window will include everything
-					if (taskIdx == numTasks - 1)
+					// Update the stats
+					int numCompleteFiles = 0;
+					if (chunk.Offset + chunk.Length == file.FileEntry.Length)
 					{
-						maxOffset++;
+						numCompleteFiles = 1;
 					}
 
-					tasks.Add(Task.Run(() => CopyToDirectoryInternalAsync(directoryInfo, minOffset, maxOffset - minOffset, copyStats, options, logger, cancellationToken), cancellationToken));
-					offset = maxOffset;
+					copyStats?.Update(numCompleteFiles, chunk.Length);
+
+					// Read the next chunk
+					chunk = await ReadNextChunkAsync(chunkReader, cancellationToken);
 				}
 			}
-			finally
+		}
+
+		static async ValueTask<OutputChunk?> ReadNextChunkAsync(ChannelReader<OutputChunk> chunkReader, CancellationToken cancellationToken)
+		{
+			await chunkReader.WaitToReadAsync(cancellationToken);
+
+			OutputChunk? chunk;
+			if (chunkReader.TryRead(out chunk))
 			{
-				await Task.WhenAll(tasks);
-				copyStats?.Flush();
+				return chunk;
+			}
+			else
+			{
+				return null;
 			}
 		}
 
-		async Task CopyToDirectoryInternalAsync(DirectoryInfo directoryInfo, long windowOffset, long windowLength, CopyStats? copyStats, BlobSerializerOptions? options, ILogger logger, CancellationToken cancellationToken)
+		async Task FindOutputChunksRootAsync(ChannelWriter<OutputChunk> chunks, BlobSerializerOptions? options, CancellationToken cancellationToken)
+		{
+			await FindOutputChunksAsync("", this, chunks, options, cancellationToken);
+			chunks.Complete();
+		}
+
+		static async Task FindOutputChunksAsync(string path, DirectoryNode node, ChannelWriter<OutputChunk> chunks, BlobSerializerOptions? options, CancellationToken cancellationToken)
 		{
-			directoryInfo.Create();
+			OutputDir outputDir = new OutputDir(path, node);
 
-			foreach (FileEntry fileEntry in _nameToFileEntry.Values)
+			foreach (FileEntry fileEntry in node._nameToFileEntry.Values)
 			{
-				// Extract any file that starts within the window (window starts before this file, window ends after the start of the file)
-				if (windowOffset <= 0 && windowOffset + windowLength > 0)
-				{
-					FileInfo fileInfo = new FileInfo(Path.Combine(directoryInfo.FullName, fileEntry.Name.ToString()));
-					await fileEntry.CopyToFileAsync(fileInfo, cancellationToken);
-					copyStats?.Update(1, fileEntry.Length);
-				}
-				windowOffset -= fileEntry.Length;
+				OutputFile outputFile = new OutputFile(outputDir, fileEntry);
+				await FindOutputChunksAsync(outputFile, 0, fileEntry.Target, chunks, cancellationToken);
+			}
+
+			foreach (DirectoryEntry directoryEntry in node._nameToDirectoryEntry.Values)
+			{
+				DirectoryNode subDirectoryNode = await directoryEntry.Handle.ReadBlobAsync(options, cancellationToken);
+				await FindOutputChunksAsync(CombinePaths(outputDir.Path, directoryEntry.Name), subDirectoryNode, chunks, options, cancellationToken);
 			}
+		}
 
-			foreach (DirectoryEntry directoryEntry in _nameToDirectoryEntry.Values)
+		static async Task FindOutputChunksAsync(OutputFile outputFile, long offset, ChunkedDataNodeRef dataRef, ChannelWriter<OutputChunk> chunks, CancellationToken cancellationToken)
+		{
+			if (dataRef.Type == ChunkedDataNodeType.Leaf)
 			{
-				// Traverse into any directory that overlaps with the window (window starts before end of the directory, and window ends at or beyond the start of the directory)
-				if (windowOffset < directoryEntry.Length && windowOffset + windowLength >= 0)
+				await chunks.WriteAsync(new OutputChunk(outputFile, offset, dataRef.Length, dataRef.Handle), cancellationToken);
+			}
+			else
+			{
+				using BlobData data = await dataRef.Handle.ReadBlobDataAsync(cancellationToken);
+				if (data.Type.Guid == LeafChunkedDataNodeConverter.BlobType.Guid)
+				{
+					await chunks.WriteAsync(new OutputChunk(outputFile, offset, dataRef.Length, dataRef.Handle), cancellationToken);
+				}
+				else
 				{
-					DirectoryInfo subDirectoryInfo = directoryInfo.CreateSubdirectory(directoryEntry.Name.ToString());
-					DirectoryNode subDirectoryNode = await directoryEntry.Handle.ReadBlobAsync(options, cancellationToken);
-					await subDirectoryNode.CopyToDirectoryInternalAsync(subDirectoryInfo, windowOffset, windowLength, copyStats, options, logger, cancellationToken);
+					InteriorChunkedDataNode interiorNode = BlobSerializer.Deserialize<InteriorChunkedDataNode>(data);
+					foreach (ChunkedDataNodeRef childRef in interiorNode.Children)
+					{
+						await FindOutputChunksAsync(outputFile, offset, childRef, chunks, cancellationToken);
+						offset += childRef.Length;
+					}
 				}
-				windowOffset -= directoryEntry.Length;
+			}
+		}
+
+		static string CombinePaths(string basePath, string nextPath)
+		{
+			if (basePath.Length > 0)
+			{
+				return $"{basePath}/{nextPath}";
+			}
+			else
+			{
+				return nextPath;
 			}
 		}
 
