Horde: Merge IStorageWriter and IBlobWriter.

diff --git a/Engine/Source/Programs/AutomationTool/BuildGraph/Tasks/CreateArtifactTask.cs b/Engine/Source/Programs/AutomationTool/BuildGraph/Tasks/CreateArtifactTask.cs
--- a/Engine/Source/Programs/AutomationTool/BuildGraph/Tasks/CreateArtifactTask.cs
+++ b/Engine/Source/Programs/AutomationTool/BuildGraph/Tasks/CreateArtifactTask.cs
@@ -117,7 +117,7 @@ namespace AutomationTool.Tasks
 			HttpStorageClientFactory httpStorageClientFactory = serviceProvider.GetRequiredService<HttpStorageClientFactory>();
 			using (IStorageClient client = httpStorageClientFactory.CreateClient(response.NamespaceId, response.Token))
 			{
-				await using (IStorageWriter writer = client.CreateWriter(response.RefName))
+				await using (IBlobWriter writer = client.CreateBlobWriter(response.RefName))
 				{
 					DirectoryReference baseDir = ResolveDirectory(Parameters.BaseDir);
 					List<FileInfo> files = ResolveFilespec(baseDir, Parameters.Files, TagNameToFileSet).Select(x => x.ToFileInfo()).ToList();
diff --git a/Engine/Source/Programs/AutomationTool/BuildGraph/Tasks/DeployToolTask.cs b/Engine/Source/Programs/AutomationTool/BuildGraph/Tasks/DeployToolTask.cs
--- a/Engine/Source/Programs/AutomationTool/BuildGraph/Tasks/DeployToolTask.cs
+++ b/Engine/Source/Programs/AutomationTool/BuildGraph/Tasks/DeployToolTask.cs
@@ -156,7 +156,7 @@ namespace AutomationTool.Tasks
 			using HttpStorageBackend httpStorageBackend = new HttpStorageBackend(basePath, CreateHttpClient, Logger);
 			using BundleStorageClient storageClient = new BundleStorageClient(httpStorageBackend, BundleCache.None, Logger);
 
-			await using (IStorageWriter treeWriter = storageClient.CreateWriter())
+			await using (IBlobWriter treeWriter = storageClient.CreateBlobWriter())
 			{
 				DirectoryNode sandbox = new DirectoryNode();
 				if (Parameters.File != null)
diff --git a/Engine/Source/Programs/Horde/Horde.Agent/Execution/JobExecutor.cs b/Engine/Source/Programs/Horde/Horde.Agent/Execution/JobExecutor.cs
--- a/Engine/Source/Programs/Horde/Horde.Agent/Execution/JobExecutor.cs
+++ b/Engine/Source/Programs/Horde/Horde.Agent/Execution/JobExecutor.cs
@@ -429,20 +429,20 @@ namespace Horde.Agent.Execution
 			return StorageFactory.CreateClient(namespaceId, token, withBackendCache);
 		}
 
-		IStorageWriter CreateStorageWriter(IStorageClient client, string? basePath, ILogger logger)
+		IBlobWriter CreateStorageWriter(IStorageClient client, string? basePath, ILogger logger)
 		{
 			if (JobOptions.BundleVersion != 0 && client is BundleStorageClient bundleClient)
 			{
 				logger.LogInformation("Using bundle version {Version}", JobOptions.BundleVersion);
-				return bundleClient.CreateWriter(basePath, new BundleOptions { MaxVersion = (BundleVersion)JobOptions.BundleVersion });
+				return bundleClient.CreateBlobWriter(basePath, new BundleOptions { MaxVersion = (BundleVersion)JobOptions.BundleVersion });
 			}
 			else
 			{
-				return client.CreateWriter(basePath);
+				return client.CreateBlobWriter(basePath);
 			}
 		}
 
-		IStorageWriter CreateStorageWriter(IStorageClient client, RefName refName, ILogger logger) => CreateStorageWriter(client, refName.Text.ToString(), logger);
+		IBlobWriter CreateStorageWriter(IStorageClient client, RefName refName, ILogger logger) => CreateStorageWriter(client, refName.Text.ToString(), logger);
 
 		protected virtual async Task<bool> SetupAsync(BeginStepResponse step, DirectoryReference workspaceDir, DirectoryReference? sharedStorageDir, bool? useP4, ILogger logger, CancellationToken cancellationToken)
 		{
@@ -519,7 +519,7 @@ namespace Horde.Agent.Execution
 					Stopwatch timer = Stopwatch.StartNew();
 
 					IBlobHandle<DirectoryNode> rootNodeRef;
-					await using (IStorageWriter treeWriter = CreateStorageWriter(storage, artifact.RefName, logger))
+					await using (IBlobWriter treeWriter = CreateStorageWriter(storage, artifact.RefName, logger))
 					{
 						DirectoryNode buildGraphNode = new DirectoryNode();
 						await buildGraphNode.AddFilesAsync(workspaceDir, buildGraphFiles, treeWriter, cancellationToken: cancellationToken);
@@ -849,7 +849,7 @@ namespace Horde.Agent.Execution
 				using IStorageClient storage = CreateStorageClient(new NamespaceId(artifact.NamespaceId), artifact.Token);
 
 				IBlobHandle<DirectoryNode> rootRef;
-				await using (IStorageWriter writer = CreateStorageWriter(storage, new RefName(artifact.RefName), logger))
+				await using (IBlobWriter writer = CreateStorageWriter(storage, new RefName(artifact.RefName), logger))
 				{
 					try
 					{
@@ -1052,7 +1052,7 @@ namespace Horde.Agent.Execution
 				Stopwatch timer = Stopwatch.StartNew();
 
 				IBlobHandle<DirectoryNode> outputNodeRef;
-				await using (IStorageWriter treeWriter = CreateStorageWriter(storage, artifact.RefName, logger))
+				await using (IBlobWriter treeWriter = CreateStorageWriter(storage, artifact.RefName, logger))
 				{
 					DirectoryNode outputNode = new DirectoryNode();
 
diff --git a/Engine/Source/Programs/Horde/Horde.Agent/Utility/JsonRpcLogSink.cs b/Engine/Source/Programs/Horde/Horde.Agent/Utility/JsonRpcLogSink.cs
--- a/Engine/Source/Programs/Horde/Horde.Agent/Utility/JsonRpcLogSink.cs
+++ b/Engine/Source/Programs/Horde/Horde.Agent/Utility/JsonRpcLogSink.cs
@@ -79,7 +79,7 @@ namespace Horde.Agent.Utility
 		readonly LogBuilder _builder;
 		readonly IJsonRpcLogSink? _inner;
 		readonly IStorageClient _store;
-		readonly IStorageWriter _writer;
+		readonly IBlobWriter _writer;
 		readonly ILogger _logger;
 
 		int _bufferLength;
@@ -99,7 +99,7 @@ namespace Horde.Agent.Utility
 			_builder = new LogBuilder(LogFormat.Json, logger);
 			_inner = inner;
 			_store = store;
-			_writer = store.CreateWriter();
+			_writer = store.CreateBlobWriter();
 			_logger = logger;
 
 			_tailTaskStop = new AsyncEvent();
diff --git a/Engine/Source/Programs/Horde/Horde.Agent/Utility/TempStorage.cs b/Engine/Source/Programs/Horde/Horde.Agent/Utility/TempStorage.cs
--- a/Engine/Source/Programs/Horde/Horde.Agent/Utility/TempStorage.cs
+++ b/Engine/Source/Programs/Horde/Horde.Agent/Utility/TempStorage.cs
@@ -638,7 +638,7 @@ namespace Horde.Storage.Utility
 		/// <param name="logger">Logger for output</param>
 		/// <param name="cancellationToken">Cancellation token for the operation</param>
 		/// <returns></returns>
-		public static async Task<FileEntry> ArchiveTagAsync(DirectoryReference manifestDir, string nodeName, string tagName, DirectoryReference workspaceDir, IEnumerable<FileReference> files, TempStorageBlockRef[] blocks, IStorageWriter writer, ILogger logger, CancellationToken cancellationToken)
+		public static async Task<FileEntry> ArchiveTagAsync(DirectoryReference manifestDir, string nodeName, string tagName, DirectoryReference workspaceDir, IEnumerable<FileReference> files, TempStorageBlockRef[] blocks, IBlobWriter writer, ILogger logger, CancellationToken cancellationToken)
 		{
 			logger.LogInformation("Creating output tag \"{TagName}\"", tagName);
 
@@ -682,7 +682,7 @@ namespace Horde.Storage.Utility
 		/// <param name="logger">Logger for output</param>
 		/// <param name="cancellationToken">Cancellation token for the operation</param>
 		/// <returns>The created manifest instance (which has already been saved to disk).</returns>
-		public static async Task<DirectoryEntry> ArchiveBlockAsync(DirectoryReference manifestDir, string nodeName, string blockName, DirectoryReference workspaceDir, IEnumerable<FileReference> files, IStorageWriter writer, ILogger logger, CancellationToken cancellationToken)
+		public static async Task<DirectoryEntry> ArchiveBlockAsync(DirectoryReference manifestDir, string nodeName, string blockName, DirectoryReference workspaceDir, IEnumerable<FileReference> files, IBlobWriter writer, ILogger logger, CancellationToken cancellationToken)
 		{
 			string blockDirectoryName = GetBlockDirectoryName(blockName);
 			logger.LogInformation("Creating output block \"{BlockName}\" (as {DirectoryName})", blockName, blockDirectoryName);
diff --git a/Engine/Source/Programs/Horde/Horde.Server.Tests/BlobStoreTests.cs b/Engine/Source/Programs/Horde/Horde.Server.Tests/BlobStoreTests.cs
--- a/Engine/Source/Programs/Horde/Horde.Server.Tests/BlobStoreTests.cs
+++ b/Engine/Source/Programs/Horde/Horde.Server.Tests/BlobStoreTests.cs
@@ -54,18 +54,31 @@ namespace Horde.Server.Tests
 		{
 			using (BlobData blobData = await store.CreateBlobHandle(locator).ReadBlobDataAsync())
 			{
-				byte[] data = blobData.Data.ToArray();
-				List<BlobLocator> locators = blobData.Refs.ConvertAll(x => x.GetLocator());
+				BlobReader reader = new BlobReader(blobData);
+				ReadOnlyMemory<byte> data = reader.ReadVariableLengthBytes();
+
+				List<IBlobHandle> handles = new List<IBlobHandle>();
+				while (reader.RemainingMemory.Length > 0)
+				{
+					handles.Add(reader.ReadBlobHandle<object>());
+				}
+
+				List<BlobLocator> locators = handles.ConvertAll(x => x.GetLocator());
 				return new Blob(data, locators);
 			}
 		}
 
 		static async ValueTask<BlobLocator> WriteBlobAsync(IStorageClient store, Blob blob)
 		{
-			await using IStorageWriter writer = store.CreateWriter();
-			blob.Data.CopyTo(writer.GetOutputBuffer(0, blob.Data.Length));
+			await using IBlobWriter writer = store.CreateBlobWriter();
+			writer.WriteVariableLengthBytes(blob.Data.Span);
+
+			foreach (BlobLocator reference in blob.References)
+			{
+				writer.WriteBlobHandle(store.CreateBlobHandle(reference).ForType<object>(IoHash.Zero));
+			}
 
-			IBlobHandle handle = await writer.WriteBlobAsync(s_blobType, blob.Data.Length, blob.References.ConvertAll(x => store.CreateBlobHandle(x)), Array.Empty<AliasInfo>());
+			IBlobHandle handle = await writer.CompleteAsync(s_blobType);
 			await handle.FlushAsync();
 
 			return handle.GetLocator();
diff --git a/Engine/Source/Programs/Horde/Horde.Server.Tests/GcServiceTests.cs b/Engine/Source/Programs/Horde/Horde.Server.Tests/GcServiceTests.cs
--- a/Engine/Source/Programs/Horde/Horde.Server.Tests/GcServiceTests.cs
+++ b/Engine/Source/Programs/Horde/Horde.Server.Tests/GcServiceTests.cs
@@ -30,15 +30,15 @@ namespace Horde.Server.Tests
 			using IStorageClient store = StorageService.CreateClient(new NamespaceId("default"));
 
 			Random random = new Random(0);
-			BlobLocator[] blobs = await CreateTestDataAsync(store, 30, 50, 30, 5, random);
+			IBlobHandle<object>[] blobs = await CreateTestDataAsync(store, 30, 50, 30, 5, random);
 
-			HashSet<BlobLocator> roots = new HashSet<BlobLocator>();
+			HashSet<IBlobHandle> roots = new HashSet<IBlobHandle>();
 			for (int idx = 0; idx < 10; idx++)
 			{
 				int blobIdx = (int)(random.NextDouble() * blobs.Length);
 				if (roots.Add(blobs[blobIdx]))
 				{
-					IBlobHandle handle = store.CreateBlobHandle(blobs[blobIdx]);
+					IBlobHandle handle = blobs[blobIdx];
 					await store.WriteRefTargetAsync(new RefName($"ref-{idx}"), handle);
 				}
 			}
@@ -56,27 +56,27 @@ namespace Horde.Server.Tests
 			Assert.IsTrue(remaining.All(x => nodePaths.Contains(x)));
 		}
 
-		static async Task<HashSet<BlobLocator>> FindNodesAsync(IStorageClient store, IEnumerable<BlobLocator> roots)
+		static async Task<HashSet<BlobLocator>> FindNodesAsync(IStorageClient store, IEnumerable<IBlobHandle> roots)
 		{
 			HashSet<BlobLocator> nodes = new HashSet<BlobLocator>();
 			await FindNodesAsync(store, roots, nodes);
 			return nodes;
 		}
 
-		static async Task FindNodesAsync(IStorageClient store, IEnumerable<BlobLocator> roots, HashSet<BlobLocator> nodes)
+		static async Task FindNodesAsync(IStorageClient store, IEnumerable<IBlobHandle> roots, HashSet<BlobLocator> nodes)
 		{
-			foreach (BlobLocator root in roots)
+			foreach (IBlobHandle root in roots)
 			{
-				if (nodes.Add(root))
+				BlobLocator locator = root.GetLocator();
+				if (nodes.Add(locator))
 				{
-					IBlobHandle handle = store.CreateBlobHandle(root);
-					BlobData data = await handle.ReadBlobDataAsync();
-					await FindNodesAsync(store, data.Refs.ConvertAll(x => x.GetLocator()), nodes);
+					BlobData data = await root.ReadBlobDataAsync();
+					await FindNodesAsync(store, data.Refs, nodes);
 				}
 			}
 		}
 
-		static async ValueTask<BlobLocator[]> CreateTestDataAsync(IStorageClient store, int numRoots, int numInterior, int numLeaves, int avgChildren, Random random)
+		static async ValueTask<IBlobHandle<object>[]> CreateTestDataAsync(IStorageClient store, int numRoots, int numInterior, int numLeaves, int avgChildren, Random random)
 		{
 			int firstRoot = 0;
 			int firstInterior = firstRoot + numRoots;
@@ -102,19 +102,23 @@ namespace Horde.Server.Tests
 
 			BlobType blobType = new BlobType(Guid.Parse("{AFDF76A7-4DEE-5333-F5B5-37B8451251CA}"), 0);
 
-			BlobLocator[] locators = new BlobLocator[children.Length];
+			IBlobHandle<object>[] handles = new IBlobHandle<object>[children.Length];
 			for (int idx = numNodes - 1; idx >= 0; idx--)
 			{
-				IBlobHandle handle;
-				await using (IStorageWriter writer = store.CreateWriter("gctest"))
+				IBlobHandle<object> handle;
+				await using (IBlobWriter writer = store.CreateBlobWriter("gctest"))
 				{
-					List<IBlobHandle> imports = children[idx].ConvertAll(x => store.CreateBlobHandle(locators[x]));
-					handle = await writer.WriteBlobAsync(blobType, 0, imports, Array.Empty<AliasInfo>());
+					List<IBlobHandle<object>> imports = children[idx].ConvertAll(x => handles[x]);
+					foreach (IBlobHandle<object> import in imports)
+					{
+						writer.WriteBlobHandle(import);
+					}
+					handle = await writer.CompleteAsync<object>(blobType);
 				}
-				locators[idx] = handle.GetLocator();
+				handles[idx] = handle;
 			}
 
-			return locators;
+			return handles;
 		}
 	}
 }
diff --git a/Engine/Source/Programs/Horde/Horde.Server.Tests/StorageServiceTests.cs b/Engine/Source/Programs/Horde/Horde.Server.Tests/StorageServiceTests.cs
--- a/Engine/Source/Programs/Horde/Horde.Server.Tests/StorageServiceTests.cs
+++ b/Engine/Source/Programs/Horde/Horde.Server.Tests/StorageServiceTests.cs
@@ -2,6 +2,7 @@
 
 using System;
 using System.Threading.Tasks;
+using EpicGames.Core;
 using EpicGames.Horde.Storage;
 using Horde.Server.Storage;
 using Microsoft.Extensions.DependencyInjection;
@@ -27,16 +28,19 @@ namespace Horde.Server.Tests
 			IBlobHandle handle1a;
 			IBlobHandle handle1b;
 			IBlobHandle handle2;
-			await using (IStorageWriter writer = client.CreateWriter())
+			await using (IBlobWriter writer = client.CreateBlobWriter())
 			{
-				data1.CopyTo(writer.GetOutputBuffer(0, data1.Length));
-				handle1a = await writer.WriteBlobAsync(type1, data1.Length, Array.Empty<IBlobHandle>(), new AliasInfo[] { new AliasInfo("foo", 2) });
+				writer.WriteFixedLengthBytes(data1);
+				writer.AddAlias("foo", 2);
+				handle1a = await writer.CompleteAsync(type1);
 
-				data1.CopyTo(writer.GetOutputBuffer(0, data1.Length));
-				handle1b = await writer.WriteBlobAsync(type1, data1.Length, Array.Empty<IBlobHandle>(), new AliasInfo[] { new AliasInfo("foo", 1) });
+				writer.WriteFixedLengthBytes(data1);
+				writer.AddAlias("foo", 1);
+				handle1b = await writer.CompleteAsync(type1);
 
-				data2.CopyTo(writer.GetOutputBuffer(0, data2.Length));
-				handle2 = await writer.WriteBlobAsync(type2, data2.Length, Array.Empty<IBlobHandle>(), new AliasInfo[] { new AliasInfo("bar") });
+				writer.WriteFixedLengthBytes(data2);
+				writer.AddAlias("bar", 0);
+				handle2 = await writer.CompleteAsync(type2);
 			}
 
 			BlobAlias[] aliases;
diff --git a/Engine/Source/Programs/Horde/Horde.Server/Commands/Generate/ToolDataCommand.cs b/Engine/Source/Programs/Horde/Horde.Server/Commands/Generate/ToolDataCommand.cs
--- a/Engine/Source/Programs/Horde/Horde.Server/Commands/Generate/ToolDataCommand.cs
+++ b/Engine/Source/Programs/Horde/Horde.Server/Commands/Generate/ToolDataCommand.cs
@@ -54,7 +54,7 @@ namespace Horde.Server.Commands.Generate
 				using (IStorageClient client = BundleStorageClient.CreateFromDirectory(bundleDir, bundleCache, logger))
 				{
 					IBlobHandle<DirectoryNode> dirNodeRef;
-					await using (IStorageWriter writer = client.CreateWriter(refName))
+					await using (IBlobWriter writer = client.CreateBlobWriter(refName))
 					{
 						DirectoryNode dirNode = new DirectoryNode();
 						await dirNode.AddFilesAsync(InputDir.ToDirectoryInfo(), writer);
diff --git a/Engine/Source/Programs/Horde/Horde.Server/Ddc/BlobService.cs b/Engine/Source/Programs/Horde/Horde.Server/Ddc/BlobService.cs
--- a/Engine/Source/Programs/Horde/Horde.Server/Ddc/BlobService.cs
+++ b/Engine/Source/Programs/Horde/Horde.Server/Ddc/BlobService.cs
@@ -144,15 +144,16 @@ namespace Horde.Server.Ddc
 			using IStorageClient storageClient = _storageClientFactory.CreateClient(ns);
 
 			IBlobHandle blobHandle;
-			await using (IStorageWriter writer = storageClient.CreateWriter())
+			await using (IBlobWriter writer = storageClient.CreateBlobWriter())
 			{
 				int contentLength = (int)content.Length;
-				Memory<byte> memory = writer.GetOutputBuffer(0, contentLength).Slice(0, contentLength);
+				Memory<byte> memory = writer.GetMemory(contentLength).Slice(0, contentLength);
 
 				using Stream stream = content.GetStream();
 				await stream.ReadFixedLengthBytesAsync(memory, cancellationToken);
+				writer.Advance(contentLength);
 
-				blobHandle = await writer.WriteBlobAsync(s_rawBlobType, memory.Length, Array.Empty<IBlobHandle>(), Array.Empty<AliasInfo>(), cancellationToken);
+				blobHandle = await writer.CompleteAsync(s_rawBlobType, cancellationToken);
 				await writer.FlushAsync(cancellationToken);
 			}
 
diff --git a/Engine/Source/Programs/Horde/Horde.Server/Ddc/RefService.cs b/Engine/Source/Programs/Horde/Horde.Server/Ddc/RefService.cs
--- a/Engine/Source/Programs/Horde/Horde.Server/Ddc/RefService.cs
+++ b/Engine/Source/Programs/Horde/Horde.Server/Ddc/RefService.cs
@@ -79,7 +79,7 @@ namespace Horde.Server.Ddc
 				RefName refName = GetRefName(bucket, key);
 
 				IBlobHandle<DdcRefNode> refNodeRef;
-				await using (IStorageWriter writer = storageClient.CreateWriter(refName))
+				await using (IBlobWriter writer = storageClient.CreateBlobWriter(refName))
 				{
 					DdcRefNode refNode = new DdcRefNode(blobHash.Hash);
 					refNode.References.Add((blobHash.Hash, blobHandle));
diff --git a/Engine/Source/Programs/Horde/Horde.Server/Replicators/PerforceReplicator.cs b/Engine/Source/Programs/Horde/Horde.Server/Replicators/PerforceReplicator.cs
--- a/Engine/Source/Programs/Horde/Horde.Server/Replicators/PerforceReplicator.cs
+++ b/Engine/Source/Programs/Horde/Horde.Server/Replicators/PerforceReplicator.cs
@@ -170,7 +170,7 @@ namespace Horde.Server.Replicators
 				public long _sizeWritten;
 				public readonly IncrementalHash Hash;
 
-				public Handle(IStorageWriter writer, ChunkingOptions options)
+				public Handle(IBlobWriter writer, ChunkingOptions options)
 				{
 					Hash = IncrementalHash.CreateHash(HashAlgorithmName.MD5);
 					FileWriter = new ChunkedDataWriter(writer, options, BlobSerializerOptions.Default);
@@ -183,13 +183,13 @@ namespace Horde.Server.Replicators
 				}
 			}
 
-			readonly IStorageWriter _writer;
+			readonly IBlobWriter _writer;
 			readonly ChunkingOptions _options;
 			readonly Stack<Handle> _freeHandles = new Stack<Handle>();
 			readonly Dictionary<int, Handle> _openHandles = new Dictionary<int, Handle>();
 			readonly ILogger _logger;
 
-			public FileWriter(IStorageWriter writer, ChunkingOptions options, ILogger logger)
+			public FileWriter(IBlobWriter writer, ChunkingOptions options, ILogger logger)
 			{
 				_writer = writer;
 				_options = options;
@@ -540,7 +540,7 @@ namespace Horde.Server.Replicators
 
 			// Create the tree writer
 			BundleStorageClient? bundleStore = store as BundleStorageClient;
-			await using IStorageWriter writer = (bundleStore != null)? bundleStore.CreateWriter(refName.ToString(), new BundleOptions { MaxVersion = BundleVersion.LatestV2 }) : store.CreateWriter(refName);
+			await using IBlobWriter writer = (bundleStore != null)? bundleStore.CreateBlobWriter(refName.ToString(), new BundleOptions { MaxVersion = BundleVersion.LatestV2 }) : store.CreateBlobWriter(refName);
 
 			// Keep track of changes to make to the directory structure
 			DirectoryUpdate rootUpdate = new DirectoryUpdate();
diff --git a/Engine/Source/Programs/Horde/Horde.Server/Storage/StorageService.cs b/Engine/Source/Programs/Horde/Horde.Server/Storage/StorageService.cs
--- a/Engine/Source/Programs/Horde/Horde.Server/Storage/StorageService.cs
+++ b/Engine/Source/Programs/Horde/Horde.Server/Storage/StorageService.cs
@@ -316,7 +316,7 @@ namespace Horde.Server.Storage
 			public IBlobHandle CreateBlobHandle(BlobLocator locator) => _inner.CreateBlobHandle(locator);
 
 			/// <inheritdoc/>
-			public IStorageWriter CreateWriter(string? basePath = null) => _inner.CreateWriter(basePath);
+			public IBlobWriter CreateBlobWriter(string? basePath = null) => _inner.CreateBlobWriter(basePath);
 
 			#endregion
 
diff --git a/Engine/Source/Programs/Horde/Horde.Server/Tools/ToolCollection.cs b/Engine/Source/Programs/Horde/Horde.Server/Tools/ToolCollection.cs
--- a/Engine/Source/Programs/Horde/Horde.Server/Tools/ToolCollection.cs
+++ b/Engine/Source/Programs/Horde/Horde.Server/Tools/ToolCollection.cs
@@ -254,7 +254,7 @@ namespace Horde.Server.Tools
 			using IStorageClient client = _storageService.CreateClient(tool.Config.NamespaceId);
 
 			IBlobHandle<DirectoryNode> nodeRef;
-			await using (IStorageWriter writer = client.CreateWriter(refName))
+			await using (IBlobWriter writer = client.CreateBlobWriter(refName))
 			{
 				DirectoryNode directoryNode = new DirectoryNode();
 				await directoryNode.CopyFromZipStreamAsync(stream, writer, new ChunkingOptions(), cancellationToken: cancellationToken);
diff --git a/Engine/Source/Programs/Horde/Horde/Commands/Bundles/BundleCreate.cs b/Engine/Source/Programs/Horde/Horde/Commands/Bundles/BundleCreate.cs
--- a/Engine/Source/Programs/Horde/Horde/Commands/Bundles/BundleCreate.cs
+++ b/Engine/Source/Programs/Horde/Horde/Commands/Bundles/BundleCreate.cs
@@ -81,7 +81,7 @@ namespace Horde.Commands.Bundles
 			// Create the bundle
 			Stopwatch timer = Stopwatch.StartNew();
 
-			await using (IStorageWriter writer = store.CreateWriter())
+			await using (IBlobWriter writer = store.CreateBlobWriter())
 			{
 				ChunkingOptions options = new ChunkingOptions();
 
diff --git a/Engine/Source/Programs/Horde/Horde/Commands/Compute/ComputeRun.cs b/Engine/Source/Programs/Horde/Horde/Commands/Compute/ComputeRun.cs
--- a/Engine/Source/Programs/Horde/Horde/Commands/Compute/ComputeRun.cs
+++ b/Engine/Source/Programs/Horde/Horde/Commands/Compute/ComputeRun.cs
@@ -150,7 +150,7 @@ namespace Horde.Commands.Compute
 
 		static async Task<BlobLocator> CreateSandboxAsync(FileReference taskFile, IStorageClient storage, CancellationToken cancellationToken)
 		{
-			await using IStorageWriter writer = storage.CreateWriter();
+			await using IBlobWriter writer = storage.CreateBlobWriter();
 
 			IBlobHandle<DirectoryNode> sandbox = await writer.WriteFilesAsync(taskFile.Directory, cancellationToken: cancellationToken);
 			await writer.FlushAsync(cancellationToken);
diff --git a/Engine/Source/Programs/Horde/Horde/Commands/Tools/ToolUpload.cs b/Engine/Source/Programs/Horde/Horde/Commands/Tools/ToolUpload.cs
--- a/Engine/Source/Programs/Horde/Horde/Commands/Tools/ToolUpload.cs
+++ b/Engine/Source/Programs/Horde/Horde/Commands/Tools/ToolUpload.cs
@@ -41,7 +41,7 @@ namespace Horde.Commands
 			using IStorageClient storageClient = _httpStorageClientFactory.CreateClientWithPath($"api/v1/tools/{ToolId}");
 
 			IBlobHandle<DirectoryNode> target;
-			await using (IStorageWriter writer = storageClient.CreateWriter())
+			await using (IBlobWriter writer = storageClient.CreateBlobWriter())
 			{
 				target = await writer.WriteFilesAsync(InputDir);
 			}
diff --git a/Engine/Source/Programs/Horde/Samples/RemoteClient/RemoteClientApp.cs b/Engine/Source/Programs/Horde/Samples/RemoteClient/RemoteClientApp.cs
--- a/Engine/Source/Programs/Horde/Samples/RemoteClient/RemoteClientApp.cs
+++ b/Engine/Source/Programs/Horde/Samples/RemoteClient/RemoteClientApp.cs
@@ -112,7 +112,7 @@ namespace RemoteClient
 				// Upload the sandbox to the primary channel.
 				using BundleStorageClient storage =  BundleStorageClient.CreateInMemory(logger);
 
-				await using (IStorageWriter writer = storage.CreateWriter())
+				await using (IBlobWriter writer = storage.CreateBlobWriter())
 				{
 					IBlobHandle<DirectoryNode> sandbox = await writer.WriteFilesAsync(uploadDir);
 					await writer.FlushAsync();
diff --git a/Engine/Source/Programs/Shared/EpicGames.Horde.Tests/BundleStoreTests.cs b/Engine/Source/Programs/Shared/EpicGames.Horde.Tests/BundleStoreTests.cs
--- a/Engine/Source/Programs/Shared/EpicGames.Horde.Tests/BundleStoreTests.cs
+++ b/Engine/Source/Programs/Shared/EpicGames.Horde.Tests/BundleStoreTests.cs
@@ -62,7 +62,7 @@ namespace EpicGames.Horde.Tests
 		{
 			using MemoryStorageBackend memoryStore = new MemoryStorageBackend();
 			using BundleStorageClient store = new BundleStorageClient(memoryStore, BundleCache.None, NullLogger.Instance);
-			await using IStorageWriter writer = store.CreateWriter(options: new BundleOptions { MaxVersion = BundleVersion.ImportHashes, CompressionFormat = BundleCompressionFormat.None });
+			await using IBlobWriter writer = store.CreateBlobWriter(options: new BundleOptions { MaxVersion = BundleVersion.ImportHashes, CompressionFormat = BundleCompressionFormat.None });
 
 			TextNode node = new TextNode("Hello world");
 			IBlobHandle<TextNode> handle = await writer.WriteBlobAsync(node);
@@ -159,7 +159,7 @@ namespace EpicGames.Horde.Tests
 		{
 			// Generate a tree
 			{
-				await using IStorageWriter writer = store.CreateWriter("test", options);
+				await using IBlobWriter writer = store.CreateBlobWriter("test", options);
 
 				SimpleNode node1 = new SimpleNode(new ReadOnlySequence<byte>(new byte[] { 1 }), Array.Empty<IBlobHandle<SimpleNode>>());
 				SimpleNode node2 = new SimpleNode(new ReadOnlySequence<byte>(new byte[] { 2 }), new[] { await writer.WriteBlobAsync(node1) });
@@ -224,7 +224,7 @@ namespace EpicGames.Horde.Tests
 			RefName refName = new RefName("test");
 
 			IBlobHandle<SimpleNode> inputRef;
-			await using (IStorageWriter writer = store.CreateWriter(refName))
+			await using (IBlobWriter writer = store.CreateBlobWriter(refName))
 			{
 				inputRef = await writer.WriteBlobAsync(new SimpleNode(new ReadOnlySequence<byte>(new byte[] { (byte)123 }), Array.Empty<IBlobHandle<SimpleNode>>()));
 			}
@@ -242,7 +242,7 @@ namespace EpicGames.Horde.Tests
 
 			// Generate a tree
 			{
-				await using (IStorageWriter writer = store.CreateWriter(new RefName("test")))
+				await using (IBlobWriter writer = store.CreateBlobWriter(new RefName("test")))
 				{
 					DirectoryNode world = new DirectoryNode();
 					IBlobHandle<DirectoryNode> worldRef = await writer.WriteBlobAsync(world);
@@ -290,7 +290,7 @@ namespace EpicGames.Horde.Tests
 
 			// Generate a tree
 			{
-				await using IStorageWriter writer = store.CreateWriter();
+				await using IBlobWriter writer = store.CreateBlobWriter();
 
 				using ChunkedDataWriter fileWriter = new ChunkedDataWriter(writer, new ChunkingOptions(), BlobSerializerOptions.Default);
 				ChunkedData fileHandle = await fileWriter.CreateAsync(Encoding.UTF8.GetBytes("world"), CancellationToken.None);
@@ -344,7 +344,7 @@ namespace EpicGames.Horde.Tests
 			// Generate a tree
 			ChunkedDataNodeRef nodeRef;
 			{
-				await using IStorageWriter writer = store.CreateWriter(options: new BundleOptions { MaxBlobSize = 1024 });
+				await using IBlobWriter writer = store.CreateBlobWriter(options: new BundleOptions { MaxBlobSize = 1024 });
 
 				ChunkingOptions options = new ChunkingOptions();
 				options.LeafOptions = new LeafChunkedDataNodeOptions(128, 256, 64 * 1024);
@@ -392,7 +392,7 @@ namespace EpicGames.Horde.Tests
 			// Generate a tree
 			DirectoryNode root;
 			{
-				await using DedupeStorageWriter writer = new DedupeStorageWriter(store.CreateWriter(options: new BundleOptions { MaxBlobSize = 1024 }));
+				await using DedupeStorageWriter writer = new DedupeStorageWriter(store.CreateBlobWriter(options: new BundleOptions { MaxBlobSize = 1024 }));
 
 				ChunkingOptions options = new ChunkingOptions();
 				options.LeafOptions = new LeafChunkedDataNodeOptions(128, 256, 64 * 1024);
diff --git a/Engine/Source/Programs/Shared/EpicGames.Horde.Tests/BundleTests.cs b/Engine/Source/Programs/Shared/EpicGames.Horde.Tests/BundleTests.cs
--- a/Engine/Source/Programs/Shared/EpicGames.Horde.Tests/BundleTests.cs
+++ b/Engine/Source/Programs/Shared/EpicGames.Horde.Tests/BundleTests.cs
@@ -72,7 +72,7 @@ namespace EpicGames.Horde.Tests
 		public async Task BasicChunkingTestsAsync()
 		{
 			RefName refName = new RefName("test");
-			await using IStorageWriter writer = _storage.CreateWriter(refName);
+			await using IBlobWriter writer = _storage.CreateBlobWriter(refName);
 
 			ChunkingOptions options = new ChunkingOptions();
 			options.LeafOptions = new LeafChunkedDataNodeOptions(8, 8, 8);
@@ -195,7 +195,7 @@ namespace EpicGames.Horde.Tests
 			IStorageClient store = _storage;
 
 			IBlobHandle<DirectoryNode> rootRef;
-			await using (IStorageWriter writer = store.CreateWriter())
+			await using (IBlobWriter writer = store.CreateBlobWriter())
 			{
 				DirectoryNode world = new DirectoryNode();
 				IBlobHandle<DirectoryNode> worldRef = await writer.WriteBlobAsync(world);
@@ -254,7 +254,7 @@ namespace EpicGames.Horde.Tests
 			BundleOptions options = new BundleOptions();
 			options.MaxBlobSize = 1;
 
-			await using (IStorageWriter writer = new DedupeStorageWriter(_storage.CreateWriter(options: options)))
+			await using (IBlobWriter writer = new DedupeStorageWriter(_storage.CreateBlobWriter(options: options)))
 			{
 				DirectoryNode root = new DirectoryNode();
 				for (int idx = 1; idx <= 3; idx++)
@@ -282,7 +282,7 @@ namespace EpicGames.Horde.Tests
 			RefName refName = new RefName("ref");
 
 			{
-				await using (IStorageWriter writer = _storage.CreateWriter(options: options))
+				await using (IBlobWriter writer = _storage.CreateBlobWriter(options: options))
 				{
 					IBlobHandle<DirectoryNode> rootRef = await writer.WriteBlobAsync(new DirectoryNode());
 					for (int idx = 4; idx >= 1; idx--)
diff --git a/Engine/Source/Programs/Shared/EpicGames.Horde.Tests/ChunkingTests.cs b/Engine/Source/Programs/Shared/EpicGames.Horde.Tests/ChunkingTests.cs
--- a/Engine/Source/Programs/Shared/EpicGames.Horde.Tests/ChunkingTests.cs
+++ b/Engine/Source/Programs/Shared/EpicGames.Horde.Tests/ChunkingTests.cs
@@ -52,7 +52,7 @@ namespace EpicGames.Horde.Tests
 			using KeyValueStorageClient store = KeyValueStorageClient.CreateInMemory();
 
 			const string RefName = "hello";
-			await using (IStorageWriter writer = store.CreateWriter(RefName))
+			await using (IBlobWriter writer = store.CreateBlobWriter(RefName))
 			{
 				ChunkingOptions options = new ChunkingOptions();
 				options.LeafOptions = new LeafChunkedDataNodeOptions(64, 64, 64);
@@ -96,7 +96,7 @@ namespace EpicGames.Horde.Tests
 
 			using KeyValueStorageClient store = KeyValueStorageClient.CreateInMemory();
 
-			await using IStorageWriter writer = store.CreateWriter();
+			await using IBlobWriter writer = store.CreateBlobWriter();
 
 			byte[] data = new byte[4096];
 			new Random(0).NextBytes(data);
diff --git a/Engine/Source/Programs/Shared/EpicGames.Horde.Tests/ComputeSocketTests.cs b/Engine/Source/Programs/Shared/EpicGames.Horde.Tests/ComputeSocketTests.cs
--- a/Engine/Source/Programs/Shared/EpicGames.Horde.Tests/ComputeSocketTests.cs
+++ b/Engine/Source/Programs/Shared/EpicGames.Horde.Tests/ComputeSocketTests.cs
@@ -174,7 +174,7 @@ namespace EpicGames.Horde.Tests
 					}
 
 					using BundleStorageClient storage = BundleStorageClient.CreateInMemory(NullLogger.Instance);
-					await using (IStorageWriter treeWriter = storage.CreateWriter())
+					await using (IBlobWriter blobWriter = storage.CreateBlobWriter())
 					{
 						FileReference file = FileReference.Combine(tempDir, "subdir/hello.txt");
 						if (FileReference.Exists(file))
@@ -185,19 +185,19 @@ namespace EpicGames.Horde.Tests
 
 						byte[] data = Encoding.UTF8.GetBytes("Hello world");
 
-						using ChunkedDataWriter writer = new ChunkedDataWriter(treeWriter, new ChunkingOptions(), BlobSerializerOptions.Default);
+						using ChunkedDataWriter writer = new ChunkedDataWriter(blobWriter, new ChunkingOptions(), BlobSerializerOptions.Default);
 						ChunkedData chunkedData = await writer.CreateAsync(data, cancellationToken);
 
 						DirectoryNode directory = new DirectoryNode();
 						directory.AddFile("hello.txt", FileEntryFlags.None, data.Length, chunkedData);
 
-						IBlobHandle<DirectoryNode> directoryRef = await treeWriter.WriteBlobAsync(directory, cancellationToken: cancellationToken);
+						IBlobHandle<DirectoryNode> directoryRef = await blobWriter.WriteBlobAsync(directory, cancellationToken: cancellationToken);
 
 						DirectoryNode root = new DirectoryNode();
 						root.AddDirectory(new DirectoryEntry("subdir", directory.Length, directoryRef));
 
-						IBlobHandle<DirectoryNode> handle = await treeWriter.WriteBlobAsync(root, cancellationToken: cancellationToken);
-						await treeWriter.FlushAsync(cancellationToken);
+						IBlobHandle<DirectoryNode> handle = await blobWriter.WriteBlobAsync(root, cancellationToken: cancellationToken);
+						await blobWriter.FlushAsync(cancellationToken);
 
 						await channel.UploadFilesAsync("", handle.GetLocator(), storage.Backend, cancellationToken);
 
diff --git a/Engine/Source/Programs/Shared/EpicGames.Horde.Tests/LogTests.cs b/Engine/Source/Programs/Shared/EpicGames.Horde.Tests/LogTests.cs
--- a/Engine/Source/Programs/Shared/EpicGames.Horde.Tests/LogTests.cs
+++ b/Engine/Source/Programs/Shared/EpicGames.Horde.Tests/LogTests.cs
@@ -70,7 +70,7 @@ namespace EpicGames.Horde.Tests
 
 			// Flush it to storage, and read the finished log node
 			IBlobHandle<LogNode> logRef;
-			await using (IStorageWriter writer = store.CreateWriter())
+			await using (IBlobWriter writer = store.CreateBlobWriter())
 			{
 				logRef = await builder.FlushAsync(writer, true, null, CancellationToken.None);
 			}
@@ -119,7 +119,7 @@ namespace EpicGames.Horde.Tests
 			}
 
 			IBlobHandle<LogNode> rootNodeRef;
-			await using (IStorageWriter writer = store.CreateWriter())
+			await using (IBlobWriter writer = store.CreateBlobWriter())
 			{
 				rootNodeRef = await builder.FlushAsync(writer, true, null, CancellationToken.None);
 			}
diff --git a/Engine/Source/Programs/Shared/EpicGames.Horde.Tests/StorageClientTests.cs b/Engine/Source/Programs/Shared/EpicGames.Horde.Tests/StorageClientTests.cs
--- a/Engine/Source/Programs/Shared/EpicGames.Horde.Tests/StorageClientTests.cs
+++ b/Engine/Source/Programs/Shared/EpicGames.Horde.Tests/StorageClientTests.cs
@@ -76,7 +76,7 @@ namespace EpicGames.Horde.Tests
 		static async Task TestBasicAsync(IStorageClient store)
 		{
 			IBlobHandle<TestNode> nodeRef;
-			await using (IStorageWriter writer = store.CreateWriter())
+			await using (IBlobWriter writer = store.CreateBlobWriter())
 			{
 				nodeRef = await writer.WriteBlobAsync(new TestNode(123));
 			}
@@ -94,7 +94,7 @@ namespace EpicGames.Horde.Tests
 			using BundleStorageClient store = BundleStorageClient.CreateInMemory(NullLogger.Instance);
 
 			IBlobHandle<TestNode> nodeRef2;
-			await using (IStorageWriter writer = store.CreateWriter())
+			await using (IBlobWriter writer = store.CreateBlobWriter())
 			{
 				IBlobHandle<TestNode> nodeRef1 = await writer.WriteBlobAsync(new TestNode(123));
 				nodeRef2 = await writer.WriteBlobAsync(new TestNode(456, nodeRef1));
@@ -117,7 +117,7 @@ namespace EpicGames.Horde.Tests
 
 			using BundleStorageClient store = BundleStorageClient.CreateInMemory(NullLogger.Instance);
 
-			await using IStorageWriter writer = store.CreateWriter(options: new BundleOptions { MinCompressionPacketSize = 100, MaxBlobSize = 1024 * 1024, MaxVersion = BundleVersion.LatestV2 });
+			await using IBlobWriter writer = store.CreateBlobWriter(options: new BundleOptions { MinCompressionPacketSize = 100, MaxBlobSize = 1024 * 1024, MaxVersion = BundleVersion.LatestV2 });
 			IBlobHandle<TestNode> nodeRef1 = await writer.WriteBlobAsync(new TestNode(123) { Padding = new byte[1024] });
 			await writer.FlushAsync();
 			IBlobHandle<TestNode> nodeRef2 = await writer.WriteBlobAsync(new TestNode(456, nodeRef1) { Padding = new byte[1024] });
diff --git a/Engine/Source/Programs/Shared/EpicGames.Horde/Logs/LogIndexNode.cs b/Engine/Source/Programs/Shared/EpicGames.Horde/Logs/LogIndexNode.cs
--- a/Engine/Source/Programs/Shared/EpicGames.Horde/Logs/LogIndexNode.cs
+++ b/Engine/Source/Programs/Shared/EpicGames.Horde/Logs/LogIndexNode.cs
@@ -87,7 +87,7 @@ namespace EpicGames.Horde.Logs
 		/// <param name="options">Options to control serialization</param>
 		/// <param name="cancellationToken"></param>
 		/// <returns>New log index with the given blocks appended</returns>
-		public async ValueTask<LogIndexNode> AppendAsync(IStorageWriter writer, IReadOnlyList<LogChunkNode> appendPlainTextChunks, BlobSerializerOptions? options, CancellationToken cancellationToken)
+		public async ValueTask<LogIndexNode> AppendAsync(IBlobWriter writer, IReadOnlyList<LogChunkNode> appendPlainTextChunks, BlobSerializerOptions? options, CancellationToken cancellationToken)
 		{
 			using IScope scope = GlobalTracer.Instance.BuildSpan("LogIndex.Append").StartActive();
 
diff --git a/Engine/Source/Programs/Shared/EpicGames.Horde/Logs/LogNode.cs b/Engine/Source/Programs/Shared/EpicGames.Horde/Logs/LogNode.cs
--- a/Engine/Source/Programs/Shared/EpicGames.Horde/Logs/LogNode.cs
+++ b/Engine/Source/Programs/Shared/EpicGames.Horde/Logs/LogNode.cs
@@ -263,7 +263,7 @@ namespace EpicGames.Horde.Logs
 		/// <param name="complete">Whether the log is complete</param>
 		/// <param name="options">Options for serialization</param>
 		/// <param name="cancellationToken">Cancellation token for the operation</param>
-		public async Task<IBlobHandle<LogNode>> FlushAsync(IStorageWriter writer, bool complete, BlobSerializerOptions? options, CancellationToken cancellationToken)
+		public async Task<IBlobHandle<LogNode>> FlushAsync(IBlobWriter writer, bool complete, BlobSerializerOptions? options, CancellationToken cancellationToken)
 		{
 			// Capture the new data that needs to be written
 			IReadOnlyList<LogChunkNode> writeTextChunks;
diff --git a/Engine/Source/Programs/Shared/EpicGames.Horde/Storage/BlobSerializer.cs b/Engine/Source/Programs/Shared/EpicGames.Horde/Storage/BlobSerializer.cs
--- a/Engine/Source/Programs/Shared/EpicGames.Horde/Storage/BlobSerializer.cs
+++ b/Engine/Source/Programs/Shared/EpicGames.Horde/Storage/BlobSerializer.cs
@@ -175,7 +175,7 @@ namespace EpicGames.Horde.Storage
 		/// <param name="value">The object to serialize</param>
 		/// <param name="cancellationToken">Cancellation token for the operation</param>
 		/// <returns>Handle to the serialized blob</returns>
-		public static ValueTask<IBlobHandle<T>> WriteBlobAsync<T>(this IStorageWriter writer, T value, CancellationToken cancellationToken)
+		public static ValueTask<IBlobHandle<T>> WriteBlobAsync<T>(this IBlobWriter writer, T value, CancellationToken cancellationToken)
 		{
 			return WriteBlobAsync<T>(writer, value, null, cancellationToken);
 		}
@@ -188,11 +188,10 @@ namespace EpicGames.Horde.Storage
 		/// <param name="options">Options to control serialization</param>
 		/// <param name="cancellationToken">Cancellation token for the operation</param>
 		/// <returns>Handle to the serialized blob</returns>
-		public static ValueTask<IBlobHandle<T>> WriteBlobAsync<T>(this IStorageWriter writer, T value, BlobSerializerOptions? options = null, CancellationToken cancellationToken = default)
+		public static async ValueTask<IBlobHandle<T>> WriteBlobAsync<T>(this IBlobWriter writer, T value, BlobSerializerOptions? options = null, CancellationToken cancellationToken = default)
 		{
-			BlobWriter blobWriter = new BlobWriter(writer);
-			BlobType blobType = BlobSerializer.Serialize<T>(blobWriter, value, options);
-			return writer.WriteBlobAsync<T>(blobType, blobWriter.Length, blobWriter.References, cancellationToken);
+			BlobType blobType = BlobSerializer.Serialize<T>(writer, value, options);
+			return await writer.CompleteAsync<T>(blobType, cancellationToken);
 		}
 
 		/// <summary>
diff --git a/Engine/Source/Programs/Shared/EpicGames.Horde/Storage/Bundles/BundleStorageClient.cs b/Engine/Source/Programs/Shared/EpicGames.Horde/Storage/Bundles/BundleStorageClient.cs
--- a/Engine/Source/Programs/Shared/EpicGames.Horde/Storage/Bundles/BundleStorageClient.cs
+++ b/Engine/Source/Programs/Shared/EpicGames.Horde/Storage/Bundles/BundleStorageClient.cs
@@ -166,7 +166,7 @@ namespace EpicGames.Horde.Storage.Bundles
 		}
 
 		/// <inheritdoc/>
-		public IStorageWriter CreateWriter(string? basePath = null, BundleOptions? options = null)
+		public IBlobWriter CreateBlobWriter(string? basePath = null, BundleOptions? options = null)
 		{
 			options ??= BundleOptions.Default;
 
@@ -185,7 +185,7 @@ namespace EpicGames.Horde.Storage.Bundles
 		}
 
 		/// <inheritdoc/>
-		IStorageWriter IStorageClient.CreateWriter(string? basePath) => CreateWriter(basePath);
+		IBlobWriter IStorageClient.CreateBlobWriter(string? basePath) => CreateBlobWriter(basePath);
 
 		#endregion
 
diff --git a/Engine/Source/Programs/Shared/EpicGames.Horde/Storage/Bundles/V1/BundleWriter.cs b/Engine/Source/Programs/Shared/EpicGames.Horde/Storage/Bundles/V1/BundleWriter.cs
--- a/Engine/Source/Programs/Shared/EpicGames.Horde/Storage/Bundles/V1/BundleWriter.cs
+++ b/Engine/Source/Programs/Shared/EpicGames.Horde/Storage/Bundles/V1/BundleWriter.cs
@@ -88,7 +88,7 @@ namespace EpicGames.Horde.Storage.Bundles.V1
 	/// Writes nodes of a tree to an <see cref="IStorageClient"/>, packed into bundles. Each <see cref="BundleWriter"/> instance is single threaded,
 	/// but multiple instances may be written to in parallel.
 	/// </summary>
-	public sealed class BundleWriter : IStorageWriter
+	public sealed class BundleWriter : BlobWriter
 	{
 		// Information about a unique output node. Note that multiple node refs may de-duplicate to the same output node.
 		internal class PendingNode : IBlobHandle
@@ -710,10 +710,10 @@ namespace EpicGames.Horde.Storage.Bundles.V1
 		}
 
 		/// <inheritdoc/>
-		IStorageWriter IStorageWriter.Fork() => new BundleWriter(this);
+		public override IBlobWriter Fork() => new BundleWriter(this);
 
 		/// <inheritdoc/>
-		public async ValueTask DisposeAsync()
+		public override async ValueTask DisposeAsync()
 		{
 			if (!_disposed)
 			{
@@ -747,7 +747,7 @@ namespace EpicGames.Horde.Storage.Bundles.V1
 		/// <param name="usedSize">Current size in the existing buffer that has been written to</param>
 		/// <param name="desiredSize">Desired size of the returned buffer</param>
 		/// <returns>Buffer to be written into.</returns>
-		public Memory<byte> GetOutputBuffer(int usedSize, int desiredSize)
+		public override Memory<byte> GetOutputBuffer(int usedSize, int desiredSize)
 		{
 			PendingBundle currentBundle = GetCurrentBundle();
 			return currentBundle.GetBuffer(usedSize, desiredSize);
@@ -762,7 +762,7 @@ namespace EpicGames.Horde.Storage.Bundles.V1
 		/// <param name="aliases">Aliases for the node</param>
 		/// <param name="cancellationToken">Cancellation token for the operation</param>
 		/// <returns>Handle to the written node</returns>
-		public async ValueTask<IBlobHandle> WriteBlobAsync(BlobType type, int size, IReadOnlyList<IBlobHandle> references, IReadOnlyList<AliasInfo> aliases, CancellationToken cancellationToken = default)
+		public override async ValueTask<IBlobHandle> WriteBlobAsync(BlobType type, int size, IReadOnlyList<IBlobHandle> references, IReadOnlyList<AliasInfo> aliases, CancellationToken cancellationToken = default)
 		{
 			PendingBundle currentBundle = GetCurrentBundle();
 
@@ -804,7 +804,7 @@ namespace EpicGames.Horde.Storage.Bundles.V1
 		/// </summary>
 		/// <param name="cancellationToken">Cancellation token for the operation</param>
 		/// <returns></returns>
-		public async Task FlushAsync(CancellationToken cancellationToken = default)
+		public override async Task FlushAsync(CancellationToken cancellationToken = default)
 		{
 			if (_disposed)
 			{
diff --git a/Engine/Source/Programs/Shared/EpicGames.Horde/Storage/Bundles/V2/BundleWriter.cs b/Engine/Source/Programs/Shared/EpicGames.Horde/Storage/Bundles/V2/BundleWriter.cs
--- a/Engine/Source/Programs/Shared/EpicGames.Horde/Storage/Bundles/V2/BundleWriter.cs
+++ b/Engine/Source/Programs/Shared/EpicGames.Horde/Storage/Bundles/V2/BundleWriter.cs
@@ -16,7 +16,7 @@ namespace EpicGames.Horde.Storage.Bundles.V2
 	/// <summary>
 	/// Implements the primary storage writer interface for V2 bundles. Writes exports into packets, and flushes them to storage in bundles.
 	/// </summary>
-	public sealed class BundleWriter : IStorageWriter
+	public sealed class BundleWriter : BlobWriter
 	{
 		// Packet that is still being built, but may be redirected to a flushed packet
 		internal sealed class PendingPacketHandle : PacketHandle
@@ -296,14 +296,14 @@ namespace EpicGames.Horde.Storage.Bundles.V2
 		}
 
 		/// <inheritdoc/>
-		public async ValueTask DisposeAsync()
+		public override async ValueTask DisposeAsync()
 		{
 			await FlushAsync();
 			_currentBundle.Dispose();
 		}
 
 		/// <inheritdoc/>
-		public async Task FlushAsync(CancellationToken cancellationToken = default)
+		public override async Task FlushAsync(CancellationToken cancellationToken = default)
 		{
 			await _currentBundle.FlushAsync(cancellationToken);
 			_currentBundle.Dispose();
@@ -311,15 +311,15 @@ namespace EpicGames.Horde.Storage.Bundles.V2
 		}
 
 		/// <inheritdoc/>
-		public IStorageWriter Fork()
+		public override IBlobWriter Fork()
 			=> new BundleWriter(_storageClient, _basePath, _bundleCache, _options);
 
 		/// <inheritdoc/>
-		public Memory<byte> GetOutputBuffer(int usedSize, int desiredSize)
+		public override Memory<byte> GetOutputBuffer(int usedSize, int desiredSize)
 			=> _currentBundle.GetOutputBuffer(usedSize, desiredSize);
 
 		/// <inheritdoc/>
-		public async ValueTask<IBlobHandle> WriteBlobAsync(BlobType type, int size, IReadOnlyList<IBlobHandle> references, IReadOnlyList<AliasInfo> aliases, CancellationToken cancellationToken = default)
+		public override async ValueTask<IBlobHandle> WriteBlobAsync(BlobType type, int size, IReadOnlyList<IBlobHandle> references, IReadOnlyList<AliasInfo> aliases, CancellationToken cancellationToken = default)
 		{
 			ExportHandle exportHandle = _currentBundle.CompleteExport(type, size, references, aliases);
 			if (_currentBundle.Length > _options.MaxBlobSize)
diff --git a/Engine/Source/Programs/Shared/EpicGames.Horde/Storage/Clients/KeyValueStorageClient.cs b/Engine/Source/Programs/Shared/EpicGames.Horde/Storage/Clients/KeyValueStorageClient.cs
--- a/Engine/Source/Programs/Shared/EpicGames.Horde/Storage/Clients/KeyValueStorageClient.cs
+++ b/Engine/Source/Programs/Shared/EpicGames.Horde/Storage/Clients/KeyValueStorageClient.cs
@@ -50,7 +50,7 @@ namespace EpicGames.Horde.Storage.Clients
 			public override int GetHashCode() => _locator.GetHashCode();
 		}
 
-		class Writer : IStorageWriter
+		class Writer : BlobWriter
 		{
 			readonly KeyValueStorageClient _outer;
 			readonly string _basePath;
@@ -72,16 +72,16 @@ namespace EpicGames.Horde.Storage.Clients
 			}
 
 			/// <inheritdoc/>
-			public ValueTask DisposeAsync() => new ValueTask();
+			public override ValueTask DisposeAsync() => new ValueTask();
 
 			/// <inheritdoc/>
-			public Task FlushAsync(CancellationToken cancellationToken = default) => Task.CompletedTask;
+			public override Task FlushAsync(CancellationToken cancellationToken = default) => Task.CompletedTask;
 
 			/// <inheritdoc/>
-			public IStorageWriter Fork() => new Writer(_outer, _basePath);
+			public override IBlobWriter Fork() => new Writer(_outer, _basePath);
 
 			/// <inheritdoc/>
-			public Memory<byte> GetOutputBuffer(int usedSize, int desiredSize)
+			public override Memory<byte> GetOutputBuffer(int usedSize, int desiredSize)
 			{
 				if (_offset + desiredSize > _data.Length)
 				{
@@ -94,7 +94,7 @@ namespace EpicGames.Horde.Storage.Clients
 			}
 
 			/// <inheritdoc/>
-			public async ValueTask<IBlobHandle> WriteBlobAsync(BlobType type, int size, IReadOnlyList<IBlobHandle> references, IReadOnlyList<AliasInfo> aliases, CancellationToken cancellationToken = default)
+			public override async ValueTask<IBlobHandle> WriteBlobAsync(BlobType type, int size, IReadOnlyList<IBlobHandle> references, IReadOnlyList<AliasInfo> aliases, CancellationToken cancellationToken = default)
 			{
 				ReadOnlyMemory<byte> data = _data.AsMemory(_offset, size);
 				_offset += size;
@@ -143,7 +143,7 @@ namespace EpicGames.Horde.Storage.Clients
 		}
 
 		/// <inheritdoc/>
-		public IStorageWriter CreateWriter(string? basePath = null) => new Writer(this, basePath);
+		public IBlobWriter CreateBlobWriter(string? basePath = null) => new Writer(this, basePath);
 
 		/// <inheritdoc/>
 		public async ValueTask<BlobData> ReadBlobAsync(BlobLocator locator, CancellationToken cancellationToken = default)
diff --git a/Engine/Source/Programs/Shared/EpicGames.Horde/Storage/IBlobReader.cs b/Engine/Source/Programs/Shared/EpicGames.Horde/Storage/IBlobReader.cs
--- a/Engine/Source/Programs/Shared/EpicGames.Horde/Storage/IBlobReader.cs
+++ b/Engine/Source/Programs/Shared/EpicGames.Horde/Storage/IBlobReader.cs
@@ -35,7 +35,7 @@ namespace EpicGames.Horde.Storage
 	/// <summary>
 	/// Reader for blob objects
 	/// </summary>
-	sealed class BlobReader : MemoryReader, IBlobReader
+	public sealed class BlobReader : MemoryReader, IBlobReader
 	{
 		/// <summary>
 		/// Type to deserialize
diff --git a/Engine/Source/Programs/Shared/EpicGames.Horde/Storage/IBlobWriter.cs b/Engine/Source/Programs/Shared/EpicGames.Horde/Storage/IBlobWriter.cs
--- a/Engine/Source/Programs/Shared/EpicGames.Horde/Storage/IBlobWriter.cs
+++ b/Engine/Source/Programs/Shared/EpicGames.Horde/Storage/IBlobWriter.cs
@@ -2,6 +2,8 @@
 
 using System;
 using System.Collections.Generic;
+using System.Threading;
+using System.Threading.Tasks;
 using EpicGames.Core;
 
 namespace EpicGames.Horde.Storage
@@ -9,8 +11,49 @@ namespace EpicGames.Horde.Storage
 	/// <summary>
 	/// Interface for a writer of node objects
 	/// </summary>
-	public interface IBlobWriter : IMemoryWriter
+	public interface IBlobWriter : IMemoryWriter, IAsyncDisposable
 	{
+		/// <summary>
+		/// Accessor for the memory written to the current blob
+		/// </summary>
+		ReadOnlyMemory<byte> WrittenMemory { get; }
+
+		/// <summary>
+		/// Adds an alias to the blob currently being written
+		/// </summary>
+		/// <param name="name">Name of the alias</param>
+		/// <param name="rank">Rank to use when finding blobs by alias</param>
+		/// <param name="data">Inline data to store with the alias</param>
+		void AddAlias(string name, int rank, ReadOnlyMemory<byte> data = default);
+
+		/// <summary>
+		/// Flush any pending nodes to storage
+		/// </summary>
+		/// <param name="cancellationToken">Cancellation token for the operation</param>
+		Task FlushAsync(CancellationToken cancellationToken = default);
+
+		/// <summary>
+		/// Create another writer instance, allowing multiple threads to write in parallel.
+		/// </summary>
+		/// <returns>New writer instance</returns>
+		IBlobWriter Fork();
+
+		/// <summary>
+		/// Finish writing a blob that has been written into the output buffer.
+		/// </summary>
+		/// <param name="type">Type of the node that was written</param>
+		/// <param name="cancellationToken">Cancellation token for the operation</param>
+		/// <returns>Handle to the written node</returns>
+		ValueTask<IBlobHandle> CompleteAsync(BlobType type, CancellationToken cancellationToken = default);
+
+		/// <summary>
+		/// Finish writing a blob that has been written into the output buffer.
+		/// </summary>
+		/// <param name="type">Type of the node that was written</param>
+		/// <param name="cancellationToken">Cancellation token for the operation</param>
+		/// <returns>Handle to the written node</returns>
+		ValueTask<IBlobHandle<T>> CompleteAsync<T>(BlobType type, CancellationToken cancellationToken = default);
+
 		/// <summary>
 		/// Writes a reference to another blob. The blob's hash is serialized to the output stream.
 		/// </summary>
@@ -19,13 +62,20 @@ namespace EpicGames.Horde.Storage
 	}
 
 	/// <summary>
-	/// Writer for node objects, which tracks references to other nodes
+	/// Information about an alias to be added alongside a blob
 	/// </summary>
-	sealed class BlobWriter : IBlobWriter
-	{
-		readonly IStorageWriter _treeWriter;
+	/// <param name="Name">Name of the alias</param>
+	/// <param name="Rank">Rank of the alias</param>
+	/// <param name="Data">Inline data to be stored for the alias</param>
+	public record class AliasInfo(string Name, int Rank, ReadOnlyMemory<byte> Data);
 
+	/// <summary>
+	/// Base class for <see cref="IBlobWriter"/> implementations.
+	/// </summary>
+	public abstract class BlobWriter : IBlobWriter
+	{
 		Memory<byte> _memory;
+		readonly List<AliasInfo> _aliases = new List<AliasInfo>();
 		readonly List<IBlobHandle> _refs = new List<IBlobHandle>();
 		int _length;
 
@@ -42,16 +92,6 @@ namespace EpicGames.Horde.Storage
 		/// </summary>
 		public ReadOnlyMemory<byte> WrittenMemory => _memory.Slice(0, _length);
 
-		/// <summary>
-		/// Constructor
-		/// </summary>
-		/// <param name="treeWriter"></param>
-		public BlobWriter(IStorageWriter treeWriter)
-		{
-			_treeWriter = treeWriter;
-			_memory = treeWriter.GetOutputBuffer(0, 2 * 1024);
-		}
-
 		/// <summary>
 		/// Computes the hash of the written data
 		/// </summary>
@@ -76,12 +116,65 @@ namespace EpicGames.Horde.Storage
 			if (newLength > _memory.Length)
 			{
 				newLength = _length + Math.Max(sizeHint, 1024);
-				_memory = _treeWriter.GetOutputBuffer(_length, Math.Max(_memory.Length * 2, newLength));
+				_memory = GetOutputBuffer(_length, Math.Max(_memory.Length * 2, newLength));
 			}
 			return _memory.Slice(_length);
 		}
 
 		/// <inheritdoc/>
 		public void Advance(int length) => _length += length;
+
+		/// <summary>
+		/// Request a new buffer to write to
+		/// </summary>
+		/// <param name="usedSize">Size of data written to the current buffer</param>
+		/// <param name="desiredSize">Desired size for the buffer</param>
+		/// <returns>New buffer</returns>
+		public abstract Memory<byte> GetOutputBuffer(int usedSize, int desiredSize);
+
+		/// <summary>
+		/// Write the current blob to storage
+		/// </summary>
+		/// <param name="type">Type of the blob</param>
+		/// <param name="size">Size of the blob to write</param>
+		/// <param name="references">References to other blobs</param>
+		/// <param name="aliases">Aliases for the new blob</param>
+		/// <param name="cancellationToken">Cancellation token for the operation</param>
+		/// <returns>New buffer</returns>
+		public abstract ValueTask<IBlobHandle> WriteBlobAsync(BlobType type, int size, IReadOnlyList<IBlobHandle> references, IReadOnlyList<AliasInfo> aliases, CancellationToken cancellationToken);
+
+		/// <inheritdoc/>
+		public void AddAlias(string name, int rank, ReadOnlyMemory<byte> data)
+			=> _aliases.Add(new AliasInfo(name, rank, data));
+
+		/// <inheritdoc/>
+		public abstract Task FlushAsync(CancellationToken cancellationToken = default);
+
+		/// <inheritdoc/>
+		public abstract IBlobWriter Fork();
+
+		/// <inheritdoc/>
+		public async ValueTask<IBlobHandle> CompleteAsync(BlobType type, CancellationToken cancellationToken = default)
+		{
+			IBlobHandle handle = await WriteBlobAsync(type, _length, _refs, _aliases, cancellationToken);
+
+			_memory = default;
+			_length = 0;
+			_refs.Clear();
+			_aliases.Clear();
+
+			return handle;
+		}
+
+		/// <inheritdoc/>
+		public async ValueTask<IBlobHandle<T>> CompleteAsync<T>(BlobType type, CancellationToken cancellationToken = default)
+		{
+			IoHash hash = IoHash.Compute(_memory.Span.Slice(0, _length));
+			IBlobHandle handle = await CompleteAsync(type, cancellationToken);
+			return handle.ForType<T>(hash);
+		}
+
+		/// <inheritdoc/>
+		public abstract ValueTask DisposeAsync();
 	}
 }
diff --git a/Engine/Source/Programs/Shared/EpicGames.Horde/Storage/IStorageClient.cs b/Engine/Source/Programs/Shared/EpicGames.Horde/Storage/IStorageClient.cs
--- a/Engine/Source/Programs/Shared/EpicGames.Horde/Storage/IStorageClient.cs
+++ b/Engine/Source/Programs/Shared/EpicGames.Horde/Storage/IStorageClient.cs
@@ -12,25 +12,9 @@ namespace EpicGames.Horde.Storage
 	/// <summary>
 	/// Options for a new ref
 	/// </summary>
-	public class RefOptions
-	{
-		/// <summary>
-		/// Time until a ref is expired
-		/// </summary>
-		public TimeSpan? Lifetime { get; set; }
-
-		/// <summary>
-		/// Whether to extend the remaining lifetime of a ref whenever it is fetched. Defaults to true.
-		/// </summary>
-		public bool? Extend { get; set; }
-	}
-
-	/// <summary>
-	/// Stored value for a ref
-	/// </summary>
-	/// <param name="Target">Target blob</param>
-	/// <param name="Data">Inline data stored with the ref</param>
-	public record class RefValue(IBlobHandle Target, ReadOnlyMemory<byte> Data);
+	/// <param name="Lifetime">Time until a ref is expired</param>
+	/// <param name="Extend">Whether to extend the remaining lifetime of a ref whenever it is fetched. Defaults to true.</param>
+	public record class RefOptions(TimeSpan? Lifetime = null, bool? Extend = null);
 
 	/// <summary>
 	/// Interface for the storage system.
@@ -51,7 +35,7 @@ namespace EpicGames.Horde.Storage
 		/// </summary>
 		/// <param name="basePath">Base path for any nodes written from the writer.</param>
 		/// <returns>New writer instance. Must be disposed after use.</returns>
-		IStorageWriter CreateWriter(string? basePath = null);
+		IBlobWriter CreateBlobWriter(string? basePath = null);
 
 		#endregion
 
@@ -250,7 +234,7 @@ namespace EpicGames.Horde.Storage
 		/// </summary>
 		/// <param name="store">The store instance to read from</param>
 		/// <param name="refName">Ref name to use as a base path</param>
-		public static IStorageWriter CreateWriter(this IStorageClient store, RefName refName) => store.CreateWriter(refName.ToString());
+		public static IBlobWriter CreateBlobWriter(this IStorageClient store, RefName refName) => store.CreateBlobWriter(refName.ToString());
 
 		#endregion
 
diff --git a/Engine/Source/Programs/Shared/EpicGames.Horde/Storage/IStorageWriter.cs b/Engine/Source/Programs/Shared/EpicGames.Horde/Storage/IStorageWriter.cs
--- a/Engine/Source/Programs/Shared/EpicGames.Horde/Storage/IStorageWriter.cs
+++ b/Engine/Source/Programs/Shared/EpicGames.Horde/Storage/IStorageWriter.cs
@@ -9,62 +9,10 @@ using EpicGames.Core;
 
 namespace EpicGames.Horde.Storage
 {
-	/// <summary>
-	/// Information about an alias to be added alongside a blob
-	/// </summary>
-	/// <param name="Name">Name of the alias</param>
-	/// <param name="Rank">Rank of the alias</param>
-	/// <param name="Data">Inline data to be stored for the alias</param>
-	public record class AliasInfo(string Name, int Rank, ReadOnlyMemory<byte> Data)
-	{
-		/// <summary>
-		/// Constructor
-		/// </summary>
-		public AliasInfo(string name, int rank = 0) : this(name, rank, ReadOnlyMemory<byte>.Empty)
-		{ }
-	}
-
-	/// <summary>
-	/// Interface for writing new nodes to the store
-	/// </summary>
-	public interface IStorageWriter : IAsyncDisposable
-	{
-		/// <summary>
-		/// Create another writer instance, allowing multiple threads to write in parallel.
-		/// </summary>
-		/// <returns>New writer instance</returns>
-		IStorageWriter Fork();
-
-		/// <summary>
-		/// Flush any pending nodes to storage
-		/// </summary>
-		/// <param name="cancellationToken">Cancellation token for the operation</param>
-		Task FlushAsync(CancellationToken cancellationToken = default);
-
-		/// <summary>
-		/// Gets an output buffer for writing.
-		/// </summary>
-		/// <param name="usedSize">Current size in the existing buffer that has been written to</param>
-		/// <param name="desiredSize">Desired size of the returned buffer</param>
-		/// <returns>Buffer to be written into.</returns>
-		Memory<byte> GetOutputBuffer(int usedSize, int desiredSize);
-
-		/// <summary>
-		/// Finish writing a blob that has been written into the output buffer.
-		/// </summary>
-		/// <param name="type">Type of the node that was written</param>
-		/// <param name="size">Used size of the buffer</param>
-		/// <param name="references">References to other nodes</param>
-		/// <param name="aliases">Aliases for this node</param>
-		/// <param name="cancellationToken">Cancellation token for the operation</param>
-		/// <returns>Handle to the written node</returns>
-		ValueTask<IBlobHandle> WriteBlobAsync(BlobType type, int size, IReadOnlyList<IBlobHandle> references, IReadOnlyList<AliasInfo> aliases, CancellationToken cancellationToken = default);
-	}
-
 	/// <summary>
 	/// Index of known nodes that can be used for deduplication.
 	/// </summary>
-	public sealed class DedupeStorageWriter : IStorageWriter
+	public sealed class DedupeStorageWriter : BlobWriter
 	{
 		record BlobKey(IoHash Hash, BlobType Type);
 
@@ -123,7 +71,7 @@ namespace EpicGames.Horde.Storage
 		/// </summary>
 		public const int DefaultMaxKeys = 64 * 1024;
 
-		readonly IStorageWriter _inner;
+		readonly BlobWriter _inner;
 		readonly DedupeCache _cache;
 
 		/// <summary>
@@ -131,32 +79,32 @@ namespace EpicGames.Horde.Storage
 		/// </summary>
 		/// <param name="inner"></param>
 		/// <param name="maxKeys"></param>
-		public DedupeStorageWriter(IStorageWriter inner, int maxKeys = DefaultMaxKeys)
+		public DedupeStorageWriter(IBlobWriter inner, int maxKeys = DefaultMaxKeys)
 		{
-			_inner = inner;
+			_inner = (BlobWriter)inner;
 			_cache = new DedupeCache(maxKeys);
 		}
 
-		private DedupeStorageWriter(IStorageWriter inner, DedupeCache cache)
+		private DedupeStorageWriter(IBlobWriter inner, DedupeCache cache)
 		{
-			_inner = inner;
+			_inner = (BlobWriter)inner;
 			_cache = cache;
 		}
 
 		/// <inheritdoc/>
-		public ValueTask DisposeAsync() => _inner.DisposeAsync();
+		public override ValueTask DisposeAsync() => _inner.DisposeAsync();
 
 		/// <inheritdoc/>
-		public Task FlushAsync(CancellationToken cancellationToken = default) => _inner.FlushAsync(cancellationToken);
+		public override Task FlushAsync(CancellationToken cancellationToken = default) => _inner.FlushAsync(cancellationToken);
 
 		/// <inheritdoc/>
-		public IStorageWriter Fork() => new DedupeStorageWriter(_inner.Fork(), _cache);
+		public override IBlobWriter Fork() => new DedupeStorageWriter(_inner.Fork(), _cache);
 
 		/// <inheritdoc/>
-		public Memory<byte> GetOutputBuffer(int usedSize, int desiredSize) => _inner.GetOutputBuffer(usedSize, desiredSize);
+		public override Memory<byte> GetOutputBuffer(int usedSize, int desiredSize) => _inner.GetOutputBuffer(usedSize, desiredSize);
 
 		/// <inheritdoc/>
-		public async ValueTask<IBlobHandle> WriteBlobAsync(BlobType type, int size, IReadOnlyList<IBlobHandle> references, IReadOnlyList<AliasInfo> aliases, CancellationToken cancellationToken = default)
+		public override async ValueTask<IBlobHandle> WriteBlobAsync(BlobType type, int size, IReadOnlyList<IBlobHandle> references, IReadOnlyList<AliasInfo> aliases, CancellationToken cancellationToken = default)
 		{
 			ReadOnlyMemory<byte> data = _inner.GetOutputBuffer(size, size).Slice(0, size);
 			IoHash hash = IoHash.Compute(data.Span);
@@ -181,29 +129,13 @@ namespace EpicGames.Horde.Storage
 	}
 
 	/// <summary>
-	/// Extension methods for <see cref="IStorageWriter"/>
+	/// Extension methods for <see cref="IBlobWriter"/>
 	/// </summary>
 	public static class StorageWriterExtensions
 	{
 		/// <summary>
-		/// Wraps a <see cref="IStorageWriter"/> with a <see cref="DedupeStorageWriter"/>
+		/// Wraps a <see cref="IBlobWriter"/> with a <see cref="DedupeStorageWriter"/>
 		/// </summary>
-		public static DedupeStorageWriter WithDedupe(this IStorageWriter writer, int maxKeys = DedupeStorageWriter.DefaultMaxKeys) => new DedupeStorageWriter(writer, maxKeys);
-
-		/// <summary>
-		/// Finish writing a node.
-		/// </summary>
-		/// <param name="writer">Writer instance to manipulate</param>
-		/// <param name="type">Type of the node that was written</param>
-		/// <param name="size">Used size of the buffer</param>
-		/// <param name="references">References to other nodes</param>
-		/// <param name="cancellationToken">Cancellation token for the operation</param>
-		/// <returns>Handle to the written node</returns>
-		public static async ValueTask<IBlobHandle<T>> WriteBlobAsync<T>(this IStorageWriter writer, BlobType type, int size, IReadOnlyList<IBlobHandle> references, CancellationToken cancellationToken = default)
-		{
-			IoHash hash = IoHash.Compute(writer.GetOutputBuffer(size, size).Span.Slice(0, size));
-			IBlobHandle blobHandle = await writer.WriteBlobAsync(type, size, references, Array.Empty<AliasInfo>(), cancellationToken);
-			return blobHandle.ForType<T>(hash);
-		}
+		public static DedupeStorageWriter WithDedupe(this IBlobWriter writer, int maxKeys = DedupeStorageWriter.DefaultMaxKeys) => new DedupeStorageWriter(writer, maxKeys);
 	}
 }
diff --git a/Engine/Source/Programs/Shared/EpicGames.Horde/Storage/Nodes/ChunkedDataNode.cs b/Engine/Source/Programs/Shared/EpicGames.Horde/Storage/Nodes/ChunkedDataNode.cs
--- a/Engine/Source/Programs/Shared/EpicGames.Horde/Storage/Nodes/ChunkedDataNode.cs
+++ b/Engine/Source/Programs/Shared/EpicGames.Horde/Storage/Nodes/ChunkedDataNode.cs
@@ -241,7 +241,7 @@ namespace EpicGames.Horde.Storage.Nodes
 		/// <param name="options">Options for finding chunk boundaries</param>
 		/// <param name="cancellationToken">Cancellation token for the operation</param>
 		/// <returns>Hash of the full file data</returns>
-		public static async Task<LeafChunkedData> CreateFromFileAsync(IStorageWriter writer, FileReference file, LeafChunkedDataNodeOptions options, CancellationToken cancellationToken)
+		public static async Task<LeafChunkedData> CreateFromFileAsync(IBlobWriter writer, FileReference file, LeafChunkedDataNodeOptions options, CancellationToken cancellationToken)
 		{
 			using (FileStream stream = FileReference.Open(file, FileMode.Open, FileAccess.Read, FileShare.Read))
 			{
@@ -257,7 +257,7 @@ namespace EpicGames.Horde.Storage.Nodes
 		/// <param name="options">Options for finding chunk boundaries</param>
 		/// <param name="cancellationToken">Cancellation token for the operation</param>
 		/// <returns>Hash of the full file data</returns>
-		public static async Task<LeafChunkedData> CreateFromFileAsync(IStorageWriter writer, FileInfo file, LeafChunkedDataNodeOptions options, CancellationToken cancellationToken)
+		public static async Task<LeafChunkedData> CreateFromFileAsync(IBlobWriter writer, FileInfo file, LeafChunkedDataNodeOptions options, CancellationToken cancellationToken)
 		{
 			using (FileStream stream = file.Open(FileMode.Open, FileAccess.Read, FileShare.Read))
 			{
@@ -273,7 +273,7 @@ namespace EpicGames.Horde.Storage.Nodes
 		/// <param name="options">Options for finding chunk boundaries</param>
 		/// <param name="cancellationToken">Cancellation token for the operation</param>
 		/// <returns>Hash of the full file data</returns>
-		public static Task<LeafChunkedData> CreateFromStreamAsync(IStorageWriter writer, Stream stream, LeafChunkedDataNodeOptions options, CancellationToken cancellationToken)
+		public static Task<LeafChunkedData> CreateFromStreamAsync(IBlobWriter writer, Stream stream, LeafChunkedDataNodeOptions options, CancellationToken cancellationToken)
 		{
 			return CreateFromStreamAsync(writer, stream, options, null, cancellationToken);
 		}
@@ -287,7 +287,7 @@ namespace EpicGames.Horde.Storage.Nodes
 		/// <param name="copyStats">Stats for the copy operation</param>
 		/// <param name="cancellationToken">Cancellation token for the operation</param>
 		/// <returns>Hash of the full file data</returns>
-		internal static async Task<LeafChunkedData> CreateFromStreamAsync(IStorageWriter writer, Stream stream, LeafChunkedDataNodeOptions options, CopyStats? copyStats, CancellationToken cancellationToken)
+		internal static async Task<LeafChunkedData> CreateFromStreamAsync(IBlobWriter writer, Stream stream, LeafChunkedDataNodeOptions options, CopyStats? copyStats, CancellationToken cancellationToken)
 		{
 			using Blake3.Hasher hasher = Blake3.Hasher.New();
 			using IMemoryOwner<byte> readBuffer = MemoryPool<byte>.Shared.Rent(options.MaxSize);
@@ -306,9 +306,9 @@ namespace EpicGames.Horde.Storage.Nodes
 
 				int nextLength = GetChunkLength(readBuffer.Memory.Span.Slice(0, size), options);
 
-				Memory<byte> outputBuffer = writer.GetOutputBuffer(0, nextLength).Slice(0, nextLength);
-				readBuffer.Memory.Slice(0, nextLength).CopyTo(outputBuffer);
-				hasher.Update(outputBuffer.Span);
+				ReadOnlyMemory<byte> nextBlobData = readBuffer.Memory.Slice(0, nextLength);
+				writer.WriteFixedLengthBytes(nextBlobData.Span);
+				hasher.Update(nextBlobData.Span);
 
 				sizeSinceProgressUpdate += nextLength;
 				if (sizeSinceProgressUpdate > 512 * 1024)
@@ -317,7 +317,7 @@ namespace EpicGames.Horde.Storage.Nodes
 					sizeSinceProgressUpdate = 0;
 				}
 
-				IBlobHandle<LeafChunkedDataNode> blobHandle = await writer.WriteBlobAsync<LeafChunkedDataNode>(LeafChunkedDataNodeConverter.BlobType, nextLength, Array.Empty<IBlobHandle>(), cancellationToken);
+				IBlobHandle<LeafChunkedDataNode> blobHandle = await writer.CompleteAsync<LeafChunkedDataNode>(LeafChunkedDataNodeConverter.BlobType, cancellationToken);
 				leafNodeRefs.Add(new ChunkedDataNodeRef(nextLength, blobHandle));
 
 				readBuffer.Memory.Slice(nextLength, size - nextLength).CopyTo(readBuffer.Memory);
@@ -449,7 +449,7 @@ namespace EpicGames.Horde.Storage.Nodes
 		/// <param name="writer">Output writer for new interior nodes</param>
 		/// <param name="cancellationToken">Cancellation token for the operation</param>
 		/// <returns>Handle to the root node of the tree</returns>
-		public static async Task<ChunkedData> CreateTreeAsync(LeafChunkedData leafChunkedData, InteriorChunkedDataNodeOptions options, IStorageWriter writer, BlobSerializerOptions? serializerOptions, CancellationToken cancellationToken)
+		public static async Task<ChunkedData> CreateTreeAsync(LeafChunkedData leafChunkedData, InteriorChunkedDataNodeOptions options, IBlobWriter writer, BlobSerializerOptions? serializerOptions, CancellationToken cancellationToken)
 		{
 			ChunkedDataNodeRef rootRef = await CreateTreeAsync(leafChunkedData.LeafHandles, options, writer, serializerOptions, cancellationToken);
 			return new ChunkedData(leafChunkedData.Hash, rootRef);
@@ -464,7 +464,7 @@ namespace EpicGames.Horde.Storage.Nodes
 		/// <param name="serializerOptions">Options for serializing blobs</param>
 		/// <param name="cancellationToken">Cancellation token for the operation</param>
 		/// <returns>Handle to the root node of the tree</returns>
-		public static async Task<ChunkedDataNodeRef> CreateTreeAsync(List<ChunkedDataNodeRef> nodeRefs, InteriorChunkedDataNodeOptions chunkingOptions, IStorageWriter writer, BlobSerializerOptions? serializerOptions, CancellationToken cancellationToken)
+		public static async Task<ChunkedDataNodeRef> CreateTreeAsync(List<ChunkedDataNodeRef> nodeRefs, InteriorChunkedDataNodeOptions chunkingOptions, IBlobWriter writer, BlobSerializerOptions? serializerOptions, CancellationToken cancellationToken)
 		{
 			List<ChunkedDataNodeRef> handleBuffer = new List<ChunkedDataNodeRef>();
 
diff --git a/Engine/Source/Programs/Shared/EpicGames.Horde/Storage/Nodes/ChunkedDataWriter.cs b/Engine/Source/Programs/Shared/EpicGames.Horde/Storage/Nodes/ChunkedDataWriter.cs
--- a/Engine/Source/Programs/Shared/EpicGames.Horde/Storage/Nodes/ChunkedDataWriter.cs
+++ b/Engine/Source/Programs/Shared/EpicGames.Horde/Storage/Nodes/ChunkedDataWriter.cs
@@ -61,7 +61,7 @@ namespace EpicGames.Horde.Storage.Nodes
 	public record class ChunkedData(IoHash StreamHash, ChunkedDataNodeRef Root);
 
 	/// <summary>
-	/// Utility class for generating FileNode data directly into <see cref="IStorageWriter"/> instances, without constructing node representations first.
+	/// Utility class for generating FileNode data directly into <see cref="IBlobWriter"/> instances, without constructing node representations first.
 	/// </summary>
 	public sealed class ChunkedDataWriter : IDisposable
 	{
@@ -70,7 +70,7 @@ namespace EpicGames.Horde.Storage.Nodes
 		/// </summary>
 		public const int DefaultBufferLength = 32 * 1024;
 
-		readonly IStorageWriter _writer;
+		readonly IBlobWriter _writer;
 		readonly ChunkingOptions _chunkingOptions;
 		readonly BlobSerializerOptions? _serializerOptions;
 		readonly Blake3.Hasher _hasher;
@@ -81,7 +81,6 @@ namespace EpicGames.Horde.Storage.Nodes
 
 		// Leaf node state
 		uint _leafHash;
-		int _leafLength;
 
 		/// <summary>
 		/// Length of the file so far
@@ -94,7 +93,7 @@ namespace EpicGames.Horde.Storage.Nodes
 		/// <param name="writer">Writer for new nodes</param>
 		/// <param name="chunkingOptions">Chunking options</param>
 		/// <param name="serializerOptions">Options for serialization</param>
-		public ChunkedDataWriter(IStorageWriter writer, ChunkingOptions chunkingOptions, BlobSerializerOptions? serializerOptions)
+		public ChunkedDataWriter(IBlobWriter writer, ChunkingOptions chunkingOptions, BlobSerializerOptions? serializerOptions)
 		{
 			_writer = writer;
 			_chunkingOptions = chunkingOptions;
@@ -125,7 +124,6 @@ namespace EpicGames.Horde.Storage.Nodes
 		void ResetLeafState()
 		{
 			_leafHash = 0;
-			_leafLength = 0;
 		}
 
 		/// <summary>
@@ -217,13 +215,9 @@ namespace EpicGames.Horde.Storage.Nodes
 		{
 			for (; ; )
 			{
-				Memory<byte> buffer = _writer.GetOutputBuffer(_leafLength, _leafLength);
-				int appendLength = AppendToLeafNode(buffer.Span.Slice(0, _leafLength), data.Span, ref _leafHash, _chunkingOptions.LeafOptions);
+				int appendLength = AppendToLeafNode(_writer.WrittenMemory.Span, data.Span, ref _leafHash, _chunkingOptions.LeafOptions);
+				_writer.WriteFixedLengthBytes(data.Slice(0, appendLength).Span);
 
-				buffer = _writer.GetOutputBuffer(_leafLength, _leafLength + appendLength);
-				data.Slice(0, appendLength).CopyTo(buffer.Slice(_leafLength));
-
-				_leafLength += appendLength;
 				data = data.Slice(appendLength);
 
 				_totalLength += appendLength;
@@ -346,8 +340,9 @@ namespace EpicGames.Horde.Storage.Nodes
 		/// <returns>Handle to the written leaf node</returns>
 		async ValueTask FlushLeafNodeAsync(CancellationToken cancellationToken)
 		{
-			IBlobHandle<LeafChunkedDataNode> leafHandle = await _writer.WriteBlobAsync<LeafChunkedDataNode>(LeafChunkedDataNodeConverter.BlobType, _leafLength, Array.Empty<IBlobHandle>(), cancellationToken);
-			_leafHandles.Add(new ChunkedDataNodeRef(_leafLength, leafHandle));
+			int leafLength = _writer.WrittenMemory.Length;
+			IBlobHandle<LeafChunkedDataNode> leafHandle = await _writer.CompleteAsync<LeafChunkedDataNode>(LeafChunkedDataNodeConverter.BlobType, cancellationToken);
+			_leafHandles.Add(new ChunkedDataNodeRef(leafLength, leafHandle));
 			ResetLeafState();
 		}
 	}
diff --git a/Engine/Source/Programs/Shared/EpicGames.Horde/Storage/Nodes/DirectoryNode.cs b/Engine/Source/Programs/Shared/EpicGames.Horde/Storage/Nodes/DirectoryNode.cs
--- a/Engine/Source/Programs/Shared/EpicGames.Horde/Storage/Nodes/DirectoryNode.cs
+++ b/Engine/Source/Programs/Shared/EpicGames.Horde/Storage/Nodes/DirectoryNode.cs
@@ -459,20 +459,20 @@ namespace EpicGames.Horde.Storage.Nodes
 
 		#endregion
 
-		/// <inheritdoc cref="AddFilesAsync(DirectoryReference, IEnumerable{FileInfo}, IStorageWriter, ChunkingOptions?, IProgress{ICopyStats}?, BlobSerializerOptions?, CancellationToken)"/>
-		public async Task AddFilesAsync(DirectoryInfo directoryInfo, IStorageWriter writer, ChunkingOptions? options = null, IProgress<ICopyStats>? progress = null, BlobSerializerOptions? serializerOptions = null, CancellationToken cancellationToken = default)
+		/// <inheritdoc cref="AddFilesAsync(DirectoryReference, IEnumerable{FileInfo}, IBlobWriter, ChunkingOptions?, IProgress{ICopyStats}?, BlobSerializerOptions?, CancellationToken)"/>
+		public async Task AddFilesAsync(DirectoryInfo directoryInfo, IBlobWriter writer, ChunkingOptions? options = null, IProgress<ICopyStats>? progress = null, BlobSerializerOptions? serializerOptions = null, CancellationToken cancellationToken = default)
 		{
 			await AddFilesAsync(new DirectoryReference(directoryInfo), directoryInfo.EnumerateFiles("*", SearchOption.AllDirectories).ToList(), writer, options, progress, serializerOptions, cancellationToken);
 		}
 
-		/// <inheritdoc cref="AddFilesAsync(DirectoryReference, IEnumerable{FileInfo}, IStorageWriter, ChunkingOptions?, IProgress{ICopyStats}?, BlobSerializerOptions?, CancellationToken)"/>
-		public Task AddFilesAsync(DirectoryReference baseDir, IEnumerable<FileReference> files, IStorageWriter writer, ChunkingOptions? options = null, IProgress<ICopyStats>? progress = null, BlobSerializerOptions? serializerOptions = null, CancellationToken cancellationToken = default)
+		/// <inheritdoc cref="AddFilesAsync(DirectoryReference, IEnumerable{FileInfo}, IBlobWriter, ChunkingOptions?, IProgress{ICopyStats}?, BlobSerializerOptions?, CancellationToken)"/>
+		public Task AddFilesAsync(DirectoryReference baseDir, IEnumerable<FileReference> files, IBlobWriter writer, ChunkingOptions? options = null, IProgress<ICopyStats>? progress = null, BlobSerializerOptions? serializerOptions = null, CancellationToken cancellationToken = default)
 		{
 			return AddFilesAsync(baseDir, files.Select(x => x.ToFileInfo()).ToList(), writer, options, progress, serializerOptions, cancellationToken);
 		}
 
-		/// <inheritdoc cref="AddFilesAsync(DirectoryReference, IEnumerable{FileInfo}, IStorageWriter, ChunkingOptions?, IProgress{ICopyStats}?, BlobSerializerOptions?, CancellationToken)"/>
-		public Task AddFilesAsync(DirectoryInfo baseDir, IEnumerable<FileInfo> files, IStorageWriter writer, ChunkingOptions? options = null, IProgress<ICopyStats>? progress = null, BlobSerializerOptions? serializerOptions = null, CancellationToken cancellationToken = default)
+		/// <inheritdoc cref="AddFilesAsync(DirectoryReference, IEnumerable{FileInfo}, IBlobWriter, ChunkingOptions?, IProgress{ICopyStats}?, BlobSerializerOptions?, CancellationToken)"/>
+		public Task AddFilesAsync(DirectoryInfo baseDir, IEnumerable<FileInfo> files, IBlobWriter writer, ChunkingOptions? options = null, IProgress<ICopyStats>? progress = null, BlobSerializerOptions? serializerOptions = null, CancellationToken cancellationToken = default)
 		{
 			return AddFilesAsync(new DirectoryReference(baseDir), files.ToList(), writer, options, progress, serializerOptions, cancellationToken);
 		}
@@ -487,7 +487,7 @@ namespace EpicGames.Horde.Storage.Nodes
 		/// <param name="progress">Feedback interface for progress updates</param>
 		/// <param name="serializerOptions">Options controlling serialization</param>
 		/// <param name="cancellationToken">Cancellation token for the operation</param>
-		public async Task AddFilesAsync(DirectoryReference baseDir, IEnumerable<FileInfo> files, IStorageWriter writer, ChunkingOptions? options = null, IProgress<ICopyStats>? progress = null, BlobSerializerOptions? serializerOptions = null, CancellationToken cancellationToken = default)
+		public async Task AddFilesAsync(DirectoryReference baseDir, IEnumerable<FileInfo> files, IBlobWriter writer, ChunkingOptions? options = null, IProgress<ICopyStats>? progress = null, BlobSerializerOptions? serializerOptions = null, CancellationToken cancellationToken = default)
 		{
 			options ??= new ChunkingOptions();
 
@@ -595,9 +595,9 @@ namespace EpicGames.Horde.Storage.Nodes
 			return partitions;
 		}
 
-		static async ValueTask CreateLeafChunkNodesAsync(IStorageWriter writer, IReadOnlyList<FileInfo> files, LeafChunkedData[] leafChunks, int start, int count, CopyStats? copyStats, ChunkingOptions options, CancellationToken cancellationToken)
+		static async ValueTask CreateLeafChunkNodesAsync(IBlobWriter writer, IReadOnlyList<FileInfo> files, LeafChunkedData[] leafChunks, int start, int count, CopyStats? copyStats, ChunkingOptions options, CancellationToken cancellationToken)
 		{
-			await using IStorageWriter writerFork = writer.Fork();
+			await using IBlobWriter writerFork = writer.Fork();
 			for (int idx = start; idx < start + count; idx++)
 			{
 				FileInfo file = files[idx];
@@ -609,7 +609,7 @@ namespace EpicGames.Horde.Storage.Nodes
 			await writerFork.FlushAsync(cancellationToken);
 		}
 
-		static async Task<ChunkedData[]> CreateInteriorChunkNodesAsync(LeafChunkedData[] leafChunkedFiles, InteriorChunkedDataNodeOptions chunkingOptions, IStorageWriter writer, BlobSerializerOptions? serializerOptions, CancellationToken cancellationToken)
+		static async Task<ChunkedData[]> CreateInteriorChunkNodesAsync(LeafChunkedData[] leafChunkedFiles, InteriorChunkedDataNodeOptions chunkingOptions, IBlobWriter writer, BlobSerializerOptions? serializerOptions, CancellationToken cancellationToken)
 		{
 			ChunkedData[] chunkedFiles = new ChunkedData[leafChunkedFiles.Length];
 			for (int idx = 0; idx < leafChunkedFiles.Length; idx++)
@@ -626,7 +626,7 @@ namespace EpicGames.Horde.Storage.Nodes
 		/// <param name="writer">Writer for new node data</param>
 		/// <param name="options">Options controling serialiation</param>
 		/// <param name="cancellationToken">Cancellation token for the operation</param>
-		public Task UpdateAsync(IEnumerable<FileUpdate> updates, IStorageWriter writer, BlobSerializerOptions? options = null, CancellationToken cancellationToken = default)
+		public Task UpdateAsync(IEnumerable<FileUpdate> updates, IBlobWriter writer, BlobSerializerOptions? options = null, CancellationToken cancellationToken = default)
 		{
 			DirectoryUpdate update = new DirectoryUpdate();
 			update.AddFiles(updates);
@@ -640,7 +640,7 @@ namespace EpicGames.Horde.Storage.Nodes
 		/// <param name="writer">Writer for new node data</param>
 		/// <param name="options">Options controling serialiation</param>
 		/// <param name="cancellationToken">Cancellation token for the operation</param>
-		public async Task UpdateAsync(DirectoryUpdate update, IStorageWriter writer, BlobSerializerOptions? options = null, CancellationToken cancellationToken = default)
+		public async Task UpdateAsync(DirectoryUpdate update, IBlobWriter writer, BlobSerializerOptions? options = null, CancellationToken cancellationToken = default)
 		{
 			foreach ((string name, DirectoryUpdate? directory) in update.Directories)
 			{
@@ -678,7 +678,7 @@ namespace EpicGames.Horde.Storage.Nodes
 		/// <param name="options"></param>
 		/// <param name="serializerOptions">Options controling serialiation</param>
 		/// <param name="cancellationToken">Cancellation token for the operation</param>
-		public async Task CopyFromZipStreamAsync(Stream stream, IStorageWriter writer, ChunkingOptions options, BlobSerializerOptions? serializerOptions = null, CancellationToken cancellationToken = default)
+		public async Task CopyFromZipStreamAsync(Stream stream, IBlobWriter writer, ChunkingOptions options, BlobSerializerOptions? serializerOptions = null, CancellationToken cancellationToken = default)
 		{
 			// Create all the leaf nodes
 			List<(ZipArchiveEntry, LeafChunkedData)> entries = new List<(ZipArchiveEntry, LeafChunkedData)>();
@@ -1308,7 +1308,7 @@ namespace EpicGames.Horde.Storage.Nodes
 		/// <summary>
 		/// Writes a tree of files to a storage writer
 		/// </summary>
-		public static async Task<IBlobHandle<DirectoryNode>> WriteFilesAsync(this IStorageWriter writer, DirectoryReference baseDir, ChunkingOptions? options = null, IProgress<ICopyStats>? progress = null, CancellationToken cancellationToken = default)
+		public static async Task<IBlobHandle<DirectoryNode>> WriteFilesAsync(this IBlobWriter writer, DirectoryReference baseDir, ChunkingOptions? options = null, IProgress<ICopyStats>? progress = null, CancellationToken cancellationToken = default)
 		{
 			DirectoryNode outputNode = new DirectoryNode();
 			await outputNode.AddFilesAsync(baseDir, DirectoryReference.EnumerateFiles(baseDir, "*", SearchOption.AllDirectories), writer, options, progress, null, cancellationToken);
@@ -1318,7 +1318,7 @@ namespace EpicGames.Horde.Storage.Nodes
 		/// <summary>
 		/// Writes a tree of files to a storage writer
 		/// </summary>
-		public static async Task<IBlobHandle<DirectoryNode>> WriteFilesAsync(this IStorageWriter writer, DirectoryInfo baseDir, IReadOnlyList<FileInfo> files, ChunkingOptions? options = null, IProgress<ICopyStats>? progress = null, CancellationToken cancellationToken = default)
+		public static async Task<IBlobHandle<DirectoryNode>> WriteFilesAsync(this IBlobWriter writer, DirectoryInfo baseDir, IReadOnlyList<FileInfo> files, ChunkingOptions? options = null, IProgress<ICopyStats>? progress = null, CancellationToken cancellationToken = default)
 		{
 			DirectoryNode outputNode = new DirectoryNode();
 			await outputNode.AddFilesAsync(baseDir, files, writer, options, progress, null, cancellationToken);
@@ -1328,7 +1328,7 @@ namespace EpicGames.Horde.Storage.Nodes
 		/// <summary>
 		/// Writes a tree of files to a storage writer
 		/// </summary>
-		public static async Task<IBlobHandle<DirectoryNode>> WriteFilesAsync(this IStorageWriter writer, DirectoryReference baseDir, IReadOnlyList<FileReference> files, ChunkingOptions? options = null, IProgress<ICopyStats>? progress = null, CancellationToken cancellationToken = default)
+		public static async Task<IBlobHandle<DirectoryNode>> WriteFilesAsync(this IBlobWriter writer, DirectoryReference baseDir, IReadOnlyList<FileReference> files, ChunkingOptions? options = null, IProgress<ICopyStats>? progress = null, CancellationToken cancellationToken = default)
 		{
 			DirectoryNode outputNode = new DirectoryNode();
 			await outputNode.AddFilesAsync(baseDir, files, writer, options, progress, null, cancellationToken);
diff --git a/Engine/Source/Programs/UnrealBuildTool/Artifacts/HordeStorageArtifactCache.cs b/Engine/Source/Programs/UnrealBuildTool/Artifacts/HordeStorageArtifactCache.cs
--- a/Engine/Source/Programs/UnrealBuildTool/Artifacts/HordeStorageArtifactCache.cs
+++ b/Engine/Source/Programs/UnrealBuildTool/Artifacts/HordeStorageArtifactCache.cs
@@ -70,7 +70,7 @@ namespace UnrealBuildTool.Artifacts
 		/// <param name="writer">Destination writer</param>
 		/// <param name="cancellationToken">Cancellation token</param>
 		/// <returns>Task</returns>
-		public async Task WriteFilesAsync(IStorageWriter writer, CancellationToken cancellationToken)
+		public async Task WriteFilesAsync(IBlobWriter writer, CancellationToken cancellationToken)
 		{
 			LeafChunkedDataNodeOptions leafOptions = new(512 * 1024, 1 * 1024 * 1024, 2 * 1024 * 1024);
 			InteriorChunkedDataNodeOptions interiorOptions = new(1, 10, 20);
@@ -476,7 +476,7 @@ namespace UnrealBuildTool.Artifacts
 					node.ArtifactActions[artifactAction.ActionKey] = hordeArtifactAction;
 
 					// Save the artifact action file
-					await using IStorageWriter writer = _store!.CreateWriter();
+					await using IBlobWriter writer = _store!.CreateBlobWriter();
 					await hordeArtifactAction.WriteFilesAsync(writer, cancellationToken);
 					IBlobHandle<ArtifactActionCollectionNode> nodeRef = await writer.WriteBlobAsync(node);
 					await writer.FlushAsync();
diff --git a/Engine/Source/Programs/UnrealBuildTool/Executors/Experimental/UBAHordeSession.cs b/Engine/Source/Programs/UnrealBuildTool/Executors/Experimental/UBAHordeSession.cs
--- a/Engine/Source/Programs/UnrealBuildTool/Executors/Experimental/UBAHordeSession.cs
+++ b/Engine/Source/Programs/UnrealBuildTool/Executors/Experimental/UBAHordeSession.cs
@@ -305,7 +305,7 @@ namespace UnrealBuildTool
 			BlobSerializerOptions serializerOptions = new BlobSerializerOptions();
 			serializerOptions.Converters.Add(new InteriorChunkedDataNodeConverter(2)); // Lock to v2 for now. Could change based on protocol version.
 
-			await using IStorageWriter writer = _storage.CreateWriter();
+			await using IBlobWriter writer = _storage.CreateBlobWriter();
 			DirectoryNode sandbox = new();
 			await sandbox.AddFilesAsync(baseDir, files, writer, serializerOptions: serializerOptions, cancellationToken: cancellationToken);
 			IBlobHandle<DirectoryNode> handle = await writer.WriteBlobAsync(sandbox);
