Horde: Allow all StringId characters within object keys, and fix invalid paths being used for prefixed object stores.

diff --git a/Engine/Source/Programs/Horde/Horde.Server/Storage/StorageService.cs b/Engine/Source/Programs/Horde/Horde.Server/Storage/StorageService.cs
--- a/Engine/Source/Programs/Horde/Horde.Server/Storage/StorageService.cs
+++ b/Engine/Source/Programs/Horde/Horde.Server/Storage/StorageService.cs
@@ -675,7 +675,7 @@ namespace Horde.Server.Storage
 
 							if (!String.IsNullOrEmpty(namespaceConfig.Prefix))
 							{
-								objectStore = new PrefixedObjectStore(new ObjectKey(namespaceConfig.Prefix), objectStore);
+								objectStore = new PrefixedObjectStore(namespaceConfig.Prefix, objectStore);
 							}
 
 #pragma warning disable CA2000 // Dispose objects before losing scope
diff --git a/Engine/Source/Programs/Shared/EpicGames.Horde/Storage/Backends/FileStorageBackend.cs b/Engine/Source/Programs/Shared/EpicGames.Horde/Storage/Backends/FileStorageBackend.cs
--- a/Engine/Source/Programs/Shared/EpicGames.Horde/Storage/Backends/FileStorageBackend.cs
+++ b/Engine/Source/Programs/Shared/EpicGames.Horde/Storage/Backends/FileStorageBackend.cs
@@ -45,7 +45,7 @@ namespace EpicGames.Horde.Storage.Backends
 		/// <summary>
 		/// Gets the path for storing a file on disk
 		/// </summary>
-		static ObjectKey GetBlobFile(BlobLocator locator) => new ObjectKey(locator.Path);
+		static ObjectKey GetBlobFile(BlobLocator locator) => new ObjectKey($"{locator}.blob");
 
 		/// <inheritdoc/>
 		public Task<Stream> OpenBlobAsync(BlobLocator locator, int offset, int? length, CancellationToken cancellationToken)
diff --git a/Engine/Source/Programs/Shared/EpicGames.Horde/Storage/ObjectKey.cs b/Engine/Source/Programs/Shared/EpicGames.Horde/Storage/ObjectKey.cs
--- a/Engine/Source/Programs/Shared/EpicGames.Horde/Storage/ObjectKey.cs
+++ b/Engine/Source/Programs/Shared/EpicGames.Horde/Storage/ObjectKey.cs
@@ -18,6 +18,15 @@ namespace EpicGames.Horde.Storage
 		/// </summary>
 		public Utf8String Path => _path;
 
+		/// <summary>
+		/// Constructor
+		/// </summary>
+		/// <param name="path">Path to the blob. The meaning of this string is implementation defined.</param>
+		public ObjectKey(string path)
+			: this(new Utf8String(path))
+		{
+		}
+
 		/// <summary>
 		/// Constructor
 		/// </summary>
@@ -34,13 +43,17 @@ namespace EpicGames.Horde.Storage
 				}
 				for (int idx = 0; idx < path.Length; idx++)
 				{
-					if (path[idx] == '/' && path[idx - 1] == '/')
+					byte character = path[idx];
+					if (!StringId.IsValidCharacter(character))
 					{
-						throw new FormatException($"Object locator '{path}' is invalid; locators may not contain double consecutive slashes");
-					}
-					if (!IsValidChar(path[idx]))
-					{
-						throw new FormatException($"Object locator '{path}' is invalid; character '{(char)path[idx]}' is not allowed");
+						if (path[idx] == '/' && path[idx - 1] != '/')
+						{
+							// Non-consecutive path separator; allowed.
+						}
+						else
+						{
+							throw new FormatException($"Object locator '{path}' is invalid; character '{(char)path[idx]}' is not allowed");
+						}
 					}
 				}
 			}
@@ -49,29 +62,11 @@ namespace EpicGames.Horde.Storage
 		static bool IsValidChar(byte character)
 			=> (character >= 'a' && character <= 'z') || (character >=  '0' && character <= '9') || character == '_' || character == '.' || character == '/';
 
-		/// <summary>
-		/// Constructor
-		/// </summary>
-		/// <param name="path">Path to the blob. The meaning of this string is implementation defined.</param>
-		public ObjectKey(string path) => _path = new Utf8String(path);
-
 		/// <summary>
 		/// Whether the blob locator is valid
 		/// </summary>
 		public bool IsValid() => !_path.IsEmpty;
 
-		/// <summary>
-		/// Checks whether this blob is within the given folder
-		/// </summary>
-		/// <param name="folder">Name of the folder</param>
-		/// <returns>True if the the blob id is within the given folder</returns>
-		public bool WithinFolder(ObjectKey folder)
-		{
-			Utf8String path = Path;
-			Utf8String folderPath = folder.Path;
-			return path.Length > folderPath.Length && path.StartsWith(folderPath) && path[folderPath.Length] == '/';
-		}
-
 		/// <inheritdoc/>
 		public override bool Equals(object? obj) => obj is BlobLocator blobId && Equals(blobId);
 
diff --git a/Engine/Source/Programs/Shared/EpicGames.Horde/Storage/ObjectStores/FileObjectStore.cs b/Engine/Source/Programs/Shared/EpicGames.Horde/Storage/ObjectStores/FileObjectStore.cs
--- a/Engine/Source/Programs/Shared/EpicGames.Horde/Storage/ObjectStores/FileObjectStore.cs
+++ b/Engine/Source/Programs/Shared/EpicGames.Horde/Storage/ObjectStores/FileObjectStore.cs
@@ -19,7 +19,7 @@ namespace EpicGames.Horde.Storage.ObjectStores
 		// Item which has been opened from the cache using a memory mapped file
 		class MappedFile : IDisposable
 		{
-			public ObjectKey Locator { get; }
+			public ObjectKey Key { get; }
 			public LinkedListNode<MappedFile> ListNode { get; }
 
 			readonly FileInfo _fileInfo;
@@ -38,7 +38,7 @@ namespace EpicGames.Horde.Storage.ObjectStores
 
 			public MappedFile(ObjectKey locator, FileInfo fileInfo)
 			{
-				Locator = locator;
+				Key = locator;
 				ListNode = new LinkedListNode<MappedFile>(this);
 
 				_fileInfo = fileInfo;
@@ -118,7 +118,7 @@ namespace EpicGames.Horde.Storage.ObjectStores
 				_deleteOnDispose = true;
 			}
 
-			public override string ToString() => Locator.ToString();
+			public override string ToString() => Key.ToString();
 		}
 
 		// Handle to a file in memory
@@ -193,17 +193,17 @@ namespace EpicGames.Horde.Storage.ObjectStores
 		/// <summary>
 		/// Gets the path for storing a file on disk
 		/// </summary>
-		FileReference GetBlobFile(ObjectKey locator) => FileReference.Combine(_baseDir, locator.Path.ToString());
+		FileReference GetBlobFile(ObjectKey key) => FileReference.Combine(_baseDir, key.Path.ToString());
 
 		/// <summary>
 		/// Finds an existing mapped file or adds a new one for the given path
 		/// </summary>
-		MappedFile FindOrAddMappedFile(ObjectKey locator)
+		MappedFile FindOrAddMappedFile(ObjectKey key)
 		{
 			MappedFile? mappedFile;
-			if (!_pathToMappedFile.TryGetValue(locator, out mappedFile))
+			if (!_pathToMappedFile.TryGetValue(key, out mappedFile))
 			{
-				FileInfo fileInfo = GetBlobFile(locator).ToFileInfo();
+				FileInfo fileInfo = GetBlobFile(key).ToFileInfo();
 				if (fileInfo.Length == 0)
 				{
 					throw new StorageException($"Unable to map empty memory mapped file: {fileInfo.FullName}");
@@ -215,8 +215,8 @@ namespace EpicGames.Horde.Storage.ObjectStores
 					UnmapFiles(maxSize, MaxMappedCount - 1);
 				}
 
-				mappedFile = new MappedFile(locator, fileInfo);
-				_pathToMappedFile.Add(locator, mappedFile);
+				mappedFile = new MappedFile(key, fileInfo);
+				_pathToMappedFile.Add(key, mappedFile);
 				_mappedFiles.AddFirst(mappedFile.ListNode);
 
 				_mappedSize += (long)mappedFile.MappedSize;
@@ -236,7 +236,7 @@ namespace EpicGames.Horde.Storage.ObjectStores
 				LinkedListNode<MappedFile>? nextListNode = listNode.Previous;
 				if (listNode.Value.RefCount == 1)
 				{
-					_pathToMappedFile.Remove(listNode.Value.Locator);
+					_pathToMappedFile.Remove(listNode.Value.Key);
 					_mappedFiles.Remove(listNode);
 					listNode.Value.Release();
 				}
@@ -245,10 +245,10 @@ namespace EpicGames.Horde.Storage.ObjectStores
 		}
 
 		/// <inheritdoc/>
-		public async Task<Stream> OpenAsync(ObjectKey locator, int offset, int? length, CancellationToken cancellationToken)
+		public async Task<Stream> OpenAsync(ObjectKey key, int offset, int? length, CancellationToken cancellationToken)
 		{
 #pragma warning disable CA2000 // Dispose objects before losing scope
-			IReadOnlyMemoryOwner<byte> storageObject = await ReadAsync(locator, offset, length, cancellationToken);
+			IReadOnlyMemoryOwner<byte> storageObject = await ReadAsync(key, offset, length, cancellationToken);
 			return storageObject.AsStream();
 #pragma warning restore CA2000 // Dispose objects before losing scope
 		}
@@ -256,29 +256,29 @@ namespace EpicGames.Horde.Storage.ObjectStores
 		/// <summary>
 		/// Maps a file into memory for reading, and returns a handle to it
 		/// </summary>
-		/// <param name="locator">Path to the file</param>
+		/// <param name="key">Path to the file</param>
 		/// <param name="offset">Offset of the data to retrieve</param>
 		/// <param name="length">Length of the data</param>
 		/// <returns>Handle to the data. Must be disposed by the caller.</returns>
-		public IReadOnlyMemoryOwner<byte> Read(ObjectKey locator, int offset, int? length)
+		public IReadOnlyMemoryOwner<byte> Read(ObjectKey key, int offset, int? length)
 		{
 			lock (_lockObject)
 			{
-				MappedFile mappedFile = FindOrAddMappedFile(locator);
+				MappedFile mappedFile = FindOrAddMappedFile(key);
 				return new MappedFileHandle(mappedFile, mappedFile.GetData(offset, length));
 			}
 		}
 
 		/// <inheritdoc/>
-		public Task<IReadOnlyMemoryOwner<byte>> ReadAsync(ObjectKey locator, int offset, int? length, CancellationToken cancellationToken = default)
+		public Task<IReadOnlyMemoryOwner<byte>> ReadAsync(ObjectKey key, int offset, int? length, CancellationToken cancellationToken = default)
 		{
-			return Task.FromResult(Read(locator, offset, length));
+			return Task.FromResult(Read(key, offset, length));
 		}
 
 		/// <inheritdoc/>
-		public async Task WriteAsync(ObjectKey locator, Stream stream, CancellationToken cancellationToken = default)
+		public async Task WriteAsync(ObjectKey key, Stream stream, CancellationToken cancellationToken = default)
 		{
-			FileReference finalLocation = GetBlobFile(locator);
+			FileReference finalLocation = GetBlobFile(key);
 			DirectoryReference.CreateDirectory(finalLocation.Directory);
 			FileReference tempLocation = new FileReference($"{finalLocation}.tmp");
 
@@ -306,42 +306,42 @@ namespace EpicGames.Horde.Storage.ObjectStores
 		}
 
 		/// <inheritdoc/>
-		public Task<bool> ExistsAsync(ObjectKey locator, CancellationToken cancellationToken)
+		public Task<bool> ExistsAsync(ObjectKey key, CancellationToken cancellationToken)
 		{
-			FileReference location = GetBlobFile(locator);
+			FileReference location = GetBlobFile(key);
 			return Task.FromResult(FileReference.Exists(location));
 		}
 
 		/// <summary>
 		/// Delete a file from the store
 		/// </summary>
-		/// <param name="locator"></param>
-		public void Delete(ObjectKey locator)
+		/// <param name="key"></param>
+		public void Delete(ObjectKey key)
 		{
 			lock(_lockObject)
 			{
 				MappedFile? mappedFile;
-				if (_pathToMappedFile.TryGetValue(locator, out mappedFile))
+				if (_pathToMappedFile.TryGetValue(key, out mappedFile))
 				{
 					mappedFile.DeleteOnDispose();
 
-					_pathToMappedFile.Remove(locator);
+					_pathToMappedFile.Remove(key);
 					_mappedFiles.Remove(mappedFile.ListNode);
 
 					mappedFile.Release();
 				}
 				else
 				{
-					FileReference location = GetBlobFile(locator);
+					FileReference location = GetBlobFile(key);
 					FileReference.Delete(location);
 				}
 			}
 		}
 
 		/// <inheritdoc/>
-		public Task DeleteAsync(ObjectKey locator, CancellationToken cancellationToken)
+		public Task DeleteAsync(ObjectKey key, CancellationToken cancellationToken)
 		{
-			Delete(locator);
+			Delete(key);
 			return Task.CompletedTask;
 		}
 
@@ -385,10 +385,10 @@ namespace EpicGames.Horde.Storage.ObjectStores
 		}
 
 		/// <inheritdoc/>
-		public ValueTask<Uri?> TryGetReadRedirectAsync(ObjectKey locator, CancellationToken cancellationToken = default) => default;
+		public ValueTask<Uri?> TryGetReadRedirectAsync(ObjectKey key, CancellationToken cancellationToken = default) => default;
 
 		/// <inheritdoc/>
-		public ValueTask<Uri?> TryGetWriteRedirectAsync(ObjectKey locator, CancellationToken cancellationToken = default) => default;
+		public ValueTask<Uri?> TryGetWriteRedirectAsync(ObjectKey key, CancellationToken cancellationToken = default) => default;
 
 		/// <inheritdoc/>
 		public void GetStats(StorageStats stats) { }
diff --git a/Engine/Source/Programs/Shared/EpicGames.Horde/Storage/ObjectStores/PrefixedObjectStore.cs b/Engine/Source/Programs/Shared/EpicGames.Horde/Storage/ObjectStores/PrefixedObjectStore.cs
--- a/Engine/Source/Programs/Shared/EpicGames.Horde/Storage/ObjectStores/PrefixedObjectStore.cs
+++ b/Engine/Source/Programs/Shared/EpicGames.Horde/Storage/ObjectStores/PrefixedObjectStore.cs
@@ -15,20 +15,22 @@ namespace EpicGames.Horde.Storage.ObjectStores
 	/// </summary>
 	public sealed class PrefixedObjectStore : IObjectStore
 	{
-		readonly ObjectKey _prefix;
+		readonly Utf8String _prefix;
 		readonly IObjectStore _inner;
 
 		/// <inheritdoc/>
 		public bool SupportsRedirects => _inner.SupportsRedirects;
 
-		ObjectKey GetFullLocator(ObjectKey locator) => new ObjectKey($"{_prefix}/{locator}");
+		ObjectKey GetKeyWithPrefix(ObjectKey key) => new ObjectKey(_prefix + key.Path);
 
 		/// <summary>
 		/// Constructor
 		/// </summary>
-		public PrefixedObjectStore(ObjectKey prefix, IObjectStore inner)
+		public PrefixedObjectStore(string prefix, IObjectStore inner)
 		{
-			_prefix = prefix;
+			prefix = prefix.TrimEnd('/').ToLowerInvariant();
+			_ = new ObjectKey(prefix); // Validate the syntax
+			_prefix = new Utf8String($"{prefix}/");
 			_inner = inner;
 		}
 
@@ -38,37 +40,37 @@ namespace EpicGames.Horde.Storage.ObjectStores
 		}
 
 		/// <inheritdoc/>
-		public Task DeleteAsync(ObjectKey locator, CancellationToken cancellationToken = default) => _inner.DeleteAsync(GetFullLocator(locator), cancellationToken);
+		public Task DeleteAsync(ObjectKey locator, CancellationToken cancellationToken = default) => _inner.DeleteAsync(GetKeyWithPrefix(locator), cancellationToken);
 
 		/// <inheritdoc/>
 		public async IAsyncEnumerable<ObjectKey> EnumerateAsync([EnumeratorCancellation] CancellationToken cancellationToken = default)
 		{
-			await foreach (ObjectKey locator in _inner.EnumerateAsync(cancellationToken))
+			await foreach (ObjectKey key in _inner.EnumerateAsync(cancellationToken))
 			{
-				if (locator.WithinFolder(_prefix))
+				if (key.Path.StartsWith(_prefix))
 				{
-					yield return new ObjectKey(locator.Path.Substring(_prefix.Path.Length));
+					yield return new ObjectKey(key.Path.Substring(_prefix.Length));
 				}
 			}
 		}
 
 		/// <inheritdoc/>
-		public Task<bool> ExistsAsync(ObjectKey locator, CancellationToken cancellationToken = default) => _inner.ExistsAsync(GetFullLocator(locator), cancellationToken);
+		public Task<bool> ExistsAsync(ObjectKey locator, CancellationToken cancellationToken = default) => _inner.ExistsAsync(GetKeyWithPrefix(locator), cancellationToken);
 
 		/// <inheritdoc/>
-		public Task<Stream> OpenAsync(ObjectKey locator, int offset, int? length, CancellationToken cancellationToken = default) => _inner.OpenAsync(GetFullLocator(locator), offset, length, cancellationToken);
+		public Task<Stream> OpenAsync(ObjectKey locator, int offset, int? length, CancellationToken cancellationToken = default) => _inner.OpenAsync(GetKeyWithPrefix(locator), offset, length, cancellationToken);
 
 		/// <inheritdoc/>
-		public Task<IReadOnlyMemoryOwner<byte>> ReadAsync(ObjectKey locator, int offset, int? length, CancellationToken cancellationToken = default) => _inner.ReadAsync(GetFullLocator(locator), offset, length, cancellationToken);
+		public Task<IReadOnlyMemoryOwner<byte>> ReadAsync(ObjectKey locator, int offset, int? length, CancellationToken cancellationToken = default) => _inner.ReadAsync(GetKeyWithPrefix(locator), offset, length, cancellationToken);
 
 		/// <inheritdoc/>
-		public ValueTask<Uri?> TryGetReadRedirectAsync(ObjectKey locator, CancellationToken cancellationToken = default) => _inner.TryGetReadRedirectAsync(GetFullLocator(locator), cancellationToken);
+		public ValueTask<Uri?> TryGetReadRedirectAsync(ObjectKey locator, CancellationToken cancellationToken = default) => _inner.TryGetReadRedirectAsync(GetKeyWithPrefix(locator), cancellationToken);
 
 		/// <inheritdoc/>
-		public ValueTask<Uri?> TryGetWriteRedirectAsync(ObjectKey locator, CancellationToken cancellationToken = default) => _inner.TryGetWriteRedirectAsync(GetFullLocator(locator), cancellationToken);
+		public ValueTask<Uri?> TryGetWriteRedirectAsync(ObjectKey locator, CancellationToken cancellationToken = default) => _inner.TryGetWriteRedirectAsync(GetKeyWithPrefix(locator), cancellationToken);
 
 		/// <inheritdoc/>
-		public Task WriteAsync(ObjectKey locator, Stream stream, CancellationToken cancellationToken = default) => _inner.WriteAsync(GetFullLocator(locator), stream, cancellationToken);
+		public Task WriteAsync(ObjectKey locator, Stream stream, CancellationToken cancellationToken = default) => _inner.WriteAsync(GetKeyWithPrefix(locator), stream, cancellationToken);
 
 		/// <inheritdoc/>
 		public void GetStats(StorageStats stats) => _inner.GetStats(stats);
diff --git a/Engine/Source/Programs/Shared/EpicGames.Horde/StringId.cs b/Engine/Source/Programs/Shared/EpicGames.Horde/StringId.cs
--- a/Engine/Source/Programs/Shared/EpicGames.Horde/StringId.cs
+++ b/Engine/Source/Programs/Shared/EpicGames.Horde/StringId.cs
@@ -135,7 +135,7 @@ namespace EpicGames.Horde
 
 			for (int idx = 0; idx < text.Length; idx++)
 			{
-				char character = (char)text[idx];
+				byte character = text[idx];
 				if (!IsValidCharacter(character))
 				{
 					if (character >= 'A' && character <= 'Z')
@@ -177,7 +177,14 @@ namespace EpicGames.Horde
 		/// </summary>
 		/// <param name="character">The character to check</param>
 		/// <returns>True if the character is valid</returns>
-		static bool IsValidCharacter(char character)
+		public static bool IsValidCharacter(char character) => character <= 0x7f && IsValidCharacter((byte)character);
+
+		/// <summary>
+		/// Checks whether the given character is valid within a string id
+		/// </summary>
+		/// <param name="character">The character to check</param>
+		/// <returns>True if the character is valid</returns>
+		public static bool IsValidCharacter(byte character)
 		{
 			if (character >= 'a' && character <= 'z')
 			{
