Horde: Add a constructor for MemoryMappedView which uses P/Invoke to map the file into memory. MemoryMappedViewAccessor calls Flush when disposed, which introduces a significant performance hit.

diff --git a/Engine/Source/Programs/Horde/Horde.Server/Storage/StorageService.cs b/Engine/Source/Programs/Horde/Horde.Server/Storage/StorageService.cs
--- a/Engine/Source/Programs/Horde/Horde.Server/Storage/StorageService.cs
+++ b/Engine/Source/Programs/Horde/Horde.Server/Storage/StorageService.cs
@@ -639,7 +639,9 @@ namespace Horde.Server.Storage
 		/// <inheritdoc/>
 		public IStorageClient? TryCreateClient(GlobalConfig globalConfig, NamespaceId namespaceId)
 		{
+#pragma warning disable CA2000 // Call dispose on backend; will be disposed by BundleStorageClient
 			IStorageBackend? backend = TryCreateBackend(globalConfig, namespaceId);
+#pragma warning restore CA2000
 			if (backend == null)
 			{
 				return null;
@@ -1123,7 +1125,7 @@ namespace Horde.Server.Storage
 
 		async Task TickGcForNamespaceAsync(NamespaceInfo namespaceInfo, ObjectId lastImportBlobInfoId, DateTime utcNow, CancellationToken cancellationToken)
 		{
-			IStorageClient client = this.CreateClient(namespaceInfo.Id);
+			using IStorageClient client = this.CreateClient(namespaceInfo.Id);
 
 			double score = GetGcTimestamp(utcNow);
 
diff --git a/Engine/Source/Programs/Shared/EpicGames.Core/MemoryMappedView.cs b/Engine/Source/Programs/Shared/EpicGames.Core/MemoryMappedView.cs
--- a/Engine/Source/Programs/Shared/EpicGames.Core/MemoryMappedView.cs
+++ b/Engine/Source/Programs/Shared/EpicGames.Core/MemoryMappedView.cs
@@ -2,7 +2,10 @@
 
 using System;
 using System.Buffers;
+using System.Diagnostics;
 using System.IO.MemoryMappedFiles;
+using System.Runtime.InteropServices;
+using Microsoft.Win32.SafeHandles;
 
 namespace EpicGames.Core
 {
@@ -35,8 +38,46 @@ namespace EpicGames.Core
 			protected override void Dispose(bool disposing) { }
 		}
 
-		readonly MemoryMappedViewAccessor _memoryMappedViewAccessor;
+		const uint FILE_MAP_ALL_ACCESS = 0x000F001F;
+
+		[DllImport("kernel32.dll")]
+		static extern unsafe void* MapViewOfFile(SafeMemoryMappedFileHandle handle, uint desiredAccess, uint fileOffsetHigh, uint fileOffsetLow, long numberOfBytes);
+
+		[DllImport("kernel32.dll")]
+		static extern unsafe bool UnmapViewOfFile(void* ptr);
+		
 		byte* _data;
+		Action? _disposeMethod;
+
+		/// <summary>
+		/// Create a view of a memory mapped file
+		/// </summary>
+		/// <param name="memoryMappedFile"></param>
+		public MemoryMappedView(MemoryMappedFile memoryMappedFile, long offset, long length)
+		{
+			if (RuntimeInformation.IsOSPlatform(OSPlatform.Windows))
+			{
+				_data = (byte*)MapViewOfFile(memoryMappedFile.SafeMemoryMappedFileHandle, FILE_MAP_ALL_ACCESS, (uint)(offset >> 32), (uint)offset, length);
+				Trace.Assert(_data != null);
+
+				_disposeMethod = () =>
+				{
+					UnmapViewOfFile(_data);
+				};
+			}
+			else
+			{
+				MemoryMappedViewAccessor memoryMappedViewAccessor = memoryMappedFile.CreateViewAccessor(offset, length);
+				memoryMappedViewAccessor.SafeMemoryMappedViewHandle.AcquirePointer(ref _data);
+				Trace.Assert(_data != null);
+
+				_disposeMethod = () =>
+				{
+					memoryMappedViewAccessor.SafeMemoryMappedViewHandle.ReleasePointer();
+					memoryMappedViewAccessor.Dispose();
+				};
+			}
+		}
 
 		/// <summary>
 		/// Constructor
@@ -44,8 +85,8 @@ namespace EpicGames.Core
 		/// <param name="memoryMappedViewAccessor"></param>
 		public MemoryMappedView(MemoryMappedViewAccessor memoryMappedViewAccessor)
 		{
-			_memoryMappedViewAccessor = memoryMappedViewAccessor;
 			memoryMappedViewAccessor.SafeMemoryMappedViewHandle.AcquirePointer(ref _data);
+			_disposeMethod = () => memoryMappedViewAccessor.SafeMemoryMappedViewHandle.ReleasePointer();
 		}
 
 		/// <summary>
@@ -68,11 +109,12 @@ namespace EpicGames.Core
 		/// <inheritdoc/>
 		public void Dispose()
 		{
-			if (_data != null)
+			if (_disposeMethod != null)
 			{
-				_memoryMappedViewAccessor.SafeMemoryMappedViewHandle.ReleasePointer();
-				_data = null;
+				_disposeMethod();
+				_disposeMethod = null;
 			}
+			_data = null;
 		}
 	}
 }
diff --git a/Engine/Source/Programs/Shared/EpicGames.Horde/Storage/Nodes/DirectoryNode.cs b/Engine/Source/Programs/Shared/EpicGames.Horde/Storage/Nodes/DirectoryNode.cs
--- a/Engine/Source/Programs/Shared/EpicGames.Horde/Storage/Nodes/DirectoryNode.cs
+++ b/Engine/Source/Programs/Shared/EpicGames.Horde/Storage/Nodes/DirectoryNode.cs
@@ -738,7 +738,7 @@ namespace EpicGames.Horde.Storage.Nodes
 		/// <param name="cancellationToken">Cancellation token for the operation</param>
 		public async Task CopyToDirectoryAsync(DirectoryInfo directoryInfo, IProgress<ICopyStats>? progress, BlobSerializerOptions? options, ILogger logger, CancellationToken cancellationToken)
 		{
-			int numTasks = Math.Min(1 + (int)(Length / (16 * 1024 * 1024)), 4);
+			int numTasks = Math.Min(1 + (int)(Length / (16 * 1024 * 1024)), 16);
 			logger.LogInformation("Splitting read into {NumThreads} threads", numTasks);
 
 			CopyStats? copyStats = null;
@@ -801,8 +801,7 @@ namespace EpicGames.Horde.Storage.Nodes
 
 				// Process as many chunks as we can for this file
 				using MemoryMappedFile memoryMappedFile = MemoryMappedFile.CreateFromFile(stream, null, file.FileEntry.Length, MemoryMappedFileAccess.ReadWrite, HandleInheritability.None, false);
-				using MemoryMappedViewAccessor memoryMappedViewAccessor = memoryMappedFile.CreateViewAccessor(0, file.FileEntry.Length);
-				using MemoryMappedView memoryMappedView = new MemoryMappedView(memoryMappedViewAccessor);
+				using MemoryMappedView memoryMappedView = new MemoryMappedView(memoryMappedFile, 0, file.FileEntry.Length);
 
 				while (chunk != null && chunk.File == file)
 				{
diff --git a/Engine/Source/Programs/Shared/EpicGames.Horde/Storage/ObjectKey.cs b/Engine/Source/Programs/Shared/EpicGames.Horde/Storage/ObjectKey.cs
--- a/Engine/Source/Programs/Shared/EpicGames.Horde/Storage/ObjectKey.cs
+++ b/Engine/Source/Programs/Shared/EpicGames.Horde/Storage/ObjectKey.cs
@@ -59,9 +59,6 @@ namespace EpicGames.Horde.Storage
 			}
 		}
 
-		static bool IsValidChar(byte character)
-			=> (character >= 'a' && character <= 'z') || (character >=  '0' && character <= '9') || character == '_' || character == '.' || character == '/';
-
 		/// <summary>
 		/// Whether the blob locator is valid
 		/// </summary>
diff --git a/Engine/Source/Programs/Shared/EpicGames.Horde/Storage/ObjectStores/PrefixedObjectStore.cs b/Engine/Source/Programs/Shared/EpicGames.Horde/Storage/ObjectStores/PrefixedObjectStore.cs
--- a/Engine/Source/Programs/Shared/EpicGames.Horde/Storage/ObjectStores/PrefixedObjectStore.cs
+++ b/Engine/Source/Programs/Shared/EpicGames.Horde/Storage/ObjectStores/PrefixedObjectStore.cs
@@ -28,7 +28,9 @@ namespace EpicGames.Horde.Storage.ObjectStores
 		/// </summary>
 		public PrefixedObjectStore(string prefix, IObjectStore inner)
 		{
+#pragma warning disable CA1308 // Use ToUpperInvariant()
 			prefix = prefix.TrimEnd('/').ToLowerInvariant();
+#pragma warning restore CA1308
 			_ = new ObjectKey(prefix); // Validate the syntax
 			_prefix = new Utf8String($"{prefix}/");
 			_inner = inner;
