Fix static analysis warning.

diff --git a/Engine/Source/Programs/Shared/EpicGames.Core/MemoryMappedView.cs b/Engine/Source/Programs/Shared/EpicGames.Core/MemoryMappedView.cs
--- a/Engine/Source/Programs/Shared/EpicGames.Core/MemoryMappedView.cs
+++ b/Engine/Source/Programs/Shared/EpicGames.Core/MemoryMappedView.cs
@@ -52,7 +52,9 @@ namespace EpicGames.Core
 		/// <summary>
 		/// Create a view of a memory mapped file
 		/// </summary>
-		/// <param name="memoryMappedFile"></param>
+		/// <param name="memoryMappedFile">Handle of the file to map into memory</param>
+		/// <param name="offset">Offset within the file to map</param>
+		/// <param name="length">Length of the region to map</param>
 		public MemoryMappedView(MemoryMappedFile memoryMappedFile, long offset, long length)
 		{
 			if (RuntimeInformation.IsOSPlatform(OSPlatform.Windows))
