Horde: Verify the contents of files written through compute leases after writing them.

diff --git a/Engine/Source/Programs/Horde/Horde.Server.Tests/GcServiceTests.cs b/Engine/Source/Programs/Horde/Horde.Server.Tests/GcServiceTests.cs
--- a/Engine/Source/Programs/Horde/Horde.Server.Tests/GcServiceTests.cs
+++ b/Engine/Source/Programs/Horde/Horde.Server.Tests/GcServiceTests.cs
@@ -10,7 +10,6 @@ using Horde.Server.Storage;
 using System.Threading;
 using Horde.Server.Server;
 using Microsoft.Extensions.DependencyInjection;
-using EpicGames.Core;
 
 namespace Horde.Server.Tests
 {
diff --git a/Engine/Source/Programs/Shared/EpicGames.Horde/Compute/AgentMessageHandler.cs b/Engine/Source/Programs/Shared/EpicGames.Horde/Compute/AgentMessageHandler.cs
--- a/Engine/Source/Programs/Shared/EpicGames.Horde/Compute/AgentMessageHandler.cs
+++ b/Engine/Source/Programs/Shared/EpicGames.Horde/Compute/AgentMessageHandler.cs
@@ -186,6 +186,7 @@ namespace EpicGames.Horde.Compute
 			}
 
 			await directoryNode.CopyToDirectoryAsync(outputDir.ToDirectoryInfo(), options, _logger, cancellationToken);
+			await VerifyFilesAsync(outputDir, directoryNode, cancellationToken);
 
 			using (IAgentMessageBuilder message = await channel.CreateMessageAsync(AgentMessageType.WriteFilesResponse, cancellationToken))
 			{
@@ -193,6 +194,46 @@ namespace EpicGames.Horde.Compute
 			}
 		}
 
+		async Task<bool> VerifyFilesAsync(DirectoryReference outputDir, DirectoryNode directoryNode, CancellationToken cancellationToken = default)
+		{
+			bool result = true;
+
+			foreach (FileEntry fileEntry in directoryNode.Files)
+			{
+				FileReference file = FileReference.Combine(outputDir, fileEntry.Name);
+				if (!FileReference.Exists(file))
+				{
+					_logger.LogError("Extracted file {File} does not exist", file);
+					result = false;
+				}
+				else
+				{
+					IoHash hash;
+					using (FileStream stream = FileReference.Open(file, FileMode.Open))
+					{
+						hash = await IoHash.ComputeAsync(stream, cancellationToken);
+					}
+					if (hash == fileEntry.StreamHash)
+					{
+						_logger.LogInformation("Hash of {File} is correct ({Hash})", file, hash);
+					}
+					if (hash != fileEntry.StreamHash)
+					{
+						_logger.LogError("Hash mismatch for {File}; expected {ExpectedHash}, got {ActualHash}", file, fileEntry.StreamHash, fileEntry.StreamHash);
+						result = false;
+					}
+				}
+			}
+
+			foreach (DirectoryEntry directoryEntry in directoryNode.Directories)
+			{
+				DirectoryNode subNode = await directoryEntry.Handle.ReadBlobAsync(cancellationToken: cancellationToken);
+				result &= await VerifyFilesAsync(DirectoryReference.Combine(outputDir, directoryEntry.Name), subNode, cancellationToken);
+			}
+
+			return result;
+		}
+
 		void DeleteFiles(IReadOnlyList<string> deleteFiles)
 		{
 			FileFilter filter = new FileFilter(deleteFiles);
