Horde: Temporarily revert change to use custom code path for extracting files from storage system on Windows. Suspect some bad interaction with WINE.

diff --git a/Engine/Source/Programs/Shared/EpicGames.Horde/Compute/AgentMessageHandler.cs b/Engine/Source/Programs/Shared/EpicGames.Horde/Compute/AgentMessageHandler.cs
--- a/Engine/Source/Programs/Shared/EpicGames.Horde/Compute/AgentMessageHandler.cs
+++ b/Engine/Source/Programs/Shared/EpicGames.Horde/Compute/AgentMessageHandler.cs
@@ -219,7 +219,7 @@ namespace EpicGames.Horde.Compute
 					}
 					if (hash != fileEntry.StreamHash)
 					{
-						_logger.LogError("Hash mismatch for {File}; expected {ExpectedHash}, got {ActualHash}", file, fileEntry.StreamHash, fileEntry.StreamHash);
+						_logger.LogError("Hash mismatch for {File}; expected {ExpectedHash}, got {ActualHash}", file, fileEntry.StreamHash, hash);
 						result = false;
 					}
 				}
diff --git a/Engine/Source/Programs/Shared/EpicGames.Horde/Storage/Nodes/DirectoryNode.cs b/Engine/Source/Programs/Shared/EpicGames.Horde/Storage/Nodes/DirectoryNode.cs
--- a/Engine/Source/Programs/Shared/EpicGames.Horde/Storage/Nodes/DirectoryNode.cs
+++ b/Engine/Source/Programs/Shared/EpicGames.Horde/Storage/Nodes/DirectoryNode.cs
@@ -801,7 +801,10 @@ namespace EpicGames.Horde.Storage.Nodes
 
 				// Process as many chunks as we can for this file
 				using MemoryMappedFile memoryMappedFile = MemoryMappedFile.CreateFromFile(stream, null, file.FileEntry.Length, MemoryMappedFileAccess.ReadWrite, HandleInheritability.None, false);
-				using MemoryMappedView memoryMappedView = new MemoryMappedView(memoryMappedFile, 0, file.FileEntry.Length);
+				using MemoryMappedViewAccessor memoryMappedViewAccessor = memoryMappedFile.CreateViewAccessor(0, file.FileEntry.Length);
+				using MemoryMappedView memoryMappedView = new MemoryMappedView(memoryMappedViewAccessor);
+
+//				using MemoryMappedView memoryMappedView = new MemoryMappedView(memoryMappedFile, 0, file.FileEntry.Length);
 
 				while (chunk != null && chunk.File == file)
 				{
