Horde: Fix support for extracting v2 format chunked blob data, and restore change to use fast memory mapped files for writing data on Windows.

diff --git a/Engine/Source/Programs/Shared/EpicGames.Horde.Tests/ChunkingTests.cs b/Engine/Source/Programs/Shared/EpicGames.Horde.Tests/ChunkingTests.cs
--- a/Engine/Source/Programs/Shared/EpicGames.Horde.Tests/ChunkingTests.cs
+++ b/Engine/Source/Programs/Shared/EpicGames.Horde.Tests/ChunkingTests.cs
@@ -11,6 +11,8 @@ using System.Threading;
 using EpicGames.Horde.Storage;
 using EpicGames.Horde.Storage.Clients;
 using Microsoft.Extensions.Caching.Memory;
+using System.Text;
+using Microsoft.Extensions.Logging.Abstractions;
 
 namespace EpicGames.Horde.Tests
 {
@@ -162,5 +164,46 @@ namespace EpicGames.Horde.Tests
 				}
 			}
 		}
+
+		[TestMethod]
+		public async Task ChunkingCompatV2()
+		{
+			ChunkingOptions chunkingOptions = new ChunkingOptions();
+			chunkingOptions.LeafOptions = new LeafChunkedDataNodeOptions(2, 2, 2);
+			chunkingOptions.InteriorOptions = new InteriorChunkedDataNodeOptions(2, 2, 2);
+
+			BlobSerializerOptions serializerOptions = new BlobSerializerOptions();
+			serializerOptions.Converters.Add(new InteriorChunkedDataNodeConverter(2)); // Does not include length fields in interior nodes
+
+			using KeyValueStorageClient store = KeyValueStorageClient.CreateInMemory();
+
+			byte[] data = Encoding.UTF8.GetBytes("hello world");
+
+			IBlobHandle<DirectoryNode> handle;
+			await using (IBlobWriter writer = store.CreateBlobWriter())
+			{
+				using ChunkedDataWriter chunkedWriter = new ChunkedDataWriter(writer, chunkingOptions, serializerOptions);
+				await chunkedWriter.AppendAsync(data, CancellationToken.None);
+				ChunkedData chunkedData = await chunkedWriter.FlushAsync(CancellationToken.None);
+
+				DirectoryNode directoryNode = new DirectoryNode();
+				directoryNode.AddFile("test", FileEntryFlags.None, data.Length, chunkedData);
+				handle = await writer.WriteBlobAsync(directoryNode);
+			}
+
+			DirectoryInfo tempDir = new DirectoryInfo(Path.Combine(Path.GetTempPath(), Guid.NewGuid().ToString()));
+			try
+			{
+				DirectoryNode expandedNode = await handle.ReadBlobAsync();
+				await expandedNode.CopyToDirectoryAsync(tempDir, null, NullLogger.Instance, CancellationToken.None);
+
+				byte[] outputData = await File.ReadAllBytesAsync(Path.Combine(tempDir.FullName, "test"));
+				Assert.IsTrue(outputData.SequenceEqual(data));
+			}
+			finally
+			{
+				tempDir.Delete(true);
+			}
+		}
 	}
 }
diff --git a/Engine/Source/Programs/Shared/EpicGames.Horde/Storage/Nodes/ChunkedDataNode.cs b/Engine/Source/Programs/Shared/EpicGames.Horde/Storage/Nodes/ChunkedDataNode.cs
--- a/Engine/Source/Programs/Shared/EpicGames.Horde/Storage/Nodes/ChunkedDataNode.cs
+++ b/Engine/Source/Programs/Shared/EpicGames.Horde/Storage/Nodes/ChunkedDataNode.cs
@@ -600,7 +600,7 @@ namespace EpicGames.Horde.Storage.Nodes
 					type = (ChunkedDataNodeType)reader.ReadUnsignedVarInt();
 				}
 
-				long length = 0;
+				long length = -1;
 				if (reader.Version >= 3)
 				{
 					length = (long)reader.ReadUnsignedVarInt();
diff --git a/Engine/Source/Programs/Shared/EpicGames.Horde/Storage/Nodes/DirectoryNode.cs b/Engine/Source/Programs/Shared/EpicGames.Horde/Storage/Nodes/DirectoryNode.cs
--- a/Engine/Source/Programs/Shared/EpicGames.Horde/Storage/Nodes/DirectoryNode.cs
+++ b/Engine/Source/Programs/Shared/EpicGames.Horde/Storage/Nodes/DirectoryNode.cs
@@ -801,10 +801,7 @@ namespace EpicGames.Horde.Storage.Nodes
 
 				// Process as many chunks as we can for this file
 				using MemoryMappedFile memoryMappedFile = MemoryMappedFile.CreateFromFile(stream, null, file.FileEntry.Length, MemoryMappedFileAccess.ReadWrite, HandleInheritability.None, false);
-				using MemoryMappedViewAccessor memoryMappedViewAccessor = memoryMappedFile.CreateViewAccessor(0, file.FileEntry.Length);
-				using MemoryMappedView memoryMappedView = new MemoryMappedView(memoryMappedViewAccessor);
-
-//				using MemoryMappedView memoryMappedView = new MemoryMappedView(memoryMappedFile, 0, file.FileEntry.Length);
+				using MemoryMappedView memoryMappedView = new MemoryMappedView(memoryMappedFile, 0, file.FileEntry.Length);
 
 				while (chunk != null && chunk.File == file)
 				{
@@ -867,11 +864,18 @@ namespace EpicGames.Horde.Storage.Nodes
 			}
 		}
 
-		static async Task FindOutputChunksAsync(OutputFile outputFile, long offset, ChunkedDataNodeRef dataRef, ChannelWriter<OutputChunk> chunks, CancellationToken cancellationToken)
+		static async Task<long> FindOutputChunksAsync(OutputFile outputFile, long offset, ChunkedDataNodeRef dataRef, ChannelWriter<OutputChunk> chunks, CancellationToken cancellationToken)
 		{
 			if (dataRef.Type == ChunkedDataNodeType.Leaf)
 			{
 				await chunks.WriteAsync(new OutputChunk(outputFile, offset, dataRef.Length, dataRef.Handle), cancellationToken);
+				if (dataRef.Length < 0)
+				{
+					// Backwards compatibility hack for v2 format
+					using BlobData data = await dataRef.Handle.ReadBlobDataAsync(cancellationToken);
+					return data.Data.Length;
+				}
+				return dataRef.Length;
 			}
 			else
 			{
@@ -879,15 +883,19 @@ namespace EpicGames.Horde.Storage.Nodes
 				if (data.Type.Guid == LeafChunkedDataNodeConverter.BlobType.Guid)
 				{
 					await chunks.WriteAsync(new OutputChunk(outputFile, offset, dataRef.Length, dataRef.Handle), cancellationToken);
+					return data.Data.Length;
 				}
 				else
 				{
+					long length = 0;
+
 					InteriorChunkedDataNode interiorNode = BlobSerializer.Deserialize<InteriorChunkedDataNode>(data);
 					foreach (ChunkedDataNodeRef childRef in interiorNode.Children)
 					{
-						await FindOutputChunksAsync(outputFile, offset, childRef, chunks, cancellationToken);
-						offset += childRef.Length;
+						length += await FindOutputChunksAsync(outputFile, offset + length, childRef, chunks, cancellationToken);
 					}
+
+					return length;
 				}
 			}
 		}
