Horde: Add a singleton IHordeClient object which allows constructing new HordeHttpClient instances. Allows keeping all the configuration/state in one place.

diff --git a/Engine/Source/Programs/AutomationTool/BuildGraph/Tasks/CreateArtifactTask.cs b/Engine/Source/Programs/AutomationTool/BuildGraph/Tasks/CreateArtifactTask.cs
--- a/Engine/Source/Programs/AutomationTool/BuildGraph/Tasks/CreateArtifactTask.cs
+++ b/Engine/Source/Programs/AutomationTool/BuildGraph/Tasks/CreateArtifactTask.cs
@@ -100,7 +100,7 @@ namespace AutomationTool.Tasks
 			// Create a DI container for building the graph
 			ServiceCollection serviceCollection = new ServiceCollection();
 			serviceCollection.AddLogging(builder => builder.AddProvider(new LoggerProviderAdapter(Log.Logger)));
-			serviceCollection.AddHorde(new HordeOptions { AllowAuthPrompt = !Automation.IsBuildMachine });
+			serviceCollection.Configure<HordeOptions>(x => x.AllowAuthPrompt = !Automation.IsBuildMachine);
 			serviceCollection.Configure<LoggerFilterOptions>(options => options.AddFilter(typeof(HttpClient).FullName, LogLevel.Warning));
 
 			await using ServiceProvider serviceProvider = serviceCollection.BuildServiceProvider();
diff --git a/Engine/Source/Programs/Horde/Horde.Agent/AgentApp.cs b/Engine/Source/Programs/Horde/Horde.Agent/AgentApp.cs
--- a/Engine/Source/Programs/Horde/Horde.Agent/AgentApp.cs
+++ b/Engine/Source/Programs/Horde/Horde.Agent/AgentApp.cs
@@ -173,7 +173,9 @@ namespace Horde.Agent
 				return builder.WaitAndRetryAsync(new[] { TimeSpan.FromSeconds(1), TimeSpan.FromSeconds(5), TimeSpan.FromSeconds(10) });
 			});
 
-			services.AddHordeHttpClient(client => client.BaseAddress = serverProfile.Url, useAuthChallenge: false)
+			services.Configure<HordeOptions>(options => options.AllowAuthPrompt = false);
+
+			services.AddHordeHttpClient(client => client.BaseAddress = serverProfile.Url)
 				.ConfigurePrimaryHttpMessageHandler(() => new SocketsHttpHandler
 				{
 					MaxConnectionsPerServer = 16,
diff --git a/Engine/Source/Programs/Horde/Horde/CmdApp.cs b/Engine/Source/Programs/Horde/Horde/CmdApp.cs
--- a/Engine/Source/Programs/Horde/Horde/CmdApp.cs
+++ b/Engine/Source/Programs/Horde/Horde/CmdApp.cs
@@ -9,11 +9,8 @@ using System.Runtime.InteropServices;
 using Serilog.Sinks.SystemConsole.Themes;
 using Serilog;
 using Serilog.Formatting.Json;
-using EpicGames.Horde.Storage;
 using Microsoft.Extensions.Options;
 using EpicGames.Horde;
-using EpicGames.Horde.Storage.Clients;
-using EpicGames.Horde.Storage.Backends;
 
 namespace Horde
 {
@@ -32,49 +29,25 @@ namespace Horde
 				.AddEnvironmentVariables()
 				.Build();
 
+			CmdConfig cmdConfig = CmdConfig.Read();
+
 			using ILoggerFactory loggerFactory = CreateLoggerFactory(configuration, arguments);
 
 			IServiceCollection services = new ServiceCollection();
+			services.Configure<HordeOptions>(options => configuration.Bind("Horde", options));
+			services.Configure<HordeOptions>(options => options.ServerUrl = cmdConfig.Server);
 			services.AddCommandsFromAssembly(Assembly.GetExecutingAssembly());
 			services.AddSingleton(loggerFactory);
 			services.AddLogging();
 			services.AddMemoryCache();
-			services.AddSingleton(sp => Options.Create(CmdConfig.Read()));
-			services.AddHordeHttpClient((sp, client) => client.BaseAddress = sp.GetRequiredService<IOptions<CmdConfig>>().Value.Server);
-			services.AddSingleton<BundleCache>(CreateStorageClientCache);
-			services.AddSingleton<StorageBackendCache>(CreateStorageBackendCache);
-			services.AddSingleton<HttpStorageBackendFactory>();
-			services.AddSingleton<HttpStorageClientFactory>();
+			services.AddSingleton(Options.Create(cmdConfig));
+			services.AddHorde();
 
 			// Execute all the commands
 			await using ServiceProvider serviceProvider = services.BuildServiceProvider();
 			return await CommandHost.RunAsync(arguments, serviceProvider, null, ToolDescription);
 		}
 
-		static BundleCache CreateStorageClientCache(IServiceProvider serviceProvider)
-		{
-			CmdConfig cmdConfig = serviceProvider.GetRequiredService<IOptions<CmdConfig>>().Value;
-
-			BundleCacheOptions options = new BundleCacheOptions();
-			if (cmdConfig.Cache.HeaderCacheSize.HasValue)
-			{
-				options.HeaderCacheSize = cmdConfig.Cache.HeaderCacheSize.Value * 1024 * 1024;
-			}
-			if (cmdConfig.Cache.PacketCacheSize.HasValue)
-			{
-				options.PacketCacheSize = cmdConfig.Cache.PacketCacheSize.Value * 1024 * 1024;
-			}
-
-			return new BundleCache(options);
-		}
-
-		static StorageBackendCache CreateStorageBackendCache(IServiceProvider serviceProvider)
-		{
-			CmdConfig cmdConfig = serviceProvider.GetRequiredService<IOptions<CmdConfig>>().Value;
-			DirectoryReference cacheDir = DirectoryReference.Combine(GetDataDir(), String.IsNullOrEmpty(cmdConfig.Cache.CacheDir)? "Cache" : cmdConfig.Cache.CacheDir);
-			return new StorageBackendCache(cacheDir, cmdConfig.Cache.CacheSize * 1024 * 1024, serviceProvider.GetRequiredService<ILogger<StorageBackendCache>>());
-		}
-
 		static DirectoryReference GetAppDir()
 		{
 			return new DirectoryReference(Path.GetDirectoryName(Assembly.GetExecutingAssembly().Location)!);
diff --git a/Engine/Source/Programs/Horde/Horde/CmdConfig.cs b/Engine/Source/Programs/Horde/Horde/CmdConfig.cs
--- a/Engine/Source/Programs/Horde/Horde/CmdConfig.cs
+++ b/Engine/Source/Programs/Horde/Horde/CmdConfig.cs
@@ -34,11 +34,6 @@ namespace Horde
 			}
 		}
 
-		/// <summary>
-		/// Settings for the storage cache
-		/// </summary>
-		public CmdCacheConfig Cache { get; set; } = new CmdCacheConfig();
-
 		/// <summary>
 		/// Gets the location for the config file
 		/// </summary>
@@ -74,30 +69,4 @@ namespace Horde
 			await FileReference.WriteAllBytesAsync(Location, data, cancellationToken);
 		}
 	}
-
-	/// <summary>
-	/// Settings for configuring the storage system
-	/// </summary>
-	public class CmdCacheConfig
-	{
-		/// <summary>
-		/// Directory to use for the persistent backend cache
-		/// </summary>
-		public string? CacheDir { get; set; }
-
-		/// <summary>
-		/// Size of the disk cache, in megabytes
-		/// </summary>
-		public long CacheSize { get; set; } = 1024;
-
-		/// <summary>
-		/// Size of the in-memory header cache, in megabytes 
-		/// </summary>
-		public long? HeaderCacheSize { get; set; }
-
-		/// <summary>
-		/// Size of the in-memory object cache, in megabytes
-		/// </summary>
-		public long? PacketCacheSize { get; set; }
-	}
 }
diff --git a/Engine/Source/Programs/Horde/Horde/Commands/LoginCommand.cs b/Engine/Source/Programs/Horde/Horde/Commands/LoginCommand.cs
--- a/Engine/Source/Programs/Horde/Horde/Commands/LoginCommand.cs
+++ b/Engine/Source/Programs/Horde/Horde/Commands/LoginCommand.cs
@@ -28,13 +28,13 @@ namespace Horde.Commands
 		public bool Token { get; set; }
 
 		readonly IServiceProvider _serviceProvider;
-		readonly IHttpClientFactory _httpClientFactory;
+		readonly IHordeClient _hordeClient;
 		readonly CmdConfig _config;
 
-		public LoginCommand(IServiceProvider serviceProvider, IHttpClientFactory httpClientFactory, IOptions<CmdConfig> config)
+		public LoginCommand(IServiceProvider serviceProvider, IHordeClient hordeClient, IOptions<CmdConfig> config)
 		{
 			_serviceProvider = serviceProvider;
-			_httpClientFactory = httpClientFactory;
+			_hordeClient = hordeClient;
 			_config = config.Value;
 		}
 
@@ -47,7 +47,7 @@ namespace Horde.Commands
 				await _config.WriteAsync();
 			}
 
-			using HordeHttpClient httpClient = _httpClientFactory.CreateHordeClient();
+			using HordeHttpClient httpClient = _hordeClient.CreateHttpClient();
 
 			GetServerInfoResponse serverInfo = await httpClient.GetServerInfoAsync();
 			logger.LogInformation("Connected to server version: {Version}", serverInfo.ServerVersion);
diff --git a/Engine/Source/Programs/Horde/Horde/Commands/Tools/ToolList.cs b/Engine/Source/Programs/Horde/Horde/Commands/Tools/ToolList.cs
--- a/Engine/Source/Programs/Horde/Horde/Commands/Tools/ToolList.cs
+++ b/Engine/Source/Programs/Horde/Horde/Commands/Tools/ToolList.cs
@@ -23,7 +23,7 @@ namespace Horde.Commands
 			GetToolsSummaryResponse response = await _hordeHttpClient.GetToolsAsync();
 			foreach (GetToolSummaryResponse tool in response.Tools)
 			{
-				logger.LogInformation("  {ToolId,-20} {Deployment,-30} {Version,-15}", tool.Id, tool.DeploymentId, tool.Version);
+				logger.LogInformation("  {ToolId,-30} {Deployment,-30} {Version,-15}", tool.Id, tool.DeploymentId, tool.Version);
 			}
 			return 0;
 		}
diff --git a/Engine/Source/Programs/Shared/EpicGames.Horde/HordeClient.cs b/Engine/Source/Programs/Shared/EpicGames.Horde/HordeClient.cs
new file mode 100644
--- /dev/null
+++ b/Engine/Source/Programs/Shared/EpicGames.Horde/HordeClient.cs
@@ -0,0 +1,26 @@
+﻿// Copyright Epic Games, Inc. All Rights Reserved.
+
+using System.Net.Http;
+
+namespace EpicGames.Horde
+{
+	/// <summary>
+	/// Default implementation of <see cref="IHordeClient"/>
+	/// </summary>
+	class HordeClient : IHordeClient
+	{
+		readonly IHttpClientFactory _httpClientFactory;
+
+		/// <summary>
+		/// Constructor
+		/// </summary>
+		public HordeClient(IHttpClientFactory httpClientFactory)
+		{
+			_httpClientFactory = httpClientFactory;
+		}
+
+		/// <inheritdoc/>
+		public HordeHttpClient CreateHttpClient()
+			=> _httpClientFactory.CreateHordeHttpClient();
+	}
+}
diff --git a/Engine/Source/Programs/Shared/EpicGames.Horde/HordeExtensions.cs b/Engine/Source/Programs/Shared/EpicGames.Horde/HordeExtensions.cs
--- a/Engine/Source/Programs/Shared/EpicGames.Horde/HordeExtensions.cs
+++ b/Engine/Source/Programs/Shared/EpicGames.Horde/HordeExtensions.cs
@@ -1,12 +1,11 @@
 // Copyright Epic Games, Inc. All Rights Reserved.
 
 using System;
-using System.Net.Http;
-using System.Net.Http.Headers;
 using EpicGames.Horde.Storage;
 using EpicGames.Horde.Storage.Backends;
 using EpicGames.Horde.Storage.Clients;
 using Microsoft.Extensions.DependencyInjection;
+using Microsoft.Extensions.Options;
 
 namespace EpicGames.Horde
 {
@@ -21,47 +20,23 @@ namespace EpicGames.Horde
 		/// <param name="serviceCollection">Collection to register services with</param>
 		public static void AddHorde(this IServiceCollection serviceCollection)
 		{
-			AddHorde(serviceCollection, new HordeOptions());
-		}
-
-		/// <summary>
-		/// Adds Horde-related services with the default settings
-		/// </summary>
-		/// <param name="serviceCollection">Collection to register services with</param>
-		/// <param name="configure">Callback to configure options</param>
-		public static void AddHorde(this IServiceCollection serviceCollection, Action<HordeOptions> configure)
-		{
-			HordeOptions options = new HordeOptions();
-			configure(options);
-			AddHorde(serviceCollection, options);
+			serviceCollection.AddHordeHttpClient();
+			serviceCollection.AddSingleton<BundleCache>(sp => new BundleCache(sp.GetRequiredService<IOptions<HordeOptions>>().Value.BundleCache));
+			serviceCollection.AddSingleton<StorageBackendCache>();
+			serviceCollection.AddSingleton<HttpStorageBackendFactory>();
+			serviceCollection.AddSingleton<HttpStorageClientFactory>();
+			serviceCollection.AddSingleton<IHordeClient, HordeClient>();
 		}
 
 		/// <summary>
-		/// Adds Horde-related services with the default settings
+		/// Adds Horde-related services
 		/// </summary>
 		/// <param name="serviceCollection">Collection to register services with</param>
-		/// <param name="options">Options for the collection</param>
-		public static void AddHorde(this IServiceCollection serviceCollection, HordeOptions options)
+		/// <param name="configureHorde">Callback to configure options</param>
+		public static void AddHorde(this IServiceCollection serviceCollection, Action<HordeOptions> configureHorde)
 		{
-			void ConfigureHttpClient(HttpClient httpClient)
-			{
-				if (options.ServerUrl != null)
-				{
-					httpClient.BaseAddress = options.ServerUrl;
-				}
-				if (options.AccessToken != null)
-				{
-					httpClient.DefaultRequestHeaders.Authorization = new AuthenticationHeaderValue("Bearer", options.AccessToken);
-				}
-
-				options.ConfigureHttpClient?.Invoke(httpClient);
-			}
-
-			serviceCollection.AddHordeHttpClient(ConfigureHttpClient, options.AllowAuthPrompt);
-			serviceCollection.AddSingleton<BundleCache>(sp => new BundleCache(options.BundleCache));
-			serviceCollection.AddSingleton<StorageBackendCache>();
-			serviceCollection.AddSingleton<HttpStorageBackendFactory>();
-			serviceCollection.AddSingleton<HttpStorageClientFactory>();
+			serviceCollection.Configure(configureHorde);
+			AddHorde(serviceCollection);
 		}
 	}
 }
diff --git a/Engine/Source/Programs/Shared/EpicGames.Horde/HordeHttpAuthHandler.cs b/Engine/Source/Programs/Shared/EpicGames.Horde/HordeHttpAuthHandler.cs
--- a/Engine/Source/Programs/Shared/EpicGames.Horde/HordeHttpAuthHandler.cs
+++ b/Engine/Source/Programs/Shared/EpicGames.Horde/HordeHttpAuthHandler.cs
@@ -14,6 +14,7 @@ using EpicGames.OIDC;
 using System.Net.Http.Json;
 using System.Text.Json;
 using EpicGames.Horde.Server;
+using Microsoft.Extensions.Options;
 
 namespace EpicGames.Horde
 {
@@ -23,20 +24,27 @@ namespace EpicGames.Horde
 	public class HordeHttpAuthHandler : DelegatingHandler
 	{
 		readonly HordeHttpAuthHandlerState _authState;
+		readonly IOptions<HordeOptions> _options;
 
 		AuthenticationHeaderValue? _authHeader;
 
 		/// <summary>
 		/// Constructor
 		/// </summary>
-		public HordeHttpAuthHandler(HordeHttpAuthHandlerState authState)
+		public HordeHttpAuthHandler(HordeHttpAuthHandlerState authState, IOptions<HordeOptions> options)
 		{
 			_authState = authState;
+			_options = options;
 		}
 
 		/// <inheritdoc/>
 		protected override async Task<HttpResponseMessage> SendAsync(HttpRequestMessage request, CancellationToken cancellationToken)
 		{
+			if (!_options.Value.AllowAuthPrompt)
+			{
+				return await base.SendAsync(request, cancellationToken);
+			}
+
 			// Do not try to override the auth header if the user has specified it explicitly
 			if (request.Headers.Authorization != null)
 			{
diff --git a/Engine/Source/Programs/Shared/EpicGames.Horde/HordeHttpClient.cs b/Engine/Source/Programs/Shared/EpicGames.Horde/HordeHttpClient.cs
--- a/Engine/Source/Programs/Shared/EpicGames.Horde/HordeHttpClient.cs
+++ b/Engine/Source/Programs/Shared/EpicGames.Horde/HordeHttpClient.cs
@@ -23,6 +23,7 @@ using EpicGames.Horde.Streams;
 using EpicGames.Horde.Tools;
 using Microsoft.Extensions.DependencyInjection;
 using Microsoft.Extensions.Logging;
+using Microsoft.Extensions.Options;
 using Polly;
 using Polly.Extensions.Http;
 using Polly.Retry;
@@ -445,7 +446,7 @@ namespace EpicGames.Horde
 		/// <summary>
 		/// Creates a <see cref="HordeHttpClient"/> instance from an http client factory
 		/// </summary>
-		public static HordeHttpClient CreateHordeClient(this IHttpClientFactory factory)
+		public static HordeHttpClient CreateHordeHttpClient(this IHttpClientFactory factory)
 		{
 			return new HordeHttpClient(factory.CreateClient(HordeHttpClient.HttpClientName));
 		}
@@ -454,10 +455,9 @@ namespace EpicGames.Horde
 		/// Registers a Horde HTTP client type, and configures it to use the default OIDC message handler.
 		/// </summary>
 		/// <param name="services">Service collection to add services to</param>
-		/// <param name="useAuthChallenge">Whether to prompt the user to authenticate if necessary</param>
-		public static IHttpClientBuilder AddHordeHttpClient(this IServiceCollection services, bool useAuthChallenge = true)
+		public static IHttpClientBuilder AddHordeHttpClient(this IServiceCollection services)
 		{
-			return services.AddHordeHttpClient((sp, client) => { }, useAuthChallenge);
+			return services.AddHordeHttpClient((sp, client) => { });
 		}
 
 		/// <summary>
@@ -465,10 +465,9 @@ namespace EpicGames.Horde
 		/// </summary>
 		/// <param name="services">Service collection to add services to</param>
 		/// <param name="configureClient">Callback to modify options for the http client</param>
-		/// <param name="useAuthChallenge">Whether to prompt the user to authenticate if necessary</param>
-		public static IHttpClientBuilder AddHordeHttpClient(this IServiceCollection services, Action<HttpClient> configureClient, bool useAuthChallenge = true)
+		public static IHttpClientBuilder AddHordeHttpClient(this IServiceCollection services, Action<HttpClient> configureClient)
 		{
-			return services.AddHordeHttpClient((sp, client) => configureClient(client), useAuthChallenge);
+			return services.AddHordeHttpClient((sp, client) => configureClient(client));
 		}
 
 		/// <summary>
@@ -476,14 +475,25 @@ namespace EpicGames.Horde
 		/// </summary>
 		/// <param name="services">Service collection to add services to</param>
 		/// <param name="configureClient">Callback to modify options for the http client</param>
-		/// <param name="useAuthChallenge">Whether to prompt the user to authenticate if necessary</param>
-		public static IHttpClientBuilder AddHordeHttpClient(this IServiceCollection services, Action<IServiceProvider, HttpClient> configureClient, bool useAuthChallenge = true)
+		public static IHttpClientBuilder AddHordeHttpClient(this IServiceCollection services, Action<IServiceProvider, HttpClient> configureClient)
 		{
 			// Sets defaults from the environment before calling the user provided configuration method
 			void ConfigureClientFromEnvironment(IServiceProvider serviceProvider, HttpClient httpClient)
 			{
+				IOptions<HordeOptions> options = serviceProvider.GetRequiredService<IOptions<HordeOptions>>();
+				if (options.Value.ServerUrl != null)
+				{
+					httpClient.BaseAddress = options.Value.ServerUrl;
+				}
+				if (options.Value.AccessToken != null)
+				{
+					httpClient.DefaultRequestHeaders.Authorization = new AuthenticationHeaderValue("Bearer", options.Value.AccessToken);
+				}
+
 				httpClient.Timeout = TimeSpan.FromSeconds(240); // Global timeout
 
+				// Run the user callbacks
+				options.Value.ConfigureHttpClient?.Invoke(httpClient);
 				configureClient(serviceProvider, httpClient);
 
 				// Only use the token from the environment if the configured base address is missing or matches the one configured in the environment
@@ -517,14 +527,11 @@ namespace EpicGames.Horde
 				.AddPolicyHandler((serviceProvider, request) => CreateDefaultTimeoutRetryPolicy(request, serviceProvider.GetRequiredService<ILogger<HttpStorageBackend>>()))
 				.AddPolicyHandler((serviceProvider, request) => CreateDefaultTransientErrorPolicy(request, serviceProvider.GetRequiredService<ILogger<HttpStorageBackend>>()));
 
-			if (useAuthChallenge)
-			{
-				services.AddSingleton<HordeHttpAuthHandlerState>();
-				services.AddTransient<HordeHttpAuthHandler>();
-				services.AddHttpClient(HordeHttpAuthHandlerState.HttpClientName, configureClient);
+			services.AddSingleton<HordeHttpAuthHandlerState>();
+			services.AddTransient<HordeHttpAuthHandler>();
+			services.AddHttpClient(HordeHttpAuthHandlerState.HttpClientName, configureClient);
 
-				builder = builder.AddHttpMessageHandler<HordeHttpAuthHandler>();
-			}
+			builder = builder.AddHttpMessageHandler<HordeHttpAuthHandler>();
 
 			return builder;
 		}
diff --git a/Engine/Source/Programs/Shared/EpicGames.Horde/IHordeClient.cs b/Engine/Source/Programs/Shared/EpicGames.Horde/IHordeClient.cs
new file mode 100644
--- /dev/null
+++ b/Engine/Source/Programs/Shared/EpicGames.Horde/IHordeClient.cs
@@ -0,0 +1,16 @@
+// Copyright Epic Games, Inc. All Rights Reserved.
+
+namespace EpicGames.Horde
+{
+	/// <summary>
+	/// Base interface for Horde functionality.
+	/// </summary>
+	public interface IHordeClient
+	{
+		/// <summary>
+		/// Creates a http client 
+		/// </summary>
+		/// <returns></returns>
+		HordeHttpClient CreateHttpClient();
+	}
+}
