Horde: Add an API version number that can be queried from the /api/v1/server/info endpoint. Allows coordinating client requests to the appropriate server version.

diff --git a/Engine/Source/Programs/Horde/Horde.Server/Agents/Pools/PoolConfig.cs b/Engine/Source/Programs/Horde/Horde.Server/Agents/Pools/PoolConfig.cs
--- a/Engine/Source/Programs/Horde/Horde.Server/Agents/Pools/PoolConfig.cs
+++ b/Engine/Source/Programs/Horde/Horde.Server/Agents/Pools/PoolConfig.cs
@@ -7,7 +7,6 @@ using System.Linq;
 using System.Text;
 using System.Text.Json.Serialization;
 using EpicGames.Core;
-using EpicGames.Horde;
 using EpicGames.Horde.Agents.Pools;
 using EpicGames.Horde.Common;
 using EpicGames.Horde.Compute;
diff --git a/Engine/Source/Programs/Horde/Horde.Server/Server/ServerController.cs b/Engine/Source/Programs/Horde/Horde.Server/Server/ServerController.cs
--- a/Engine/Source/Programs/Horde/Horde.Server/Server/ServerController.cs
+++ b/Engine/Source/Programs/Horde/Horde.Server/Server/ServerController.cs
@@ -8,6 +8,7 @@ using System.Runtime.InteropServices;
 using System.Threading;
 using System.Threading.Tasks;
 using EpicGames.Core;
+using EpicGames.Horde;
 using EpicGames.Horde.Server;
 using EpicGames.Perforce;
 using Horde.Server.Agents;
@@ -70,6 +71,7 @@ namespace Horde.Server.Server
 		public async Task<ActionResult<GetServerInfoResponse>> GetServerInfoAsync()
 		{
 			GetServerInfoResponse response = new GetServerInfoResponse();
+			response.ApiVersion = (int)HordeApiVersion.Latest;
 
 			FileVersionInfo versionInfo = FileVersionInfo.GetVersionInfo(Assembly.GetExecutingAssembly().Location);
 			response.ServerVersion = versionInfo.ProductVersion ?? String.Empty;
diff --git a/Engine/Source/Programs/Shared/EpicGames.Horde/HordeApiVersion.cs b/Engine/Source/Programs/Shared/EpicGames.Horde/HordeApiVersion.cs
new file mode 100644
--- /dev/null
+++ b/Engine/Source/Programs/Shared/EpicGames.Horde/HordeApiVersion.cs
@@ -0,0 +1,58 @@
+// Copyright Epic Games, Inc. All Rights Reserved.
+
+using System;
+using System.Text.Json;
+using System.Text.Json.Serialization;
+
+#pragma warning disable CA1027
+#pragma warning disable CA1069
+
+namespace EpicGames.Horde
+{
+	/// <summary>
+	/// Version number for the Horde public API. Can be retrieved from the /api/v1/server/info endpoint via the
+	/// <see cref="Server.GetServerInfoResponse"/> response. Should be serialized as an integer in messages to allow
+	/// clients with missing enum names to still parse the result correctly.
+	/// </summary>
+	public enum HordeApiVersion
+	{
+		/// <summary>
+		/// Unknown version
+		/// </summary>
+		Unknown,
+
+		/// <summary>
+		/// Initial version
+		/// </summary>
+		Initial,
+
+		/// <summary>
+		/// One past the latest known version number. Add new version numbers above this point.
+		/// </summary>
+		LatestPlusOne,
+
+		/// <summary>
+		/// Latest API version
+		/// </summary>
+		Latest = (int)(LatestPlusOne - 1),
+	}
+
+	/// <summary>
+	/// Converter for <see cref="HordeApiVersion"/> which forces serialization as an integer, to override any 
+	/// default <see cref="JsonStringEnumConverter"/> that may be enabled.
+	/// </summary>
+	public class HordeApiVersionConverter : JsonConverter<HordeApiVersion>
+	{
+		/// <inheritdoc/>
+		public override HordeApiVersion Read(ref Utf8JsonReader reader, Type typeToConvert, JsonSerializerOptions options)
+		{
+			return (HordeApiVersion)reader.GetInt32();
+		}
+
+		/// <inheritdoc/>
+		public override void Write(Utf8JsonWriter writer, HordeApiVersion value, JsonSerializerOptions options)
+		{
+			writer.WriteNumberValue((int)value);
+		}
+	}
+}
diff --git a/Engine/Source/Programs/Shared/EpicGames.Horde/Server/ServerMessages.cs b/Engine/Source/Programs/Shared/EpicGames.Horde/Server/ServerMessages.cs
--- a/Engine/Source/Programs/Shared/EpicGames.Horde/Server/ServerMessages.cs
+++ b/Engine/Source/Programs/Shared/EpicGames.Horde/Server/ServerMessages.cs
@@ -1,6 +1,7 @@
 // Copyright Epic Games, Inc. All Rights Reserved.
 
 using System;
+using System.Text.Json.Serialization;
 
 namespace EpicGames.Horde.Server
 {
@@ -9,6 +10,12 @@ namespace EpicGames.Horde.Server
 	/// </summary>
 	public class GetServerInfoResponse
 	{
+		/// <summary>
+		/// Current API version number of the server
+		/// </summary>
+		[JsonConverter(typeof(HordeApiVersionConverter))]
+		public HordeApiVersion ApiVersion { get; set; }
+
 		/// <summary>
 		/// Server version info
 		/// </summary>
