Horde: Sanitize object keys before writing artifacts to storage.

diff --git a/Engine/Source/Programs/Horde/Horde.Server/Jobs/Artifacts/ArtifactCollectionV1.cs b/Engine/Source/Programs/Horde/Horde.Server/Jobs/Artifacts/ArtifactCollectionV1.cs
--- a/Engine/Source/Programs/Horde/Horde.Server/Jobs/Artifacts/ArtifactCollectionV1.cs
+++ b/Engine/Source/Programs/Horde/Horde.Server/Jobs/Artifacts/ArtifactCollectionV1.cs
@@ -11,6 +11,8 @@ using MongoDB.Bson.Serialization.Attributes;
 using MongoDB.Driver;
 using EpicGames.Horde.Storage;
 using EpicGames.Horde.Jobs;
+using EpicGames.Horde;
+using EpicGames.Core;
 
 namespace Horde.Server.Jobs.Artifacts
 {
@@ -211,11 +213,11 @@ namespace Horde.Server.Jobs.Artifacts
 		{
 			if (stepId == null)
 			{
-				return new ObjectKey($"{jobId}/{name}");
+				return ObjectKey.Sanitize(new Utf8String($"{jobId}/{name}"));
 			}
 			else
 			{
-				return new ObjectKey($"{jobId}/{stepId.Value}/{name}");
+				return ObjectKey.Sanitize(new Utf8String($"{jobId}/{stepId.Value}/{name}"));
 			}
 		}
 
diff --git a/Engine/Source/Programs/Shared/EpicGames.Horde/Storage/ObjectKey.cs b/Engine/Source/Programs/Shared/EpicGames.Horde/Storage/ObjectKey.cs
--- a/Engine/Source/Programs/Shared/EpicGames.Horde/Storage/ObjectKey.cs
+++ b/Engine/Source/Programs/Shared/EpicGames.Horde/Storage/ObjectKey.cs
@@ -59,6 +59,31 @@ namespace EpicGames.Horde.Storage
 			}
 		}
 
+		/// <summary>
+		/// Makes an object key from the given path
+		/// </summary>
+		public static ObjectKey Sanitize(Utf8String path)
+		{
+			byte[]? data = null; 
+			for (int idx = 0; idx < path.Length; idx++)
+			{
+				byte character = path[idx];
+				if (!StringId.IsValidCharacter(character) && character != '/')
+				{
+					data ??= path.Memory.ToArray();
+					if (character >= 'A' && character <= 'Z')
+					{
+						data[idx] = (byte)('a' + (character - 'A'));
+					}
+					else
+					{
+						data[idx] = (byte)'_';
+					}
+				}
+			}
+			return new ObjectKey((data == null) ? path : new Utf8String(data));
+		}
+
 		/// <summary>
 		/// Whether the blob locator is valid
 		/// </summary>
