StateTree: Introduced TStateTreePropertyRef - type-safe FStateTreePropertyRef wrapper

diff --git a/Engine/Plugins/Runtime/StateTree/Source/StateTreeModule/Private/StateTreePropertyRefHelpers.cpp b/Engine/Plugins/Runtime/StateTree/Source/StateTreeModule/Private/StateTreePropertyRefHelpers.cpp
--- a/Engine/Plugins/Runtime/StateTree/Source/StateTreeModule/Private/StateTreePropertyRefHelpers.cpp
+++ b/Engine/Plugins/Runtime/StateTree/Source/StateTreeModule/Private/StateTreePropertyRefHelpers.cpp
@@ -115,7 +115,8 @@ namespace UE::StateTree::PropertyRefHelpers
 
 			if (const FObjectProperty* ObjectProperty = CastField<FObjectProperty>(TestProperty))
 			{
-				return ObjectProperty->PropertyClass->IsChildOf(Cast<UStruct>(TargetRefField));
+				// Only referencing object of the same exact class should be allowed. Otherwise one could e.g assign UObject to AActor property through reference to UObject.
+				return ObjectProperty->PropertyClass == Cast<UStruct>(TargetRefField);
 			}
 
 			if (const FEnumProperty* EnumProperty = CastField<FEnumProperty>(TestProperty))
diff --git a/Engine/Plugins/Runtime/StateTree/Source/StateTreeModule/Public/StateTreePropertyRef.h b/Engine/Plugins/Runtime/StateTree/Source/StateTreeModule/Public/StateTreePropertyRef.h
--- a/Engine/Plugins/Runtime/StateTree/Source/StateTreeModule/Public/StateTreePropertyRef.h
+++ b/Engine/Plugins/Runtime/StateTree/Source/StateTreeModule/Public/StateTreePropertyRef.h
@@ -33,7 +33,7 @@
  *	UPROPERTY(EditAnywhere, meta = (RefType = "/Script/ModuleName.TestStructBase", IsRefToArray))
  *	FStateTreePropertyRef RefToArrayOfTests;
  */
-USTRUCT(BlueprintType)
+USTRUCT()
 struct STATETREEMODULE_API FStateTreePropertyRef
 {
 	GENERATED_BODY()
@@ -61,4 +61,35 @@ private:
 	FStateTreeIndex16 RefAccessIndex;
 
 	friend FStateTreePropertyBindingCompiler;
+};
+
+/**
+ * TStateTreePropertyRef is a type-safe FStateTreePropertyRef wrapper against the given type.
+ * @note When used as a property, this automatically defines PropertyRef property meta-data.
+ * 
+ * Example:
+ *
+ *  // Reference to float
+ *	UPROPERTY(EditAnywhere)
+ *	TStateTreePropertyRef<float> RefToFloat;
+ * 
+ *  // Reference to FTestStructBase
+ *	UPROPERTY(EditAnywhere)
+ *	TStateTreePropertyRef<FTestStructBase> RefToTest;
+ *
+ *  // Reference to TArray<FTestStructBase>
+ *	UPROPERTY(EditAnywhere)
+ *	TStateTreePropertyRef<TArray<FTestStructBase>> RefToArrayOfTests;
+ */
+template<class TRef>
+struct TStateTreePropertyRef
+{
+	/** @return pointer to the property if possible, nullptr otherwise. */
+	TRef* GetMutablePtr(FStateTreeExecutionContext& Context) const
+	{
+		return PropertyRef.GetMutablePtr<TRef>(Context);
+	}
+
+private:
+	FStateTreePropertyRef PropertyRef;
 };
\ No newline at end of file
diff --git a/Engine/Plugins/Runtime/StateTree/Source/StateTreeModule/Public/StateTreePropertyRefHelpers.h b/Engine/Plugins/Runtime/StateTree/Source/StateTreeModule/Public/StateTreePropertyRefHelpers.h
--- a/Engine/Plugins/Runtime/StateTree/Source/StateTreeModule/Public/StateTreePropertyRefHelpers.h
+++ b/Engine/Plugins/Runtime/StateTree/Source/StateTreeModule/Public/StateTreePropertyRefHelpers.h
@@ -133,7 +133,7 @@ namespace UE::StateTree::PropertyRefHelpers
 		{		
 			if (const FStructProperty* StructProperty = CastField<FStructProperty>(&Property))
 			{
-				return StructProperty->Struct == TBaseStructure<T>::Get();
+				return StructProperty->Struct->IsChildOf(TBaseStructure<T>::Get());
 			}
 
 			return false;
diff --git a/Engine/Source/Programs/Shared/EpicGames.UHT/Types/Properties/UhtStructProperty.cs b/Engine/Source/Programs/Shared/EpicGames.UHT/Types/Properties/UhtStructProperty.cs
--- a/Engine/Source/Programs/Shared/EpicGames.UHT/Types/Properties/UhtStructProperty.cs
+++ b/Engine/Source/Programs/Shared/EpicGames.UHT/Types/Properties/UhtStructProperty.cs
@@ -592,6 +592,90 @@ namespace EpicGames.UHT.Types
 			instancedStructProperty.MetaData.Add("BaseStruct", baseScriptStruct.PathName);
 			return instancedStructProperty;
 		}
+
+		[UhtPropertyType(Keyword = "TStateTreePropertyRef")]
+		[SuppressMessage("CodeQuality", "IDE0051:Remove unused private members", Justification = "Attribute accessed method")]
+		[SuppressMessage("Style", "IDE0060:Remove unused parameter", Justification = "Attribute accessed method")]
+		private static UhtProperty? StateTreePropertyRefProperty(UhtPropertyResolvePhase resolvePhase, UhtPropertySettings propertySettings, IUhtTokenReader tokenReader, UhtToken matchedToken)
+		{
+			const string RefTypeName = "RefType";
+			const string IsRefToArrayName = "IsRefToArray";
+			UhtSession session = propertySettings.Outer.Session;
+
+			if(session.FStateTreePropertyRef == null)
+			{
+				return null;
+			}
+
+			if (propertySettings.MetaData.ContainsKey(RefTypeName))
+			{
+				tokenReader.LogError("{0} metadata is implicitly set from the TStateTreePropertyRef template argument and should not be explicitly specified.", RefTypeName);
+				return null;
+			}
+
+			if (propertySettings.MetaData.ContainsKey(IsRefToArrayName))
+			{
+				tokenReader.LogError("{0} metadata is implicitly set from the TStateTreePropertyRef template argument and should not be explicitly specified.", IsRefToArrayName);
+				return null;
+			}
+
+			if(!tokenReader.SkipExpectedType(matchedToken.Value, propertySettings.PropertyCategory == UhtPropertyCategory.Member))
+			{
+				return null;
+			}
+
+			tokenReader.Require('<');
+
+			bool isRefToArray = tokenReader.TryOptional("TArray");
+
+			if (isRefToArray)
+			{
+				tokenReader.Require('<');
+			}
+
+			UhtToken identifier = new();
+			tokenReader
+				.Optional("struct")
+				.Optional("class")
+				.RequireIdentifier((ref UhtToken token) => { identifier = token; });
+
+			session.Config!.RedirectTypeIdentifier(ref identifier);
+
+			UhtStructProperty instancedStructProperty = new UhtStructProperty(propertySettings, session.FStateTreePropertyRef);
+
+			// TStateTreePropertyRef supports UStructs, UClasses, enums and primitive types.
+			UhtType? foundType = propertySettings.Outer.FindType(UhtFindOptions.SourceName | UhtFindOptions.TypesMask, ref identifier);
+			if (foundType is UhtStruct foundStruct)
+			{
+				if (foundStruct.IsChildOf(session.UObject))
+				{
+					tokenReader.Require('*');
+				}
+
+				// It's a UStruct or UClass
+				instancedStructProperty.MetaData.Add(RefTypeName, foundStruct.PathName);
+			}
+			else if(foundType is UhtEnum foundEnum)
+			{
+				// It's an enum
+				instancedStructProperty.MetaData.Add(RefTypeName, foundEnum.PathName);
+			}
+			else
+			{
+				// It's a primitive or unknown type.
+				instancedStructProperty.MetaData.Add(RefTypeName, identifier.ToString());
+			}
+
+			if (isRefToArray)
+			{
+				tokenReader.Require(">");
+				instancedStructProperty.MetaData.Add(IsRefToArrayName, true);
+			}
+
+			tokenReader.Require(">");
+
+			return instancedStructProperty;
+		}
 		#endregion
 	}
 }
diff --git a/Engine/Source/Programs/Shared/EpicGames.UHT/Utils/UhtSession.cs b/Engine/Source/Programs/Shared/EpicGames.UHT/Utils/UhtSession.cs
--- a/Engine/Source/Programs/Shared/EpicGames.UHT/Utils/UhtSession.cs
+++ b/Engine/Source/Programs/Shared/EpicGames.UHT/Utils/UhtSession.cs
@@ -777,6 +777,12 @@ namespace EpicGames.UHT.Utils
 		/// </summary>
 		public UhtScriptStruct? FInstancedStruct { get; set; } = null;
 
+		/// <summary>
+		/// After headers are parsed, returns the FStateTreePropertyRef script struct.
+		/// There is no requirement for FStateTreePropertyRef to be defined. May be null.
+		/// </summary>
+		public UhtScriptStruct? FStateTreePropertyRef { get; set; } = null;
+
 		private readonly List<UhtPackage> _packages = new();
 		private readonly List<UhtHeaderFile> _headerFiles = new();
 		private readonly List<UhtHeaderFile> _sortedHeaderFiles = new();
@@ -2046,6 +2052,7 @@ namespace EpicGames.UHT.Utils
 				AActor = (UhtClass?)FindType(null, UhtFindOptions.SourceName | UhtFindOptions.Class, "AActor");
 				INotifyFieldValueChanged = (UhtClass?)FindType(null, UhtFindOptions.SourceName | UhtFindOptions.Class, "INotifyFieldValueChanged");
 				FInstancedStruct = (UhtScriptStruct?)FindType(null, UhtFindOptions.SourceName | UhtFindOptions.ScriptStruct, "FInstancedStruct");
+				FStateTreePropertyRef = (UhtScriptStruct?)FindType(null, UhtFindOptions.SourceName | UhtFindOptions.ScriptStruct, "FStateTreePropertyRef");
 			});
 		}
 
