Horde: Disable validation for object keys passed to the object store from the legacy V1 endpoint.

diff --git a/Engine/Source/Programs/Horde/Horde.Server/Jobs/Artifacts/ArtifactCollectionV1.cs b/Engine/Source/Programs/Horde/Horde.Server/Jobs/Artifacts/ArtifactCollectionV1.cs
--- a/Engine/Source/Programs/Horde/Horde.Server/Jobs/Artifacts/ArtifactCollectionV1.cs
+++ b/Engine/Source/Programs/Horde/Horde.Server/Jobs/Artifacts/ArtifactCollectionV1.cs
@@ -212,11 +212,11 @@ namespace Horde.Server.Jobs.Artifacts
 		{
 			if (stepId == null)
 			{
-				return new ObjectKey(new Utf8String($"{jobId}/{name}.blob"));
+				return new ObjectKey(new Utf8String($"{jobId}/{name}.blob"), ObjectKey.Validate.None);
 			}
 			else
 			{
-				return new ObjectKey(new Utf8String($"{jobId}/{stepId.Value}/{name}.blob"));
+				return new ObjectKey(new Utf8String($"{jobId}/{stepId.Value}/{name}.blob"), ObjectKey.Validate.None);
 			}
 		}
 
diff --git a/Engine/Source/Programs/Shared/EpicGames.Horde/Storage/ObjectKey.cs b/Engine/Source/Programs/Shared/EpicGames.Horde/Storage/ObjectKey.cs
--- a/Engine/Source/Programs/Shared/EpicGames.Horde/Storage/ObjectKey.cs
+++ b/Engine/Source/Programs/Shared/EpicGames.Horde/Storage/ObjectKey.cs
@@ -11,6 +11,17 @@ namespace EpicGames.Horde.Storage
 	[JsonSchemaString]
 	public readonly struct ObjectKey : IEquatable<ObjectKey>
 	{
+		/// <summary>
+		/// Dummy enum to allow invoking the constructor which takes a sanitized full path
+		/// </summary>
+		public enum Validate
+		{
+			/// <summary>
+			/// Dummy value
+			/// </summary>
+			None
+		}
+
 		readonly Utf8String _path;
 
 		/// <summary>
@@ -63,6 +74,17 @@ namespace EpicGames.Horde.Storage
 			}
 		}
 
+		/// <summary>
+		/// Constructor
+		/// </summary>
+		/// <param name="path">Path to the blob. The meaning of this string is implementation defined.</param>
+		/// <param name="validate"></param>
+		public ObjectKey(Utf8String path, Validate validate)
+		{
+			_path = path;
+			_ = validate;
+		}
+
 		/// <summary>
 		/// Makes an object key from the given path
 		/// </summary>
