Horde: Splitting Dashboard config out to top level, adding telemetry config and responses
#rnx

diff --git a/Engine/Source/Programs/Horde/Horde.Server/Dashboard/DashboardConfig.cs b/Engine/Source/Programs/Horde/Horde.Server/Dashboard/DashboardConfig.cs
new file mode 100644
--- /dev/null
+++ b/Engine/Source/Programs/Horde/Horde.Server/Dashboard/DashboardConfig.cs
@@ -0,0 +1,82 @@
+// Copyright Epic Games, Inc. All Rights Reserved.
+
+using System.Collections.Generic;
+using EpicGames.Core;
+using EpicGames.Horde.Common;
+using Horde.Server.Configuration;
+using Horde.Server.Server;
+
+namespace Horde.Server.Dashboard
+{
+
+	/// <summary>
+	/// Configuration for dashboard features
+	/// </summary>
+	[JsonSchema("https://unrealengine.com/horde/telemetry")]
+	[JsonSchemaCatalog("Horde Dashboard", "Horde dashboard configuration file", new[] { "*.dashboard.json", "Dashboard/*.json" })]
+	[ConfigIncludeRoot]
+	public class DashboardConfig
+	{
+		/// <summary>
+		/// Navigate to the landing page by default
+		/// </summary>
+		public bool ShowLandingPage { get; set; } = false;
+
+		/// <summary>
+		/// Enable CI functionality
+		/// </summary>
+		public bool ShowCI { get; set; } = true;
+
+		/// <summary>
+		/// Whether to show functionality related to agents, pools, and utilization on the dashboard.
+		/// </summary>
+		public bool ShowAgents { get; set; } = true;
+
+		/// <summary>
+		/// Show the Perforce server option on the server menu
+		/// </summary>
+		public bool ShowPerforceServers { get; set; } = true;
+
+		/// <summary>
+		/// Show the device manager on the server menu
+		/// </summary>
+		public bool ShowDeviceManager { get; set; } = true;
+
+		/// <summary>
+		/// Show automated tests on the server menu
+		/// </summary>
+		public bool ShowTests { get; set; } = true;
+
+		/// <summary>
+		/// Configuration for different agent pages
+		/// </summary>
+		public List<DashboardAgentCategoryConfig> AgentCategories { get; set; } = new List<DashboardAgentCategoryConfig>();
+
+		/// <summary>
+		/// Configuration for telemetry views
+		/// </summary>
+		public List<TelemetryViewConfig> Telemetry { get; set; } = new List<TelemetryViewConfig>();
+
+		/// <summary>
+		/// Includes for other configuration files
+		/// </summary>
+		public List<ConfigInclude> Include { get; set; } = new List<ConfigInclude>();
+
+	}
+
+	/// <summary>
+	/// Configuration for a category of agents
+	/// </summary>
+	public class DashboardAgentCategoryConfig
+	{
+		/// <summary>
+		/// Name of the category
+		/// </summary>
+		public string Name { get; set; } = "Unnamed";
+
+		/// <summary>
+		/// Condition string to be evaluated for this page
+		/// </summary>
+		public Condition? Condition { get; set; }
+	}
+}
diff --git a/Engine/Source/Programs/Horde/Horde.Server/Dashboard/DashboardController.cs b/Engine/Source/Programs/Horde/Horde.Server/Dashboard/DashboardController.cs
--- a/Engine/Source/Programs/Horde/Horde.Server/Dashboard/DashboardController.cs
+++ b/Engine/Source/Programs/Horde/Horde.Server/Dashboard/DashboardController.cs
@@ -147,6 +147,37 @@ namespace Horde.Server.Dashboard
 				dashboardConfigResponse.AgentCategories.Add(new GetDashboardAgentCategoryResponse { Name = category.Name, Condition = category.Condition });
 			}
 
+			foreach (TelemetryViewConfig telemetry in _globalConfig.Value.Dashboard.Telemetry)
+			{
+				TelemetryViewResponse rview = new TelemetryViewResponse();
+				rview.Id = telemetry.Id.ToString();
+				rview.Name = telemetry.Name;
+
+				foreach (TelemetryVariableConfig variable in telemetry.Variables)
+				{
+					rview.Variables.Add(new TelemetryVariableResponse { Name = variable.Name, Group = variable.Group });
+				}
+
+				foreach (TelemetryCategoryConfig category in telemetry.Categories)
+				{
+					TelemetryCategoryResponse rcategory = new TelemetryCategoryResponse { Name = category.Name };
+
+					foreach (TelemetryChartConfig chart in category.Charts)
+					{
+						TelemetryChartResponse rchart = new TelemetryChartResponse { Name = chart.Name, Display = chart.Display.ToString(), Graph = chart.Graph.ToString(), Max = chart.Max, Metrics = new List<TelemetryChartMetricResponse>() };
+
+						foreach (TelemetryChartMetricConfig metric  in chart.Metrics)
+						{
+							rchart.Metrics.Add(new TelemetryChartMetricResponse { MetricId = metric.MetricId.ToString(), Threshold = metric.Threshold, Alias = metric.Alias });
+						}
+					}
+
+					rview.Categories.Add(rcategory);
+				}
+
+				dashboardConfigResponse.TelemetryViews.Add(rview);
+			}
+
 			return dashboardConfigResponse;
 		}
 
diff --git a/Engine/Source/Programs/Horde/Horde.Server/Server/GlobalConfig.cs b/Engine/Source/Programs/Horde/Horde.Server/Server/GlobalConfig.cs
--- a/Engine/Source/Programs/Horde/Horde.Server/Server/GlobalConfig.cs
+++ b/Engine/Source/Programs/Horde/Horde.Server/Server/GlobalConfig.cs
@@ -27,6 +27,7 @@ using Horde.Server.Acls;
 using Horde.Server.Agents;
 using Horde.Server.Agents.Pools;
 using Horde.Server.Configuration;
+using Horde.Server.Dashboard;
 using Horde.Server.Devices;
 using Horde.Server.Projects;
 using Horde.Server.Secrets;
@@ -51,63 +52,6 @@ namespace Horde.Server.Server
 		public string Path { get; set; } = null!;
 	}
 
-	/// <summary>
-	/// Configuration for global features
-	/// </summary>
-	public class DashboardConfig
-	{
-		/// <summary>
-		/// Navigate to the landing page by default
-		/// </summary>
-		public bool ShowLandingPage { get; set; } = false;
-
-		/// <summary>
-		/// Enable CI functionality
-		/// </summary>
-		public bool ShowCI { get; set; } = true;
-
-		/// <summary>
-		/// Whether to show functionality related to agents, pools, and utilization on the dashboard.
-		/// </summary>
-		public bool ShowAgents { get; set; } = true;
-
-		/// <summary>
-		/// Show the Perforce server option on the server menu
-		/// </summary>
-		public bool ShowPerforceServers { get; set; } = true;
-
-		/// <summary>
-		/// Show the device manager on the server menu
-		/// </summary>
-		public bool ShowDeviceManager { get; set; } = true;
-
-		/// <summary>
-		/// Show automated tests on the server menu
-		/// </summary>
-		public bool ShowTests { get; set; } = true;
-
-		/// <summary>
-		/// Configuration for different agent pages
-		/// </summary>
-		public List<DashboardAgentCategoryConfig> AgentCategories { get; set; } = new List<DashboardAgentCategoryConfig>();
-	}
-
-	/// <summary>
-	/// Configuration for a category of agents
-	/// </summary>
-	public class DashboardAgentCategoryConfig
-	{
-		/// <summary>
-		/// Name of the category
-		/// </summary>
-		public string Name { get; set; } = "Unnamed";
-
-		/// <summary>
-		/// Condition string to be evaluated for this page
-		/// </summary>
-		public Condition? Condition { get; set; }
-	}
-
 	/// <summary>
 	/// Configuration for an artifact
 	/// </summary>
diff --git a/Engine/Source/Programs/Shared/EpicGames.Horde/Dashboard/DashboardMessages.cs b/Engine/Source/Programs/Shared/EpicGames.Horde/Dashboard/DashboardMessages.cs
--- a/Engine/Source/Programs/Shared/EpicGames.Horde/Dashboard/DashboardMessages.cs
+++ b/Engine/Source/Programs/Shared/EpicGames.Horde/Dashboard/DashboardMessages.cs
@@ -47,6 +47,12 @@ namespace EpicGames.Horde.Dashboard
 		/// Categories to display on the agents page
 		/// </summary>
 		public List<GetDashboardAgentCategoryResponse> AgentCategories { get; set; } = new List<GetDashboardAgentCategoryResponse>();
+
+		/// <summary>
+		/// Telemetry to display on the telemetry page
+		/// </summary>
+		public List<TelemetryViewResponse> TelemetryViews { get; set; } = new List<TelemetryViewResponse>();
+
 	}
 
 	/// <summary>
@@ -182,4 +188,117 @@ namespace EpicGames.Horde.Dashboard
 		/// </summary>
 		public string? TrackingLink { get; set; }
 	}
+
+	#region Telemetry
+
+	/// <summary>
+	/// Metric attached to a telemetry chart
+	/// </summary>
+	public class TelemetryChartMetricResponse
+	{
+		/// <summary>
+		/// Associated metric id
+		/// </summary>		
+		public string MetricId { get; set; } = null!;
+
+		/// <summary>
+		/// The threshold for KPI values
+		/// </summary>
+		public int? Threshold { get; set; }
+
+		/// <summary>
+		/// The metric alias for display purposes
+		/// </summary>
+		public string? Alias { get; set; }
+	}
+
+	/// <summary>
+	/// Telemetry chart configuraton
+	/// </summary>
+	public class TelemetryChartResponse
+	{
+		/// <summary>
+		/// The name of the chart, will be displayed on the dashboard
+		/// </summary>		
+		public string Name { get; set; } = null!;
+
+		/// <summary>
+		/// The unit to display
+		/// </summary>
+		public string Display { get; set; } = null!;
+
+		/// <summary>
+		/// The graph type 
+		/// </summary>
+		public string Graph { get; set; } = null!;
+
+		/// <summary>
+		/// List of configured metrics
+		/// </summary>
+		public List<TelemetryChartMetricResponse> Metrics { get; set; } = new List<TelemetryChartMetricResponse>();
+
+		/// <summary>
+		/// The max unit value for clamping chart
+		/// </summary>
+		public int? Max { get; set; }
+	}
+
+	/// <summary>
+	/// A chart categody, will be displayed on the dashbord under an associated pivot
+	/// </summary>
+	public class TelemetryCategoryResponse
+	{
+		/// <summary>
+		/// The name of the category
+		/// </summary>		
+		public string Name { get; set; } = null!;
+
+		/// <summary>
+		/// The charts contained within the category
+		/// </summary>
+		public List<TelemetryChartResponse> Charts { get; set; } = new List<TelemetryChartResponse> { };
+	}
+
+	/// <summary>
+	/// A telemetry view variable used for filtering the charting data
+	/// </summary>
+	public class TelemetryVariableResponse
+	{
+		/// <summary>
+		/// The name of the variable for display purposes
+		/// </summary>
+		public string Name { get; set; } = null!;
+
+		/// <summary>
+		/// The associated data group attached to the variable 
+		/// </summary>
+		public string Group { get; set; } = null!;
+	}
+
+	/// <summary>
+	/// A telemetry view of related metrics, divided into categofies
+	/// </summary>
+	public class TelemetryViewResponse
+	{
+		/// <summary>
+		/// Identifier for the view
+		/// </summary>
+		public string Id { get; set; } = null!;
+
+		/// <summary>
+		/// The name of the view
+		/// </summary>
+		public string Name { get; set; } = null!;
+
+		/// <summary>
+		///  The variables used to filter the view data
+		/// </summary>
+		public List<TelemetryVariableResponse> Variables { get; set; } = new List<TelemetryVariableResponse> { };
+
+		/// <summary>
+		/// The categories contained within the view
+		/// </summary>
+		public List<TelemetryCategoryResponse> Categories { get; set; } = new List<TelemetryCategoryResponse> { };
+	}
+	#endregion
 }
