Horde: Fix configuration of HTTP client used for interactive auth requests.

diff --git a/Engine/Source/Programs/Shared/EpicGames.Horde/HordeHttpClient.cs b/Engine/Source/Programs/Shared/EpicGames.Horde/HordeHttpClient.cs
--- a/Engine/Source/Programs/Shared/EpicGames.Horde/HordeHttpClient.cs
+++ b/Engine/Source/Programs/Shared/EpicGames.Horde/HordeHttpClient.cs
@@ -529,7 +529,17 @@ namespace EpicGames.Horde
 
 			services.AddSingleton<HordeHttpAuthHandlerState>();
 			services.AddTransient<HordeHttpAuthHandler>();
-			services.AddHttpClient(HordeHttpAuthHandlerState.HttpClientName, configureClient);
+
+			// Sets defaults from the configured HordeOptions
+			void ConfigureClientFromOptions(IServiceProvider serviceProvider, HttpClient httpClient)
+			{
+				IOptions<HordeOptions> options = serviceProvider.GetRequiredService<IOptions<HordeOptions>>();
+				if (options.Value.ServerUrl != null)
+				{
+					httpClient.BaseAddress = options.Value.ServerUrl;
+				}
+			}
+			services.AddHttpClient(HordeHttpAuthHandlerState.HttpClientName, ConfigureClientFromOptions);
 
 			builder = builder.AddHttpMessageHandler<HordeHttpAuthHandler>();
 
