Horde: Address naming consistency
#rnx

diff --git a/Engine/Source/Programs/Horde/Horde.Server/Dashboard/DashboardController.cs b/Engine/Source/Programs/Horde/Horde.Server/Dashboard/DashboardController.cs
--- a/Engine/Source/Programs/Horde/Horde.Server/Dashboard/DashboardController.cs
+++ b/Engine/Source/Programs/Horde/Horde.Server/Dashboard/DashboardController.cs
@@ -149,26 +149,26 @@ namespace Horde.Server.Dashboard
 
 			foreach (TelemetryViewConfig telemetry in _globalConfig.Value.Dashboard.Telemetry)
 			{
-				TelemetryViewResponse rview = new TelemetryViewResponse();
+				GetTelemetryViewResponse rview = new GetTelemetryViewResponse();
 				rview.Id = telemetry.Id.ToString();
 				rview.Name = telemetry.Name;
 
 				foreach (TelemetryVariableConfig variable in telemetry.Variables)
 				{
-					rview.Variables.Add(new TelemetryVariableResponse { Name = variable.Name, Group = variable.Group });
+					rview.Variables.Add(new GetTelemetryVariableResponse { Name = variable.Name, Group = variable.Group });
 				}
 
 				foreach (TelemetryCategoryConfig category in telemetry.Categories)
 				{
-					TelemetryCategoryResponse rcategory = new TelemetryCategoryResponse { Name = category.Name };
+					GetTelemetryCategoryResponse rcategory = new GetTelemetryCategoryResponse { Name = category.Name };
 
 					foreach (TelemetryChartConfig chart in category.Charts)
 					{
-						TelemetryChartResponse rchart = new TelemetryChartResponse { Name = chart.Name, Display = chart.Display.ToString(), Graph = chart.Graph.ToString(), Max = chart.Max, Metrics = new List<TelemetryChartMetricResponse>() };
+						GetTelemetryChartResponse rchart = new GetTelemetryChartResponse { Name = chart.Name, Display = chart.Display.ToString(), Graph = chart.Graph.ToString(), Max = chart.Max, Metrics = new List<GetTelemetryChartMetricResponse>() };
 
 						foreach (TelemetryChartMetricConfig metric  in chart.Metrics)
 						{
-							rchart.Metrics.Add(new TelemetryChartMetricResponse { MetricId = metric.MetricId.ToString(), Threshold = metric.Threshold, Alias = metric.Alias });
+							rchart.Metrics.Add(new GetTelemetryChartMetricResponse { MetricId = metric.MetricId.ToString(), Threshold = metric.Threshold, Alias = metric.Alias });
 						}
 
 						rcategory.Charts.Add(rchart);
diff --git a/Engine/Source/Programs/Shared/EpicGames.Horde/Dashboard/DashboardMessages.cs b/Engine/Source/Programs/Shared/EpicGames.Horde/Dashboard/DashboardMessages.cs
--- a/Engine/Source/Programs/Shared/EpicGames.Horde/Dashboard/DashboardMessages.cs
+++ b/Engine/Source/Programs/Shared/EpicGames.Horde/Dashboard/DashboardMessages.cs
@@ -51,7 +51,7 @@ namespace EpicGames.Horde.Dashboard
 		/// <summary>
 		/// Telemetry to display on the telemetry page
 		/// </summary>
-		public List<TelemetryViewResponse> TelemetryViews { get; set; } = new List<TelemetryViewResponse>();
+		public List<GetTelemetryViewResponse> TelemetryViews { get; set; } = new List<GetTelemetryViewResponse>();
 
 	}
 
@@ -194,7 +194,7 @@ namespace EpicGames.Horde.Dashboard
 	/// <summary>
 	/// Metric attached to a telemetry chart
 	/// </summary>
-	public class TelemetryChartMetricResponse
+	public class GetTelemetryChartMetricResponse
 	{
 		/// <summary>
 		/// Associated metric id
@@ -215,7 +215,7 @@ namespace EpicGames.Horde.Dashboard
 	/// <summary>
 	/// Telemetry chart configuraton
 	/// </summary>
-	public class TelemetryChartResponse
+	public class GetTelemetryChartResponse
 	{
 		/// <summary>
 		/// The name of the chart, will be displayed on the dashboard
@@ -235,7 +235,7 @@ namespace EpicGames.Horde.Dashboard
 		/// <summary>
 		/// List of configured metrics
 		/// </summary>
-		public List<TelemetryChartMetricResponse> Metrics { get; set; } = new List<TelemetryChartMetricResponse>();
+		public List<GetTelemetryChartMetricResponse> Metrics { get; set; } = new List<GetTelemetryChartMetricResponse>();
 
 		/// <summary>
 		/// The max unit value for clamping chart
@@ -246,7 +246,7 @@ namespace EpicGames.Horde.Dashboard
 	/// <summary>
 	/// A chart categody, will be displayed on the dashbord under an associated pivot
 	/// </summary>
-	public class TelemetryCategoryResponse
+	public class GetTelemetryCategoryResponse
 	{
 		/// <summary>
 		/// The name of the category
@@ -256,13 +256,13 @@ namespace EpicGames.Horde.Dashboard
 		/// <summary>
 		/// The charts contained within the category
 		/// </summary>
-		public List<TelemetryChartResponse> Charts { get; set; } = new List<TelemetryChartResponse> { };
+		public List<GetTelemetryChartResponse> Charts { get; set; } = new List<GetTelemetryChartResponse> { };
 	}
 
 	/// <summary>
 	/// A telemetry view variable used for filtering the charting data
 	/// </summary>
-	public class TelemetryVariableResponse
+	public class GetTelemetryVariableResponse
 	{
 		/// <summary>
 		/// The name of the variable for display purposes
@@ -278,7 +278,7 @@ namespace EpicGames.Horde.Dashboard
 	/// <summary>
 	/// A telemetry view of related metrics, divided into categofies
 	/// </summary>
-	public class TelemetryViewResponse
+	public class GetTelemetryViewResponse
 	{
 		/// <summary>
 		/// Identifier for the view
@@ -293,12 +293,12 @@ namespace EpicGames.Horde.Dashboard
 		/// <summary>
 		///  The variables used to filter the view data
 		/// </summary>
-		public List<TelemetryVariableResponse> Variables { get; set; } = new List<TelemetryVariableResponse> { };
+		public List<GetTelemetryVariableResponse> Variables { get; set; } = new List<GetTelemetryVariableResponse> { };
 
 		/// <summary>
 		/// The categories contained within the view
 		/// </summary>
-		public List<TelemetryCategoryResponse> Categories { get; set; } = new List<TelemetryCategoryResponse> { };
+		public List<GetTelemetryCategoryResponse> Categories { get; set; } = new List<GetTelemetryCategoryResponse> { };
 	}
 	#endregion
 }
