Horde: Fix compile errors.

diff --git a/Engine/Source/Programs/Shared/EpicGames.Horde/Storage/Nodes/ChunkedDataNode.cs b/Engine/Source/Programs/Shared/EpicGames.Horde/Storage/Nodes/ChunkedDataNode.cs
--- a/Engine/Source/Programs/Shared/EpicGames.Horde/Storage/Nodes/ChunkedDataNode.cs
+++ b/Engine/Source/Programs/Shared/EpicGames.Horde/Storage/Nodes/ChunkedDataNode.cs
@@ -411,6 +411,11 @@ namespace EpicGames.Horde.Storage.Nodes
 	/// <param name="SliceThreshold">Threshold hash value for splitting interior nodes</param>
 	public record class InteriorChunkedDataNodeOptions(int MinChildCount, int TargetChildCount, int MaxChildCount, uint SliceThreshold)
 	{
+		/// <summary>
+		/// Default settings
+		/// </summary>
+		public static InteriorChunkedDataNodeOptions Default { get; } = new InteriorChunkedDataNodeOptions(1, 100, 500);
+
 		/// <summary>
 		/// Constructor
 		/// </summary>
diff --git a/Engine/Source/Programs/Shared/EpicGames.Horde/Storage/Nodes/ChunkedDataWriter.cs b/Engine/Source/Programs/Shared/EpicGames.Horde/Storage/Nodes/ChunkedDataWriter.cs
--- a/Engine/Source/Programs/Shared/EpicGames.Horde/Storage/Nodes/ChunkedDataWriter.cs
+++ b/Engine/Source/Programs/Shared/EpicGames.Horde/Storage/Nodes/ChunkedDataWriter.cs
@@ -30,8 +30,8 @@ namespace EpicGames.Horde.Storage.Nodes
 		/// </summary>
 		public ChunkingOptions()
 		{
-			LeafOptions = new LeafChunkedDataNodeOptions(32 * 1024, 256 * 1024, 64 * 1024);
-			InteriorOptions = new InteriorChunkedDataNodeOptions(1, 5, 10);
+			LeafOptions = LeafChunkedDataNodeOptions.Default;
+			InteriorOptions = InteriorChunkedDataNodeOptions.Default;
 		}
 	}
 
@@ -43,6 +43,11 @@ namespace EpicGames.Horde.Storage.Nodes
 	/// <param name="TargetSize">Target chunk size for content-slicing</param>
 	public record class LeafChunkedDataNodeOptions(int MinSize, int MaxSize, int TargetSize)
 	{
+		/// <summary>
+		/// Default settings
+		/// </summary>
+		public static LeafChunkedDataNodeOptions Default { get; } = new LeafChunkedDataNodeOptions(32 * 1024, 256 * 1024, 64 * 1024);
+
 		/// <summary>
 		/// Constructor
 		/// </summary>
@@ -53,13 +58,6 @@ namespace EpicGames.Horde.Storage.Nodes
 		}
 	}
 
-	/// <summary>
-	/// Describes a chunked data stream
-	/// </summary>
-	/// <param name="StreamHash">Hash of the stream as a contiguous buffer</param>
-	/// <param name="Root">Handle to the root chunk containing the data</param>
-	public record class ChunkedData(IoHash StreamHash, ChunkedDataNodeRef Root);
-
 	/// <summary>
 	/// Utility class for generating FileNode data directly into <see cref="IBlobWriter"/> instances, without constructing node representations first.
 	/// </summary>
@@ -86,6 +84,15 @@ namespace EpicGames.Horde.Storage.Nodes
 		/// </summary>
 		public long Length => _totalLength;
 
+		/// <summary>
+		/// Constructor
+		/// </summary>
+		/// <param name="leafNodeWriter">Writer for new nodes</param>
+		public LeafChunkedDataWriter(IBlobWriter leafNodeWriter)
+			: this(leafNodeWriter, LeafChunkedDataNodeOptions.Default)
+		{
+		}
+
 		/// <summary>
 		/// Constructor
 		/// </summary>
@@ -336,4 +343,110 @@ namespace EpicGames.Horde.Storage.Nodes
 			ResetLeafState();
 		}
 	}
+
+
+	/// <summary>
+	/// Describes a chunked data stream
+	/// </summary>
+	/// <param name="StreamHash">Hash of the stream as a contiguous buffer</param>
+	/// <param name="Root">Handle to the root chunk containing the data</param>
+	public record class ChunkedData(IoHash StreamHash, ChunkedDataNodeRef Root);
+
+	/// <summary>
+	/// Writes chunked data to an output writer
+	/// </summary>
+	public sealed class ChunkedDataWriter : IDisposable
+	{
+		readonly IBlobWriter _writer;
+		readonly ChunkingOptions _chunkingOptions;
+		readonly BlobSerializerOptions _blobSerializerOptions;
+		readonly LeafChunkedDataWriter _leafWriter;
+
+		/// <summary>
+		/// Length of the current stream
+		/// </summary>
+		public long Length => _leafWriter.Length;
+
+		/// <summary>
+		/// Constructor
+		/// </summary>
+		public ChunkedDataWriter(IBlobWriter writer, ChunkingOptions chunkingOptions, BlobSerializerOptions blobSerializerOptions)
+		{
+			_writer = writer;
+			_chunkingOptions = chunkingOptions;
+			_blobSerializerOptions = blobSerializerOptions;
+			_leafWriter = new LeafChunkedDataWriter(writer, chunkingOptions.LeafOptions);
+		}
+		
+		/// <inheritdoc/>
+		public void Dispose()
+		{
+			_leafWriter.Dispose();
+		}
+
+		/// <summary>
+		/// Reset the current state
+		/// </summary>
+		public void Reset()
+		{
+			_leafWriter.Reset();
+		}
+
+		/// <summary>
+		/// Creates data for the given file
+		/// </summary>
+		/// <param name="fileInfo">File to append</param>
+		/// <param name="cancellationToken">Cancellation token for the operation</param>
+		public async Task<ChunkedData> CreateAsync(FileInfo fileInfo, CancellationToken cancellationToken)
+		{
+			Reset();
+			LeafChunkedData leafChunkedData = await _leafWriter.CreateAsync(fileInfo, cancellationToken);
+			return await InteriorChunkedDataNode.CreateTreeAsync(leafChunkedData, _chunkingOptions.InteriorOptions, _writer, _blobSerializerOptions, cancellationToken);
+		}
+
+		/// <summary>
+		/// Creates data from the given data
+		/// </summary>
+		/// <param name="data">Stream to append</param>
+		/// <param name="cancellationToken">Cancellation token for the operation</param>
+		public async Task<ChunkedData> CreateAsync(ReadOnlyMemory<byte> data, CancellationToken cancellationToken)
+		{
+			Reset();
+			await _leafWriter.AppendAsync(data, cancellationToken);
+			return await CompleteAsync(cancellationToken);
+		}
+
+		/// <summary>
+		/// Appends data to the current file
+		/// </summary>
+		/// <param name="data">Data to append</param>
+		/// <param name="cancellationToken">Cancellation token for the operation</param>
+		public async Task AppendAsync(ReadOnlyMemory<byte> data, CancellationToken cancellationToken)
+		{
+			await _leafWriter.AppendAsync(data, cancellationToken);
+		}
+
+		/// <summary>
+		/// Complete the current file, and write all open nodes to the underlying writer
+		/// </summary>
+		/// <param name="cancellationToken">Cancellation token for the operation</param>
+		/// <returns>Handle to the root node</returns>
+		public async Task<ChunkedData> CompleteAsync(CancellationToken cancellationToken = default)
+		{
+			LeafChunkedData leafChunkedData = await _leafWriter.CompleteAsync(cancellationToken);
+			return await InteriorChunkedDataNode.CreateTreeAsync(leafChunkedData, _chunkingOptions.InteriorOptions, _writer, _blobSerializerOptions, cancellationToken);
+		}
+
+		/// <summary>
+		/// Complete the current file, and write all open nodes to the underlying writer
+		/// </summary>
+		/// <param name="cancellationToken">Cancellation token for the operation</param>
+		/// <returns>Handle to the root node</returns>
+		public async Task<ChunkedData> FlushAsync(CancellationToken cancellationToken = default)
+		{
+			ChunkedData chunkedData = await CompleteAsync(cancellationToken);
+			await _writer.FlushAsync(cancellationToken);
+			return chunkedData;
+		}
+	}
 }
diff --git a/Engine/Source/Programs/Shared/EpicGames.Horde/Storage/Nodes/DirectoryNode.cs b/Engine/Source/Programs/Shared/EpicGames.Horde/Storage/Nodes/DirectoryNode.cs
--- a/Engine/Source/Programs/Shared/EpicGames.Horde/Storage/Nodes/DirectoryNode.cs
+++ b/Engine/Source/Programs/Shared/EpicGames.Horde/Storage/Nodes/DirectoryNode.cs
@@ -606,16 +606,6 @@ namespace EpicGames.Horde.Storage.Nodes
 			await writerFork.FlushAsync(cancellationToken);
 		}
 
-		static async Task<ChunkedData[]> CreateInteriorChunkNodesAsync(LeafChunkedData[] leafChunkedFiles, InteriorChunkedDataNodeOptions chunkingOptions, IBlobWriter writer, BlobSerializerOptions? serializerOptions, CancellationToken cancellationToken)
-		{
-			ChunkedData[] chunkedFiles = new ChunkedData[leafChunkedFiles.Length];
-			for (int idx = 0; idx < leafChunkedFiles.Length; idx++)
-			{
-				chunkedFiles[idx] = await InteriorChunkedDataNode.CreateTreeAsync(leafChunkedFiles[idx], chunkingOptions, writer, serializerOptions, cancellationToken);
-			}
-			return chunkedFiles;
-		}
-
 		/// <summary>
 		/// Updates this tree of directory objects
 		/// </summary>
