Horde: Allow printing progress of copy operations without file counts.

diff --git a/Engine/Source/Programs/Shared/EpicGames.Horde/Storage/Nodes/DirectoryNode.cs b/Engine/Source/Programs/Shared/EpicGames.Horde/Storage/Nodes/DirectoryNode.cs
--- a/Engine/Source/Programs/Shared/EpicGames.Horde/Storage/Nodes/DirectoryNode.cs
+++ b/Engine/Source/Programs/Shared/EpicGames.Horde/Storage/Nodes/DirectoryNode.cs
@@ -92,32 +92,20 @@ namespace EpicGames.Horde.Storage.Nodes
 	/// </summary>
 	public class CopyStatsLogger : IProgress<ICopyStats>
 	{
+		readonly int _totalCount;
+		readonly long _totalSize;
 		readonly ILogger _logger;
 
 		/// <summary>
 		/// Constructor
 		/// </summary>
-		public CopyStatsLogger(ILogger logger) 
+		public CopyStatsLogger(int totalCount, ILogger logger)
 			=> _logger = logger;
 
-		/// <inheritdoc/>
-		public void Report(ICopyStats stats)
-			=> _logger.LogInformation("Copied {NumFiles:n0} files ({Size:n1}mb, {Rate:n1}mb/s)", stats.Count, stats.Size / (1024.0 * 1024.0), stats.Rate / (1024.0 * 1024.0));
-	}
-
-	/// <summary>
-	/// Progress logger for writing copy stats
-	/// </summary>
-	public class CopyStatsLoggerWithTotals : IProgress<ICopyStats>
-	{
-		readonly int _totalCount;
-		readonly long _totalSize;
-		readonly ILogger _logger;
-
 		/// <summary>
 		/// Constructor
 		/// </summary>
-		public CopyStatsLoggerWithTotals(int totalCount, long totalSize, ILogger logger)
+		public CopyStatsLogger(int totalCount, long totalSize, ILogger logger)
 		{
 			_totalCount = totalCount;
 			_totalSize = totalSize;
@@ -126,7 +114,24 @@ namespace EpicGames.Horde.Storage.Nodes
 
 		/// <inheritdoc/>
 		public void Report(ICopyStats stats)
-			=> _logger.LogInformation("Copied {NumFiles:n0}/{TotalFiles:n0} files ({Size:n1}/{TotalSize:n1}mb, {Rate:n1}mb/s, {Pct}%)", stats.Count, _totalCount, stats.Size / (1024.0 * 1024.0), _totalSize / (1024.0 * 1024.0), stats.Rate / (1024.0 * 1024.0), (int)((Math.Max(stats.Size, 1) * 100) / Math.Max(_totalSize, 1)));
+		{
+			if (_totalCount > 0 && _totalSize > 0)
+			{
+				_logger.LogInformation("Copied {NumFiles:n0}/{TotalFiles:n0} files ({Size:n1}/{TotalSize:n1}mb, {Rate:n1}mb/s, {Pct}%)", stats.Count, _totalCount, stats.Size / (1024.0 * 1024.0), _totalSize / (1024.0 * 1024.0), stats.Rate / (1024.0 * 1024.0), (int)((Math.Max(stats.Size, 1) * 100) / Math.Max(_totalSize, 1)));
+			}
+			else if (_totalCount > 0)
+			{
+				_logger.LogInformation("Copied {NumFiles:n0}/{TotalFiles:n0} files ({Size:n1}mb, {Rate:n1}mb/s)", stats.Count, _totalCount, stats.Size / (1024.0 * 1024.0), stats.Rate / (1024.0 * 1024.0));
+			}
+			else if (_totalSize > 0)
+			{
+				_logger.LogInformation("Copied {NumFiles:n0} files ({Size:n1}/{TotalSize:n1}mb, {Rate:n1}mb/s, {Pct}%)", stats.Count, stats.Size / (1024.0 * 1024.0), _totalSize / (1024.0 * 1024.0), stats.Rate / (1024.0 * 1024.0), (int)((Math.Max(stats.Size, 1) * 100) / Math.Max(_totalSize, 1)));
+			}
+			else
+			{
+				_logger.LogInformation("Copied {NumFiles:n0} files ({Size:n1}mb, {Rate:n1}mb/s)", stats.Count, stats.Size / (1024.0 * 1024.0), stats.Rate / (1024.0 * 1024.0));
+			}
+		}
 	}
 
 	/// <summary>
