Horde: Remove unused argument to CopyStatsLogger constructor.

diff --git a/Engine/Source/Programs/Shared/EpicGames.Horde/Storage/Nodes/DirectoryNode.cs b/Engine/Source/Programs/Shared/EpicGames.Horde/Storage/Nodes/DirectoryNode.cs
--- a/Engine/Source/Programs/Shared/EpicGames.Horde/Storage/Nodes/DirectoryNode.cs
+++ b/Engine/Source/Programs/Shared/EpicGames.Horde/Storage/Nodes/DirectoryNode.cs
@@ -99,7 +99,7 @@ namespace EpicGames.Horde.Storage.Nodes
 		/// <summary>
 		/// Constructor
 		/// </summary>
-		public CopyStatsLogger(int totalCount, ILogger logger)
+		public CopyStatsLogger(ILogger logger)
 			=> _logger = logger;
 
 		/// <summary>
@@ -721,6 +721,11 @@ namespace EpicGames.Horde.Storage.Nodes
 		record class OutputFile(OutputDir Directory, FileEntry FileEntry);
 		record class OutputChunk(OutputFile File, long Offset, long Length, IBlobHandle Handle);
 
+		static void TraceBlobRead(string type, string path, IBlobHandle handle, ILogger logger)
+		{
+			logger.LogInformation("Blob [{Type,-20}] Path=\"{Path}\", Locator={Locator}", type, path, handle.GetLocator());
+		}
+
 		/// <summary>
 		/// Utility function to allow extracting a packed directory to disk
 		/// </summary>
@@ -740,7 +745,7 @@ namespace EpicGames.Horde.Storage.Nodes
 				copyStats = new CopyStats(progress);
 			}
 
-			Channel<OutputChunk> chunks = Channel.CreateBounded<OutputChunk>(new BoundedChannelOptions(128) { FullMode = BoundedChannelFullMode.Wait });
+			Channel<OutputChunk> chunks = Channel.CreateBounded<OutputChunk>(new BoundedChannelOptions(128 * 1024) { FullMode = BoundedChannelFullMode.Wait });
 			using (CancellationTokenSource cancellationSource = CancellationTokenSource.CreateLinkedTokenSource(cancellationToken))
 			{
 				// Helper method to run a background task and set a cancellation source on error
@@ -761,18 +766,20 @@ namespace EpicGames.Horde.Storage.Nodes
 					}
 				}
 
+				await FindOutputChunksRootAsync(chunks.Writer, options, logger, cancellationSource.Token);
+
 				List<Task> tasks = new List<Task>();
-				tasks.Add(RunBackgroundTask(ctx => FindOutputChunksRootAsync(chunks.Writer, options, ctx)));
+//				tasks.Add(RunBackgroundTask(ctx => FindOutputChunksRootAsync(chunks.Writer, options, logger, ctx)));
 				for (int idx = 0; idx < numTasks; idx++)
 				{
-					tasks.Add(RunBackgroundTask(ctx => ExtractAsync(chunks.Reader, new DirectoryReference(directoryInfo), copyStats, ctx)));
+					tasks.Add(RunBackgroundTask(ctx => ExtractAsync(chunks.Reader, new DirectoryReference(directoryInfo), copyStats, logger, ctx)));
 				}
 
 				await Task.WhenAll(tasks);
 			}
 		}
 
-		static async Task ExtractAsync(ChannelReader<OutputChunk> chunkReader, DirectoryReference baseDir, CopyStats? copyStats, CancellationToken cancellationToken)
+		static async Task ExtractAsync(ChannelReader<OutputChunk> chunkReader, DirectoryReference baseDir, CopyStats? copyStats, ILogger logger, CancellationToken cancellationToken)
 		{
 			OutputChunk? chunk = await ReadNextChunkAsync(chunkReader, cancellationToken);
 			while (chunk != null)
@@ -801,6 +808,7 @@ namespace EpicGames.Horde.Storage.Nodes
 					// Write this chunk
 					using (BlobData data = await chunk.Handle.ReadBlobDataAsync(cancellationToken))
 					{
+						TraceBlobRead("Leaf", CombinePaths(chunk.File.Directory.Path, chunk.File.FileEntry.Name), chunk.Handle, logger);
 						data.Data.CopyTo(memoryMappedView!.GetMemory(chunk.Offset, data.Data.Length));
 					}
 
@@ -834,30 +842,34 @@ namespace EpicGames.Horde.Storage.Nodes
 			}
 		}
 
-		async Task FindOutputChunksRootAsync(ChannelWriter<OutputChunk> chunks, BlobSerializerOptions? options, CancellationToken cancellationToken)
+		async Task FindOutputChunksRootAsync(ChannelWriter<OutputChunk> chunks, BlobSerializerOptions? options, ILogger logger, CancellationToken cancellationToken)
 		{
-			await FindOutputChunksAsync("", this, chunks, options, cancellationToken);
+			await FindOutputChunksAsync("", this, chunks, options, logger, cancellationToken);
 			chunks.Complete();
 		}
 
-		static async Task FindOutputChunksAsync(string path, DirectoryNode node, ChannelWriter<OutputChunk> chunks, BlobSerializerOptions? options, CancellationToken cancellationToken)
+		static async Task FindOutputChunksAsync(string path, DirectoryNode node, ChannelWriter<OutputChunk> chunks, BlobSerializerOptions? options, ILogger logger, CancellationToken cancellationToken)
 		{
 			OutputDir outputDir = new OutputDir(path, node);
 
 			foreach (FileEntry fileEntry in node._nameToFileEntry.Values)
 			{
 				OutputFile outputFile = new OutputFile(outputDir, fileEntry);
-				await FindOutputChunksAsync(outputFile, 0, fileEntry.Target, chunks, cancellationToken);
+				await FindOutputChunksAsync(outputFile, 0, fileEntry.Target, chunks, logger, cancellationToken);
 			}
 
 			foreach (DirectoryEntry directoryEntry in node._nameToDirectoryEntry.Values)
 			{
 				DirectoryNode subDirectoryNode = await directoryEntry.Handle.ReadBlobAsync(options, cancellationToken);
-				await FindOutputChunksAsync(CombinePaths(outputDir.Path, directoryEntry.Name), subDirectoryNode, chunks, options, cancellationToken);
+
+				string subPath = CombinePaths(outputDir.Path, directoryEntry.Name);
+				TraceBlobRead("Directory", subPath, directoryEntry.Handle, logger);
+
+				await FindOutputChunksAsync(subPath, subDirectoryNode, chunks, options, logger, cancellationToken);
 			}
 		}
 
-		static async Task<long> FindOutputChunksAsync(OutputFile outputFile, long offset, ChunkedDataNodeRef dataRef, ChannelWriter<OutputChunk> chunks, CancellationToken cancellationToken)
+		static async Task<long> FindOutputChunksAsync(OutputFile outputFile, long offset, ChunkedDataNodeRef dataRef, ChannelWriter<OutputChunk> chunks, ILogger logger, CancellationToken cancellationToken)
 		{
 			if (dataRef.Type == ChunkedDataNodeType.Leaf)
 			{
@@ -873,6 +885,8 @@ namespace EpicGames.Horde.Storage.Nodes
 			else
 			{
 				using BlobData data = await dataRef.Handle.ReadBlobDataAsync(cancellationToken);
+				TraceBlobRead("Interior", CombinePaths(outputFile.Directory.Path, outputFile.FileEntry.Name), dataRef.Handle, logger);
+
 				if (data.Type.Guid == LeafChunkedDataNodeConverter.BlobType.Guid)
 				{
 					await chunks.WriteAsync(new OutputChunk(outputFile, offset, dataRef.Length, dataRef.Handle), cancellationToken);
@@ -885,7 +899,7 @@ namespace EpicGames.Horde.Storage.Nodes
 					InteriorChunkedDataNode interiorNode = BlobSerializer.Deserialize<InteriorChunkedDataNode>(data);
 					foreach (ChunkedDataNodeRef childRef in interiorNode.Children)
 					{
-						length += await FindOutputChunksAsync(outputFile, offset + length, childRef, chunks, cancellationToken);
+						length += await FindOutputChunksAsync(outputFile, offset + length, childRef, chunks, logger, cancellationToken);
 					}
 
 					return length;
