Updated how templated script structs are parsed and represented in UHT.

diff --git a/Engine/Source/Programs/Shared/EpicGames.UHT/Types/Properties/UhtStructProperty.cs b/Engine/Source/Programs/Shared/EpicGames.UHT/Types/Properties/UhtStructProperty.cs
--- a/Engine/Source/Programs/Shared/EpicGames.UHT/Types/Properties/UhtStructProperty.cs
+++ b/Engine/Source/Programs/Shared/EpicGames.UHT/Types/Properties/UhtStructProperty.cs
@@ -12,48 +12,6 @@ using EpicGames.UHT.Utils;
 
 namespace EpicGames.UHT.Types
 {
-	/// <summary>
-	/// Extra information needed for a FStructProperty that is generated from a template wrapper struct
-	/// </summary>
-	/// <param name="TemplateSourceName">Source name of the template wrapper struct</param>
-	/// <param name="InnerScriptStruct">USTRUCTs referenced by the args of the template instance</param>
-	record class UhtStructTemplateWrapper(string TemplateSourceName, UhtScriptStruct InnerScriptStruct)
-	{
-		public StringBuilder AppendSourceName(StringBuilder builder)
-		{
-			builder.Append(TemplateSourceName);
-			builder.Append('<');
-			builder.Append(InnerScriptStruct.SourceName);
-			builder.Append('>');
-
-			return builder;
-		}
-
-		public StringBuilder AppendForwardDeclarations(StringBuilder builder)
-		{
-			builder.Append("struct ");
-			builder.Append(InnerScriptStruct.SourceName);
-			builder.Append(';');
-			builder.Append("\r\n");
-
-			builder.Append("template struct ");
-			AppendSourceName(builder);
-			builder.Append(';');
-
-			return builder;
-		}
-
-		public void ParseTemplateType(IUhtTokenReader tokenReader)
-		{
-			tokenReader
-				.Require(TemplateSourceName)
-				.Require('<')
-				.Optional("struct")
-				.Require(InnerScriptStruct.SourceName)
-				.Require('>');
-		}
-	};
-
 	/// <summary>
 	/// FStructProperty
 	/// </summary>
@@ -79,11 +37,6 @@ namespace EpicGames.UHT.Types
 		[JsonConverter(typeof(UhtTypeSourceNameJsonConverter<UhtScriptStruct>))]
 		public UhtScriptStruct ScriptStruct { get; set; }
 
-		/// <summary>
-		/// Extra information for struct properties generated from a template wrapper struct
-		/// </summary>
-		private UhtStructTemplateWrapper? TemplateWrapper { get; set; }
-
 		/// <summary>
 		/// Construct property
 		/// </summary>
@@ -157,14 +110,6 @@ namespace EpicGames.UHT.Types
 			{
 				return null;
 			}
-
-			if (TemplateWrapper != null)
-			{
-				StringBuilder builder = new();
-				TemplateWrapper.AppendForwardDeclarations(builder);
-				return builder.ToString();
-			}
-
 			return $"struct {ScriptStruct.SourceName};";
 		}
 
@@ -180,14 +125,7 @@ namespace EpicGames.UHT.Types
 			switch (textType)
 			{
 				default:
-					if (TemplateWrapper != null)
-					{
-						TemplateWrapper.AppendSourceName(builder);
-					}
-					else
-					{
-						builder.Append(ScriptStruct.SourceName);
-					}
+					builder.Append(ScriptStruct.SourceName);
 					break;
 			}
 			return builder;
@@ -303,6 +241,23 @@ namespace EpicGames.UHT.Types
 			return false;
 		}
 
+		/// <summary>
+		/// Perform default, default value parsing
+		/// </summary>
+		/// <param name="defaultValueReader">Default value reader</param>
+		/// <param name="innerDefaultValue">Sanitized default value</param>
+		/// <returns></returns>
+		public virtual bool DefaultDefaultValue(IUhtTokenReader defaultValueReader, StringBuilder innerDefaultValue)
+		{
+			defaultValueReader
+				.Require(ScriptStruct.SourceName);
+			defaultValueReader
+				.Require('(')
+				.Require(')');
+			innerDefaultValue.Append("()");
+			return true;
+		}
+
 		#region Structure default value sanitizers
 		[UhtStructDefaultValue(Name = "FVector")]
 		[SuppressMessage("CodeQuality", "IDE0051:Remove unused private members", Justification = "Attribute accessed method")]
@@ -515,20 +470,7 @@ namespace EpicGames.UHT.Types
 		[SuppressMessage("CodeQuality", "IDE0051:Remove unused private members", Justification = "Attribute accessed method")]
 		private static bool DefaultStructDefaultValue(UhtStructProperty property, IUhtTokenReader defaultValueReader, StringBuilder innerDefaultValue)
 		{
-			if (property.TemplateWrapper != null)
-			{
-				property.TemplateWrapper.ParseTemplateType(defaultValueReader);
-			}
-			else
-			{
-				defaultValueReader
-					.Require(property.ScriptStruct.SourceName);
-			}
-			defaultValueReader
-				.Require('(')
-				.Require(')');
-			innerDefaultValue.Append("()");
-			return true;
+			return property.DefaultDefaultValue(defaultValueReader, innerDefaultValue);
 		}
 		#endregion
 
@@ -567,6 +509,78 @@ namespace EpicGames.UHT.Types
 			return propertySettings.Outer.FindType(UhtFindOptions.SourceName | UhtFindOptions.ScriptStruct, ref identifier, tokenReader) as UhtScriptStruct;
 		}
 		#endregion
+	}
+
+	/// <summary>
+	/// FStructProperty
+	/// </summary>
+	[UnrealHeaderTool]
+	public class UhtTemplateStructProperty : UhtStructProperty
+	{
+		/// <summary>
+		/// When using the template wrapper pattern that provides a template wrapper to an existing structure that
+		/// can reference types, this is the name of the template.  For example, FInstancedStruct has TInstancedStruct
+		/// as a wrapper template.
+		/// </summary>
+		public string TemplateWrapperName { get; init; }
+
+		/// <summary>
+		/// The structure being wrapped by the template wrapper
+		/// </summary>
+		[JsonConverter(typeof(UhtTypeSourceNameJsonConverter<UhtScriptStruct>))]
+		public UhtScriptStruct TemplateArgumentStruct { get; init; }
+
+		/// <summary>
+		/// Construct property
+		/// </summary>
+		/// <param name="propertySettings">Property settings</param>
+		/// <param name="scriptStruct">USTRUCT being referenced</param>
+		/// <param name="templateWrapperName">The name of wrapping template type</param>
+		/// <param name="templateArgumentStruct">The path name of the type being managed by the template</param>
+		public UhtTemplateStructProperty(UhtPropertySettings propertySettings, UhtScriptStruct scriptStruct, string templateWrapperName, UhtScriptStruct templateArgumentStruct) : base(propertySettings, scriptStruct)
+		{
+			TemplateWrapperName = templateWrapperName;
+			TemplateArgumentStruct = templateArgumentStruct;
+			MetaData.Add("BaseStruct", templateArgumentStruct.PathName);
+		}
+
+		/// <inheritdoc/>
+		public override string? GetForwardDeclarations()
+		{
+			if (TemplateArgumentStruct.IsCoreType)
+			{
+				return null;
+			}
+			return $"struct {TemplateArgumentStruct.SourceName};\r\ntemplate struct {TemplateWrapperName}<{TemplateArgumentStruct.SourceName}>;";
+		}
+
+		/// <inheritdoc/>
+		public override StringBuilder AppendText(StringBuilder builder, UhtPropertyTextType textType, bool isTemplateArgument)
+		{
+			switch (textType)
+			{
+				default:
+					builder.Append(TemplateWrapperName).Append('<').Append(TemplateArgumentStruct.SourceName).Append('>');
+					break;
+			}
+			return builder;
+		}
+
+		/// <inheritdoc/>
+		public override bool DefaultDefaultValue(IUhtTokenReader defaultValueReader, StringBuilder innerDefaultValue)
+		{
+			defaultValueReader
+				.Require(TemplateWrapperName)
+				.Require('<')
+				.Optional("struct")
+				.Require(TemplateArgumentStruct.SourceName)
+				.Require('>');
+			defaultValueReader
+				.Require('(')
+				.Require(')');
+			innerDefaultValue.Append("()");
+			return true;
+		}
 
 		#region Keyword
 		[UhtPropertyType(Keyword = "TInstancedStruct")]
@@ -587,10 +601,7 @@ namespace EpicGames.UHT.Types
 			}
 
 			// With TInstancedStruct, BaseStruct is used as a type limiter.
-			UhtStructProperty instancedStructProperty = new UhtStructProperty(propertySettings, baseScriptStruct.Session.FInstancedStruct);
-			instancedStructProperty.TemplateWrapper = new UhtStructTemplateWrapper("TInstancedStruct", baseScriptStruct);
-			instancedStructProperty.MetaData.Add("BaseStruct", baseScriptStruct.PathName);
-			return instancedStructProperty;
+			return new UhtTemplateStructProperty(propertySettings, baseScriptStruct.Session.FInstancedStruct, "TInstancedStruct", baseScriptStruct);
 		}
 
 		[UhtPropertyType(Keyword = "TStateTreePropertyRef")]
