Horde: Add container execution support for compute tasks

diff --git a/Engine/Source/Programs/Horde/Horde.Agent/AgentSettings.cs b/Engine/Source/Programs/Horde/Horde.Agent/AgentSettings.cs
--- a/Engine/Source/Programs/Horde/Horde.Agent/AgentSettings.cs
+++ b/Engine/Source/Programs/Horde/Horde.Agent/AgentSettings.cs
@@ -235,10 +235,15 @@ namespace Horde.Agent
 		public List<ProcessToTerminate> ProcessesToTerminate { get; } = new List<ProcessToTerminate>();
 		
 		/// <summary>
-		/// Path to Wine executable. Set to null to disable.
+		/// Path to Wine executable. If null, execution under Wine is disabled
 		/// </summary>
 		public string? WineExecutablePath { get; set; }
 
+		/// <summary>
+		/// Path to container engine executable, such as /usr/bin/podman. If null, execution of compute workloads inside a container is disabled
+		/// </summary>
+		public string? ContainerEngineExecutablePath { get; set; }
+
 		/// <summary>
 		/// Whether to write step output to the logging device
 		/// </summary>
diff --git a/Engine/Source/Programs/Horde/Horde.Agent/Commands/Compute/ComputeWorkerCommand.cs b/Engine/Source/Programs/Horde/Horde.Agent/Commands/Compute/ComputeWorkerCommand.cs
--- a/Engine/Source/Programs/Horde/Horde.Agent/Commands/Compute/ComputeWorkerCommand.cs
+++ b/Engine/Source/Programs/Horde/Horde.Agent/Commands/Compute/ComputeWorkerCommand.cs
@@ -44,7 +44,7 @@ namespace Horde.Agent.Commands.Compute
 		{
 			DirectoryReference sandboxDir = DirectoryReference.Combine(AgentApp.DataDir, "Sandbox");
 
-			AgentMessageHandler worker = new AgentMessageHandler(sandboxDir, null, false, null, logger);
+			AgentMessageHandler worker = new AgentMessageHandler(sandboxDir, null, false, null, null, logger);
 			await worker.RunAsync(socket, cancellationToken);
 		}
 	}
diff --git a/Engine/Source/Programs/Horde/Horde.Agent/Leases/Handlers/ComputeHandler.cs b/Engine/Source/Programs/Horde/Horde.Agent/Leases/Handlers/ComputeHandler.cs
--- a/Engine/Source/Programs/Horde/Horde.Agent/Leases/Handlers/ComputeHandler.cs
+++ b/Engine/Source/Programs/Horde/Horde.Agent/Leases/Handlers/ComputeHandler.cs
@@ -110,7 +110,7 @@ namespace Horde.Agent.Leases.Handlers
 								newEnvVars["UE_HORDE_SHARED_DIR"] = sharedDir.FullName;
 								newEnvVars["UE_HORDE_TERMINATION_SIGNAL_FILE"] = _settings.GetTerminationSignalFile().FullName;
 
-								AgentMessageHandler worker = new AgentMessageHandler(sandboxDir, newEnvVars, false, _settings.WineExecutablePath, serverLogger ?? _logger);
+								AgentMessageHandler worker = new AgentMessageHandler(sandboxDir, newEnvVars, false, _settings.WineExecutablePath, _settings.ContainerEngineExecutablePath, serverLogger ?? _logger);
 								await worker.RunAsync(socket, cts.Token);
 								await socket.CloseAsync(cts.Token);
 								return LeaseResult.Success;
diff --git a/Engine/Source/Programs/Horde/Horde/Commands/Compute/ComputeRun.cs b/Engine/Source/Programs/Horde/Horde/Commands/Compute/ComputeRun.cs
--- a/Engine/Source/Programs/Horde/Horde/Commands/Compute/ComputeRun.cs
+++ b/Engine/Source/Programs/Horde/Horde/Commands/Compute/ComputeRun.cs
@@ -27,6 +27,9 @@ namespace Horde.Commands.Compute
 			public string WorkingDir { get; set; } = String.Empty;
 			public Dictionary<string, string?> EnvVars { get; set; } = new Dictionary<string, string?>();
 			public List<string> OutputPaths { get; set; } = new List<string>();
+			public string? ContainerImageUrl { get; set; }
+			public bool ContainerReplaceEntrypoint { get; set; } = false;
+			public bool UseWine { get; set; } = false;
 		}
 
 		[CommandLine("-Cluster")]
@@ -123,7 +126,7 @@ namespace Horde.Commands.Compute
 
 			// Read the task definition
 			byte[] data = await FileReference.ReadAllBytesAsync(TaskFile, cancellationToken);
-			JsonComputeTask jsonComputeTask = JsonSerializer.Deserialize<JsonComputeTask>(data, new JsonSerializerOptions { AllowTrailingCommas = true, PropertyNameCaseInsensitive = true, PropertyNamingPolicy = JsonNamingPolicy.CamelCase })!;
+			JsonComputeTask task = JsonSerializer.Deserialize<JsonComputeTask>(data, new JsonSerializerOptions { AllowTrailingCommas = true, PropertyNameCaseInsensitive = true, PropertyNamingPolicy = JsonNamingPolicy.CamelCase })!;
 
 			// Create a sandbox from the data to be uploaded
 			using BundleStorageClient storage = BundleStorageClient.CreateInMemory(logger);
@@ -135,7 +138,11 @@ namespace Horde.Commands.Compute
 				await channel.WaitForAttachAsync(cancellationToken);
 				await channel.UploadFilesAsync("", sandbox, storage.Backend, cancellationToken);
 
-				await using (AgentManagedProcess process = await channel.ExecuteAsync(jsonComputeTask.Executable, jsonComputeTask.Arguments, jsonComputeTask.WorkingDir, jsonComputeTask.EnvVars, ExecuteProcessFlags.None, cancellationToken))
+				ExecuteProcessFlags execFlags = ExecuteProcessFlags.None;
+				execFlags |= task.UseWine ? ExecuteProcessFlags.UseWine : 0;
+				execFlags |= task.ContainerReplaceEntrypoint ? ExecuteProcessFlags.ReplaceContainerEntrypoint : 0;
+
+				await using (AgentManagedProcess process = await channel.ExecuteAsync(task.Executable, task.Arguments, task.WorkingDir, task.EnvVars, execFlags, task.ContainerImageUrl, cancellationToken))
 				{
 					string? line;
 					while ((line = await process.ReadLineAsync(cancellationToken)) != null)
diff --git a/Engine/Source/Programs/Horde/Horde/Commands/Compute/ComputeWorker.cs b/Engine/Source/Programs/Horde/Horde/Commands/Compute/ComputeWorker.cs
--- a/Engine/Source/Programs/Horde/Horde/Commands/Compute/ComputeWorker.cs
+++ b/Engine/Source/Programs/Horde/Horde/Commands/Compute/ComputeWorker.cs
@@ -44,7 +44,7 @@ namespace Horde.Commands.Compute
 		{
 			DirectoryReference sandboxDir = DirectoryReference.Combine(DirectoryReference.GetSpecialFolder(Environment.SpecialFolder.LocalApplicationData)!, "Horde", "Sandbox");
 
-			AgentMessageHandler worker = new AgentMessageHandler(sandboxDir, null, false, null, logger);
+			AgentMessageHandler worker = new AgentMessageHandler(sandboxDir, null, false, null, null, logger);
 			await worker.RunAsync(socket, cancellationToken);
 		}
 	}
diff --git a/Engine/Source/Programs/Shared/EpicGames.Horde.Tests/ComputeSocketTests.cs b/Engine/Source/Programs/Shared/EpicGames.Horde.Tests/ComputeSocketTests.cs
--- a/Engine/Source/Programs/Shared/EpicGames.Horde.Tests/ComputeSocketTests.cs
+++ b/Engine/Source/Programs/Shared/EpicGames.Horde.Tests/ComputeSocketTests.cs
@@ -240,7 +240,7 @@ namespace EpicGames.Horde.Tests
 		{
 			try
 			{
-				AgentMessageHandler handler = new AgentMessageHandler(tempDir, null, true, null, NullLogger.Instance);
+				AgentMessageHandler handler = new AgentMessageHandler(tempDir, null, true, null, null, NullLogger.Instance);
 				await handler.RunAsync(socket, cancellationToken);
 			}
 			catch (Exception e)
diff --git a/Engine/Source/Programs/Shared/EpicGames.Horde/Compute/AgentMessage.cs b/Engine/Source/Programs/Shared/EpicGames.Horde/Compute/AgentMessage.cs
--- a/Engine/Source/Programs/Shared/EpicGames.Horde/Compute/AgentMessage.cs
+++ b/Engine/Source/Programs/Shared/EpicGames.Horde/Compute/AgentMessage.cs
@@ -10,6 +10,8 @@ using System.Threading.Tasks;
 using EpicGames.Core;
 using EpicGames.Horde.Storage;
 
+#pragma warning disable CA1054 // URI-like parameters should not be strings
+
 namespace EpicGames.Horde.Compute
 {
 	/// <summary>
@@ -68,6 +70,11 @@ namespace EpicGames.Horde.Compute
 		/// Execute a process in a sandbox (Initiator -> Remote)
 		/// </summary>
 		ExecuteV2 = 0x22,
+		
+		/// <summary>
+		/// Execute a process in a sandbox (Initiator -> Remote)
+		/// </summary>
+		ExecuteV3 = 0x23,
 
 		/// <summary>
 		/// Returns output from the child process to the caller (Remote -> Initiator)
@@ -126,6 +133,12 @@ namespace EpicGames.Horde.Compute
 		/// Agent still reserves the right to refuse it (e.g no Wine executable configured, mismatching OS etc)
 		/// </summary>
 		UseWine = 1,
+		
+		/// <summary>
+		/// Use compute process executable as entrypoint for container
+		/// If not set, path to the executable is passed as the first parameter to the container invocation
+		/// </summary>
+		ReplaceContainerEntrypoint = 2,
 	}
 
 	/// <summary>
@@ -229,7 +242,8 @@ namespace EpicGames.Horde.Compute
 	/// <param name="WorkingDir">Working directory to execute in</param>
 	/// <param name="EnvVars">Environment variables for the child process. Null values unset variables.</param>
 	/// <param name="Flags">Additional execution flags</param>
-	public record struct ExecuteProcessMessage(string Executable, IReadOnlyList<string> Arguments, string? WorkingDir, IReadOnlyDictionary<string, string?> EnvVars, ExecuteProcessFlags Flags);
+	/// <param name="ContainerImageUrl">URL to container image. If specified, process will be executed inside this container</param>
+	public record struct ExecuteProcessMessage(string Executable, IReadOnlyList<string> Arguments, string? WorkingDir, IReadOnlyDictionary<string, string?> EnvVars, ExecuteProcessFlags Flags, string? ContainerImageUrl);
 
 	/// <summary>
 	/// Response from executing a child process
@@ -471,6 +485,32 @@ namespace EpicGames.Horde.Compute
 			}
 			return new AgentManagedProcess(channel);
 		}
+		
+		/// <summary>
+		/// Executes a remote process (using ExecuteV3)
+		/// </summary>
+		/// <param name="channel">Current channel</param>
+		/// <param name="executable">Executable to run, relative to the sandbox root</param>
+		/// <param name="arguments">Arguments for the child process</param>
+		/// <param name="workingDir">Working directory for the process</param>
+		/// <param name="envVars">Environment variables for the child process</param>
+		/// <param name="flags">Additional execution flags</param>
+		/// <param name="containerImageUrl">Optional container image URL. If set, execution will happen inside this container</param>
+		/// <param name="cancellationToken">Cancellation token for the operation</param>
+		public static async Task<AgentManagedProcess> ExecuteAsync(this AgentMessageChannel channel, string executable, IReadOnlyList<string> arguments, string? workingDir, IReadOnlyDictionary<string, string?>? envVars, ExecuteProcessFlags flags, string? containerImageUrl, CancellationToken cancellationToken = default)
+		{
+			using (IAgentMessageBuilder request = await channel.CreateMessageAsync(AgentMessageType.ExecuteV3, cancellationToken))
+			{
+				request.WriteString(executable);
+				request.WriteList(arguments, MemoryWriterExtensions.WriteString);
+				request.WriteOptionalString(workingDir);
+				request.WriteDictionary(envVars ?? new Dictionary<string, string?>(), MemoryWriterExtensions.WriteString, MemoryWriterExtensions.WriteOptionalString);
+				request.WriteInt32((int)flags);
+				request.WriteString(containerImageUrl ?? "");
+				request.Send();
+			}
+			return new AgentManagedProcess(channel);
+		}
 
 		/// <summary>
 		/// Parses a message as a <see cref="ExecuteProcessMessage"/>
@@ -481,7 +521,7 @@ namespace EpicGames.Horde.Compute
 			List<string> arguments = message.ReadList(MemoryReaderExtensions.ReadString);
 			string? workingDir = message.ReadOptionalString();
 			Dictionary<string, string?> envVars = message.ReadDictionary(MemoryReaderExtensions.ReadString, MemoryReaderExtensions.ReadOptionalString);
-			return new ExecuteProcessMessage(executable, arguments, workingDir, envVars, ExecuteProcessFlags.None);
+			return new ExecuteProcessMessage(executable, arguments, workingDir, envVars, ExecuteProcessFlags.None, null);
 		}
 		
 		/// <summary>
@@ -494,7 +534,21 @@ namespace EpicGames.Horde.Compute
 			string? workingDir = message.ReadOptionalString();
 			Dictionary<string, string?> envVars = message.ReadDictionary(MemoryReaderExtensions.ReadString, MemoryReaderExtensions.ReadOptionalString);
 			ExecuteProcessFlags flags = (ExecuteProcessFlags)message.ReadInt32();
-			return new ExecuteProcessMessage(executable, arguments, workingDir, envVars, flags);
+			return new ExecuteProcessMessage(executable, arguments, workingDir, envVars, flags, null);
+		}
+		
+		/// <summary>
+		/// Parses a message as a <see cref="ExecuteProcessMessage"/>
+		/// </summary>
+		public static ExecuteProcessMessage ParseExecuteProcessV3Message(this AgentMessage message)
+		{
+			string executable = message.ReadString();
+			List<string> arguments = message.ReadList(MemoryReaderExtensions.ReadString);
+			string? workingDir = message.ReadOptionalString();
+			Dictionary<string, string?> envVars = message.ReadDictionary(MemoryReaderExtensions.ReadString, MemoryReaderExtensions.ReadOptionalString);
+			ExecuteProcessFlags flags = (ExecuteProcessFlags)message.ReadInt32();
+			string containerImageUrl = message.ReadString();
+			return new ExecuteProcessMessage(executable, arguments, workingDir, envVars, flags, String.IsNullOrEmpty(containerImageUrl) ? null : containerImageUrl);
 		}
 
 		/// <summary>
diff --git a/Engine/Source/Programs/Shared/EpicGames.Horde/Compute/AgentMessageHandler.cs b/Engine/Source/Programs/Shared/EpicGames.Horde/Compute/AgentMessageHandler.cs
--- a/Engine/Source/Programs/Shared/EpicGames.Horde/Compute/AgentMessageHandler.cs
+++ b/Engine/Source/Programs/Shared/EpicGames.Horde/Compute/AgentMessageHandler.cs
@@ -28,6 +28,7 @@ namespace EpicGames.Horde.Compute
 		readonly Dictionary<string, string?> _envVars;
 		readonly bool _executeInProcess;
 		readonly string? _wineExecutablePath;
+		readonly string? _containerEngineExecutable;
 		readonly ILogger _logger;
 
 		/// <summary>
@@ -37,13 +38,15 @@ namespace EpicGames.Horde.Compute
 		/// <param name="envVars">Environment variables to set for any child processes</param>
 		/// <param name="executeInProcess">Whether to execute any external assemblies in the current process</param>
 		/// <param name="wineExecutablePath">Path to Wine executable. If null, execution under Wine is disabled</param>
+		/// <param name="containerEngineExecutable">Path to container engine executable, e.g /usr/bin/podman. If null, execution inside a container is disabled</param>
 		/// <param name="logger">Logger for diagnostics</param>
-		public AgentMessageHandler(DirectoryReference sandboxDir, Dictionary<string, string?>? envVars, bool executeInProcess, string? wineExecutablePath, ILogger logger)
+		public AgentMessageHandler(DirectoryReference sandboxDir, Dictionary<string, string?>? envVars, bool executeInProcess, string? wineExecutablePath, string? containerEngineExecutable, ILogger logger)
 		{
 			_sandboxDir = sandboxDir;
 			_envVars = envVars ?? new Dictionary<string, string?>();
 			_executeInProcess = executeInProcess;
 			_wineExecutablePath = wineExecutablePath;
+			_containerEngineExecutable = containerEngineExecutable;
 			_logger = logger;
 		}
 
@@ -118,14 +121,20 @@ namespace EpicGames.Horde.Compute
 							break;
 						case AgentMessageType.ExecuteV1:
 							{
-								ExecuteProcessMessage executeProcess = message.ParseExecuteProcessV1Message();
-								await ExecuteProcessAsync(socket, channel, executeProcess.Executable, executeProcess.Arguments, executeProcess.WorkingDir, executeProcess.EnvVars, executeProcess.Flags, cancellationToken);
+								ExecuteProcessMessage ep = message.ParseExecuteProcessV1Message();
+								await ExecuteProcessAsync(socket, channel, ep.Executable, ep.Arguments, ep.WorkingDir, ep.ContainerImageUrl, ep.EnvVars, ep.Flags, cancellationToken);
 							}
 							break;
 						case AgentMessageType.ExecuteV2:
 							{
-								ExecuteProcessMessage executeProcess = message.ParseExecuteProcessV2Message();
-								await ExecuteProcessAsync(socket, channel, executeProcess.Executable, executeProcess.Arguments, executeProcess.WorkingDir, executeProcess.EnvVars, executeProcess.Flags, cancellationToken);
+								ExecuteProcessMessage ep = message.ParseExecuteProcessV2Message();
+								await ExecuteProcessAsync(socket, channel, ep.Executable, ep.Arguments, ep.WorkingDir, ep.ContainerImageUrl, ep.EnvVars, ep.Flags, cancellationToken);
+							}
+							break;
+						case AgentMessageType.ExecuteV3:
+							{
+								ExecuteProcessMessage ep = message.ParseExecuteProcessV3Message();
+								await ExecuteProcessAsync(socket, channel, ep.Executable, ep.Arguments, ep.WorkingDir, ep.ContainerImageUrl, ep.EnvVars, ep.Flags, cancellationToken);
 							}
 							break;
 						case AgentMessageType.XorRequest:
@@ -245,7 +254,16 @@ namespace EpicGames.Horde.Compute
 			}
 		}
 
-		async Task ExecuteProcessAsync(ComputeSocket socket, AgentMessageChannel channel, string executable, IReadOnlyList<string> arguments, string? workingDir, IReadOnlyDictionary<string, string?>? envVars, ExecuteProcessFlags flags, CancellationToken cancellationToken)
+		async Task ExecuteProcessAsync(
+			ComputeSocket socket,
+			AgentMessageChannel channel,
+			string executable,
+			IReadOnlyList<string> arguments,
+			string? workingDir,
+			string? containerImageUrl,
+			IReadOnlyDictionary<string, string?>? envVars,
+			ExecuteProcessFlags flags,
+			CancellationToken cancellationToken)
 		{
 			try
 			{
@@ -253,6 +271,10 @@ namespace EpicGames.Horde.Compute
 				{
 					await ExecuteProcessWindowsAsync(socket, channel, executable, arguments, workingDir, envVars, flags, cancellationToken);
 				}
+				else if (containerImageUrl != null)
+				{
+					await ExecuteProcessInContainerAsync(channel, executable, arguments, workingDir, containerImageUrl, envVars, flags, cancellationToken);
+				}
 				else
 				{
 					await ExecuteProcessInternalAsync(channel, executable, arguments, workingDir, envVars, flags, cancellationToken);
@@ -407,49 +429,110 @@ namespace EpicGames.Horde.Compute
 			}
 		}
 
-		async Task ExecuteProcessInternalAsync(AgentMessageChannel channel, string executable, IReadOnlyList<string> arguments, string? workingDir, IReadOnlyDictionary<string, string?>? envVars, ExecuteProcessFlags flags, CancellationToken cancellationToken)
+		async Task ExecuteProcessAssemblyAsync(AgentMessageChannel channel, IReadOnlyList<string> arguments, string? workingDir, IReadOnlyDictionary<string, string?>? envVars, CancellationToken cancellationToken)
 		{
-			string resolvedExecutable = FileReference.Combine(_sandboxDir, executable).FullName;
-			string resolvedWorkingDir = DirectoryReference.Combine(_sandboxDir, workingDir ?? String.Empty).FullName;
-			if (_executeInProcess && Path.GetFileNameWithoutExtension(resolvedExecutable).Equals("dotnet", StringComparison.OrdinalIgnoreCase))
+			List<(string, string?)> prevEnvVars = new List<(string, string?)>();
+			if (envVars != null)
 			{
-				List<(string, string?)> prevEnvVars = new List<(string, string?)>();
-				if (envVars != null)
+				foreach ((string key, string? value) in envVars)
 				{
-					foreach ((string key, string? value) in envVars)
-					{
-						prevEnvVars.Add((key, Environment.GetEnvironmentVariable(key)));
-						Environment.SetEnvironmentVariable(key, value);
-					}
+					prevEnvVars.Add((key, Environment.GetEnvironmentVariable(key)));
+					Environment.SetEnvironmentVariable(key, value);
 				}
+			}
 
-				string prevWorkingDir = Directory.GetCurrentDirectory();
-				Directory.SetCurrentDirectory(resolvedWorkingDir);
+			string prevWorkingDir = Directory.GetCurrentDirectory();
+			Directory.SetCurrentDirectory(GetWorkingDirAbsPath(workingDir));
 
-				try
-				{
-					string assemblyPath = FileReference.Combine(_sandboxDir, arguments[0]).FullName;
-					string[] mainArgs = arguments.Skip(1).ToArray();
+			try
+			{
+				string assemblyPath = FileReference.Combine(_sandboxDir, arguments[0]).FullName;
+				string[] mainArgs = arguments.Skip(1).ToArray();
 
-					_logger.LogWarning("Note: Loading and running {Assembly} in process", assemblyPath);
+				_logger.LogWarning("Note: Loading and running {Assembly} in process", assemblyPath);
 
-					TaskCompletionSource<int> resultTcs = new TaskCompletionSource<int>();
+				TaskCompletionSource<int> resultTcs = new TaskCompletionSource<int>();
 
-					Thread thread = new Thread(() => resultTcs.SetResult(AppDomain.CurrentDomain.ExecuteAssembly(assemblyPath, mainArgs)));
-					thread.Start();
+				Thread thread = new Thread(() => resultTcs.SetResult(AppDomain.CurrentDomain.ExecuteAssembly(assemblyPath, mainArgs)));
+				thread.Start();
 
-					int result = await resultTcs.Task;
-					await channel.SendExecuteResultAsync(result, cancellationToken);
-				}
-				finally
+				int result = await resultTcs.Task;
+				await channel.SendExecuteResultAsync(result, cancellationToken);
+			}
+			finally
+			{
+				Directory.SetCurrentDirectory(prevWorkingDir);
+				foreach((string key, string? value) in prevEnvVars)
 				{
-					Directory.SetCurrentDirectory(prevWorkingDir);
-					foreach((string key, string? value) in prevEnvVars)
-					{
-						Environment.SetEnvironmentVariable(key, value);
-					}
+					Environment.SetEnvironmentVariable(key, value);
 				}
 			}
+		}
+
+		async Task ExecuteProcessInContainerAsync(AgentMessageChannel channel, string executable, IReadOnlyList<string> arguments, string? workingDir, string containerImageUrl, IReadOnlyDictionary<string, string?>? envVars, ExecuteProcessFlags flags, CancellationToken cancellationToken)
+		{
+			if (!RuntimeInformation.IsOSPlatform(OSPlatform.Linux))
+			{
+				throw new Exception("Only Linux is supported for executing a process inside a container");
+			}
+			
+			if (_containerEngineExecutable == null)
+			{
+				throw new Exception("Container execution requested but agent has no container engine configured");
+			}
+
+			string resolvedExecutable = FileReference.Combine(_sandboxDir, executable).FullName;
+			uint linuxUid = getuid();
+			uint linuxGid = getgid();
+
+			// Resolve env vars here even if they are resolved later in ExecuteProcessInternalAsync
+			// The environment file must be written at this step
+			Dictionary<string, string> resolvedEnvVars = ResolveEnvVars(envVars);
+			string envFilePath = Path.GetTempFileName();
+			StringBuilder sb = new();
+			foreach ((string key, string value) in resolvedEnvVars)
+			{
+				sb.AppendLine($"{key}={value}");
+			}
+			await File.WriteAllTextAsync(envFilePath, sb.ToString(), cancellationToken);
+			
+			List<string> resolvedArguments = new()
+			{
+				"run",
+				"--tty", // Allocate a pseudo-TTY
+				"--rm", // Ensure container is removed after run
+				$"--user={linuxUid}:{linuxGid}", // Run container as current user (important for mounted dirs)
+				$"--volume={_sandboxDir}:{_sandboxDir}:rw",
+				"--env-file=" + envFilePath,
+			};
+			
+			if (flags.HasFlag(ExecuteProcessFlags.ReplaceContainerEntrypoint))
+			{
+				resolvedArguments.Add("--entrypoint=" + resolvedExecutable);
+				resolvedArguments.Add(containerImageUrl);
+			}
+			else
+			{
+				resolvedArguments.Add(containerImageUrl);
+				resolvedArguments.Add(resolvedExecutable); // Add executable as first argument and assume the entrypoint inside the container image will handle this
+			}
+			
+			resolvedArguments.AddRange(arguments);
+			_logger.LogInformation("Executing {File} {Arguments} in container", _containerEngineExecutable, arguments);
+
+			// Skip forwarding of env vars as they are explicitly set above as arguments to container run
+			await ExecuteProcessInternalAsync(channel, _containerEngineExecutable, resolvedArguments, workingDir, new Dictionary<string, string?>(), flags, cancellationToken);
+		}
+
+		async Task ExecuteProcessInternalAsync(AgentMessageChannel channel, string executable, IReadOnlyList<string> arguments, string? workingDir, IReadOnlyDictionary<string, string?>? envVars, ExecuteProcessFlags flags, CancellationToken cancellationToken)
+		{
+			string resolvedExecutable = GetExecutableAbsPath(executable);
+			string resolvedWorkingDir = GetWorkingDirAbsPath(workingDir);
+
+			if (_executeInProcess && Path.GetFileNameWithoutExtension(resolvedExecutable).Equals("dotnet", StringComparison.OrdinalIgnoreCase))
+			{
+				await ExecuteProcessAssemblyAsync(channel, arguments, workingDir, envVars, cancellationToken);
+			}
 			else
 			{
 				string resolvedCommandLine = CommandLineArguments.Join(arguments);
@@ -461,31 +544,7 @@ namespace EpicGames.Horde.Compute
 					resolvedExecutable = _wineExecutablePath;
 				}
 
-				Dictionary<string, string> resolvedEnvVars = ManagedProcess.GetCurrentEnvVars();
-
-				foreach ((string key, string? value) in _envVars)
-				{
-					if (value != null)
-					{
-						resolvedEnvVars[key] = value;
-					}
-				}
-				
-				if (envVars != null)
-				{
-					foreach ((string key, string? value) in envVars)
-					{
-						if (value == null)
-						{
-							resolvedEnvVars.Remove(key);
-						}
-						else
-						{
-							resolvedEnvVars[key] = value;
-						}
-					}
-				}
-
+				Dictionary<string, string> resolvedEnvVars = ResolveEnvVars(envVars);
 				if (!File.Exists(resolvedExecutable))
 				{
 					_logger.LogWarning("Executable {Path} does not exist", resolvedExecutable);	
@@ -518,5 +577,64 @@ namespace EpicGames.Horde.Compute
 				}
 			}
 		}
+		
+		private string GetExecutableAbsPath(string relPath)
+		{
+			return FileReference.Combine(_sandboxDir, relPath).FullName;
+		}
+		
+		private string GetWorkingDirAbsPath(string? relPath)
+		{
+			return DirectoryReference.Combine(_sandboxDir, relPath ?? String.Empty).FullName;
+		}
+
+		/// <summary>
+		/// Flattens and merges available env vars to be used for compute process execution
+		/// </summary>
+		/// <param name="envVars">Optional extra env vars</param>
+		/// <returns>Merged environment variables</returns>
+		private Dictionary<string, string> ResolveEnvVars(IReadOnlyDictionary<string, string?>? envVars)
+		{
+			Dictionary<string, string> resolvedEnvVars = ManagedProcess.GetCurrentEnvVars();
+
+			foreach ((string key, string? value) in _envVars)
+			{
+				if (value != null)
+				{
+					resolvedEnvVars[key] = value;
+				}
+			}
+
+			if (envVars != null)
+			{
+				foreach ((string key, string? value) in envVars)
+				{
+					if (value == null)
+					{
+						resolvedEnvVars.Remove(key);
+					}
+					else
+					{
+						resolvedEnvVars[key] = value;
+					}
+				}
+			}
+
+			return resolvedEnvVars;
+		}
+
+		/// <summary>
+		/// Get user identity (Linux only)
+		/// </summary>
+		/// <returns>Real user ID of the calling process</returns>
+		[DllImport("libc", SetLastError = true)]
+		internal static extern uint getuid();
+
+		/// <summary>
+		/// Get group identity (Linux only)
+		/// </summary>
+		/// <returns>Real group ID of the calling process</returns>
+		[DllImport("libc", SetLastError = true)]
+		internal static extern uint getgid();
 	}
 }
diff --git a/Engine/Source/Programs/Shared/EpicGames.Horde/Compute/Clients/LocalComputeClient.cs b/Engine/Source/Programs/Shared/EpicGames.Horde/Compute/Clients/LocalComputeClient.cs
--- a/Engine/Source/Programs/Shared/EpicGames.Horde/Compute/Clients/LocalComputeClient.cs
+++ b/Engine/Source/Programs/Shared/EpicGames.Horde/Compute/Clients/LocalComputeClient.cs
@@ -78,7 +78,7 @@ namespace EpicGames.Horde.Compute.Clients
 			using Socket tcpSocket = await listener.AcceptAsync(cancellationToken);
 			await using TcpTransport tcpTransport = new (tcpSocket);
 			await using RemoteComputeSocket socket = new (tcpTransport, ComputeProtocol.Latest, logger);
-			AgentMessageHandler worker = new (sandboxDir, null, executeInProcess, null, logger);
+			AgentMessageHandler worker = new (sandboxDir, null, executeInProcess, null, null, logger);
 			await worker.RunAsync(socket, cancellationToken);
 			await socket.CloseAsync(cancellationToken);
 		}
