Horde: Fix PrefixedObjectStore not disposing of the inner store instance.

diff --git a/Engine/Source/Programs/Shared/EpicGames.Horde/Storage/ObjectStores/PrefixedObjectStore.cs b/Engine/Source/Programs/Shared/EpicGames.Horde/Storage/ObjectStores/PrefixedObjectStore.cs
--- a/Engine/Source/Programs/Shared/EpicGames.Horde/Storage/ObjectStores/PrefixedObjectStore.cs
+++ b/Engine/Source/Programs/Shared/EpicGames.Horde/Storage/ObjectStores/PrefixedObjectStore.cs
@@ -39,6 +39,7 @@ namespace EpicGames.Horde.Storage.ObjectStores
 		/// <inheritdoc/>
 		public void Dispose()
 		{
+			_inner.Dispose();
 		}
 
 		/// <inheritdoc/>
