Horde: Move functionality for managing a cache of memory mapped files into EpicGames.Core.

diff --git a/Engine/Source/Programs/Shared/EpicGames.Core/MemoryMappedFileCache.cs b/Engine/Source/Programs/Shared/EpicGames.Core/MemoryMappedFileCache.cs
new file mode 100644
--- /dev/null
+++ b/Engine/Source/Programs/Shared/EpicGames.Core/MemoryMappedFileCache.cs
@@ -0,0 +1,280 @@
+// Copyright Epic Games, Inc. All Rights Reserved.
+
+using System;
+using System.Collections.Generic;
+using System.IO;
+using System.IO.MemoryMappedFiles;
+using System.Threading;
+
+namespace EpicGames.Core
+{
+	/// <summary>
+	/// Manages a cache of memory mapped file handles
+	/// </summary>
+	public class MemoryMappedFileCache
+	{
+		// Item which has been opened from the cache using a memory mapped file
+		class MappedFile : IDisposable
+		{
+			public FileReference Key { get; }
+			public LinkedListNode<MappedFile> ListNode { get; }
+
+			readonly FileInfo _fileInfo;
+
+			MemoryMappedFile? _memoryMappedFile;
+			MemoryMappedViewAccessor? _memoryMappedViewAccessor;
+			MemoryMappedView? _memoryMappedView;
+
+			int _refCount = 1;
+			ReadOnlyMemory<byte> _data;
+			bool _deleteOnDispose;
+
+			public int RefCount => _refCount;
+
+			public ulong MappedSize => _memoryMappedViewAccessor?.SafeMemoryMappedViewHandle.ByteLength ?? 0UL;
+
+			public MappedFile(FileReference locator, FileInfo fileInfo)
+			{
+				Key = locator;
+				ListNode = new LinkedListNode<MappedFile>(this);
+
+				_fileInfo = fileInfo;
+
+				try
+				{
+					_memoryMappedFile = MemoryMappedFile.CreateFromFile(fileInfo.FullName, FileMode.Open, null, 0, MemoryMappedFileAccess.Read);
+					_memoryMappedViewAccessor = _memoryMappedFile.CreateViewAccessor(0, 0, MemoryMappedFileAccess.Read);
+					_memoryMappedView = new MemoryMappedView(_memoryMappedViewAccessor);
+
+					_data = _memoryMappedView.GetMemory(0, (int)fileInfo.Length);
+				}
+				catch
+				{
+					Dispose();
+					throw;
+				}
+			}
+
+			public void Dispose()
+			{
+				_data = ReadOnlyMemory<byte>.Empty;
+
+				if (_memoryMappedView != null)
+				{
+					_memoryMappedView.Dispose();
+					_memoryMappedView = null;
+				}
+				if (_memoryMappedViewAccessor != null)
+				{
+					_memoryMappedViewAccessor.Dispose();
+					_memoryMappedViewAccessor = null;
+				}
+				if (_memoryMappedFile != null)
+				{
+					_memoryMappedFile.Dispose();
+					_memoryMappedFile = null;
+				}
+
+				if (_deleteOnDispose)
+				{
+					try
+					{
+						_fileInfo.Delete();
+					}
+					catch { }
+				}
+			}
+
+			public ReadOnlyMemory<byte> GetData(int offset, int? length)
+			{
+				if (length == null)
+				{
+					return _data.Slice(offset);
+				}
+				else
+				{
+					return _data.Slice(offset, Math.Min(length.Value, _data.Length - offset));
+				}
+			}
+
+			public void AddRef()
+			{
+				Interlocked.Increment(ref _refCount);
+			}
+
+			public void Release()
+			{
+				if (Interlocked.Decrement(ref _refCount) == 0)
+				{
+					Dispose();
+				}
+			}
+
+			public void DeleteOnDispose()
+			{
+				_deleteOnDispose = true;
+			}
+
+			public override string ToString() => Key.ToString();
+		}
+
+		// Handle to a file in memory
+		class MappedFileHandle : IReadOnlyMemoryOwner<byte>
+		{
+			MappedFile? _mappedFile;
+			ReadOnlyMemory<byte> _data;
+
+			public ReadOnlyMemory<byte> Memory => _data;
+
+			public MappedFileHandle(MappedFile? mappedFile, ReadOnlyMemory<byte> data)
+			{
+				_mappedFile = mappedFile;
+				_mappedFile?.AddRef();
+				_data = data;
+			}
+
+			public MappedFileHandle Clone()
+			{
+				_mappedFile?.AddRef();
+				return new MappedFileHandle(_mappedFile, _data);
+			}
+
+			public void Dispose()
+			{
+				if (_mappedFile != null)
+				{
+					_mappedFile.Release();
+					_mappedFile = null!;
+				}
+
+				_data = ReadOnlyMemory<byte>.Empty;
+			}
+		}
+
+		readonly object _lockObject = new object();
+		readonly Dictionary<FileReference, MappedFile> _pathToMappedFile = new Dictionary<FileReference, MappedFile>();
+		readonly LinkedList<MappedFile> _mappedFiles = new LinkedList<MappedFile>();
+
+		long _mappedSize;
+
+		readonly long _maxMappedSize;
+		readonly int _maxMappedCount;
+
+		/// <summary>
+		/// Constructor
+		/// </summary>
+		/// <param name="maxMappedSize">Max amount of data to keep mapped into memory</param>
+		/// <param name="maxMappedCount">Max number of memory mapped file handles</param>
+		public MemoryMappedFileCache(long maxMappedSize = 1024L * 1024 * 1024, int maxMappedCount = 128)
+		{
+			_maxMappedSize = maxMappedSize;
+			_maxMappedCount = maxMappedCount;
+		}
+
+		/// <inheritdoc/>
+		public void Dispose()
+		{
+			lock (_lockObject)
+			{
+				UnmapFiles(0, 0);
+			}
+		}
+
+		/// <summary>
+		/// Maps data for a file into memory
+		/// </summary>
+		/// <param name="file">File to open</param>
+		/// <param name="offset">Start offset of the file to map</param>
+		/// <param name="length">Length of the mapped data</param>
+		public IReadOnlyMemoryOwner<byte> Read(FileReference file, int offset, int? length)
+		{
+			lock (_lockObject)
+			{
+				MappedFile? mappedFile = FindOrAddMappedFile(file);
+				if (mappedFile == null)
+				{
+					return ReadOnlyMemoryOwner.Create(ReadOnlyMemory<byte>.Empty, null);
+				}
+				return new MappedFileHandle(mappedFile, mappedFile.GetData(offset, length));
+			}
+		}
+
+		/// <summary>
+		/// Delete a file from the store
+		/// </summary>
+		/// <param name="file"></param>
+		public void Delete(FileReference file)
+		{
+			lock (_lockObject)
+			{
+				MappedFile? mappedFile;
+				if (_pathToMappedFile.TryGetValue(file, out mappedFile))
+				{
+					mappedFile.DeleteOnDispose();
+
+					_pathToMappedFile.Remove(file);
+					_mappedFiles.Remove(mappedFile.ListNode);
+
+					mappedFile.Release();
+				}
+				else
+				{
+					FileReference.Delete(file);
+				}
+			}
+		}
+
+		/// <summary>
+		/// Finds an existing mapped file or adds a new one for the given path
+		/// </summary>
+		MappedFile? FindOrAddMappedFile(FileReference key)
+		{
+			MappedFile? mappedFile;
+			if (!_pathToMappedFile.TryGetValue(key, out mappedFile))
+			{
+				FileInfo fileInfo = key.ToFileInfo();
+				if (!fileInfo.Exists)
+				{
+					throw new FileNotFoundException($"Unable to find file '{key}'", key.FullName);
+				}
+				if (fileInfo.Length == 0)
+				{
+					return null;
+				}
+
+				long maxSize = _maxMappedSize - fileInfo.Length;
+				if (_mappedSize > maxSize || _mappedFiles.Count + 1 > _maxMappedCount)
+				{
+					UnmapFiles(maxSize, _maxMappedCount - 1);
+				}
+
+				mappedFile = new MappedFile(key, fileInfo);
+				_pathToMappedFile.Add(key, mappedFile);
+				_mappedFiles.AddFirst(mappedFile.ListNode);
+
+				_mappedSize += (long)mappedFile.MappedSize;
+			}
+			return mappedFile;
+		}
+
+		/// <summary>
+		/// Discard mapped files until only a certain size is mapped in memory
+		/// </summary>
+		/// <param name="maxMappedSize">Maximum mapped size</param>
+		/// <param name="maxMappedCount">Maximum number of mapped files</param>
+		void UnmapFiles(long maxMappedSize, int maxMappedCount)
+		{
+			for (LinkedListNode<MappedFile>? listNode = _mappedFiles.Last; listNode != null && (_mappedSize > maxMappedSize || _mappedFiles.Count > maxMappedCount);)
+			{
+				LinkedListNode<MappedFile>? nextListNode = listNode.Previous;
+				if (listNode.Value.RefCount == 1)
+				{
+					_pathToMappedFile.Remove(listNode.Value.Key);
+					_mappedFiles.Remove(listNode);
+					listNode.Value.Release();
+				}
+				listNode = nextListNode;
+			}
+		}
+	}
+}
diff --git a/Engine/Source/Programs/Shared/EpicGames.Horde/Storage/ObjectStores/FileObjectStore.cs b/Engine/Source/Programs/Shared/EpicGames.Horde/Storage/ObjectStores/FileObjectStore.cs
--- a/Engine/Source/Programs/Shared/EpicGames.Horde/Storage/ObjectStores/FileObjectStore.cs
+++ b/Engine/Source/Programs/Shared/EpicGames.Horde/Storage/ObjectStores/FileObjectStore.cs
@@ -3,7 +3,6 @@
 using System;
 using System.Collections.Generic;
 using System.IO;
-using System.IO.MemoryMappedFiles;
 using System.Runtime.CompilerServices;
 using System.Threading;
 using System.Threading.Tasks;
@@ -16,161 +15,12 @@ namespace EpicGames.Horde.Storage.ObjectStores
 	/// </summary>
 	public sealed class FileObjectStore : IObjectStore, IDisposable
 	{
-		// Item which has been opened from the cache using a memory mapped file
-		class MappedFile : IDisposable
-		{
-			public ObjectKey Key { get; }
-			public LinkedListNode<MappedFile> ListNode { get; }
-
-			readonly FileInfo _fileInfo;
-
-			MemoryMappedFile? _memoryMappedFile;
-			MemoryMappedViewAccessor? _memoryMappedViewAccessor;
-			MemoryMappedView? _memoryMappedView;
-
-			int _refCount = 1;
-			ReadOnlyMemory<byte> _data;
-			bool _deleteOnDispose;
-
-			public int RefCount => _refCount;
-
-			public ulong MappedSize => _memoryMappedViewAccessor?.SafeMemoryMappedViewHandle.ByteLength ?? 0UL;
-
-			public MappedFile(ObjectKey locator, FileInfo fileInfo)
-			{
-				Key = locator;
-				ListNode = new LinkedListNode<MappedFile>(this);
-
-				_fileInfo = fileInfo;
-
-				try
-				{
-					_memoryMappedFile = MemoryMappedFile.CreateFromFile(fileInfo.FullName, FileMode.Open, null, 0, MemoryMappedFileAccess.Read);
-					_memoryMappedViewAccessor = _memoryMappedFile.CreateViewAccessor(0, 0, MemoryMappedFileAccess.Read);
-					_memoryMappedView = new MemoryMappedView(_memoryMappedViewAccessor);
-
-					_data = _memoryMappedView.GetMemory(0, (int)fileInfo.Length);
-				}
-				catch
-				{
-					Dispose();
-					throw;
-				}
-			}
-
-			public void Dispose()
-			{
-				_data = ReadOnlyMemory<byte>.Empty;
-
-				if (_memoryMappedView != null)
-				{
-					_memoryMappedView.Dispose();
-					_memoryMappedView = null;
-				}
-				if (_memoryMappedViewAccessor != null)
-				{
-					_memoryMappedViewAccessor.Dispose();
-					_memoryMappedViewAccessor = null;
-				}
-				if (_memoryMappedFile != null)
-				{
-					_memoryMappedFile.Dispose();
-					_memoryMappedFile = null;
-				}
-
-				if (_deleteOnDispose)
-				{
-					try
-					{
-						_fileInfo.Delete();
-					}
-					catch { }
-				}
-			}
-
-			public ReadOnlyMemory<byte> GetData(int offset, int? length)
-			{
-				if (length == null)
-				{
-					return _data.Slice(offset);
-				}
-				else
-				{
-					return _data.Slice(offset, Math.Min(length.Value, _data.Length - offset));
-				}
-			}
-
-			public void AddRef()
-			{
-				Interlocked.Increment(ref _refCount);
-			}
-
-			public void Release()
-			{
-				if (Interlocked.Decrement(ref _refCount) == 0)
-				{
-					Dispose();
-				}
-			}
-
-			public void DeleteOnDispose()
-			{
-				_deleteOnDispose = true;
-			}
-
-			public override string ToString() => Key.ToString();
-		}
-
-		// Handle to a file in memory
-		class MappedFileHandle : IReadOnlyMemoryOwner<byte>
-		{
-			MappedFile? _mappedFile;
-			ReadOnlyMemory<byte> _data;
-
-			public ReadOnlyMemory<byte> Memory => _data;
-
-			public MappedFileHandle(MappedFile? mappedFile, ReadOnlyMemory<byte> data)
-			{
-				_mappedFile = mappedFile;
-				_mappedFile?.AddRef();
-				_data = data;
-			}
-
-			public MappedFileHandle Clone()
-			{
-				_mappedFile?.AddRef();
-				return new MappedFileHandle(_mappedFile, _data);
-			}
-
-			public void Dispose()
-			{
-				if (_mappedFile != null)
-				{
-					_mappedFile.Release();
-					_mappedFile = null!;
-				}
-
-				_data = ReadOnlyMemory<byte>.Empty;
-			}
-		}
-
-		/// <summary>
-		/// Base directory for log files
-		/// </summary>
-		private readonly DirectoryReference _baseDir;
+		readonly DirectoryReference _baseDir;
+		readonly MemoryMappedFileCache _mappedFileCache;
 
 		/// <inheritdoc/>
 		public bool SupportsRedirects => false;
 
-		readonly object _lockObject = new object();
-		readonly Dictionary<ObjectKey, MappedFile> _pathToMappedFile = new Dictionary<ObjectKey, MappedFile>();
-		readonly LinkedList<MappedFile> _mappedFiles = new LinkedList<MappedFile>();
-
-		long _mappedSize;
-
-		const long MaxMappedSize = 1024L * 1024 * 1024;
-		const int MaxMappedCount = 128;
-
 		/// <summary>
 		/// Constructor
 		/// </summary>
@@ -178,16 +28,15 @@ namespace EpicGames.Horde.Storage.ObjectStores
 		public FileObjectStore(DirectoryReference baseDir)
 		{
 			_baseDir = baseDir;
+			_mappedFileCache = new MemoryMappedFileCache();
+
 			DirectoryReference.CreateDirectory(_baseDir);
 		}
 
 		/// <inheritdoc/>
 		public void Dispose()
 		{
-			lock (_lockObject)
-			{
-				UnmapFiles(0, 0);
-			}
+			_mappedFileCache.Dispose();
 		}
 
 		/// <summary>
@@ -195,55 +44,6 @@ namespace EpicGames.Horde.Storage.ObjectStores
 		/// </summary>
 		FileReference GetBlobFile(ObjectKey key) => FileReference.Combine(_baseDir, key.Path.ToString());
 
-		/// <summary>
-		/// Finds an existing mapped file or adds a new one for the given path
-		/// </summary>
-		MappedFile FindOrAddMappedFile(ObjectKey key)
-		{
-			MappedFile? mappedFile;
-			if (!_pathToMappedFile.TryGetValue(key, out mappedFile))
-			{
-				FileInfo fileInfo = GetBlobFile(key).ToFileInfo();
-				if (fileInfo.Length == 0)
-				{
-					throw new StorageException($"Unable to map empty memory mapped file: {fileInfo.FullName}");
-				}
-
-				long maxSize = MaxMappedSize - fileInfo.Length;
-				if (_mappedSize > maxSize || _mappedFiles.Count + 1 > MaxMappedCount)
-				{
-					UnmapFiles(maxSize, MaxMappedCount - 1);
-				}
-
-				mappedFile = new MappedFile(key, fileInfo);
-				_pathToMappedFile.Add(key, mappedFile);
-				_mappedFiles.AddFirst(mappedFile.ListNode);
-
-				_mappedSize += (long)mappedFile.MappedSize;
-			}
-			return mappedFile;
-		}
-
-		/// <summary>
-		/// Discard mapped files until only a certain size is mapped in memory
-		/// </summary>
-		/// <param name="maxMappedSize">Maximum mapped size</param>
-		/// <param name="maxMappedCount">Maximum number of mapped files</param>
-		void UnmapFiles(long maxMappedSize, int maxMappedCount)
-		{
-			for (LinkedListNode<MappedFile>? listNode = _mappedFiles.Last; listNode != null && (_mappedSize > maxMappedSize || _mappedFiles.Count > maxMappedCount);)
-			{
-				LinkedListNode<MappedFile>? nextListNode = listNode.Previous;
-				if (listNode.Value.RefCount == 1)
-				{
-					_pathToMappedFile.Remove(listNode.Value.Key);
-					_mappedFiles.Remove(listNode);
-					listNode.Value.Release();
-				}
-				listNode = nextListNode;
-			}
-		}
-
 		/// <inheritdoc/>
 		public async Task<Stream> OpenAsync(ObjectKey key, int offset, int? length, CancellationToken cancellationToken)
 		{
@@ -262,11 +62,8 @@ namespace EpicGames.Horde.Storage.ObjectStores
 		/// <returns>Handle to the data. Must be disposed by the caller.</returns>
 		public IReadOnlyMemoryOwner<byte> Read(ObjectKey key, int offset, int? length)
 		{
-			lock (_lockObject)
-			{
-				MappedFile mappedFile = FindOrAddMappedFile(key);
-				return new MappedFileHandle(mappedFile, mappedFile.GetData(offset, length));
-			}
+			FileReference file = GetBlobFile(key);
+			return _mappedFileCache.Read(file, offset, length);
 		}
 
 		/// <inheritdoc/>
@@ -318,24 +115,8 @@ namespace EpicGames.Horde.Storage.ObjectStores
 		/// <param name="key"></param>
 		public void Delete(ObjectKey key)
 		{
-			lock(_lockObject)
-			{
-				MappedFile? mappedFile;
-				if (_pathToMappedFile.TryGetValue(key, out mappedFile))
-				{
-					mappedFile.DeleteOnDispose();
-
-					_pathToMappedFile.Remove(key);
-					_mappedFiles.Remove(mappedFile.ListNode);
-
-					mappedFile.Release();
-				}
-				else
-				{
-					FileReference location = GetBlobFile(key);
-					FileReference.Delete(location);
-				}
-			}
+			FileReference location = GetBlobFile(key);
+			_mappedFileCache.Delete(location);
 		}
 
 		/// <inheritdoc/>
