UHT: Added support for UE_INTERNAL

#rnx

diff --git a/Engine/Source/Programs/Shared/EpicGames.UHT/Tokenizer/UhtTokenReaderGeneralExtensions.cs b/Engine/Source/Programs/Shared/EpicGames.UHT/Tokenizer/UhtTokenReaderGeneralExtensions.cs
--- a/Engine/Source/Programs/Shared/EpicGames.UHT/Tokenizer/UhtTokenReaderGeneralExtensions.cs
+++ b/Engine/Source/Programs/Shared/EpicGames.UHT/Tokenizer/UhtTokenReaderGeneralExtensions.cs
@@ -235,6 +235,10 @@ namespace EpicGames.UHT.Tokenizer
 						.Require(')', "deprecation macro");
 					attributeAction?.Invoke("deprecated");
 				}
+				else if (token.IsIdentifier("UE_INTERNAL"))
+				{
+					tokenReader.ConsumeToken();
+				}
 				else if (token.IsIdentifier("UE_NODISCARD") || token.IsIdentifier("UE_NODISCARD_CTOR"))
 				{
 					tokenReader.ConsumeToken();
