EpicGames.Perforce: Expanding Merge command options

Allowing to perform stream based merging with the command and exposing appropriate options

#rnx

diff --git a/Engine/Source/Programs/Shared/EpicGames.Perforce/PerforceConnection.cs b/Engine/Source/Programs/Shared/EpicGames.Perforce/PerforceConnection.cs
--- a/Engine/Source/Programs/Shared/EpicGames.Perforce/PerforceConnection.cs
+++ b/Engine/Source/Programs/Shared/EpicGames.Perforce/PerforceConnection.cs
@@ -2532,7 +2532,34 @@ namespace EpicGames.Perforce
 			{
 				arguments.Add($"-m{maxFiles}");
 			}
+			if ((options & MergeOptions.AsStreamSpec) != 0)
+			{
+				arguments.Add("-As");
+			}
+			if ((options & MergeOptions.AsFiles) != 0)
+			{
+				arguments.Add("-Af");
+			}
+			if ((options & MergeOptions.Stream) != 0)
+			{
+				arguments.Add("-S");
+			}
+
 			arguments.Add(sourceFileSpec);
+
+			if ((options & MergeOptions.Force) != 0)
+			{
+				arguments.Add("-F");
+			}
+			if ((options & MergeOptions.ReverseMapping) != 0)
+			{
+				arguments.Add("-r");
+			}
+			if ((options & MergeOptions.Source) != 0)
+			{
+				arguments.Add("-s");
+			}
+
 			arguments.Add(targetFileSpec);
 
 			PerforceResponseList<MergeRecord> records = await CommandAsync<MergeRecord>(connection, "merge", arguments, null, cancellationToken);
diff --git a/Engine/Source/Programs/Shared/EpicGames.Perforce/Records/MergeOptions.cs b/Engine/Source/Programs/Shared/EpicGames.Perforce/Records/MergeOptions.cs
--- a/Engine/Source/Programs/Shared/EpicGames.Perforce/Records/MergeOptions.cs
+++ b/Engine/Source/Programs/Shared/EpicGames.Perforce/Records/MergeOptions.cs
@@ -19,5 +19,35 @@ namespace EpicGames.Perforce
 		/// Preview the merge
 		/// </summary>
 		Preview = 1,
+
+		/// <summary>
+		/// Forces the operation
+		/// </summary>
+		Force = 2,
+
+		/// <summary>
+		/// Performs the merge with files only
+		/// </summary>
+		AsFiles = 4,
+
+		/// <summary>
+		/// Performs the merge as a stream spec
+		/// </summary>
+		AsStreamSpec = 8,
+
+		/// <summary>
+		/// Use stream based merging
+		/// </summary>
+		Stream = 16,
+
+		/// <summary>
+		/// Revere the mapping direction
+		/// </summary>
+		ReverseMapping = 32,
+
+		/// <summary>
+		/// Treat the from path as the source
+		/// </summary>
+		Source = 64,
 	}
 }
