Horde: Include type names when returning JSON responses for storage nodes.

diff --git a/Engine/Source/Programs/Horde/Horde.Server/Server/DebugController.cs b/Engine/Source/Programs/Horde/Horde.Server/Server/DebugController.cs
--- a/Engine/Source/Programs/Horde/Horde.Server/Server/DebugController.cs
+++ b/Engine/Source/Programs/Horde/Horde.Server/Server/DebugController.cs
@@ -704,7 +704,7 @@ namespace Horde.Server.Server
 				}
 				if (currentConfig.FleetManagers != null && currentConfig.FleetManagers.Count > 0)
 				{
-					config.FleetManagers = currentConfig.FleetManagers?.ToList();
+					config.FleetManagers = currentConfig.FleetManagers.ToList();
 				}
 				config.LeaseUtilizationSettings = currentConfig.LeaseUtilizationSettings;
 				config.JobQueueSettings = currentConfig.JobQueueSettings;
diff --git a/Engine/Source/Programs/Horde/Horde.Server/Server/GlobalConfig.cs b/Engine/Source/Programs/Horde/Horde.Server/Server/GlobalConfig.cs
--- a/Engine/Source/Programs/Horde/Horde.Server/Server/GlobalConfig.cs
+++ b/Engine/Source/Programs/Horde/Horde.Server/Server/GlobalConfig.cs
@@ -39,6 +39,8 @@ using Horde.Server.Utilities;
 
 namespace Horde.Server.Server
 {
+#pragma warning disable CA1027 // Mark enums with FlagsAttribute
+#pragma warning disable CA1069 // Enum member 'Latest' has same value as ...
 	/// <summary>
 	/// Global version number for running the server. As new features are introduced that require data migrations, this version number indicates the backwards compatibility functionality that must be enabled.
 	/// When adding a new version here, also add a message to <see cref="ConfigService.CreateSnapshotAsync"/> describing the steps that need to be taken to upgrade the deployment.
@@ -65,6 +67,8 @@ namespace Horde.Server.Server
 		/// </summary>
 		Latest = (int)LatestPlusOne - 1,
 	}
+#pragma warning restore CA1069
+#pragma warning restore CA1027
 
 	/// <summary>
 	/// Directive to merge config data from another source
diff --git a/Engine/Source/Programs/Horde/Horde.Server/Storage/StorageController.cs b/Engine/Source/Programs/Horde/Horde.Server/Storage/StorageController.cs
--- a/Engine/Source/Programs/Horde/Horde.Server/Storage/StorageController.cs
+++ b/Engine/Source/Programs/Horde/Horde.Server/Storage/StorageController.cs
@@ -14,6 +14,7 @@ using EpicGames.Core;
 using EpicGames.Horde.Storage;
 using EpicGames.Horde.Storage.Nodes;
 using Horde.Server.Acls;
+using Horde.Server.Ddc;
 using Horde.Server.Server;
 using Horde.Server.Utilities;
 using Microsoft.AspNetCore.Authorization;
@@ -348,6 +349,21 @@ namespace Horde.Server.Storage
 			}
 		}
 
+		static readonly IReadOnlyDictionary<Guid, Type> s_blobGuidToType = GetBlobGuidTypeMap();
+
+		static Dictionary<Guid, Type> GetBlobGuidTypeMap()
+		{
+			Dictionary<Guid, Type> guidTypeMap = new Dictionary<Guid, Type>();
+			guidTypeMap.Add(CbNode.BlobTypeGuid, typeof(CbNode));
+			guidTypeMap.Add(LeafChunkedDataNode.BlobTypeGuid, typeof(LeafChunkedDataNode));
+			guidTypeMap.Add(InteriorChunkedDataNode.BlobTypeGuid, typeof(InteriorChunkedDataNode));
+			guidTypeMap.Add(CommitNode.BlobTypeGuid, typeof(CommitNode));
+			guidTypeMap.Add(DdcRefNode.BlobTypeGuid, typeof(DdcRefNode));
+			guidTypeMap.Add(DirectoryNode.BlobTypeGuid, typeof(DirectoryNode));
+			guidTypeMap.Add(RedirectNode<object>.BlobTypeGuid, typeof(RedirectNode<object>));
+			return guidTypeMap;
+		}
+
 		/// <summary>
 		/// Gets information about a particular bundle in storage
 		/// </summary>
@@ -390,7 +406,7 @@ namespace Horde.Server.Storage
 			object content;
 
 			using BlobData blobData = await storageClient.CreateBlobHandle(locator).ReadBlobDataAsync(cancellationToken);
-			if (blobData.Type.Guid == DirectoryNode.BlobType.Guid)
+			if (blobData.Type.Guid == DirectoryNode.BlobTypeGuid)
 			{
 				DirectoryNode directoryNode = BlobSerializer.Deserialize<DirectoryNode>(blobData);
 
@@ -413,7 +429,13 @@ namespace Horde.Server.Storage
 				content = new { references = blobData.Refs.Select(x => GetNodeLink(namespaceId, x)) };
 			}
 
-			return new { type = blobData.Type.Guid, content = content };
+			string? typeName = null;
+			if (s_blobGuidToType.TryGetValue(blobData.Type.Guid, out Type? type))
+			{
+				typeName = type.Name;
+			}
+
+			return new { type = blobData.Type.Guid, typeName = typeName, content = content };
 		}
 
 		static string GetNodeLink(NamespaceId namespaceId, IBlobHandle handle) => GetNodeLink(namespaceId, handle.GetLocator());
diff --git a/Engine/Source/Programs/Shared/EpicGames.Horde/Storage/Nodes/CbNode.cs b/Engine/Source/Programs/Shared/EpicGames.Horde/Storage/Nodes/CbNode.cs
--- a/Engine/Source/Programs/Shared/EpicGames.Horde/Storage/Nodes/CbNode.cs
+++ b/Engine/Source/Programs/Shared/EpicGames.Horde/Storage/Nodes/CbNode.cs
@@ -2,6 +2,7 @@
 
 using EpicGames.Core;
 using EpicGames.Serialization;
+using System;
 using System.Collections.Generic;
 
 namespace EpicGames.Horde.Storage.Nodes
@@ -12,6 +13,11 @@ namespace EpicGames.Horde.Storage.Nodes
 	[BlobConverter(typeof(CbNodeConverter))]
 	public class CbNode
 	{
+		/// <summary>
+		/// Static accessor for the blob type guid
+		/// </summary>
+		public static Guid BlobTypeGuid { get; } = new Guid("{34A0793F-42F4-8364-A798-32862932841C}");
+
 		/// <summary>
 		/// The compact binary object
 		/// </summary>
@@ -68,7 +74,7 @@ namespace EpicGames.Horde.Storage.Nodes
 			}
 		}
 
-		public static BlobType BlobType { get; } = new BlobType("{34A0793F-42F4-8364-A798-32862932841C}", 1);
+		public static BlobType BlobType { get; } = new BlobType(CbNode.BlobTypeGuid, 1);
 
 		/// <inheritdoc/>
 		public override CbNode Read(IBlobReader reader, BlobSerializerOptions options)
diff --git a/Engine/Source/Programs/Shared/EpicGames.Horde/Storage/Nodes/ChunkedDataNode.cs b/Engine/Source/Programs/Shared/EpicGames.Horde/Storage/Nodes/ChunkedDataNode.cs
--- a/Engine/Source/Programs/Shared/EpicGames.Horde/Storage/Nodes/ChunkedDataNode.cs
+++ b/Engine/Source/Programs/Shared/EpicGames.Horde/Storage/Nodes/ChunkedDataNode.cs
@@ -38,11 +38,11 @@ namespace EpicGames.Horde.Storage.Nodes
 		public static async Task CopyToStreamAsync(IBlobHandle handle, Stream outputStream, CancellationToken cancellationToken)
 		{
 			using BlobData blobData = await handle.ReadBlobDataAsync(cancellationToken);
-			if (blobData.Type.Guid == LeafChunkedDataNodeConverter.BlobType.Guid)
+			if (blobData.Type.Guid == LeafChunkedDataNode.BlobTypeGuid)
 			{
 				await LeafChunkedDataNode.CopyToStreamAsync(blobData, outputStream, cancellationToken);
 			}
-			else if (blobData.Type.Guid == InteriorChunkedDataNodeConverter.BlobTypeGuid)
+			else if (blobData.Type.Guid == InteriorChunkedDataNode.BlobTypeGuid)
 			{
 				await InteriorChunkedDataNode.CopyToStreamAsync(blobData, outputStream, cancellationToken);
 			}
@@ -92,11 +92,11 @@ namespace EpicGames.Horde.Storage.Nodes
 		/// <inheritdoc/>
 		public override ChunkedDataNode Read(IBlobReader reader, BlobSerializerOptions options)
 		{
-			if (reader.Type.Guid == LeafChunkedDataNodeConverter.BlobType.Guid)
+			if (reader.Type.Guid == LeafChunkedDataNode.BlobTypeGuid)
 			{
 				return options.GetConverter<LeafChunkedDataNode>().Read(reader, options);
 			}
-			else if (reader.Type.Guid == InteriorChunkedDataNodeConverter.BlobTypeGuid)
+			else if (reader.Type.Guid == InteriorChunkedDataNode.BlobTypeGuid)
 			{
 				return options.GetConverter<InteriorChunkedDataNode>().Read(reader, options);
 			}
@@ -202,6 +202,11 @@ namespace EpicGames.Horde.Storage.Nodes
 	[BlobConverter(typeof(LeafChunkedDataNodeConverter))]
 	public sealed class LeafChunkedDataNode : ChunkedDataNode
 	{
+		/// <summary>
+		/// Guid for the blob type
+		/// </summary>
+		public static Guid BlobTypeGuid { get; } = new Guid("{B27AFB68-4A4B-9E20-8A78-D8A439D49840}");
+
 		/// <summary>
 		/// Data for this node
 		/// </summary>
@@ -387,7 +392,7 @@ namespace EpicGames.Horde.Storage.Nodes
 		/// <summary>
 		/// Static accessor for the blob type
 		/// </summary>
-		public static BlobType BlobType { get; } = new BlobType("{B27AFB68-4A4B-9E20-8A78-D8A439D49840}", 1);
+		public static BlobType BlobType { get; } = new BlobType(LeafChunkedDataNode.BlobTypeGuid, 1);
 
 		public override LeafChunkedDataNode Read(IBlobReader reader, BlobSerializerOptions options)
 		{
@@ -431,6 +436,11 @@ namespace EpicGames.Horde.Storage.Nodes
 	[BlobConverter(typeof(InteriorChunkedDataNodeConverter))]
 	public class InteriorChunkedDataNode : ChunkedDataNode
 	{
+		/// <summary>
+		/// Static accessor for the blob type guid
+		/// </summary>
+		public static Guid BlobTypeGuid { get; } = Guid.Parse("{F4DEDDBC-4C7A-70CB-11F0-4783B9CDCCAF}");
+
 		/// <summary>
 		/// Child nodes
 		/// </summary>
@@ -567,11 +577,6 @@ namespace EpicGames.Horde.Storage.Nodes
 	/// </summary>
 	public class InteriorChunkedDataNodeConverter : BlobConverter<InteriorChunkedDataNode>
 	{
-		/// <summary>
-		/// Static accessor for the blob type guid
-		/// </summary>
-		public static Guid BlobTypeGuid { get; } = Guid.Parse("{F4DEDDBC-4C7A-70CB-11F0-4783B9CDCCAF}");
-
 		readonly int _writeVersion;
 
 		/// <summary>
@@ -631,7 +636,7 @@ namespace EpicGames.Horde.Storage.Nodes
 					writer.WriteUnsignedVarInt((ulong)child.Length);
 				}
 			}
-			return new BlobType(BlobTypeGuid, _writeVersion);
+			return new BlobType(InteriorChunkedDataNode.BlobTypeGuid, _writeVersion);
 		}
 	}
 }
diff --git a/Engine/Source/Programs/Shared/EpicGames.Horde/Storage/Nodes/CommitNode.cs b/Engine/Source/Programs/Shared/EpicGames.Horde/Storage/Nodes/CommitNode.cs
--- a/Engine/Source/Programs/Shared/EpicGames.Horde/Storage/Nodes/CommitNode.cs
+++ b/Engine/Source/Programs/Shared/EpicGames.Horde/Storage/Nodes/CommitNode.cs
@@ -12,6 +12,11 @@ namespace EpicGames.Horde.Storage.Nodes
 	[BlobConverter(typeof(CommitNodeConverter))]
 	public class CommitNode
 	{
+		/// <summary>
+		/// Static accessor for the blob type guid
+		/// </summary>
+		public static Guid BlobTypeGuid { get; } = new Guid("{64D50724-41C0-6B22-1CB5-90A8171824D6}");
+
 		/// <summary>
 		/// The commit number
 		/// </summary>
@@ -82,7 +87,7 @@ namespace EpicGames.Horde.Storage.Nodes
 
 	class CommitNodeConverter : BlobConverter<CommitNode>
 	{
-		static readonly BlobType s_blobType = new BlobType("{64D50724-41C0-6B22-1CB5-90A8171824D6}", 1);
+		static readonly BlobType s_blobType = new BlobType(CommitNode.BlobTypeGuid, 1);
 
 		public override CommitNode Read(IBlobReader reader, BlobSerializerOptions options)
 		{
diff --git a/Engine/Source/Programs/Shared/EpicGames.Horde/Storage/Nodes/DdcRefNode.cs b/Engine/Source/Programs/Shared/EpicGames.Horde/Storage/Nodes/DdcRefNode.cs
--- a/Engine/Source/Programs/Shared/EpicGames.Horde/Storage/Nodes/DdcRefNode.cs
+++ b/Engine/Source/Programs/Shared/EpicGames.Horde/Storage/Nodes/DdcRefNode.cs
@@ -1,5 +1,6 @@
 // Copyright Epic Games, Inc. All Rights Reserved.
 
+using System;
 using System.Collections.Generic;
 using System.Linq;
 using EpicGames.Core;
@@ -13,6 +14,11 @@ namespace Horde.Server.Ddc
 	[BlobConverter(typeof(DdcRefNodeConverter))]
 	public class DdcRefNode
 	{
+		/// <summary>
+		/// Static accessor for the blob type guid
+		/// </summary>
+		public static Guid BlobTypeGuid { get; } = new Guid("{0C7E5F25-4B55-454B-63F4-4A9B74D00651}");
+		
 		/// <summary>
 		/// Hash of the root node
 		/// </summary>
@@ -35,7 +41,7 @@ namespace Horde.Server.Ddc
 
 	class DdcRefNodeConverter : BlobConverter<DdcRefNode>
 	{
-		static BlobType s_blobType = new BlobType("{0C7E5F25-4B55-454B-63F4-4A9B74D00651}", 1);
+		static BlobType s_blobType = new BlobType(DdcRefNode.BlobTypeGuid, 1);
 
 		public override DdcRefNode Read(IBlobReader reader, BlobSerializerOptions options)
 		{
diff --git a/Engine/Source/Programs/Shared/EpicGames.Horde/Storage/Nodes/DirectoryNode.cs b/Engine/Source/Programs/Shared/EpicGames.Horde/Storage/Nodes/DirectoryNode.cs
--- a/Engine/Source/Programs/Shared/EpicGames.Horde/Storage/Nodes/DirectoryNode.cs
+++ b/Engine/Source/Programs/Shared/EpicGames.Horde/Storage/Nodes/DirectoryNode.cs
@@ -154,7 +154,7 @@ namespace EpicGames.Horde.Storage.Nodes
 		/// <summary>
 		/// Type of serialized directory node blobs
 		/// </summary>
-		public static BlobType BlobType { get; } = new BlobType("{0714EC11-4D07-291A-8AE7-7F86799980D6}", 2);
+		public static Guid BlobTypeGuid { get; } = new Guid("{0714EC11-4D07-291A-8AE7-7F86799980D6}");
 
 		readonly SortedDictionary<string, FileEntry> _nameToFileEntry = new SortedDictionary<string, FileEntry>(StringComparer.Ordinal);
 		readonly SortedDictionary<string, DirectoryEntry> _nameToDirectoryEntry = new SortedDictionary<string, DirectoryEntry>(StringComparer.Ordinal);
@@ -931,6 +931,11 @@ namespace EpicGames.Horde.Storage.Nodes
 
 	class DirectoryNodeConverter : BlobConverter<DirectoryNode>
 	{
+		/// <summary>
+		/// Type of serialized directory node blobs
+		/// </summary>
+		public static BlobType BlobType { get; } = new BlobType(DirectoryNode.BlobTypeGuid, 2);
+
 		/// <inheritdoc/>
 		public override DirectoryNode Read(IBlobReader reader, BlobSerializerOptions options)
 		{
@@ -1008,7 +1013,7 @@ namespace EpicGames.Horde.Storage.Nodes
 				writer.WriteString(directoryEntry.Name);
 			}
 
-			return DirectoryNode.BlobType;
+			return BlobType;
 		}
 	}
 
diff --git a/Engine/Source/Programs/Shared/EpicGames.Horde/Storage/Nodes/RedirectNode.cs b/Engine/Source/Programs/Shared/EpicGames.Horde/Storage/Nodes/RedirectNode.cs
--- a/Engine/Source/Programs/Shared/EpicGames.Horde/Storage/Nodes/RedirectNode.cs
+++ b/Engine/Source/Programs/Shared/EpicGames.Horde/Storage/Nodes/RedirectNode.cs
@@ -1,5 +1,7 @@
 // Copyright Epic Games, Inc. All Rights Reserved.
 
+using System;
+
 namespace EpicGames.Horde.Storage.Nodes
 {
 	/// <summary>
@@ -8,6 +10,11 @@ namespace EpicGames.Horde.Storage.Nodes
 	[BlobConverter(typeof(RedirectNodeConverter<>))]
 	public class RedirectNode<T>
 	{
+		/// <summary>
+		/// Static accessor for the blob type guid
+		/// </summary>
+		public static Guid BlobTypeGuid { get; } = new Guid("{BE09E54F-47CA-7A6B-2A97-AFBC183B1538}");
+
 		/// <summary>
 		/// The target handle
 		/// </summary>
@@ -22,7 +29,7 @@ namespace EpicGames.Horde.Storage.Nodes
 
 	class RedirectNodeConverter<T> : BlobConverter<RedirectNode<T>>
 	{
-		static readonly BlobType s_blobType = new BlobType("{BE09E54F-47CA-7A6B-2A97-AFBC183B1538}", 1);
+		static readonly BlobType s_blobType = new BlobType(RedirectNode<T>.BlobTypeGuid, 1);
 
 		/// <inheritdoc/>
 		public override RedirectNode<T> Read(IBlobReader reader, BlobSerializerOptions options)
