Horde: Fix static analysis warning.

diff --git a/Engine/Source/Programs/Horde/Horde.Server/Storage/StorageController.cs b/Engine/Source/Programs/Horde/Horde.Server/Storage/StorageController.cs
--- a/Engine/Source/Programs/Horde/Horde.Server/Storage/StorageController.cs
+++ b/Engine/Source/Programs/Horde/Horde.Server/Storage/StorageController.cs
@@ -16,7 +16,6 @@ using EpicGames.Horde.Storage;
 using EpicGames.Horde.Storage.Nodes;
 using EpicGames.Serialization;
 using Horde.Server.Acls;
-using Horde.Server.Agents;
 using Horde.Server.Ddc;
 using Horde.Server.Server;
 using Horde.Server.Utilities;
@@ -363,7 +362,7 @@ namespace Horde.Server.Storage
 			guidTypeMap.Add(CommitNode.BlobTypeGuid, typeof(CommitNode));
 			guidTypeMap.Add(DdcRefNode.BlobTypeGuid, typeof(DdcRefNode));
 			guidTypeMap.Add(DirectoryNode.BlobTypeGuid, typeof(DirectoryNode));
-			guidTypeMap.Add(RedirectNode<object>.BlobTypeGuid, typeof(RedirectNode<object>));
+			guidTypeMap.Add(RedirectNode.BlobTypeGuid, typeof(RedirectNode));
 			return guidTypeMap;
 		}
 
@@ -474,7 +473,7 @@ namespace Horde.Server.Storage
 			return new { type = typeName, guid = blobData.Type.Guid, content = content };
 		}
 
-		object? GetCbNodeObject(NamespaceId namespaceId, CbField field, IEnumerator<IBlobHandle> references)
+		static object? GetCbNodeObject(NamespaceId namespaceId, CbField field, IEnumerator<IBlobHandle> references)
 		{
 			if (field.IsAttachment())
 			{
diff --git a/Engine/Source/Programs/Shared/EpicGames.Horde/Storage/Nodes/RedirectNode.cs b/Engine/Source/Programs/Shared/EpicGames.Horde/Storage/Nodes/RedirectNode.cs
--- a/Engine/Source/Programs/Shared/EpicGames.Horde/Storage/Nodes/RedirectNode.cs
+++ b/Engine/Source/Programs/Shared/EpicGames.Horde/Storage/Nodes/RedirectNode.cs
@@ -5,16 +5,22 @@ using System;
 namespace EpicGames.Horde.Storage.Nodes
 {
 	/// <summary>
-	/// A node containing arbitrary compact binary data
+	/// Shared definitions for <see cref="RedirectNode{T}"/>
 	/// </summary>
-	[BlobConverter(typeof(RedirectNodeConverter<>))]
-	public class RedirectNode<T>
+	public static class RedirectNode
 	{
 		/// <summary>
 		/// Static accessor for the blob type guid
 		/// </summary>
 		public static Guid BlobTypeGuid { get; } = new Guid("{BE09E54F-47CA-7A6B-2A97-AFBC183B1538}");
+	}
 
+	/// <summary>
+	/// A node containing arbitrary compact binary data
+	/// </summary>
+	[BlobConverter(typeof(RedirectNodeConverter<>))]
+	public class RedirectNode<T>
+	{
 		/// <summary>
 		/// The target handle
 		/// </summary>
@@ -29,7 +35,7 @@ namespace EpicGames.Horde.Storage.Nodes
 
 	class RedirectNodeConverter<T> : BlobConverter<RedirectNode<T>>
 	{
-		static readonly BlobType s_blobType = new BlobType(RedirectNode<T>.BlobTypeGuid, 1);
+		static readonly BlobType s_blobType = new BlobType(RedirectNode.BlobTypeGuid, 1);
 
 		/// <inheritdoc/>
 		public override RedirectNode<T> Read(IBlobReader reader, BlobSerializerOptions options)
