Horde: Add separate types for blob handles with a hash (IBlobRef) and blob handles with a hash and serialization options (IBlobRef<T>). Storing serialization options inside IBlobRef<T> implementations removes a lot of boilerplate code to pass options objects around, and ensures consistency when navigating object hierarchies.

diff --git a/Engine/Source/Programs/AutomationTool/BuildGraph/Tasks/CreateArtifactTask.cs b/Engine/Source/Programs/AutomationTool/BuildGraph/Tasks/CreateArtifactTask.cs
--- a/Engine/Source/Programs/AutomationTool/BuildGraph/Tasks/CreateArtifactTask.cs
+++ b/Engine/Source/Programs/AutomationTool/BuildGraph/Tasks/CreateArtifactTask.cs
@@ -125,7 +125,7 @@ namespace AutomationTool.Tasks
 					int totalCount = files.Count;
 					long totalSize = files.Sum(x => x.Length);
 
-					IBlobHandle<DirectoryNode> outputNodeRef = await writer.WriteFilesAsync(baseDir.ToDirectoryInfo(), files, progress: new CopyStatsLogger(totalCount, totalSize, Logger));
+					IBlobRef<DirectoryNode> outputNodeRef = await writer.WriteFilesAsync(baseDir.ToDirectoryInfo(), files, progress: new CopyStatsLogger(totalCount, totalSize, Logger));
 					await writer.FlushAsync();
 
 					await client.WriteRefAsync(response.RefName, outputNodeRef);
diff --git a/Engine/Source/Programs/AutomationTool/BuildGraph/Tasks/DeployToolTask.cs b/Engine/Source/Programs/AutomationTool/BuildGraph/Tasks/DeployToolTask.cs
--- a/Engine/Source/Programs/AutomationTool/BuildGraph/Tasks/DeployToolTask.cs
+++ b/Engine/Source/Programs/AutomationTool/BuildGraph/Tasks/DeployToolTask.cs
@@ -165,25 +165,25 @@ namespace AutomationTool.Tasks
 			ToolId toolId = new ToolId(Parameters.Id);
 
 			using IStorageClient storageClient = hordeClient.CreateStorageClient(toolId);
-			await using (IBlobWriter treeWriter = storageClient.CreateBlobWriter())
+			await using (IBlobWriter blobWriter = storageClient.CreateBlobWriter(serializerOptions: serializerOptions))
 			{
 				DirectoryNode sandbox = new DirectoryNode();
 				if (Parameters.File != null)
 				{
 					using FileStream stream = FileReference.Open(ResolveFile(Parameters.File), FileMode.Open, FileAccess.Read);
-					await sandbox.CopyFromZipStreamAsync(stream, treeWriter, new ChunkingOptions(), serializerOptions);
+					await sandbox.CopyFromZipStreamAsync(stream, blobWriter, new ChunkingOptions());
 				}
 				else if (Parameters.Directory != null)
 				{
 					DirectoryInfo directoryInfo = ResolveDirectory(Parameters.Directory).ToDirectoryInfo();
-					await sandbox.AddFilesAsync(directoryInfo, treeWriter, serializerOptions: serializerOptions);
+					await sandbox.AddFilesAsync(directoryInfo, blobWriter);
 				}
 				else
 				{
 					throw new AutomationException("Either File=... or Directory=... must be specified");
 				}
-				handle = await treeWriter.WriteBlobAsync(sandbox, serializerOptions);
-				await treeWriter.FlushAsync();
+				handle = await blobWriter.WriteBlobAsync(sandbox);
+				await blobWriter.FlushAsync();
 			}
 
 			double? duration = null;
diff --git a/Engine/Source/Programs/Horde/Horde.Agent.Tests/JsonRpcLoggerTests.cs b/Engine/Source/Programs/Horde/Horde.Agent.Tests/JsonRpcLoggerTests.cs
--- a/Engine/Source/Programs/Horde/Horde.Agent.Tests/JsonRpcLoggerTests.cs
+++ b/Engine/Source/Programs/Horde/Horde.Agent.Tests/JsonRpcLoggerTests.cs
@@ -102,14 +102,14 @@ namespace Horde.Agent.Tests
 
 		class FakeJsonRpcLoggerBackend : JsonRpcAndStorageLogSink
 		{
-			public IBlobHandle? Target { get; private set; }
+			public IBlobRef? Target { get; private set; }
 
 			public FakeJsonRpcLoggerBackend(IRpcConnection connection, string logId, string? jobId, string? batchId, string? stepId, IStorageClient store, ILogger logger)
 				: base(connection, logId, jobId, batchId, stepId, store, logger)
 			{
 			}
 
-			protected override Task UpdateLogAsync(IBlobHandle target, int lineCount, bool complete, CancellationToken cancellationToken)
+			protected override Task UpdateLogAsync(IBlobRef target, int lineCount, bool complete, CancellationToken cancellationToken)
 			{
 				Target = target;
 				return Task.CompletedTask;
diff --git a/Engine/Source/Programs/Horde/Horde.Agent/Execution/JobExecutor.cs b/Engine/Source/Programs/Horde/Horde.Agent/Execution/JobExecutor.cs
--- a/Engine/Source/Programs/Horde/Horde.Agent/Execution/JobExecutor.cs
+++ b/Engine/Source/Programs/Horde/Horde.Agent/Execution/JobExecutor.cs
@@ -518,12 +518,12 @@ namespace Horde.Agent.Execution
 
 					Stopwatch timer = Stopwatch.StartNew();
 
-					IBlobHandle<DirectoryNode> rootNodeRef;
+					IBlobRef<DirectoryNode> rootNodeRef;
 					await using (IBlobWriter treeWriter = CreateStorageWriter(storage, artifact.RefName, logger))
 					{
 						DirectoryNode buildGraphNode = new DirectoryNode();
 						await buildGraphNode.AddFilesAsync(workspaceDir, buildGraphFiles, treeWriter, cancellationToken: cancellationToken);
-						IBlobHandle<DirectoryNode> outputNodeRef = await treeWriter.WriteBlobAsync(buildGraphNode, cancellationToken);
+						IBlobRef<DirectoryNode> outputNodeRef = await treeWriter.WriteBlobAsync(buildGraphNode, cancellationToken);
 
 						DirectoryNode rootNode = new DirectoryNode();
 						rootNode.AddDirectory(new DirectoryEntry(BuildGraphTempStorageDir, buildGraphNode.Length, outputNodeRef));
@@ -848,7 +848,7 @@ namespace Horde.Agent.Execution
 
 				using IStorageClient storage = CreateStorageClient(new NamespaceId(artifact.NamespaceId), artifact.Token);
 
-				IBlobHandle<DirectoryNode> rootRef;
+				IBlobRef<DirectoryNode> rootRef;
 				await using (IBlobWriter writer = CreateStorageWriter(storage, new RefName(artifact.RefName), logger))
 				{
 					try
@@ -1051,7 +1051,7 @@ namespace Horde.Agent.Execution
 				// Upload the data
 				Stopwatch timer = Stopwatch.StartNew();
 
-				IBlobHandle<DirectoryNode> outputNodeRef;
+				IBlobRef<DirectoryNode> outputNodeRef;
 				await using (IBlobWriter treeWriter = CreateStorageWriter(storage, artifact.RefName, logger))
 				{
 					DirectoryNode outputNode = new DirectoryNode();
diff --git a/Engine/Source/Programs/Horde/Horde.Agent/Utility/JsonRpcLogSink.cs b/Engine/Source/Programs/Horde/Horde.Agent/Utility/JsonRpcLogSink.cs
--- a/Engine/Source/Programs/Horde/Horde.Agent/Utility/JsonRpcLogSink.cs
+++ b/Engine/Source/Programs/Horde/Horde.Agent/Utility/JsonRpcLogSink.cs
@@ -194,7 +194,7 @@ namespace Horde.Agent.Utility
 
 			if (request.Flush || _bufferLength > FlushLength)
 			{
-				IBlobHandle<LogNode> target = await _builder.FlushAsync(_writer, request.Flush, null, cancellationToken);
+				IBlobRef<LogNode> target = await _builder.FlushAsync(_writer, request.Flush, cancellationToken);
 				await UpdateLogAsync(target, _builder.LineCount, request.Flush, cancellationToken);
 				_bufferLength = 0;
 			}
@@ -204,7 +204,7 @@ namespace Horde.Agent.Utility
 
 		#region RPC calls
 
-		protected virtual async Task UpdateLogAsync(IBlobHandle target, int lineCount, bool complete, CancellationToken cancellationToken)
+		protected virtual async Task UpdateLogAsync(IBlobRef target, int lineCount, bool complete, CancellationToken cancellationToken)
 		{
 			_logger.LogInformation("Updating log {LogId} to line {LineCount}, target {Locator}", _logId, lineCount, target.GetLocator());
 
diff --git a/Engine/Source/Programs/Horde/Horde.Agent/Utility/TempStorage.cs b/Engine/Source/Programs/Horde/Horde.Agent/Utility/TempStorage.cs
--- a/Engine/Source/Programs/Horde/Horde.Agent/Utility/TempStorage.cs
+++ b/Engine/Source/Programs/Horde/Horde.Agent/Utility/TempStorage.cs
@@ -647,7 +647,7 @@ namespace Horde.Storage.Utility
 			FileReference localFileListLocation = GetTagManifestLocation(manifestDir, nodeName, tagName);
 			fileList.Save(localFileListLocation);
 
-			using ChunkedDataWriter fileNodeWriter = new ChunkedDataWriter(writer, new ChunkingOptions(), BlobSerializerOptions.Default);
+			using ChunkedDataWriter fileNodeWriter = new ChunkedDataWriter(writer, new ChunkingOptions());
 			ChunkedData fileNodeData = await fileNodeWriter.CreateAsync(localFileListLocation.ToFileInfo(), cancellationToken);
 
 			return new FileEntry(localFileListLocation.GetFileName(), FileEntryFlags.None, fileNodeWriter.Length, fileNodeData);
@@ -701,7 +701,7 @@ namespace Horde.Storage.Utility
 			DirectoryNode rootNode = new DirectoryNode();
 			await rootNode.AddFilesAsync(workspaceDir, archiveFiles, writer, progress: new CopyStatsLogger(logger), cancellationToken: cancellationToken);
 
-			IBlobHandle<DirectoryNode> rootNodeRef = await writer.WriteBlobAsync(rootNode, cancellationToken: cancellationToken);
+			IBlobRef<DirectoryNode> rootNodeRef = await writer.WriteBlobAsync(rootNode, cancellationToken: cancellationToken);
 			return new DirectoryEntry(blockDirectoryName, rootNode.Length, rootNodeRef);
 		}
 
@@ -766,7 +766,7 @@ namespace Horde.Storage.Utility
 
 				// Add all the files and flush the ref
 				DirectoryNode rootDirNode = await rootDirEntry.Handle.ReadBlobAsync(cancellationToken: cancellationToken);
-				await rootDirNode.CopyToDirectoryAsync(rootDir.ToDirectoryInfo(), new CopyStatsLogger(logger), null, logger, cancellationToken);
+				await rootDirNode.CopyToDirectoryAsync(rootDir.ToDirectoryInfo(), new CopyStatsLogger(logger), logger, cancellationToken);
 
 				StorageStats deltaStats = StorageStats.GetDelta(initialStats, storageClient.GetStats());
 				logger.LogInformation("{Stats}", $"Elapsed: {(int)timer.Elapsed.TotalSeconds}s, {String.Join(", ", deltaStats.Values.Select(x => $"{x.Item1}: {x.Item2:n0}"))}");
diff --git a/Engine/Source/Programs/Horde/Horde.Server.Tests/BlobStoreTests.cs b/Engine/Source/Programs/Horde/Horde.Server.Tests/BlobStoreTests.cs
--- a/Engine/Source/Programs/Horde/Horde.Server.Tests/BlobStoreTests.cs
+++ b/Engine/Source/Programs/Horde/Horde.Server.Tests/BlobStoreTests.cs
@@ -54,13 +54,13 @@ namespace Horde.Server.Tests
 		{
 			using (BlobData blobData = await store.CreateBlobHandle(locator).ReadBlobDataAsync())
 			{
-				BlobReader reader = new BlobReader(blobData);
+				BlobReader reader = new BlobReader(blobData, null);
 				ReadOnlyMemory<byte> data = reader.ReadVariableLengthBytes();
 
-				List<IBlobHandle> handles = new List<IBlobHandle>();
+				List<IBlobRef> handles = new List<IBlobRef>();
 				while (reader.RemainingMemory.Length > 0)
 				{
-					handles.Add(reader.ReadBlobHandle<object>());
+					handles.Add(reader.ReadBlobRef<object>());
 				}
 
 				List<BlobLocator> locators = handles.ConvertAll(x => x.GetLocator());
@@ -75,10 +75,10 @@ namespace Horde.Server.Tests
 
 			foreach (BlobLocator reference in blob.References)
 			{
-				writer.WriteBlobHandle(store.CreateBlobHandle(reference).ForType<object>(IoHash.Zero));
+				writer.WriteBlobRef(BlobRef.Create(IoHash.Zero, store.CreateBlobHandle(reference)));
 			}
 
-			IBlobHandle handle = await writer.CompleteAsync(s_blobType);
+			IBlobRef handle = await writer.CompleteAsync(s_blobType);
 			await handle.FlushAsync();
 
 			return handle.GetLocator();
diff --git a/Engine/Source/Programs/Horde/Horde.Server.Tests/GcServiceTests.cs b/Engine/Source/Programs/Horde/Horde.Server.Tests/GcServiceTests.cs
--- a/Engine/Source/Programs/Horde/Horde.Server.Tests/GcServiceTests.cs
+++ b/Engine/Source/Programs/Horde/Horde.Server.Tests/GcServiceTests.cs
@@ -29,7 +29,7 @@ namespace Horde.Server.Tests
 			using IStorageClient store = StorageService.CreateClient(new NamespaceId("default"));
 
 			Random random = new Random(0);
-			IBlobHandle<object>[] blobs = await CreateTestDataAsync(store, 30, 50, 30, 5, random);
+			IBlobRef[] blobs = await CreateTestDataAsync(store, 30, 50, 30, 5, random);
 
 			HashSet<IBlobHandle> roots = new HashSet<IBlobHandle>();
 			for (int idx = 0; idx < 10; idx++)
@@ -70,12 +70,12 @@ namespace Horde.Server.Tests
 				if (nodes.Add(locator))
 				{
 					BlobData data = await root.ReadBlobDataAsync();
-					await FindNodesAsync(store, data.Refs, nodes);
+					await FindNodesAsync(store, data.Imports, nodes);
 				}
 			}
 		}
 
-		static async ValueTask<IBlobHandle<object>[]> CreateTestDataAsync(IStorageClient store, int numRoots, int numInterior, int numLeaves, int avgChildren, Random random)
+		static async ValueTask<IBlobRef[]> CreateTestDataAsync(IStorageClient store, int numRoots, int numInterior, int numLeaves, int avgChildren, Random random)
 		{
 			int firstRoot = 0;
 			int firstInterior = firstRoot + numRoots;
@@ -101,16 +101,16 @@ namespace Horde.Server.Tests
 
 			BlobType blobType = new BlobType(Guid.Parse("{AFDF76A7-4DEE-5333-F5B5-37B8451251CA}"), 0);
 
-			IBlobHandle<object>[] handles = new IBlobHandle<object>[children.Length];
+			IBlobRef[] handles = new IBlobRef[children.Length];
 			for (int idx = numNodes - 1; idx >= 0; idx--)
 			{
-				IBlobHandle<object> handle;
+				IBlobRef handle;
 				await using (IBlobWriter writer = store.CreateBlobWriter("gctest"))
 				{
-					List<IBlobHandle<object>> imports = children[idx].ConvertAll(x => handles[x]);
-					foreach (IBlobHandle<object> import in imports)
+					List<IBlobRef> imports = children[idx].ConvertAll(x => handles[x]);
+					foreach (IBlobRef import in imports)
 					{
-						writer.WriteBlobHandle(import);
+						writer.WriteBlobRef(import);
 					}
 					handle = await writer.CompleteAsync<object>(blobType);
 				}
diff --git a/Engine/Source/Programs/Horde/Horde.Server.Tests/IssueServiceTests.cs b/Engine/Source/Programs/Horde/Horde.Server.Tests/IssueServiceTests.cs
--- a/Engine/Source/Programs/Horde/Horde.Server.Tests/IssueServiceTests.cs
+++ b/Engine/Source/Programs/Horde/Horde.Server.Tests/IssueServiceTests.cs
@@ -74,7 +74,7 @@ namespace Horde.Server.Tests
 
 				await using (IBlobWriter writer = _storageClient.CreateBlobWriter())
 				{
-					IBlobHandle<LogNode> handle = await _builder.FlushAsync(writer, true, null, CancellationToken.None);
+					IBlobRef<LogNode> handle = await _builder.FlushAsync(writer, true, CancellationToken.None);
 					ILogFile? logFile = await _logFileService.GetLogFileAsync(_logId, CancellationToken.None);
 					await _storageClient.WriteRefAsync(logFile!.RefName, handle);
 				}
@@ -334,7 +334,7 @@ namespace Horde.Server.Tests
 					LogBuilder builder = new LogBuilder(LogFormat.Json, NullLogger.Instance);
 					builder.WriteData(data);
 
-					IBlobHandle<LogNode> handle = await builder.FlushAsync(writer, true, null, CancellationToken.None);
+					IBlobRef<LogNode> handle = await builder.FlushAsync(writer, true, CancellationToken.None);
 					await storageClient.WriteRefAsync(logFile!.RefName, handle);
 				}
 			}
diff --git a/Engine/Source/Programs/Horde/Horde.Server.Tests/LogIndexingTests.cs b/Engine/Source/Programs/Horde/Horde.Server.Tests/LogIndexingTests.cs
--- a/Engine/Source/Programs/Horde/Horde.Server.Tests/LogIndexingTests.cs
+++ b/Engine/Source/Programs/Horde/Horde.Server.Tests/LogIndexingTests.cs
@@ -52,7 +52,7 @@ namespace Horde.Server.Tests
 
 		public async Task<ILogFile> FlushAsync(bool complete = true)
 		{
-			IBlobHandle<LogNode> handle = await _builder.FlushAsync(_blobWriter, complete, null, CancellationToken.None);
+			IBlobRef<LogNode> handle = await _builder.FlushAsync(_blobWriter, complete, CancellationToken.None);
 			await _storageClient.WriteRefAsync(_logFile.RefName, handle);
 			_logFile = await _logCollection.UpdateLineCountAsync(_logFile, _lineCount, complete, CancellationToken.None);
 			return _logFile;
diff --git a/Engine/Source/Programs/Horde/Horde.Server.Tests/StorageServiceTests.cs b/Engine/Source/Programs/Horde/Horde.Server.Tests/StorageServiceTests.cs
--- a/Engine/Source/Programs/Horde/Horde.Server.Tests/StorageServiceTests.cs
+++ b/Engine/Source/Programs/Horde/Horde.Server.Tests/StorageServiceTests.cs
@@ -25,9 +25,9 @@ namespace Horde.Server.Tests
 			BlobType type2 = new BlobType(Guid.Parse("{6CB3A005-4787-26BA-3E79-D286CB7137D1}"), 0);
 			byte[] data2 = new byte[] { 4, 5, 6 };
 
-			IBlobHandle handle1a;
-			IBlobHandle handle1b;
-			IBlobHandle handle2;
+			IBlobRef handle1a;
+			IBlobRef handle1b;
+			IBlobRef handle2;
 			await using (IBlobWriter writer = client.CreateBlobWriter())
 			{
 				writer.WriteFixedLengthBytes(data1);
diff --git a/Engine/Source/Programs/Horde/Horde.Server/Commands/Generate/ToolDataCommand.cs b/Engine/Source/Programs/Horde/Horde.Server/Commands/Generate/ToolDataCommand.cs
--- a/Engine/Source/Programs/Horde/Horde.Server/Commands/Generate/ToolDataCommand.cs
+++ b/Engine/Source/Programs/Horde/Horde.Server/Commands/Generate/ToolDataCommand.cs
@@ -53,7 +53,7 @@ namespace Horde.Server.Commands.Generate
 			{
 				using (IStorageClient client = BundleStorageClient.CreateFromDirectory(bundleDir, bundleCache, logger))
 				{
-					IBlobHandle<DirectoryNode> dirNodeRef;
+					IBlobRef<DirectoryNode> dirNodeRef;
 					await using (IBlobWriter writer = client.CreateBlobWriter(refName))
 					{
 						DirectoryNode dirNode = new DirectoryNode();
diff --git a/Engine/Source/Programs/Horde/Horde.Server/Ddc/BlobService.cs b/Engine/Source/Programs/Horde/Horde.Server/Ddc/BlobService.cs
--- a/Engine/Source/Programs/Horde/Horde.Server/Ddc/BlobService.cs
+++ b/Engine/Source/Programs/Horde/Horde.Server/Ddc/BlobService.cs
@@ -143,7 +143,7 @@ namespace Horde.Server.Ddc
 		{
 			using IStorageClient storageClient = _storageClientFactory.CreateClient(ns);
 
-			IBlobHandle blobHandle;
+			IBlobRef blobRef;
 			await using (IBlobWriter writer = storageClient.CreateBlobWriter())
 			{
 				int contentLength = (int)content.Length;
@@ -153,11 +153,11 @@ namespace Horde.Server.Ddc
 				await stream.ReadFixedLengthBytesAsync(memory, cancellationToken);
 				writer.Advance(contentLength);
 
-				blobHandle = await writer.CompleteAsync(s_rawBlobType, cancellationToken);
+				blobRef = await writer.CompleteAsync(s_rawBlobType, cancellationToken);
 				await writer.FlushAsync(cancellationToken);
 			}
 
-			await storageClient.AddAliasAsync(GetAlias(identifier), blobHandle, data: identifier.Hash.ToByteArray(), cancellationToken: cancellationToken);
+			await storageClient.AddAliasAsync(GetAlias(identifier), blobRef, data: identifier.Hash.ToByteArray(), cancellationToken: cancellationToken);
 			return identifier;
 		}
 	}
diff --git a/Engine/Source/Programs/Horde/Horde.Server/Ddc/RefService.cs b/Engine/Source/Programs/Horde/Horde.Server/Ddc/RefService.cs
--- a/Engine/Source/Programs/Horde/Horde.Server/Ddc/RefService.cs
+++ b/Engine/Source/Programs/Horde/Horde.Server/Ddc/RefService.cs
@@ -78,15 +78,15 @@ namespace Horde.Server.Ddc
 				// TODO: We resolved all these blobs above... Need to just have GetReferencedBlobs just return the appropriate handles directly.
 				RefName refName = GetRefName(bucket, key);
 
-				IBlobHandle<DdcRefNode> refNodeRef;
+				IBlobRef<DdcRefNode> refNodeRef;
 				await using (IBlobWriter writer = storageClient.CreateBlobWriter(refName))
 				{
 					DdcRefNode refNode = new DdcRefNode(blobHash.Hash);
-					refNode.References.Add((blobHash.Hash, blobHandle));
+					refNode.References.Add(BlobRef.Create(blobHash.Hash, blobHandle));
 					foreach (BlobId referencedBlob in referencedBlobs)
 					{
 						BlobAlias? alias = await storageClient.FindAliasAsync(BlobService.GetAlias(referencedBlob), cancellationToken);
-						refNode.References.Add((referencedBlob.Hash, alias!.Target));
+						refNode.References.Add(BlobRef.Create(referencedBlob.Hash, alias!.Target));
 					}
 					refNodeRef = await writer.WriteBlobAsync(refNode, cancellationToken: cancellationToken);
 				}
@@ -138,7 +138,7 @@ namespace Horde.Server.Ddc
 				throw new RefNotFoundException(ns, bucket, key);
 			}
 
-			BlobData data = await node.References.First(x => x.Hash == node.RootHash).Handle.ReadBlobDataAsync(cancellationToken);
+			BlobData data = await node.References.First(x => x.Hash == node.RootHash).ReadBlobDataAsync(cancellationToken);
 			BlobContents contents = new BlobContents(data.Data.ToArray());
 
 			RefRecord record = new RefRecord(ns, bucket, key, DateTime.UtcNow, null, new BlobId(node.RootHash), true);
diff --git a/Engine/Source/Programs/Horde/Horde.Server/Replicators/PerforceReplicator.cs b/Engine/Source/Programs/Horde/Horde.Server/Replicators/PerforceReplicator.cs
--- a/Engine/Source/Programs/Horde/Horde.Server/Replicators/PerforceReplicator.cs
+++ b/Engine/Source/Programs/Horde/Horde.Server/Replicators/PerforceReplicator.cs
@@ -54,12 +54,12 @@ namespace Horde.Server.Replicators
 		{
 			public int Change { get; }
 			public int ParentChange { get; }
-			public IBlobHandle<CommitNode>? ParentHandle { get; }
+			public IBlobRef<CommitNode>? ParentHandle { get; }
 			public long CopiedSize { get; set; }
-			public IBlobHandle<DirectoryNode>? Contents { get; set; }
+			public IBlobRef<DirectoryNode>? Contents { get; set; }
 			public List<string> Paths { get; }
 
-			public StateNode(int number, int parentNumber, IBlobHandle<CommitNode>? parentHandle, long copiedSize, IBlobHandle<DirectoryNode>? contents, List<string>? paths = null)
+			public StateNode(int number, int parentNumber, IBlobRef<CommitNode>? parentHandle, long copiedSize, IBlobRef<DirectoryNode>? contents, List<string>? paths = null)
 			{
 				Change = number;
 				ParentChange = parentNumber;
@@ -79,10 +79,10 @@ namespace Horde.Server.Replicators
 				int change = (int)reader.ReadUnsignedVarInt();
 				int parentChange = (int)reader.ReadUnsignedVarInt();
 
-				IBlobHandle<CommitNode>? parentHandle = null;
+				IBlobRef<CommitNode>? parentHandle = null;
 				if (reader.ReadBoolean())
 				{
-					parentHandle = reader.ReadBlobHandle<CommitNode>();
+					parentHandle = reader.ReadBlobRef<CommitNode>();
 				}
 
 				long copiedSize = 0;
@@ -91,10 +91,10 @@ namespace Horde.Server.Replicators
 					copiedSize = (long)reader.ReadUnsignedVarInt();
 				}
 
-				IBlobHandle<DirectoryNode>? contents = null;
+				IBlobRef<DirectoryNode>? contents = null;
 				if (reader.ReadBoolean())
 				{
-					contents = reader.ReadBlobHandle<DirectoryNode>();
+					contents = reader.ReadBlobRef<DirectoryNode>();
 				}
 
 				List<string> paths = reader.ReadList(() => reader.ReadString());
@@ -110,7 +110,7 @@ namespace Horde.Server.Replicators
 				writer.WriteBoolean(value.ParentHandle != null);
 				if (value.ParentHandle != null)
 				{
-					writer.WriteBlobHandle(value.ParentHandle);
+					writer.WriteBlobRef(value.ParentHandle);
 				}
 
 				writer.WriteUnsignedVarInt((ulong)value.CopiedSize);
@@ -118,7 +118,7 @@ namespace Horde.Server.Replicators
 				writer.WriteBoolean(value.Contents != null);
 				if (value.Contents != null)
 				{
-					writer.WriteBlobHandle(value.Contents);
+					writer.WriteBlobRef(value.Contents);
 				}
 
 				writer.WriteList(value.Paths, x => writer.WriteString(x));
@@ -425,7 +425,7 @@ namespace Horde.Server.Replicators
 			}
 			else
 			{
-				stateNode = await store.TryReadRefAsync<StateNode>(incRefName, cancellationToken: cancellationToken);
+				stateNode = await store.TryReadRefAsync<StateNode>(incRefName, options: blobOptions, cancellationToken: cancellationToken);
 				if (stateNode != null)
 				{
 					if (stateNode.Change != change)
@@ -443,10 +443,10 @@ namespace Horde.Server.Replicators
 					RedirectNode<CommitNode>? lastCommit = await store.TryReadRefAsync<RedirectNode<CommitNode>>(refName, options: blobOptions, cancellationToken: cancellationToken);
 
 					CommitNode? parent = null;
-					IBlobHandle<CommitNode>? parentHandle = lastCommit?.Target;
+					IBlobRef<CommitNode>? parentHandle = lastCommit?.Target;
 					while (parentHandle != null)
 					{
-						CommitNode parentBlob = await parentHandle.ReadBlobAsync(blobOptions, cancellationToken);
+						CommitNode parentBlob = await parentHandle.ReadBlobAsync(cancellationToken);
 						if (parentBlob.Number < change)
 						{
 							parent = parentBlob;
@@ -669,7 +669,7 @@ namespace Horde.Server.Replicators
 						{
 							UnpackUnlinkPayload(io.Payload, out string path);
 							string file = GetClientRelativePath(path, clientInfo.Client.Root);
-							await root.DeleteFileByPathAsync(file, blobOptions, cancellationToken);
+							await root.DeleteFileByPathAsync(file, cancellationToken);
 						}
 						else
 						{
@@ -719,13 +719,13 @@ namespace Horde.Server.Replicators
 
 				await leafNodeWriter.FlushAsync(cancellationToken);
 
-				await rootUpdate.WriteInteriorNodesAsync(interiorChunkWriter, options.ChunkingOptions.InteriorOptions, blobOptions, cancellationToken);
+				await rootUpdate.WriteInteriorNodesAsync(interiorChunkWriter, options.ChunkingOptions.InteriorOptions, cancellationToken);
 				await interiorChunkWriter.FlushAsync(cancellationToken);
 
-				await root.UpdateAsync(rootUpdate, directoryWriter, blobOptions, cancellationToken);
-				stateNode.Contents = await directoryWriter.WriteBlobAsync(root, blobOptions, cancellationToken);
+				await root.UpdateAsync(rootUpdate, directoryWriter, cancellationToken);
+				stateNode.Contents = await directoryWriter.WriteBlobAsync(root, cancellationToken);
 
-				IBlobHandle<StateNode> stateNodeRef = await directoryWriter.WriteBlobAsync(stateNode, blobOptions, cancellationToken);
+				IBlobRef<StateNode> stateNodeRef = await directoryWriter.WriteBlobAsync(stateNode, cancellationToken);
 				await directoryWriter.FlushAsync(cancellationToken);
 
 				await store.WriteRefTargetAsync(incRefName, stateNodeRef, cancellationToken: cancellationToken);
@@ -747,11 +747,11 @@ namespace Horde.Server.Replicators
 			DirectoryNodeRef rootRef = new DirectoryNodeRef(root.Length, stateNode.Contents!);
 
 			ChangeRecord changeRecord = await perforce.GetChangeAsync(GetChangeOptions.None, change, cancellationToken);
-			CommitNode commitNode = new CommitNode(change, stateNode.ParentHandle, changeRecord.User ?? "Unknown", null, null, null, changeRecord.Description ?? String.Empty, changeRecord.Date, rootRef, new Dictionary<Guid, IBlobHandle<object>>());
-			IBlobHandle<CommitNode> commitNodeRef = await commitWriter.WriteBlobAsync(commitNode, blobOptions, cancellationToken);
+			CommitNode commitNode = new CommitNode(change, stateNode.ParentHandle, changeRecord.User ?? "Unknown", null, null, null, changeRecord.Description ?? String.Empty, changeRecord.Date, rootRef, new Dictionary<Guid, IBlobRef>());
+			IBlobRef<CommitNode> commitNodeRef = await commitWriter.WriteBlobAsync(commitNode, cancellationToken);
 
 			RedirectNode<CommitNode> redirectNode = new RedirectNode<CommitNode>(commitNodeRef);
-			IBlobHandle<RedirectNode<CommitNode>> redirectNodeRef = await commitWriter.WriteBlobAsync(redirectNode, blobOptions, cancellationToken);
+			IBlobRef<RedirectNode<CommitNode>> redirectNodeRef = await commitWriter.WriteBlobAsync(redirectNode, cancellationToken);
 
 			await commitWriter.FlushAsync(cancellationToken);
 			await store.WriteRefTargetAsync(refName, redirectNodeRef, options.RefOptions, cancellationToken: cancellationToken);
diff --git a/Engine/Source/Programs/Horde/Horde.Server/Storage/StorageController.cs b/Engine/Source/Programs/Horde/Horde.Server/Storage/StorageController.cs
--- a/Engine/Source/Programs/Horde/Horde.Server/Storage/StorageController.cs
+++ b/Engine/Source/Programs/Horde/Horde.Server/Storage/StorageController.cs
@@ -453,7 +453,7 @@ namespace Horde.Server.Storage
 				if (commitNode.Metadata.Count > 0)
 				{
 					metadata = new Dictionary<Guid, object>();
-					foreach ((Guid blobGuid, IBlobHandle<object> handle) in commitNode.Metadata)
+					foreach ((Guid blobGuid, IBlobRef handle) in commitNode.Metadata)
 					{
 						metadata.Add(blobGuid, GetNodeHandleLink(namespaceId, handle));
 					}
@@ -464,14 +464,14 @@ namespace Horde.Server.Storage
 			else if (blobData.Type.Guid == CbNode.BlobTypeGuid)
 			{
 				CbNode cbNode = BlobSerializer.Deserialize<CbNode>(blobData);
-				content = GetCbNodeObject(namespaceId, cbNode.Object.AsField(), cbNode.References.GetEnumerator()) ?? new object();
+				content = GetCbNodeObject(namespaceId, cbNode.Object.AsField(), cbNode.Imports.GetEnumerator()) ?? new object();
 			}
 			else
 			{
 				IEnumerable<string>? references = null;
-				if (blobData.Refs.Count > 0)
+				if (blobData.Imports.Count > 0)
 				{
-					references = blobData.Refs.Select(x => GetNodeLink(namespaceId, x));
+					references = blobData.Imports.Select(x => GetNodeLink(namespaceId, x));
 				}
 				content = new { length = blobData.Data.Length, references };
 			}
@@ -485,12 +485,12 @@ namespace Horde.Server.Storage
 			return new { type = typeName, guid = blobData.Type.Guid, data = $"{GetNodeLink(namespaceId, locator)}&data=true", content = content };
 		}
 
-		static object? GetCbNodeObject(NamespaceId namespaceId, CbField field, IEnumerator<IBlobHandle> references)
+		static object? GetCbNodeObject(NamespaceId namespaceId, CbField field, IEnumerator<IBlobHandle> imports)
 		{
 			if (field.IsAttachment())
 			{
-				object? link = GetNodeLink(namespaceId, references.Current);
-				references.MoveNext();
+				object? link = GetNodeLink(namespaceId, imports.Current);
+				imports.MoveNext();
 				return link;
 			}
 			else if (field.IsObject())
@@ -500,7 +500,7 @@ namespace Horde.Server.Storage
 				CbObject obj = field.AsObject();
 				foreach (CbField member in obj)
 				{
-					fields[member.Name.ToString()] = GetCbNodeObject(namespaceId, member, references);
+					fields[member.Name.ToString()] = GetCbNodeObject(namespaceId, member, imports);
 				}
 
 				return fields;
@@ -512,7 +512,7 @@ namespace Horde.Server.Storage
 				CbArray arr = field.AsArray();
 				foreach (CbField member in arr)
 				{
-					elements.Add(GetCbNodeObject(namespaceId, member, references));
+					elements.Add(GetCbNodeObject(namespaceId, member, imports));
 				}
 
 				return elements;
@@ -527,9 +527,7 @@ namespace Horde.Server.Storage
 		static object? GetNodeObject(NamespaceId namespaceId, DirectoryNodeRef? nodeRef) => (nodeRef == null) ? null : new { nodeRef.Length, nodeRef.Handle.Hash, link = GetNodeLink(namespaceId, nodeRef.Handle.GetLocator()) };
 
 		[return: NotNullIfNotNull("handle")]
-		static object? GetNodeHandleLink<T>(NamespaceId namespaceId, IBlobHandle<T>? handle) => (handle == null) ? null : GetNodeHandleLink(namespaceId, handle.Hash, handle);
-
-		static object GetNodeHandleLink(NamespaceId namespaceId, IoHash hash, IBlobHandle handle) => new { hash, link = GetNodeLink(namespaceId, handle.GetLocator()) };
+		static object? GetNodeHandleLink(NamespaceId namespaceId, IBlobRef? handle) => (handle == null)? null : new { handle.Hash, link = GetNodeLink(namespaceId, handle.GetLocator()) };
 
 		static string GetNodeLink(NamespaceId namespaceId, IBlobHandle handle) => GetNodeLink(namespaceId, handle.GetLocator());
 		
diff --git a/Engine/Source/Programs/Horde/Horde.Server/Tools/ToolCollection.cs b/Engine/Source/Programs/Horde/Horde.Server/Tools/ToolCollection.cs
--- a/Engine/Source/Programs/Horde/Horde.Server/Tools/ToolCollection.cs
+++ b/Engine/Source/Programs/Horde/Horde.Server/Tools/ToolCollection.cs
@@ -256,7 +256,7 @@ namespace Horde.Server.Tools
 
 			using IStorageClient client = _storageService.CreateClient(tool.Config.NamespaceId);
 
-			IBlobHandle<DirectoryNode> nodeRef;
+			IBlobRef<DirectoryNode> nodeRef;
 			await using (IBlobWriter writer = client.CreateBlobWriter(refName))
 			{
 				DirectoryNode directoryNode = new DirectoryNode();
@@ -264,7 +264,7 @@ namespace Horde.Server.Tools
 				nodeRef = await writer.WriteBlobAsync(directoryNode, cancellationToken: cancellationToken);
 			}
 
-			IBlobHandle target = nodeRef;
+			IBlobRef target = nodeRef;
 			await client.WriteRefTargetAsync(refName, target, cancellationToken: cancellationToken);
 
 			return await CreateDeploymentAsync(tool, options, target.GetLocator(), globalConfig, cancellationToken);
diff --git a/Engine/Source/Programs/Horde/Horde/Commands/Bundles/BundleCreate.cs b/Engine/Source/Programs/Horde/Horde/Commands/Bundles/BundleCreate.cs
--- a/Engine/Source/Programs/Horde/Horde/Commands/Bundles/BundleCreate.cs
+++ b/Engine/Source/Programs/Horde/Horde/Commands/Bundles/BundleCreate.cs
@@ -88,7 +88,7 @@ namespace Horde.Commands.Bundles
 				List<FileInfo> fileInfos = files.ConvertAll(x => x.ToFileInfo());
 				CopyStatsLogger copyStatsLogger = new CopyStatsLogger(files.Count, fileInfos.Sum(x => x.Length), logger);
 
-				IBlobHandle<DirectoryNode> nodeRef = await writer.WriteFilesAsync(baseDir.ToDirectoryInfo(), fileInfos, options, copyStatsLogger, CancellationToken.None);
+				IBlobRef<DirectoryNode> nodeRef = await writer.WriteFilesAsync(baseDir.ToDirectoryInfo(), fileInfos, options, copyStatsLogger, CancellationToken.None);
 
 				await writer.FlushAsync();
 
diff --git a/Engine/Source/Programs/Horde/Horde/Commands/Bundles/BundleExtract.cs b/Engine/Source/Programs/Horde/Horde/Commands/Bundles/BundleExtract.cs
--- a/Engine/Source/Programs/Horde/Horde/Commands/Bundles/BundleExtract.cs
+++ b/Engine/Source/Programs/Horde/Horde/Commands/Bundles/BundleExtract.cs
@@ -84,7 +84,7 @@ namespace Horde.Commands.Bundles
 			Stopwatch timer = Stopwatch.StartNew();
 
 			DirectoryNode node = await handle.ReadBlobAsync<DirectoryNode>();
-			await node.CopyToDirectoryAsync(OutputDir.ToDirectoryInfo(), new CopyStatsLogger(logger), null, logger, CancellationToken.None);
+			await node.CopyToDirectoryAsync(OutputDir.ToDirectoryInfo(), new CopyStatsLogger(logger), logger, CancellationToken.None);
 
 			logger.LogInformation("Elapsed: {Time}s", timer.Elapsed.TotalSeconds);
 
diff --git a/Engine/Source/Programs/Horde/Horde/Commands/Compute/ComputeRun.cs b/Engine/Source/Programs/Horde/Horde/Commands/Compute/ComputeRun.cs
--- a/Engine/Source/Programs/Horde/Horde/Commands/Compute/ComputeRun.cs
+++ b/Engine/Source/Programs/Horde/Horde/Commands/Compute/ComputeRun.cs
@@ -159,7 +159,7 @@ namespace Horde.Commands.Compute
 		{
 			await using IBlobWriter writer = storage.CreateBlobWriter();
 
-			IBlobHandle<DirectoryNode> sandbox = await writer.WriteFilesAsync(taskFile.Directory, cancellationToken: cancellationToken);
+			IBlobRef<DirectoryNode> sandbox = await writer.WriteFilesAsync(taskFile.Directory, cancellationToken: cancellationToken);
 			await writer.FlushAsync(cancellationToken);
 
 			return sandbox.GetLocator();
diff --git a/Engine/Source/Programs/Horde/Horde/Commands/Tools/ToolUpload.cs b/Engine/Source/Programs/Horde/Horde/Commands/Tools/ToolUpload.cs
--- a/Engine/Source/Programs/Horde/Horde/Commands/Tools/ToolUpload.cs
+++ b/Engine/Source/Programs/Horde/Horde/Commands/Tools/ToolUpload.cs
@@ -40,7 +40,7 @@ namespace Horde.Commands
 		{
 			using IStorageClient storageClient = _httpStorageClientFactory.CreateClientWithPath($"api/v1/tools/{ToolId}");
 
-			IBlobHandle<DirectoryNode> target;
+			IBlobRef<DirectoryNode> target;
 			await using (IBlobWriter writer = storageClient.CreateBlobWriter())
 			{
 				target = await writer.WriteFilesAsync(InputDir);
diff --git a/Engine/Source/Programs/Horde/Samples/RemoteClient/RemoteClientApp.cs b/Engine/Source/Programs/Horde/Samples/RemoteClient/RemoteClientApp.cs
--- a/Engine/Source/Programs/Horde/Samples/RemoteClient/RemoteClientApp.cs
+++ b/Engine/Source/Programs/Horde/Samples/RemoteClient/RemoteClientApp.cs
@@ -114,7 +114,7 @@ namespace RemoteClient
 
 				await using (IBlobWriter writer = storage.CreateBlobWriter())
 				{
-					IBlobHandle<DirectoryNode> sandbox = await writer.WriteFilesAsync(uploadDir);
+					IBlobRef<DirectoryNode> sandbox = await writer.WriteFilesAsync(uploadDir);
 					await writer.FlushAsync();
 					await channel.UploadFilesAsync("", sandbox.GetLocator(), storage.Backend);
 				}
diff --git a/Engine/Source/Programs/Shared/EpicGames.Horde.Tests/BlobDataTests.cs b/Engine/Source/Programs/Shared/EpicGames.Horde.Tests/BlobDataTests.cs
--- a/Engine/Source/Programs/Shared/EpicGames.Horde.Tests/BlobDataTests.cs
+++ b/Engine/Source/Programs/Shared/EpicGames.Horde.Tests/BlobDataTests.cs
@@ -17,6 +17,8 @@ namespace EpicGames.Horde.Tests
 		{
 			readonly BlobLocator _locator;
 
+			public IBlobHandle Innermost => this;
+
 			public DummyHandle(string locator) => _locator = new BlobLocator(new Utf8String(locator));
 			public ValueTask FlushAsync(CancellationToken cancellationToken = default) => default;
 			public ValueTask<BlobData> ReadBlobDataAsync(CancellationToken cancellationToken = default) => throw new NotImplementedException();
@@ -46,12 +48,12 @@ namespace EpicGames.Horde.Tests
 
 			Assert.AreEqual(blobData.Type.Guid, packet.Type.Guid);
 			Assert.AreEqual(blobData.Type.Version, packet.Type.Version);
-			Assert.AreEqual(blobData.Refs.Count, packet.Refs.Count);
+			Assert.AreEqual(blobData.Imports.Count, packet.Imports.Count);
 
-			for (int idx = 0; idx < packet.Refs.Count; idx++)
+			for (int idx = 0; idx < packet.Imports.Count; idx++)
 			{
-				string locator = blobData.Refs[idx].GetLocator().ToString();
-				string encoded = packet.Refs[idx].ToString();
+				string locator = blobData.Imports[idx].GetLocator().ToString();
+				string encoded = packet.Imports[idx].ToString();
 				Assert.AreEqual(locator, encoded);
 			}
 		}
diff --git a/Engine/Source/Programs/Shared/EpicGames.Horde.Tests/BundleStoreTests.cs b/Engine/Source/Programs/Shared/EpicGames.Horde.Tests/BundleStoreTests.cs
--- a/Engine/Source/Programs/Shared/EpicGames.Horde.Tests/BundleStoreTests.cs
+++ b/Engine/Source/Programs/Shared/EpicGames.Horde.Tests/BundleStoreTests.cs
@@ -62,10 +62,10 @@ namespace EpicGames.Horde.Tests
 		{
 			MemoryStorageBackend memoryStore = new MemoryStorageBackend();
 			using BundleStorageClient store = new BundleStorageClient(memoryStore, BundleCache.None, NullLogger.Instance);
-			await using IBlobWriter writer = store.CreateBlobWriter(options: new BundleOptions { MaxVersion = BundleVersion.ImportHashes, CompressionFormat = BundleCompressionFormat.None });
+			await using IBlobWriter writer = store.CreateBlobWriter(bundleOptions: new BundleOptions { MaxVersion = BundleVersion.ImportHashes, CompressionFormat = BundleCompressionFormat.None });
 
 			TextNode node = new TextNode("Hello world");
-			IBlobHandle<TextNode> handle = await writer.WriteBlobAsync(node);
+			IBlobRef<TextNode> handle = await writer.WriteBlobAsync(node);
 			await writer.FlushAsync();
 
 			IReadOnlyMemoryOwner<byte> owner = await memoryStore.ReadBlobAsync(handle.GetLocator().BaseLocator, 0, null);
@@ -125,9 +125,9 @@ namespace EpicGames.Horde.Tests
 			public static BlobType BlobType { get; } = new BlobType("{F63606D4-4061-5DBB-446F-55A69F22654F}", 1);
 
 			public ReadOnlySequence<byte> Data { get; }
-			public IReadOnlyList<IBlobHandle<SimpleNode>> Refs { get; }
+			public IReadOnlyList<IBlobRef<SimpleNode>> Refs { get; }
 
-			public SimpleNode(ReadOnlySequence<byte> data, IReadOnlyList<IBlobHandle<SimpleNode>> refs)
+			public SimpleNode(ReadOnlySequence<byte> data, IReadOnlyList<IBlobRef<SimpleNode>> refs)
 			{
 				Data = data;
 				Refs = refs;
@@ -139,14 +139,14 @@ namespace EpicGames.Horde.Tests
 			public override SimpleNode Read(IBlobReader reader, BlobSerializerOptions options)
 			{
 				ReadOnlyMemory<byte> data = reader.ReadVariableLengthBytes();
-				IReadOnlyList<IBlobHandle<SimpleNode>> refs = reader.ReadVariableLengthArray(() => reader.ReadBlobHandle<SimpleNode>());
+				IReadOnlyList<IBlobRef<SimpleNode>> refs = reader.ReadVariableLengthArray(() => reader.ReadBlobRef<SimpleNode>());
 				return new SimpleNode(new ReadOnlySequence<byte>(data), refs);
 			}
 
 			public override BlobType Write(IBlobWriter writer, SimpleNode value, BlobSerializerOptions options)
 			{
 				writer.WriteVariableLengthBytes(value.Data);
-				writer.WriteVariableLengthArray(value.Refs, x => writer.WriteBlobHandle(x));
+				writer.WriteVariableLengthArray(value.Refs, x => writer.WriteBlobRef(x));
 				return SimpleNode.BlobType;
 			}
 		}
@@ -157,13 +157,13 @@ namespace EpicGames.Horde.Tests
 			{
 				await using IBlobWriter writer = store.CreateBlobWriter("test", options);
 
-				SimpleNode node1 = new SimpleNode(new ReadOnlySequence<byte>(new byte[] { 1 }), Array.Empty<IBlobHandle<SimpleNode>>());
+				SimpleNode node1 = new SimpleNode(new ReadOnlySequence<byte>(new byte[] { 1 }), Array.Empty<IBlobRef<SimpleNode>>());
 				SimpleNode node2 = new SimpleNode(new ReadOnlySequence<byte>(new byte[] { 2 }), new[] { await writer.WriteBlobAsync(node1) });
 				SimpleNode node3 = new SimpleNode(new ReadOnlySequence<byte>(new byte[] { 3 }), new[] { await writer.WriteBlobAsync(node2) });
-				SimpleNode node4 = new SimpleNode(new ReadOnlySequence<byte>(new byte[] { 4 }), Array.Empty<IBlobHandle<SimpleNode>>());
+				SimpleNode node4 = new SimpleNode(new ReadOnlySequence<byte>(new byte[] { 4 }), Array.Empty<IBlobRef<SimpleNode>>());
 
 				SimpleNode root = new SimpleNode(new ReadOnlySequence<byte>(new byte[] { 5 }), new[] { await writer.WriteBlobAsync(node4), await writer.WriteBlobAsync(node3) });
-				IBlobHandle<SimpleNode> rootRef = await writer.WriteBlobAsync(root);
+				IBlobRef<SimpleNode> rootRef = await writer.WriteBlobAsync(root);
 				await writer.FlushAsync();
 
 				await store.WriteRefTargetAsync(new RefName("test"), rootRef);
@@ -184,31 +184,31 @@ namespace EpicGames.Horde.Tests
 			SimpleNode node5 = root;
 			byte[] data5 = node5.Data.ToArray();
 			Assert.IsTrue(data5.SequenceEqual(new byte[] { 5 }));
-			IReadOnlyList<IBlobHandle<SimpleNode>> refs5 = node5.Refs;
+			IReadOnlyList<IBlobRef<SimpleNode>> refs5 = node5.Refs;
 			Assert.AreEqual(2, refs5.Count);
 
 			SimpleNode node4 = await refs5[0].ReadBlobAsync();
 			byte[] data4 = node4.Data.ToArray();
 			Assert.IsTrue(data4.SequenceEqual(new byte[] { 4 }));
-			IReadOnlyList<IBlobHandle<SimpleNode>> refs4 = node4.Refs;
+			IReadOnlyList<IBlobRef<SimpleNode>> refs4 = node4.Refs;
 			Assert.AreEqual(0, refs4.Count);
 
 			SimpleNode node3 = await refs5[1].ReadBlobAsync();
 			byte[] data3 = node3.Data.ToArray();
 			Assert.IsTrue(data3.SequenceEqual(new byte[] { 3 }));
-			IReadOnlyList<IBlobHandle<SimpleNode>> refs3 = node3.Refs;
+			IReadOnlyList<IBlobRef<SimpleNode>> refs3 = node3.Refs;
 			Assert.AreEqual(1, refs3.Count);
 
 			SimpleNode node2 = await refs3[0].ReadBlobAsync();
 			byte[] data2 = node2.Data.ToArray();
 			Assert.IsTrue(data2.SequenceEqual(new byte[] { 2 }));
-			IReadOnlyList<IBlobHandle<SimpleNode>> refs2 = node2.Refs;
+			IReadOnlyList<IBlobRef<SimpleNode>> refs2 = node2.Refs;
 			Assert.AreEqual(1, refs2.Count);
 
 			SimpleNode node1 = await refs2[0].ReadBlobAsync();
 			byte[] data1 = node1.Data.ToArray();
 			Assert.IsTrue(data1.SequenceEqual(new byte[] { 1 }));
-			IReadOnlyList<IBlobHandle<SimpleNode>> refs1 = node1.Refs;
+			IReadOnlyList<IBlobRef<SimpleNode>> refs1 = node1.Refs;
 			Assert.AreEqual(0, refs1.Count);
 		}
 
@@ -219,10 +219,10 @@ namespace EpicGames.Horde.Tests
 
 			RefName refName = new RefName("test");
 
-			IBlobHandle<SimpleNode> inputRef;
+			IBlobRef<SimpleNode> inputRef;
 			await using (IBlobWriter writer = store.CreateBlobWriter(refName))
 			{
-				inputRef = await writer.WriteBlobAsync(new SimpleNode(new ReadOnlySequence<byte>(new byte[] { (byte)123 }), Array.Empty<IBlobHandle<SimpleNode>>()));
+				inputRef = await writer.WriteBlobAsync(new SimpleNode(new ReadOnlySequence<byte>(new byte[] { (byte)123 }), Array.Empty<IBlobRef<SimpleNode>>()));
 			}
 			await store.WriteRefAsync(refName, inputRef);
 
@@ -241,15 +241,15 @@ namespace EpicGames.Horde.Tests
 				await using (IBlobWriter writer = store.CreateBlobWriter(new RefName("test")))
 				{
 					DirectoryNode world = new DirectoryNode();
-					IBlobHandle<DirectoryNode> worldRef = await writer.WriteBlobAsync(world);
+					IBlobRef<DirectoryNode> worldRef = await writer.WriteBlobAsync(world);
 
 					DirectoryNode hello = new DirectoryNode();
 					hello.AddDirectory(new DirectoryEntry("world", 0, worldRef));
-					IBlobHandle<DirectoryNode> helloRef = await writer.WriteBlobAsync(hello);
+					IBlobRef<DirectoryNode> helloRef = await writer.WriteBlobAsync(hello);
 
 					DirectoryNode root = new DirectoryNode(DirectoryFlags.None);
 					root.AddDirectory(new DirectoryEntry("hello", 0, helloRef));
-					IBlobHandle<DirectoryNode> rootRef = await writer.WriteBlobAsync(root);
+					IBlobRef<DirectoryNode> rootRef = await writer.WriteBlobAsync(root);
 
 					await writer.FlushAsync();
 
@@ -288,7 +288,7 @@ namespace EpicGames.Horde.Tests
 			{
 				await using IBlobWriter writer = store.CreateBlobWriter();
 
-				using ChunkedDataWriter fileWriter = new ChunkedDataWriter(writer, new ChunkingOptions(), BlobSerializerOptions.Default);
+				using ChunkedDataWriter fileWriter = new ChunkedDataWriter(writer, new ChunkingOptions());
 				ChunkedData fileHandle = await fileWriter.CreateAsync(Encoding.UTF8.GetBytes("world"), CancellationToken.None);
 
 				List<FileUpdate> fileUpdates = new List<FileUpdate>();
@@ -297,7 +297,7 @@ namespace EpicGames.Horde.Tests
 				DirectoryNode root = new DirectoryNode();
 				await root.UpdateAsync(fileUpdates, writer);
 
-				IBlobHandle<DirectoryNode> rootRef = await writer.WriteBlobAsync(root);
+				IBlobRef<DirectoryNode> rootRef = await writer.WriteBlobAsync(root);
 				await store.WriteRefTargetAsync(new RefName("test"), rootRef);
 
 				await CheckFileTreeAsync(root);
@@ -339,12 +339,12 @@ namespace EpicGames.Horde.Tests
 			// Generate a tree
 			ChunkedDataNodeRef nodeRef;
 			{
-				await using IBlobWriter writer = store.CreateBlobWriter(options: new BundleOptions { MaxBlobSize = 1024 });
+				await using IBlobWriter writer = store.CreateBlobWriter(bundleOptions: new BundleOptions { MaxBlobSize = 1024 });
 
 				ChunkingOptions options = new ChunkingOptions();
 				options.LeafOptions = new LeafChunkedDataNodeOptions(128, 256, 64 * 1024);
 
-				using ChunkedDataWriter fileWriter = new ChunkedDataWriter(writer, options, BlobSerializerOptions.Default);
+				using ChunkedDataWriter fileWriter = new ChunkedDataWriter(writer, options);
 				for (int idx = 0; idx < chunk.Length / 16; idx++)
 				{
 					await fileWriter.AppendAsync(chunk.AsMemory(idx * 16, 16), CancellationToken.None);
@@ -385,18 +385,18 @@ namespace EpicGames.Horde.Tests
 			// Generate a tree
 			DirectoryNode root;
 			{
-				await using DedupeStorageWriter writer = new DedupeStorageWriter(store.CreateBlobWriter(options: new BundleOptions { MaxBlobSize = 1024 }));
+				await using DedupeStorageWriter writer = new DedupeStorageWriter(store.CreateBlobWriter(bundleOptions: new BundleOptions { MaxBlobSize = 1024 }));
 
 				ChunkingOptions options = new ChunkingOptions();
 				options.LeafOptions = new LeafChunkedDataNodeOptions(128, 256, 64 * 1024);
 
-				using ChunkedDataWriter fileWriter = new ChunkedDataWriter(writer, options, BlobSerializerOptions.Default);
+				using ChunkedDataWriter fileWriter = new ChunkedDataWriter(writer, options);
 				ChunkedData chunkedData = await fileWriter.CreateAsync(data, CancellationToken.None);
 
 				root = new DirectoryNode(DirectoryFlags.None);
 				root.AddFile("test", FileEntryFlags.None, fileWriter.Length, chunkedData);
 
-				IBlobHandle<DirectoryNode> rootRef = await writer.WriteBlobAsync(root);
+				IBlobRef<DirectoryNode> rootRef = await writer.WriteBlobAsync(root);
 				await store.WriteRefTargetAsync(new RefName("test"), rootRef);
 
 				await CheckLargeFileTreeAsync(root, data);
@@ -423,9 +423,9 @@ namespace EpicGames.Horde.Tests
 			using BlobData data = await handle.ReadBlobDataAsync();
 			locatorToSize[handle.GetLocator()] = data.Data.Length;
 
-			foreach (IBlobHandle reference in data.Refs)
+			foreach (IBlobHandle import in data.Imports)
 			{
-				await GetUniqueBlobsAsync(reference, locatorToSize);
+				await GetUniqueBlobsAsync(import, locatorToSize);
 			}
 		}
 
diff --git a/Engine/Source/Programs/Shared/EpicGames.Horde.Tests/BundleTests.cs b/Engine/Source/Programs/Shared/EpicGames.Horde.Tests/BundleTests.cs
--- a/Engine/Source/Programs/Shared/EpicGames.Horde.Tests/BundleTests.cs
+++ b/Engine/Source/Programs/Shared/EpicGames.Horde.Tests/BundleTests.cs
@@ -76,7 +76,7 @@ namespace EpicGames.Horde.Tests
 			ChunkingOptions options = new ChunkingOptions();
 			options.LeafOptions = new LeafChunkedDataNodeOptions(8, 8, 8);
 
-			using ChunkedDataWriter fileNodeWriter = new ChunkedDataWriter(writer, options, BlobSerializerOptions.Default);
+			using ChunkedDataWriter fileNodeWriter = new ChunkedDataWriter(writer, options);
 
 			ChunkedDataNode node;
 			ChunkedDataNodeRef nodeRef;
@@ -193,15 +193,15 @@ namespace EpicGames.Horde.Tests
 		{
 			IStorageClient store = _storage;
 
-			IBlobHandle<DirectoryNode> rootRef;
+			IBlobRef<DirectoryNode> rootRef;
 			await using (IBlobWriter writer = store.CreateBlobWriter())
 			{
 				DirectoryNode world = new DirectoryNode();
-				IBlobHandle<DirectoryNode> worldRef = await writer.WriteBlobAsync(world);
+				IBlobRef<DirectoryNode> worldRef = await writer.WriteBlobAsync(world);
 
 				DirectoryNode hello = new DirectoryNode();
 				hello.AddDirectory(new DirectoryEntry("world", 0, worldRef));
-				IBlobHandle<DirectoryNode> helloRef = await writer.WriteBlobAsync(hello);
+				IBlobRef<DirectoryNode> helloRef = await writer.WriteBlobAsync(hello);
 
 				DirectoryNode root = new DirectoryNode();
 				root.AddDirectory(new DirectoryEntry("hello", 0, helloRef));
@@ -250,21 +250,21 @@ namespace EpicGames.Horde.Tests
 			Assert.AreEqual(0, _memoryStore.Refs.Count);
 			Assert.AreEqual(0, _memoryStore.Blobs.Count);
 
-			BundleOptions options = new BundleOptions();
-			options.MaxBlobSize = 1;
+			BundleOptions bundleOptions = new BundleOptions();
+			bundleOptions.MaxBlobSize = 1;
 
-			await using (IBlobWriter writer = new DedupeStorageWriter(_storage.CreateBlobWriter(options: options)))
+			await using (IBlobWriter writer = new DedupeStorageWriter(_storage.CreateBlobWriter(bundleOptions: bundleOptions)))
 			{
 				DirectoryNode root = new DirectoryNode();
 				for (int idx = 1; idx <= 3; idx++)
 				{
 					DirectoryNode node = new DirectoryNode();
-					IBlobHandle<DirectoryNode> nodeRef = await writer.WriteBlobAsync(node);
+					IBlobRef<DirectoryNode> nodeRef = await writer.WriteBlobAsync(node);
 					root.AddDirectory(new DirectoryEntry($"node{idx}", 0, nodeRef));
 				}
 
 				RefName refName = new RefName("ref");
-				IBlobHandle<DirectoryNode> rootRef = await writer.WriteBlobAsync(root);
+				IBlobRef<DirectoryNode> rootRef = await writer.WriteBlobAsync(root);
 				await _storage.WriteRefTargetAsync(refName, rootRef);
 			}
 
@@ -275,15 +275,15 @@ namespace EpicGames.Horde.Tests
 		[TestMethod]
 		public async Task ReloadTestsAsync()
 		{
-			BundleOptions options = new BundleOptions();
-			options.MaxBlobSize = 1;
+			BundleOptions bundleOptions = new BundleOptions();
+			bundleOptions.MaxBlobSize = 1;
 
 			RefName refName = new RefName("ref");
 
 			{
-				await using (IBlobWriter writer = _storage.CreateBlobWriter(options: options))
+				await using (IBlobWriter writer = _storage.CreateBlobWriter(bundleOptions: bundleOptions))
 				{
-					IBlobHandle<DirectoryNode> rootRef = await writer.WriteBlobAsync(new DirectoryNode());
+					IBlobRef<DirectoryNode> rootRef = await writer.WriteBlobAsync(new DirectoryNode());
 					for (int idx = 4; idx >= 1; idx--)
 					{
 						DirectoryNode next = new DirectoryNode();
diff --git a/Engine/Source/Programs/Shared/EpicGames.Horde.Tests/ChunkingTests.cs b/Engine/Source/Programs/Shared/EpicGames.Horde.Tests/ChunkingTests.cs
--- a/Engine/Source/Programs/Shared/EpicGames.Horde.Tests/ChunkingTests.cs
+++ b/Engine/Source/Programs/Shared/EpicGames.Horde.Tests/ChunkingTests.cs
@@ -62,7 +62,7 @@ namespace EpicGames.Horde.Tests
 
 				using MemoryStream emptyStream = new MemoryStream();
 				LeafChunkedData leafChunkedData = await LeafChunkedDataNode.CreateFromStreamAsync(writer, emptyStream, new LeafChunkedDataNodeOptions(64, 64, 64), CancellationToken.None);
-				ChunkedData chunkedData = await InteriorChunkedDataNode.CreateTreeAsync(leafChunkedData, new InteriorChunkedDataNodeOptions(4, 4, 4), writer, BlobSerializerOptions.Default, CancellationToken.None); 
+				ChunkedData chunkedData = await InteriorChunkedDataNode.CreateTreeAsync(leafChunkedData, new InteriorChunkedDataNodeOptions(4, 4, 4), writer, CancellationToken.None); 
 
 				DirectoryNode directory = new DirectoryNode();
 				directory.AddFile("test.foo", FileEntryFlags.None, 0, chunkedData);
@@ -112,7 +112,7 @@ namespace EpicGames.Horde.Tests
 
 			const int NumIterations = 100;
 			{
-				using ChunkedDataWriter fileWriter = new ChunkedDataWriter(writer, options, BlobSerializerOptions.Default);
+				using ChunkedDataWriter fileWriter = new ChunkedDataWriter(writer, options);
 
 				for (int idx = 0; idx < NumIterations; idx++)
 				{
@@ -179,10 +179,10 @@ namespace EpicGames.Horde.Tests
 
 			byte[] data = Encoding.UTF8.GetBytes("hello world");
 
-			IBlobHandle<DirectoryNode> handle;
-			await using (IBlobWriter writer = store.CreateBlobWriter())
+			IBlobRef<DirectoryNode> handle;
+			await using (IBlobWriter writer = store.CreateBlobWriter(options: serializerOptions))
 			{
-				using ChunkedDataWriter chunkedWriter = new ChunkedDataWriter(writer, chunkingOptions, serializerOptions);
+				using ChunkedDataWriter chunkedWriter = new ChunkedDataWriter(writer, chunkingOptions);
 				await chunkedWriter.AppendAsync(data, CancellationToken.None);
 				ChunkedData chunkedData = await chunkedWriter.FlushAsync(CancellationToken.None);
 
diff --git a/Engine/Source/Programs/Shared/EpicGames.Horde.Tests/ComputeSocketTests.cs b/Engine/Source/Programs/Shared/EpicGames.Horde.Tests/ComputeSocketTests.cs
--- a/Engine/Source/Programs/Shared/EpicGames.Horde.Tests/ComputeSocketTests.cs
+++ b/Engine/Source/Programs/Shared/EpicGames.Horde.Tests/ComputeSocketTests.cs
@@ -185,18 +185,18 @@ namespace EpicGames.Horde.Tests
 
 						byte[] data = Encoding.UTF8.GetBytes("Hello world");
 
-						using ChunkedDataWriter writer = new ChunkedDataWriter(blobWriter, new ChunkingOptions(), BlobSerializerOptions.Default);
+						using ChunkedDataWriter writer = new ChunkedDataWriter(blobWriter, new ChunkingOptions());
 						ChunkedData chunkedData = await writer.CreateAsync(data, cancellationToken);
 
 						DirectoryNode directory = new DirectoryNode();
 						directory.AddFile("hello.txt", FileEntryFlags.None, data.Length, chunkedData);
 
-						IBlobHandle<DirectoryNode> directoryRef = await blobWriter.WriteBlobAsync(directory, cancellationToken: cancellationToken);
+						IBlobRef<DirectoryNode> directoryRef = await blobWriter.WriteBlobAsync(directory, cancellationToken: cancellationToken);
 
 						DirectoryNode root = new DirectoryNode();
 						root.AddDirectory(new DirectoryEntry("subdir", directory.Length, directoryRef));
 
-						IBlobHandle<DirectoryNode> handle = await blobWriter.WriteBlobAsync(root, cancellationToken: cancellationToken);
+						IBlobRef<DirectoryNode> handle = await blobWriter.WriteBlobAsync(root, cancellationToken: cancellationToken);
 						await blobWriter.FlushAsync(cancellationToken);
 
 						await channel.UploadFilesAsync("", handle.GetLocator(), storage.Backend, cancellationToken);
diff --git a/Engine/Source/Programs/Shared/EpicGames.Horde.Tests/LogTests.cs b/Engine/Source/Programs/Shared/EpicGames.Horde.Tests/LogTests.cs
--- a/Engine/Source/Programs/Shared/EpicGames.Horde.Tests/LogTests.cs
+++ b/Engine/Source/Programs/Shared/EpicGames.Horde.Tests/LogTests.cs
@@ -69,10 +69,10 @@ namespace EpicGames.Horde.Tests
 			}
 
 			// Flush it to storage, and read the finished log node
-			IBlobHandle<LogNode> logRef;
+			IBlobRef<LogNode> logRef;
 			await using (IBlobWriter writer = store.CreateBlobWriter())
 			{
-				logRef = await builder.FlushAsync(writer, true, null, CancellationToken.None);
+				logRef = await builder.FlushAsync(writer, true, CancellationToken.None);
 			}
 			LogNode log = await logRef.ReadBlobAsync();
 
@@ -118,10 +118,10 @@ namespace EpicGames.Horde.Tests
 				builder.WriteData(Encoding.UTF8.GetBytes(lines[lineIdx]));
 			}
 
-			IBlobHandle<LogNode> rootNodeRef;
+			IBlobRef<LogNode> rootNodeRef;
 			await using (IBlobWriter writer = store.CreateBlobWriter())
 			{
-				rootNodeRef = await builder.FlushAsync(writer, true, null, CancellationToken.None);
+				rootNodeRef = await builder.FlushAsync(writer, true, CancellationToken.None);
 			}
 
 			// Read it back in and test the index
@@ -205,7 +205,7 @@ namespace EpicGames.Horde.Tests
 		static async Task SearchLogDataTestAsync(LogIndexNode index, string text, int firstLine, int count, int[] expectedLines)
 		{
 			SearchStats stats = new SearchStats();
-			List<int> lines = await index.SearchAsync(firstLine, new SearchTerm(text), stats, null, CancellationToken.None).Take(count).ToListAsync();
+			List<int> lines = await index.SearchAsync(firstLine, new SearchTerm(text), stats, CancellationToken.None).Take(count).ToListAsync();
 			Assert.IsTrue(lines.SequenceEqual(expectedLines));
 		}
 	}
diff --git a/Engine/Source/Programs/Shared/EpicGames.Horde.Tests/StorageClientTests.cs b/Engine/Source/Programs/Shared/EpicGames.Horde.Tests/StorageClientTests.cs
--- a/Engine/Source/Programs/Shared/EpicGames.Horde.Tests/StorageClientTests.cs
+++ b/Engine/Source/Programs/Shared/EpicGames.Horde.Tests/StorageClientTests.cs
@@ -22,15 +22,15 @@ namespace EpicGames.Horde.Tests
 
 			public int Value { get; }
 			public byte[] Padding { get; set; } = Array.Empty<byte>();
-			public IBlobHandle<TestNode>[] Refs { get; }
+			public IBlobRef<TestNode>[] Refs { get; }
 
-			public TestNode(int value, params IBlobHandle<TestNode>[] refs)
+			public TestNode(int value, params IBlobRef<TestNode>[] refs)
 			{
 				Value = value;
 				Refs = refs;
 			}
 
-			public TestNode(int value, byte[] padding, IBlobHandle<TestNode>[] refs)
+			public TestNode(int value, byte[] padding, IBlobRef<TestNode>[] refs)
 			{
 				Value = value;
 				Padding = padding;
@@ -44,7 +44,7 @@ namespace EpicGames.Horde.Tests
 			{
 				int value = reader.ReadInt32();
 				byte[] padding = reader.ReadVariableLengthBytes().ToArray();
-				IBlobHandle<TestNode>[] refs = reader.ReadVariableLengthArray(() => reader.ReadBlobHandle<TestNode>());
+				IBlobRef<TestNode>[] refs = reader.ReadVariableLengthArray(() => reader.ReadBlobRef<TestNode>());
 
 				return new TestNode(value, padding, refs);
 			}
@@ -53,7 +53,7 @@ namespace EpicGames.Horde.Tests
 			{
 				writer.WriteInt32(value.Value);
 				writer.WriteVariableLengthBytes(value.Padding);
-				writer.WriteVariableLengthArray(value.Refs, x => writer.WriteBlobHandle(x));
+				writer.WriteVariableLengthArray(value.Refs, x => writer.WriteBlobRef(x));
 
 				return TestNode.BlobType;
 			}
@@ -75,7 +75,7 @@ namespace EpicGames.Horde.Tests
 
 		static async Task TestBasicAsync(IStorageClient store)
 		{
-			IBlobHandle<TestNode> nodeRef;
+			IBlobRef<TestNode> nodeRef;
 			await using (IBlobWriter writer = store.CreateBlobWriter())
 			{
 				nodeRef = await writer.WriteBlobAsync(new TestNode(123));
@@ -93,10 +93,10 @@ namespace EpicGames.Horde.Tests
 
 			using BundleStorageClient store = BundleStorageClient.CreateInMemory(NullLogger.Instance);
 
-			IBlobHandle<TestNode> nodeRef2;
+			IBlobRef<TestNode> nodeRef2;
 			await using (IBlobWriter writer = store.CreateBlobWriter())
 			{
-				IBlobHandle<TestNode> nodeRef1 = await writer.WriteBlobAsync(new TestNode(123));
+				IBlobRef<TestNode> nodeRef1 = await writer.WriteBlobAsync(new TestNode(123));
 				nodeRef2 = await writer.WriteBlobAsync(new TestNode(456, nodeRef1));
 			}
 			await store.WriteRefAsync("hello", nodeRef2);
@@ -117,14 +117,14 @@ namespace EpicGames.Horde.Tests
 
 			using BundleStorageClient store = BundleStorageClient.CreateInMemory(NullLogger.Instance);
 
-			await using IBlobWriter writer = store.CreateBlobWriter(options: new BundleOptions { MinCompressionPacketSize = 100, MaxBlobSize = 1024 * 1024, MaxVersion = BundleVersion.LatestV2 });
-			IBlobHandle<TestNode> nodeRef1 = await writer.WriteBlobAsync(new TestNode(123) { Padding = new byte[1024] });
+			await using IBlobWriter writer = store.CreateBlobWriter(bundleOptions: new BundleOptions { MinCompressionPacketSize = 100, MaxBlobSize = 1024 * 1024, MaxVersion = BundleVersion.LatestV2 });
+			IBlobRef<TestNode> nodeRef1 = await writer.WriteBlobAsync(new TestNode(123) { Padding = new byte[1024] });
 			await writer.FlushAsync();
-			IBlobHandle<TestNode> nodeRef2 = await writer.WriteBlobAsync(new TestNode(456, nodeRef1) { Padding = new byte[1024] });
-			IBlobHandle<TestNode> nodeRef3 = await writer.WriteBlobAsync(new TestNode(789, nodeRef2));
+			IBlobRef<TestNode> nodeRef2 = await writer.WriteBlobAsync(new TestNode(456, nodeRef1) { Padding = new byte[1024] });
+			IBlobRef<TestNode> nodeRef3 = await writer.WriteBlobAsync(new TestNode(789, nodeRef2));
 
 			// nodeRef1 is in a flushed bundle
-			ExportHandle export1 = (ExportHandle)nodeRef1.Unwrap();
+			ExportHandle export1 = (ExportHandle)nodeRef1.Innermost;
 			BundleWriter.PendingPacketHandle packet1 = (BundleWriter.PendingPacketHandle)export1.Packet;
 			BundleWriter.PendingBundleHandle bundle1 = (BundleWriter.PendingBundleHandle)packet1.Bundle;
 			Assert.IsNotNull(packet1.FlushedHandle);
@@ -134,7 +134,7 @@ namespace EpicGames.Horde.Tests
 			Assert.AreEqual(123, node1.Value);
 
 			// nodeRef2 is in a flushed packet, unflushed bundle
-			ExportHandle export2 = (ExportHandle)nodeRef2.Unwrap();
+			ExportHandle export2 = (ExportHandle)nodeRef2.Innermost;
 
 			BundleWriter.PendingPacketHandle packet2 = (BundleWriter.PendingPacketHandle)export2.Packet;
 			Assert.AreNotEqual(packet1, packet2);
@@ -148,7 +148,7 @@ namespace EpicGames.Horde.Tests
 			Assert.AreEqual(456, node2.Value);
 
 			// nodeRef3 is in an unflushed packet, unflushed bundle
-			ExportHandle export3 = (ExportHandle)nodeRef3.Unwrap();
+			ExportHandle export3 = (ExportHandle)nodeRef3.Innermost;
 
 			BundleWriter.PendingPacketHandle packet3 = (BundleWriter.PendingPacketHandle)export3.Packet;
 			Assert.AreNotEqual(packet2, packet3);
diff --git a/Engine/Source/Programs/Shared/EpicGames.Horde/Compute/AgentMessageHandler.cs b/Engine/Source/Programs/Shared/EpicGames.Horde/Compute/AgentMessageHandler.cs
--- a/Engine/Source/Programs/Shared/EpicGames.Horde/Compute/AgentMessageHandler.cs
+++ b/Engine/Source/Programs/Shared/EpicGames.Horde/Compute/AgentMessageHandler.cs
@@ -194,7 +194,7 @@ namespace EpicGames.Horde.Compute
 				throw new InvalidOperationException("Cannot write files outside sandbox");
 			}
 
-			await directoryNode.CopyToDirectoryAsync(outputDir.ToDirectoryInfo(), options, _logger, cancellationToken);
+			await directoryNode.CopyToDirectoryAsync(outputDir.ToDirectoryInfo(), _logger, cancellationToken);
 			await VerifyFilesAsync(outputDir, directoryNode, cancellationToken);
 
 			using (IAgentMessageBuilder message = await channel.CreateMessageAsync(AgentMessageType.WriteFilesResponse, cancellationToken))
diff --git a/Engine/Source/Programs/Shared/EpicGames.Horde/Logs/LogChunkNode.cs b/Engine/Source/Programs/Shared/EpicGames.Horde/Logs/LogChunkNode.cs
--- a/Engine/Source/Programs/Shared/EpicGames.Horde/Logs/LogChunkNode.cs
+++ b/Engine/Source/Programs/Shared/EpicGames.Horde/Logs/LogChunkNode.cs
@@ -220,7 +220,7 @@ namespace EpicGames.Horde.Logs
 		/// <summary>
 		/// Handle to the target chunk
 		/// </summary>
-		public IBlobHandle<LogChunkNode> Target { get; }
+		public IBlobRef<LogChunkNode> Target { get; }
 
 		/// <summary>
 		/// Constructor
@@ -230,7 +230,7 @@ namespace EpicGames.Horde.Logs
 		/// <param name="offset">Offset within the log file</param>
 		/// <param name="length">Length of the chunk</param>
 		/// <param name="target">Referenced log text</param>
-		public LogChunkRef(int lineIndex, int lineCount, long offset, int length, IBlobHandle<LogChunkNode> target)
+		public LogChunkRef(int lineIndex, int lineCount, long offset, int length, IBlobRef<LogChunkNode> target)
 		{
 			LineIndex = lineIndex;
 			LineCount = lineCount;
@@ -245,7 +245,7 @@ namespace EpicGames.Horde.Logs
 		/// <param name="reader"></param>
 		public LogChunkRef(IBlobReader reader)
 		{
-			Target = reader.ReadBlobHandle<LogChunkNode>();
+			Target = reader.ReadBlobRef<LogChunkNode>();
 			LineIndex = (int)reader.ReadUnsignedVarInt();
 			LineCount = (int)reader.ReadUnsignedVarInt();
 			Offset = (long)reader.ReadUnsignedVarInt();
@@ -255,7 +255,7 @@ namespace EpicGames.Horde.Logs
 		/// <inheritdoc/>
 		public void Serialize(IBlobWriter writer)
 		{
-			writer.WriteBlobHandle(Target);
+			writer.WriteBlobRef(Target);
 			writer.WriteUnsignedVarInt(LineIndex);
 			writer.WriteUnsignedVarInt(LineCount);
 			writer.WriteUnsignedVarInt((ulong)Offset);
diff --git a/Engine/Source/Programs/Shared/EpicGames.Horde/Logs/LogIndexNode.cs b/Engine/Source/Programs/Shared/EpicGames.Horde/Logs/LogIndexNode.cs
--- a/Engine/Source/Programs/Shared/EpicGames.Horde/Logs/LogIndexNode.cs
+++ b/Engine/Source/Programs/Shared/EpicGames.Horde/Logs/LogIndexNode.cs
@@ -84,10 +84,9 @@ namespace EpicGames.Horde.Logs
 		/// </summary>
 		/// <param name="writer">Writer for output nodes</param>
 		/// <param name="appendPlainTextChunks">Text blocks to append</param>
-		/// <param name="options">Options to control serialization</param>
 		/// <param name="cancellationToken"></param>
 		/// <returns>New log index with the given blocks appended</returns>
-		public async ValueTask<LogIndexNode> AppendAsync(IBlobWriter writer, IReadOnlyList<LogChunkNode> appendPlainTextChunks, BlobSerializerOptions? options, CancellationToken cancellationToken)
+		public async ValueTask<LogIndexNode> AppendAsync(IBlobWriter writer, IReadOnlyList<LogChunkNode> appendPlainTextChunks, CancellationToken cancellationToken)
 		{
 			using IScope scope = GlobalTracer.Instance.BuildSpan("LogIndex.Append").StartActive();
 
@@ -103,7 +102,7 @@ namespace EpicGames.Horde.Logs
 			for (int idx = 0; idx < appendPlainTextChunks.Count; idx++)
 			{
 				LogChunkNode newChunk = appendPlainTextChunks[idx];
-				IBlobHandle<LogChunkNode> newChunkRef = await writer.WriteBlobAsync(newChunk, options, cancellationToken);
+				IBlobRef<LogChunkNode> newChunkRef = await writer.WriteBlobAsync(newChunk, cancellationToken);
 				newChunks[_plainTextChunkRefs.Length + idx] = new LogChunkRef(lineIndex, newChunk.LineCount, offset, newChunk.Length, newChunkRef);
 				lineIndex += newChunk.LineCount;
 				offset += newChunk.Length;
@@ -148,10 +147,9 @@ namespace EpicGames.Horde.Logs
 		/// <param name="firstLineIndex">First line index to search from</param>
 		/// <param name="text">Text to search for</param>
 		/// <param name="stats">Receives stats for the search</param>
-		/// <param name="options">Options for serialization</param>
 		/// <param name="cancellationToken">Cancellation token for the operation</param>
 		/// <returns>List of line numbers for the text</returns>
-		public async IAsyncEnumerable<int> SearchAsync(int firstLineIndex, SearchTerm text, SearchStats stats, BlobSerializerOptions? options = null, [EnumeratorCancellation] CancellationToken cancellationToken = default)
+		public async IAsyncEnumerable<int> SearchAsync(int firstLineIndex, SearchTerm text, SearchStats stats, [EnumeratorCancellation] CancellationToken cancellationToken = default)
 		{
 			int lastBlockCount = 0;
 			foreach (int blockIdx in EnumeratePossibleChunks(text.Bytes, firstLineIndex))
@@ -165,7 +163,7 @@ namespace EpicGames.Horde.Logs
 				lastBlockCount = blockIdx + 1;
 
 				// Decompress the text
-				LogChunkNode chunk = await indexChunk.Target.ReadBlobAsync(options, cancellationToken);
+				LogChunkNode chunk = await indexChunk.Target.ReadBlobAsync(cancellationToken);
 
 				// Find the initial offset within this block
 				int offset = 0;
diff --git a/Engine/Source/Programs/Shared/EpicGames.Horde/Logs/LogNode.cs b/Engine/Source/Programs/Shared/EpicGames.Horde/Logs/LogNode.cs
--- a/Engine/Source/Programs/Shared/EpicGames.Horde/Logs/LogNode.cs
+++ b/Engine/Source/Programs/Shared/EpicGames.Horde/Logs/LogNode.cs
@@ -58,7 +58,7 @@ namespace EpicGames.Horde.Logs
 		/// <summary>
 		/// Index for this log
 		/// </summary>
-		public IBlobHandle<LogIndexNode> IndexRef { get; }
+		public IBlobRef<LogIndexNode> IndexRef { get; }
 
 		/// <summary>
 		/// Whether this log is complete
@@ -68,7 +68,7 @@ namespace EpicGames.Horde.Logs
 		/// <summary>
 		/// Deserializing constructor
 		/// </summary>
-		public LogNode(LogFormat format, int lineCount, long length, IReadOnlyList<LogChunkRef> textChunkRefs, IBlobHandle<LogIndexNode> indexRef, bool complete)
+		public LogNode(LogFormat format, int lineCount, long length, IReadOnlyList<LogChunkRef> textChunkRefs, IBlobRef<LogIndexNode> indexRef, bool complete)
 		{
 			Format = format;
 			LineCount = lineCount;
@@ -95,7 +95,7 @@ namespace EpicGames.Horde.Logs
 			LogFormat format = (LogFormat)reader.ReadUInt8();
 			int lineCount = (int)reader.ReadUnsignedVarInt();
 			long length = (long)reader.ReadUnsignedVarInt();
-			IBlobHandle<LogIndexNode> indexRef = reader.ReadBlobHandle<LogIndexNode>();
+			IBlobRef<LogIndexNode> indexRef = reader.ReadBlobRef<LogIndexNode>();
 			List<LogChunkRef> textChunkRefs = reader.ReadList(() => new LogChunkRef(reader));
 			bool complete = reader.ReadBoolean();
 
@@ -108,7 +108,7 @@ namespace EpicGames.Horde.Logs
 			writer.WriteUInt8((byte)value.Format);
 			writer.WriteUnsignedVarInt(value.LineCount);
 			writer.WriteUnsignedVarInt((ulong)value.Length);
-			writer.WriteBlobHandle(value.IndexRef);
+			writer.WriteBlobRef(value.IndexRef);
 			writer.WriteList(value.TextChunkRefs, x => x.Serialize(writer));
 			writer.WriteBoolean(value.Complete);
 
@@ -261,9 +261,8 @@ namespace EpicGames.Horde.Logs
 		/// </summary>
 		/// <param name="writer">Writer for the output nodes</param>
 		/// <param name="complete">Whether the log is complete</param>
-		/// <param name="options">Options for serialization</param>
 		/// <param name="cancellationToken">Cancellation token for the operation</param>
-		public async Task<IBlobHandle<LogNode>> FlushAsync(IBlobWriter writer, bool complete, BlobSerializerOptions? options, CancellationToken cancellationToken)
+		public async Task<IBlobRef<LogNode>> FlushAsync(IBlobWriter writer, bool complete, CancellationToken cancellationToken)
 		{
 			// Capture the new data that needs to be written
 			IReadOnlyList<LogChunkNode> writeTextChunks;
@@ -279,22 +278,22 @@ namespace EpicGames.Horde.Logs
 			}
 
 			// Flush any complete chunks to storage
-			LogIndexNode newIndex = await _index.AppendAsync(writer, writeIndexTextChunks, null, cancellationToken);
-			IBlobHandle<LogIndexNode> newIndexRef = await writer.WriteBlobAsync(newIndex, options, cancellationToken);
+			LogIndexNode newIndex = await _index.AppendAsync(writer, writeIndexTextChunks, cancellationToken);
+			IBlobRef<LogIndexNode> newIndexRef = await writer.WriteBlobAsync(newIndex, cancellationToken);
 
 			List<LogChunkRef> newJsonChunkRefs = new List<LogChunkRef>(_root?.TextChunkRefs ?? Array.Empty<LogChunkRef>());
 			int lineCount = _root?.LineCount ?? 0;
 			long length = _root?.Length ?? 0;
 			foreach (LogChunkNode writeTextChunk in writeTextChunks)
 			{
-				IBlobHandle<LogChunkNode> writeTextChunkRef = await writer.WriteBlobAsync(writeTextChunk, options, cancellationToken);
+				IBlobRef<LogChunkNode> writeTextChunkRef = await writer.WriteBlobAsync(writeTextChunk, cancellationToken);
 				newJsonChunkRefs.Add(new LogChunkRef(lineCount, writeTextChunk.LineCount, length, writeTextChunk.Length, writeTextChunkRef));
 				lineCount += writeTextChunk.LineCount;
 				length += writeTextChunk.Length;
 			}
 
 			LogNode newRoot = new LogNode(_format, lineCount, length, newJsonChunkRefs, newIndexRef, complete);
-			IBlobHandle<LogNode> newRootRef = await writer.WriteBlobAsync(newRoot, options, cancellationToken);
+			IBlobRef<LogNode> newRootRef = await writer.WriteBlobAsync(newRoot, cancellationToken);
 
 			await writer.FlushAsync(cancellationToken);
 
@@ -320,14 +319,13 @@ namespace EpicGames.Horde.Logs
 		/// Reads lines from a line
 		/// </summary>
 		/// <param name="logNode">Log to read from</param>
-		/// <param name="options">Options for deserialization</param>
 		/// <param name="cancellationToken">Cancellation token</param>
 		/// <returns>Sequence of line buffers</returns>
-		public static async IAsyncEnumerable<ReadOnlyMemory<byte>> ReadLogAsync(this LogNode logNode, BlobSerializerOptions? options = null, [EnumeratorCancellation] CancellationToken cancellationToken = default)
+		public static async IAsyncEnumerable<ReadOnlyMemory<byte>> ReadLogAsync(this LogNode logNode, [EnumeratorCancellation] CancellationToken cancellationToken = default)
 		{
 			foreach (LogChunkRef textChunkRef in logNode.TextChunkRefs)
 			{
-				LogChunkNode textChunk = await textChunkRef.Target.ReadBlobAsync(options, cancellationToken);
+				LogChunkNode textChunk = await textChunkRef.Target.ReadBlobAsync(cancellationToken);
 				yield return textChunk.Data;
 			}
 		}
@@ -337,17 +335,16 @@ namespace EpicGames.Horde.Logs
 		/// </summary>
 		/// <param name="logNode">Log to read from</param>
 		/// <param name="index">Zero-based index of the first line to read from</param>
-		/// <param name="options">Options for deserialization</param>
 		/// <param name="cancellationToken">Cancellation token</param>
 		/// <returns>Sequence of line buffers</returns>
-		public static async IAsyncEnumerable<ReadOnlyMemory<byte>> ReadLogLinesAsync(this LogNode logNode, int index, BlobSerializerOptions? options = null, [EnumeratorCancellation] CancellationToken cancellationToken = default)
+		public static async IAsyncEnumerable<ReadOnlyMemory<byte>> ReadLogLinesAsync(this LogNode logNode, int index, [EnumeratorCancellation] CancellationToken cancellationToken = default)
 		{
 			foreach (LogChunkRef textChunkRef in logNode.TextChunkRefs)
 			{
 				int lineIdx = Math.Max(index - textChunkRef.LineIndex, 0);
 				if (lineIdx < textChunkRef.LineCount)
 				{
-					LogChunkNode textChunk = await textChunkRef.Target.ReadBlobAsync(options, cancellationToken);
+					LogChunkNode textChunk = await textChunkRef.Target.ReadBlobAsync(cancellationToken);
 
 					int offset = textChunk.LineOffsets[lineIdx];
 					for (; lineIdx < textChunk.LineCount; lineIdx++)
diff --git a/Engine/Source/Programs/Shared/EpicGames.Horde/Storage/BlobData.cs b/Engine/Source/Programs/Shared/EpicGames.Horde/Storage/BlobData.cs
--- a/Engine/Source/Programs/Shared/EpicGames.Horde/Storage/BlobData.cs
+++ b/Engine/Source/Programs/Shared/EpicGames.Horde/Storage/BlobData.cs
@@ -24,16 +24,16 @@ namespace EpicGames.Horde.Storage
 		/// <summary>
 		/// Handles to referenced blobs
 		/// </summary>
-		public IReadOnlyList<IBlobHandle> Refs { get; }
+		public IReadOnlyList<IBlobHandle> Imports { get; }
 
 		/// <summary>
 		/// Constructor
 		/// </summary>
-		public BlobData(BlobType type, ReadOnlyMemory<byte> data, IReadOnlyList<IBlobHandle> refs)
+		public BlobData(BlobType type, ReadOnlyMemory<byte> data, IReadOnlyList<IBlobHandle> imports)
 		{
 			Type = type;
 			Data = data;
-			Refs = refs;
+			Imports = imports;
 		}
 
 		/// <inheritdoc/>
@@ -62,8 +62,8 @@ namespace EpicGames.Horde.Storage
 		/// <summary>
 		/// Constructor
 		/// </summary>
-		public BlobDataWithOwner(BlobType type, ReadOnlyMemory<byte> data, IReadOnlyList<IBlobHandle> refs, IDisposable owner)
-			: base(type, data, refs)
+		public BlobDataWithOwner(BlobType type, ReadOnlyMemory<byte> data, IReadOnlyList<IBlobHandle> imports, IDisposable owner)
+			: base(type, data, imports)
 		{
 			_owner = owner;
 		}
diff --git a/Engine/Source/Programs/Shared/EpicGames.Horde/Storage/BlobSerializer.cs b/Engine/Source/Programs/Shared/EpicGames.Horde/Storage/BlobSerializer.cs
--- a/Engine/Source/Programs/Shared/EpicGames.Horde/Storage/BlobSerializer.cs
+++ b/Engine/Source/Programs/Shared/EpicGames.Horde/Storage/BlobSerializer.cs
@@ -24,7 +24,7 @@ namespace EpicGames.Horde.Storage
 		public static T Deserialize<T>(BlobData blobData, BlobSerializerOptions? options = null)
 		{
 			options ??= BlobSerializerOptions.Default;
-			BlobReader reader = new BlobReader(blobData);
+			BlobReader reader = new BlobReader(blobData, options);
 			return options.GetConverter<T>().Read(reader, options);
 		}
 
@@ -135,12 +135,11 @@ namespace EpicGames.Horde.Storage
 		/// </summary>
 		/// <typeparam name="T">Return type for the deserialized object</typeparam>
 		/// <param name="handle">Handle to the blob to deserialize</param>
-		/// <param name="options">Options to control serialization</param>
 		/// <param name="cancellationToken">Cancellation token for the operation</param>
-		public static async ValueTask<T> ReadBlobAsync<T>(this IBlobHandle<T> handle, BlobSerializerOptions? options = null, CancellationToken cancellationToken = default)
+		public static async ValueTask<T> ReadBlobAsync<T>(this IBlobRef<T> handle, CancellationToken cancellationToken = default)
 		{
 			BlobData data = await handle.ReadBlobDataAsync(cancellationToken);
-			return BlobSerializer.Deserialize<T>(data, options);
+			return BlobSerializer.Deserialize<T>(data, handle.Options);
 		}
 
 		/// <summary>
@@ -150,22 +149,9 @@ namespace EpicGames.Horde.Storage
 		/// <param name="value">The object to serialize</param>
 		/// <param name="cancellationToken">Cancellation token for the operation</param>
 		/// <returns>Handle to the serialized blob</returns>
-		public static ValueTask<IBlobHandle<T>> WriteBlobAsync<T>(this IBlobWriter writer, T value, CancellationToken cancellationToken)
-		{
-			return WriteBlobAsync<T>(writer, value, null, cancellationToken);
-		}
-
-		/// <summary>
-		/// Serialize an object to storage
-		/// </summary>
-		/// <param name="writer">Writer for serialized data</param>
-		/// <param name="value">The object to serialize</param>
-		/// <param name="options">Options to control serialization</param>
-		/// <param name="cancellationToken">Cancellation token for the operation</param>
-		/// <returns>Handle to the serialized blob</returns>
-		public static async ValueTask<IBlobHandle<T>> WriteBlobAsync<T>(this IBlobWriter writer, T value, BlobSerializerOptions? options = null, CancellationToken cancellationToken = default)
+		public static async ValueTask<IBlobRef<T>> WriteBlobAsync<T>(this IBlobWriter writer, T value, CancellationToken cancellationToken = default)
 		{
-			BlobType blobType = BlobSerializer.Serialize<T>(writer, value, options);
+			BlobType blobType = BlobSerializer.Serialize<T>(writer, value, writer.Options);
 			return await writer.CompleteAsync<T>(blobType, cancellationToken);
 		}
 
diff --git a/Engine/Source/Programs/Shared/EpicGames.Horde/Storage/Bundles/BundleStorageClient.cs b/Engine/Source/Programs/Shared/EpicGames.Horde/Storage/Bundles/BundleStorageClient.cs
--- a/Engine/Source/Programs/Shared/EpicGames.Horde/Storage/Bundles/BundleStorageClient.cs
+++ b/Engine/Source/Programs/Shared/EpicGames.Horde/Storage/Bundles/BundleStorageClient.cs
@@ -170,26 +170,38 @@ namespace EpicGames.Horde.Storage.Bundles
 		}
 
 		/// <inheritdoc/>
-		public IBlobWriter CreateBlobWriter(string? basePath = null, BundleOptions? options = null)
+		public IBlobRef CreateBlobRef(IoHash hash, BlobLocator locator)
 		{
-			options ??= BundleOptions.Default;
+			return BlobRef.Create(hash, CreateBlobHandle(locator));
+		}
+
+		/// <inheritdoc/>
+		public IBlobRef<T> CreateBlobRef<T>(IoHash hash, BlobLocator locator, BlobSerializerOptions options)
+		{
+			return BlobRef.Create<T>(hash, CreateBlobHandle(locator), options);
+		}
+
+		/// <inheritdoc/>
+		public IBlobWriter CreateBlobWriter(string? basePath = null, BundleOptions? bundleOptions = null, BlobSerializerOptions? serializerOptions = null)
+		{
+			bundleOptions ??= BundleOptions.Default;
 
-			if (options.MaxVersion == BundleVersion.LatestV1)
+			if (bundleOptions.MaxVersion == BundleVersion.LatestV1)
 			{
-				return new Bundles.V1.BundleWriter(this, _bundleReader, basePath, options);
+				return new Bundles.V1.BundleWriter(this, _bundleReader, basePath, bundleOptions);
 			}
-			else if(options.MaxVersion == BundleVersion.LatestV2)
+			else if(bundleOptions.MaxVersion == BundleVersion.LatestV2)
 			{
-				return new Bundles.V2.BundleWriter(this, basePath, _cache, options);
+				return new Bundles.V2.BundleWriter(this, basePath, _cache, bundleOptions, serializerOptions);
 			}
 			else
 			{
-				throw new InvalidOperationException($"Unsupported bundle version: {(int)options.MaxVersion}");
+				throw new InvalidOperationException($"Unsupported bundle version: {(int)bundleOptions.MaxVersion}");
 			}
 		}
 
 		/// <inheritdoc/>
-		IBlobWriter IStorageClient.CreateBlobWriter(string? basePath) => CreateBlobWriter(basePath);
+		IBlobWriter IStorageClient.CreateBlobWriter(string? basePath, BlobSerializerOptions? options) => CreateBlobWriter(basePath, serializerOptions: options);
 
 		#endregion
 
diff --git a/Engine/Source/Programs/Shared/EpicGames.Horde/Storage/Bundles/V1/BundleReader.cs b/Engine/Source/Programs/Shared/EpicGames.Horde/Storage/Bundles/V1/BundleReader.cs
--- a/Engine/Source/Programs/Shared/EpicGames.Horde/Storage/Bundles/V1/BundleReader.cs
+++ b/Engine/Source/Programs/Shared/EpicGames.Horde/Storage/Bundles/V1/BundleReader.cs
@@ -546,7 +546,7 @@ namespace EpicGames.Horde.Storage.Bundles.V1
 
 			Dictionary<BlobLocator, BundleHandle> locatorToBundleHandle = new Dictionary<BlobLocator, BundleHandle>();
 
-			List<IBlobHandle> refs = new List<IBlobHandle>(export.References.Count);
+			List<IBlobHandle> imports = new List<IBlobHandle>(export.References.Count);
 			foreach (BundleExportRef reference in export.References)
 			{
 				BlobLocator importBlob;
@@ -567,7 +567,7 @@ namespace EpicGames.Horde.Storage.Bundles.V1
 				}
 
 				Debug.Assert(importBlob.IsValid());
-				refs.Add(new FlushedNodeHandle(this, importBlob, importBundle, reference.NodeIdx));
+				imports.Add(new FlushedNodeHandle(this, importBlob, importBundle, reference.NodeIdx));
 			}
 
 			ReadOnlyMemory<byte> nodeData = ReadOnlyMemory<byte>.Empty;
@@ -578,7 +578,7 @@ namespace EpicGames.Horde.Storage.Bundles.V1
 			}
 
 			BlobType nodeType = header.Types[export.TypeIdx];
-			return new BlobData(nodeType, nodeData, refs);
+			return new BlobData(nodeType, nodeData, imports);
 		}
 
 		/// <summary>
diff --git a/Engine/Source/Programs/Shared/EpicGames.Horde/Storage/Bundles/V1/BundleWriter.cs b/Engine/Source/Programs/Shared/EpicGames.Horde/Storage/Bundles/V1/BundleWriter.cs
--- a/Engine/Source/Programs/Shared/EpicGames.Horde/Storage/Bundles/V1/BundleWriter.cs
+++ b/Engine/Source/Programs/Shared/EpicGames.Horde/Storage/Bundles/V1/BundleWriter.cs
@@ -28,6 +28,8 @@ namespace EpicGames.Horde.Storage.Bundles.V1
 		/// <inheritdoc/>
 		public BundleHandle Outer { get; }
 
+		public IBlobHandle Innermost => this;
+
 		/// <summary>
 		/// Constructor
 		/// </summary>
@@ -91,7 +93,7 @@ namespace EpicGames.Horde.Storage.Bundles.V1
 	public sealed class BundleWriter : BlobWriter
 	{
 		// Information about a unique output node. Note that multiple node refs may de-duplicate to the same output node.
-		internal class PendingNode : IBlobHandle
+		internal class PendingNode : IBlobRef
 		{
 			readonly BundleReader _reader;
 
@@ -101,26 +103,29 @@ namespace EpicGames.Horde.Storage.Bundles.V1
 			PendingBundle? _pendingBundle;
 
 			public readonly BlobType BlobType;
+			public IoHash Hash { get; }
 			public readonly int Packet;
 			public readonly int Offset;
 			public readonly int Length;
-			public readonly IBlobHandle[] Refs;
+			public readonly IBlobHandle[] Imports;
 			public readonly AliasInfo[] Aliases;
 
+			public IBlobHandle Innermost => this;
 			public PendingBundle? PendingBundle => _pendingBundle;
 			public FlushedNodeHandle? FlushedNodeHandle => _flushedHandle;
 
 			public BundleHandle Outer => (_flushedHandle != null) ? _flushedHandle.Outer : throw new NotSupportedException();
 
-			public PendingNode(BundleReader reader, BlobType blobType, int packet, int offset, int length, IReadOnlyList<IBlobHandle> refs, IReadOnlyList<AliasInfo> aliases, PendingBundle pendingBundle)
+			public PendingNode(BundleReader reader, BlobType blobType, IoHash hash, int packet, int offset, int length, IReadOnlyList<IBlobHandle> imports, IReadOnlyList<AliasInfo> aliases, PendingBundle pendingBundle)
 			{
 				_reader = reader;
 
 				BlobType = blobType;
+				Hash = hash;
 				Packet = packet;
 				Offset = offset;
 				Length = length;
-				Refs = refs.Select(x => x.Unwrap()).ToArray();
+				Imports = imports.Select(x => x.Innermost).ToArray();
 				Aliases = aliases.ToArray();
 
 				_pendingBundle = pendingBundle;
@@ -152,7 +157,7 @@ namespace EpicGames.Horde.Storage.Bundles.V1
 						if (_flushedHandle == null)
 						{
 							ReadOnlyMemory<byte> data = _pendingBundle!.GetNodeData(Packet, Offset, Length);
-							return new BlobData(BlobType, data, Refs);
+							return new BlobData(BlobType, data, Imports);
 						}
 					}
 				}
@@ -168,9 +173,9 @@ namespace EpicGames.Horde.Storage.Bundles.V1
 					return;
 				}
 
-				foreach (IBlobHandle nodeRef in Refs)
+				foreach (IBlobHandle import in Imports)
 				{
-					await nodeRef.FlushAsync(cancellationToken);
+					await import.FlushAsync(cancellationToken);
 				}
 
 				PendingBundle? pendingBundle = _pendingBundle;
@@ -318,12 +323,14 @@ namespace EpicGames.Horde.Storage.Bundles.V1
 			// Finish a node write
 			public PendingNode WriteNode(BlobType blobType, int size, IReadOnlyList<IBlobHandle> refs, IReadOnlyList<AliasInfo> aliases)
 			{
-				PendingNode pendingNode = new PendingNode(_treeReader, blobType, _currentPacketIdx, _currentPacketLength, (int)size, refs, aliases, this);
+				IoHash hash = IoHash.Compute(GetBuffer(size, size).Span);
+
+				PendingNode pendingNode = new PendingNode(_treeReader, blobType, hash, _currentPacketIdx, _currentPacketLength, (int)size, refs, aliases, this);
 				_currentPacketLength += pendingNode.Length;
 				UncompressedLength += pendingNode.Length;
 
 				_queue.Add(pendingNode);
-				_queuedRefs += pendingNode.Refs.Length;
+				_queuedRefs += pendingNode.Imports.Length;
 
 				return pendingNode;
 			}
@@ -512,18 +519,18 @@ namespace EpicGames.Horde.Storage.Bundles.V1
 					int typeIdx = FindOrAddItemIndex(nodeInfo.BlobType, types, typeToIndex);
 
 					List<BundleExportRef> exportRefs = new List<BundleExportRef>();
-					foreach (IBlobHandle handle in nodeInfo.Refs)
+					foreach (IBlobHandle import in nodeInfo.Imports)
 					{
 						BundleExportRef exportRef;
-						if (!nodeHandleToExportRef.TryGetValue(handle, out exportRef))
+						if (!nodeHandleToExportRef.TryGetValue(import, out exportRef))
 						{
-							FlushedNodeHandle? flushedHandle = handle as FlushedNodeHandle;
+							FlushedNodeHandle? flushedHandle = import as FlushedNodeHandle;
 							if (flushedHandle == null)
 							{
-								flushedHandle = (handle as PendingNode)?.FlushedNodeHandle;
+								flushedHandle = (import as PendingNode)?.FlushedNodeHandle;
 								if (flushedHandle == null)
 								{
-									throw new InvalidOperationException($"Node {handle.GetLocator()} is not a bundle node");
+									throw new InvalidOperationException($"Node {import.GetLocator()} is not a bundle node");
 								}
 							}
 							
@@ -662,7 +669,7 @@ namespace EpicGames.Horde.Storage.Bundles.V1
 
 		readonly BundleStorageClient _store;
 		readonly BundleReader _reader;
-		readonly BundleOptions _options;
+		readonly BundleOptions _bundleOptions;
 		readonly string? _basePath;
 
 		readonly WriteQueue _writeQueue;
@@ -679,10 +686,11 @@ namespace EpicGames.Horde.Storage.Bundles.V1
 		/// <param name="store">Store to write data to</param>
 		/// <param name="reader">Reader for serialized node data</param>
 		/// <param name="basePath">Base path for new nodes</param>
-		/// <param name="options">Options for the writer</param>
+		/// <param name="bundleOptions">Options for the writer</param>
+		/// <param name="blobSerializerOptions"></param>
 		/// <param name="traceLogger">Optional logger for trace information</param>
-		public BundleWriter(BundleStorageClient store, BundleReader reader, string? basePath, BundleOptions? options = null, ILogger? traceLogger = null)
-			: this(store, reader, basePath, options, null, traceLogger)
+		public BundleWriter(BundleStorageClient store, BundleReader reader, string? basePath, BundleOptions? bundleOptions = null, BlobSerializerOptions? blobSerializerOptions = null, ILogger? traceLogger = null)
+			: this(store, reader, basePath, bundleOptions, null, blobSerializerOptions, traceLogger)
 		{
 		}
 
@@ -691,7 +699,7 @@ namespace EpicGames.Horde.Storage.Bundles.V1
 		/// </summary>
 		/// <param name="other"></param>
 		public BundleWriter(BundleWriter other)
-			: this(other._store, other._reader, other._basePath, other._options, other._writeQueue, other.TraceLogger)
+			: this(other._store, other._reader, other._basePath, other._bundleOptions, other._writeQueue, other.Options, other.TraceLogger)
 		{
 			_writeQueue.AddRef();
 		}
@@ -699,13 +707,14 @@ namespace EpicGames.Horde.Storage.Bundles.V1
 		/// <summary>
 		/// Internal constructor
 		/// </summary>
-		private BundleWriter(BundleStorageClient store, BundleReader reader, string? basePath, BundleOptions? options, WriteQueue? writeQueue, ILogger? traceLogger = null)
+		private BundleWriter(BundleStorageClient store, BundleReader reader, string? basePath, BundleOptions? bundleOptions, WriteQueue? writeQueue, BlobSerializerOptions? blobSerializerOptions, ILogger? traceLogger = null)
+			: base(blobSerializerOptions)
 		{
 			_store = store;
 			_reader = reader;
 			_basePath = basePath;
-			_options = options ?? s_defaultOptions;
-			_writeQueue = writeQueue ?? new WriteQueue(store, basePath, _options.MaxWriteQueueLength, traceLogger);
+			_bundleOptions = bundleOptions ?? s_defaultOptions;
+			_writeQueue = writeQueue ?? new WriteQueue(store, basePath, _bundleOptions.MaxWriteQueueLength, traceLogger);
 			TraceLogger = traceLogger;
 		}
 
@@ -758,22 +767,22 @@ namespace EpicGames.Horde.Storage.Bundles.V1
 		/// </summary>
 		/// <param name="type">Type of the node that was written</param>
 		/// <param name="size">Used size of the buffer</param>
-		/// <param name="references">References to other nodes</param>
+		/// <param name="imports">References to other nodes</param>
 		/// <param name="aliases">Aliases for the node</param>
 		/// <param name="cancellationToken">Cancellation token for the operation</param>
 		/// <returns>Handle to the written node</returns>
-		public override async ValueTask<IBlobHandle> WriteBlobAsync(BlobType type, int size, IReadOnlyList<IBlobHandle> references, IReadOnlyList<AliasInfo> aliases, CancellationToken cancellationToken = default)
+		public override async ValueTask<IBlobRef> WriteBlobAsync(BlobType type, int size, IReadOnlyList<IBlobHandle> imports, IReadOnlyList<AliasInfo> aliases, CancellationToken cancellationToken = default)
 		{
 			PendingBundle currentBundle = GetCurrentBundle();
 
 			// Append this node data
-			PendingNode pendingNode = currentBundle.WriteNode(type, size, references, aliases);
+			PendingNode pendingNode = currentBundle.WriteNode(type, size, imports, aliases);
 			TraceLogger?.LogInformation("Added new node for {NodeKey} in bundle {BundleId}", pendingNode, currentBundle.BundleId);
 
 			// Add dependencies on all bundles containing a dependent node
-			foreach (IBlobHandle reference in references)
+			foreach (IBlobHandle import in imports)
 			{
-				PendingNode? pendingReference = reference.Unwrap() as PendingNode;
+				PendingNode? pendingReference = import.Innermost as PendingNode;
 				if (pendingReference?.PendingBundle != null)
 				{
 					currentBundle.AddDependencyOn(pendingReference.PendingBundle, TraceLogger);
@@ -781,7 +790,7 @@ namespace EpicGames.Horde.Storage.Bundles.V1
 			}
 
 			// If the bundle is full, start the process of writing it to disk
-			if (currentBundle.UncompressedLength > _options.MaxBlobSize || currentBundle.IsFull())
+			if (currentBundle.UncompressedLength > _bundleOptions.MaxBlobSize || currentBundle.IsFull())
 			{
 				await CompleteAsync(cancellationToken);
 			}
@@ -793,8 +802,8 @@ namespace EpicGames.Horde.Storage.Bundles.V1
 		{
 			if (_currentBundle == null)
 			{
-				int bufferSize = (int)(_options.MinCompressionPacketSize * 1.2);
-				_currentBundle = new PendingBundle(_reader, this, bufferSize, _options.CompressionFormat);
+				int bufferSize = (int)(_bundleOptions.MinCompressionPacketSize * 1.2);
+				_currentBundle = new PendingBundle(_reader, this, bufferSize, _bundleOptions.CompressionFormat);
 			}
 			return _currentBundle;
 		}
diff --git a/Engine/Source/Programs/Shared/EpicGames.Horde/Storage/Bundles/V2/BundleWriter.cs b/Engine/Source/Programs/Shared/EpicGames.Horde/Storage/Bundles/V2/BundleWriter.cs
--- a/Engine/Source/Programs/Shared/EpicGames.Horde/Storage/Bundles/V2/BundleWriter.cs
+++ b/Engine/Source/Programs/Shared/EpicGames.Horde/Storage/Bundles/V2/BundleWriter.cs
@@ -59,7 +59,7 @@ namespace EpicGames.Horde.Storage.Bundles.V2
 			readonly BundleStorageClient _storageClient;
 			readonly string? _basePath;
 			readonly BundleCache _cache;
-			readonly BundleOptions _options;
+			readonly BundleOptions _bundleOptions;
 
 			FlushedBundleHandle? _flushedHandle;
 
@@ -80,12 +80,12 @@ namespace EpicGames.Horde.Storage.Bundles.V2
 			/// </summary>
 			public int Length => _encodedPacketWriter?.Length ?? throw new InvalidOperationException("Bundle has been flushed");
 
-			public PendingBundleHandle(BundleStorageClient storageClient, string? basePath, BundleCache cache, BundleOptions options)
+			public PendingBundleHandle(BundleStorageClient storageClient, string? basePath, BundleCache cache, BundleOptions bundleOptions)
 			{
 				_storageClient = storageClient;
 				_basePath = basePath;
 				_cache = cache;
-				_options = options;
+				_bundleOptions = bundleOptions;
 
 				_encodedPacketWriter = new RefCountedMemoryWriter(_cache.Allocator, 65536);
 
@@ -124,13 +124,15 @@ namespace EpicGames.Horde.Storage.Bundles.V2
 				return _packetWriter.GetExport(exportIdx);
 			}
 
-			public ExportHandle CompleteExport(BlobType type, int size, IReadOnlyList<IBlobHandle> references, IReadOnlyList<AliasInfo> aliases)
+			public ExportHandle CompleteExport(BlobType type, int size, IReadOnlyList<IBlobHandle> imports, IReadOnlyList<AliasInfo> aliases)
 			{
 				RuntimeAssert(_packetWriter != null);
 				RuntimeAssert(_packetHandle != null);
 
-				int exportIdx = _packetWriter.CompleteExport(size, type, references);
-				ExportHandle exportHandle = new ExportHandle(_packetHandle, exportIdx);
+				IoHash hash = IoHash.Compute(_packetWriter.GetOutputBuffer(size, size).Span);
+
+				int exportIdx = _packetWriter.CompleteExport(size, type, imports);
+				ExportHandle exportHandle = new ExportHandle(hash, _packetHandle, exportIdx);
 
 				if (aliases.Count > 0)
 				{
@@ -138,7 +140,7 @@ namespace EpicGames.Horde.Storage.Bundles.V2
 					_pendingExportAliases.AddRange(aliases.Select(x => (exportHandle, x)));
 				}
 
-				if (_packetWriter.Length > Math.Min(_options.MinCompressionPacketSize, _options.MaxBlobSize))
+				if (_packetWriter.Length > Math.Min(_bundleOptions.MinCompressionPacketSize, _bundleOptions.MaxBlobSize))
 				{
 					FinishPacket();
 					StartPacket();
@@ -166,7 +168,7 @@ namespace EpicGames.Horde.Storage.Bundles.V2
 				{
 					int packetOffset = _encodedPacketWriter.Length;
 					Packet packet = _packetWriter.CompletePacket();
-					packet.Encode(_options.CompressionFormat, _encodedPacketWriter);
+					packet.Encode(_bundleOptions.CompressionFormat, _encodedPacketWriter);
 					int packetLength = _encodedPacketWriter.Length - packetOffset;
 
 					// Point the packet handle to the encoded data
@@ -277,7 +279,7 @@ namespace EpicGames.Horde.Storage.Bundles.V2
 
 		readonly BundleStorageClient _storageClient;
 		readonly string? _basePath;
-		readonly BundleOptions _options;
+		readonly BundleOptions _bundleOptions;
 		readonly BundleCache _bundleCache;
 
 		PendingBundleHandle _currentBundle;
@@ -285,14 +287,15 @@ namespace EpicGames.Horde.Storage.Bundles.V2
 		/// <summary>
 		/// 
 		/// </summary>
-		public BundleWriter(BundleStorageClient storageClient, string? basePath, BundleCache bundleCache, BundleOptions options)
+		public BundleWriter(BundleStorageClient storageClient, string? basePath, BundleCache bundleCache, BundleOptions bundleOptions, BlobSerializerOptions? blobOptions)
+			: base(blobOptions)
 		{
 			_storageClient = storageClient;
 			_basePath = basePath;
-			_options = options;
+			_bundleOptions = bundleOptions;
 			_bundleCache = bundleCache;
 
-			_currentBundle = new PendingBundleHandle(storageClient, basePath, bundleCache, _options);
+			_currentBundle = new PendingBundleHandle(storageClient, basePath, bundleCache, _bundleOptions);
 		}
 
 		/// <inheritdoc/>
@@ -307,22 +310,22 @@ namespace EpicGames.Horde.Storage.Bundles.V2
 		{
 			await _currentBundle.FlushAsync(cancellationToken);
 			_currentBundle.Dispose();
-			_currentBundle = new PendingBundleHandle(_storageClient, _basePath, _bundleCache, _options);
+			_currentBundle = new PendingBundleHandle(_storageClient, _basePath, _bundleCache, _bundleOptions);
 		}
 
 		/// <inheritdoc/>
 		public override IBlobWriter Fork()
-			=> new BundleWriter(_storageClient, _basePath, _bundleCache, _options);
+			=> new BundleWriter(_storageClient, _basePath, _bundleCache, _bundleOptions, Options);
 
 		/// <inheritdoc/>
 		public override Memory<byte> GetOutputBuffer(int usedSize, int desiredSize)
 			=> _currentBundle.GetOutputBuffer(usedSize, desiredSize);
 
 		/// <inheritdoc/>
-		public override async ValueTask<IBlobHandle> WriteBlobAsync(BlobType type, int size, IReadOnlyList<IBlobHandle> references, IReadOnlyList<AliasInfo> aliases, CancellationToken cancellationToken = default)
+		public override async ValueTask<IBlobRef> WriteBlobAsync(BlobType type, int size, IReadOnlyList<IBlobHandle> references, IReadOnlyList<AliasInfo> aliases, CancellationToken cancellationToken = default)
 		{
 			ExportHandle exportHandle = _currentBundle.CompleteExport(type, size, references, aliases);
-			if (_currentBundle.Length > _options.MaxBlobSize)
+			if (_currentBundle.Length > _bundleOptions.MaxBlobSize)
 			{
 				await FlushAsync(cancellationToken);
 			}
diff --git a/Engine/Source/Programs/Shared/EpicGames.Horde/Storage/Bundles/V2/ExportHandle.cs b/Engine/Source/Programs/Shared/EpicGames.Horde/Storage/Bundles/V2/ExportHandle.cs
--- a/Engine/Source/Programs/Shared/EpicGames.Horde/Storage/Bundles/V2/ExportHandle.cs
+++ b/Engine/Source/Programs/Shared/EpicGames.Horde/Storage/Bundles/V2/ExportHandle.cs
@@ -12,7 +12,7 @@ namespace EpicGames.Horde.Storage.Bundles.V2
 	/// <summary>
 	/// Handle class for blobs exported from a bundle
 	/// </summary>
-	public class ExportHandle : IBlobHandle
+	public class ExportHandle : IBlobRef
 	{
 		static readonly Utf8String s_fragmentPrefix = new Utf8String("exp=");
 
@@ -26,14 +26,22 @@ namespace EpicGames.Horde.Storage.Bundles.V2
 		/// </summary>
 		public int ExportIdx => _exportIdx;
 
+		readonly IoHash _hash;
 		readonly PacketHandle _packet;
 		readonly int _exportIdx;
 
+		/// <inheritdoc/>
+		public IBlobHandle Innermost => this;
+
+		/// <inheritdoc/>
+		public IoHash Hash => _hash;
+
 		/// <summary>
 		/// Constructor
 		/// </summary>
-		public ExportHandle(PacketHandle packet, int exportIdx)
+		public ExportHandle(IoHash hash, PacketHandle packet, int exportIdx)
 		{
+			_hash = hash;
 			_packet = packet;
 			_exportIdx = exportIdx;
 		}
diff --git a/Engine/Source/Programs/Shared/EpicGames.Horde/Storage/Bundles/V2/PacketReader.cs b/Engine/Source/Programs/Shared/EpicGames.Horde/Storage/Bundles/V2/PacketReader.cs
--- a/Engine/Source/Programs/Shared/EpicGames.Horde/Storage/Bundles/V2/PacketReader.cs
+++ b/Engine/Source/Programs/Shared/EpicGames.Horde/Storage/Bundles/V2/PacketReader.cs
@@ -135,13 +135,13 @@ namespace EpicGames.Horde.Storage.Bundles.V2
 			IBlobHandle? blobHandle = _cachedImportHandles[blobIdx] as IBlobHandle;
 			if (blobHandle is null)
 			{
-				PacketImport blobImport = _decodedPacket.GetImport(blobIdx);
-				Trace.Assert(blobImport.BaseIdx != -1);
+				PacketImport blobImportInfo = _decodedPacket.GetImport(blobIdx);
+				Trace.Assert(blobImportInfo.BaseIdx != -1);
 
-				int packetIdx = blobImport.BaseIdx;
+				int packetIdx = blobImportInfo.BaseIdx;
 				PacketHandle packetHandle = GetImportedPacketHandle(packetIdx);
 
-				blobHandle = new ExportHandle(packetHandle, blobImport.Fragment);
+				blobHandle = new ExportHandle(packetHandle, blobImportInfo.Fragment);
 				_cachedImportHandles[blobIdx] = blobHandle;
 			}
 			return blobHandle;
diff --git a/Engine/Source/Programs/Shared/EpicGames.Horde/Storage/Bundles/V2/PacketWriter.cs b/Engine/Source/Programs/Shared/EpicGames.Horde/Storage/Bundles/V2/PacketWriter.cs
--- a/Engine/Source/Programs/Shared/EpicGames.Horde/Storage/Bundles/V2/PacketWriter.cs
+++ b/Engine/Source/Programs/Shared/EpicGames.Horde/Storage/Bundles/V2/PacketWriter.cs
@@ -88,13 +88,13 @@ namespace EpicGames.Horde.Storage.Bundles.V2
 		/// </summary>
 		/// <param name="size"></param>
 		/// <param name="type"></param>
-		/// <param name="references"></param>
-		public int CompleteExport(int size, BlobType type, IReadOnlyList<IBlobHandle> references)
+		/// <param name="imports"></param>
+		public int CompleteExport(int size, BlobType type, IReadOnlyList<IBlobHandle> imports)
 		{
-			int[] importIndices = new int[references.Count];
-			for (int idx = 0; idx < references.Count; idx++)
+			int[] importIndices = new int[imports.Count];
+			for (int idx = 0; idx < imports.Count; idx++)
 			{
-				importIndices[idx] = FindOrAddImport(references[idx].Unwrap());
+				importIndices[idx] = FindOrAddImport(imports[idx].Innermost);
 			}
 
 			int typeIdx = FindOrAddType(type);
diff --git a/Engine/Source/Programs/Shared/EpicGames.Horde/Storage/Clients/KeyValueStorageClient.cs b/Engine/Source/Programs/Shared/EpicGames.Horde/Storage/Clients/KeyValueStorageClient.cs
--- a/Engine/Source/Programs/Shared/EpicGames.Horde/Storage/Clients/KeyValueStorageClient.cs
+++ b/Engine/Source/Programs/Shared/EpicGames.Horde/Storage/Clients/KeyValueStorageClient.cs
@@ -21,6 +21,9 @@ namespace EpicGames.Horde.Storage.Clients
 			readonly KeyValueStorageClient _keyValueStorageClient;
 			readonly BlobLocator _locator;
 
+			/// <inheritdoc/>
+			public IBlobHandle Innermost => this;
+
 			/// <summary>
 			/// Constructor
 			/// </summary>
@@ -60,7 +63,8 @@ namespace EpicGames.Horde.Storage.Clients
 			/// <summary>
 			/// Constructor
 			/// </summary>
-			public Writer(KeyValueStorageClient outer, string? basePath)
+			public Writer(KeyValueStorageClient outer, string? basePath, BlobSerializerOptions? options)
+				: base(options)
 			{
 				_outer = outer;
 				_basePath = basePath ?? String.Empty;
@@ -78,7 +82,7 @@ namespace EpicGames.Horde.Storage.Clients
 			public override Task FlushAsync(CancellationToken cancellationToken = default) => Task.CompletedTask;
 
 			/// <inheritdoc/>
-			public override IBlobWriter Fork() => new Writer(_outer, _basePath);
+			public override IBlobWriter Fork() => new Writer(_outer, _basePath, Options);
 
 			/// <inheritdoc/>
 			public override Memory<byte> GetOutputBuffer(int usedSize, int desiredSize)
@@ -94,14 +98,14 @@ namespace EpicGames.Horde.Storage.Clients
 			}
 
 			/// <inheritdoc/>
-			public override async ValueTask<IBlobHandle> WriteBlobAsync(BlobType type, int size, IReadOnlyList<IBlobHandle> references, IReadOnlyList<AliasInfo> aliases, CancellationToken cancellationToken = default)
+			public override async ValueTask<IBlobRef> WriteBlobAsync(BlobType type, int size, IReadOnlyList<IBlobHandle> imports, IReadOnlyList<AliasInfo> aliases, CancellationToken cancellationToken = default)
 			{
 				ReadOnlyMemory<byte> data = _data.AsMemory(_offset, size);
 				_offset += size;
 
 				using ReadOnlyMemoryStream stream = new ReadOnlyMemoryStream(data);
 
-				IBlobHandle handle = await _outer.WriteBlobAsync(type, stream, references, _basePath, cancellationToken);
+				IBlobRef handle = await _outer.WriteBlobAsync(type, stream, imports, _basePath, cancellationToken);
 				foreach (AliasInfo aliasInfo in aliases)
 				{
 					await _outer.AddAliasAsync(aliasInfo.Name, handle, aliasInfo.Rank, aliasInfo.Data, cancellationToken);
@@ -142,7 +146,19 @@ namespace EpicGames.Horde.Storage.Clients
 		}
 
 		/// <inheritdoc/>
-		public IBlobWriter CreateBlobWriter(string? basePath = null) => new Writer(this, basePath);
+		public IBlobRef CreateBlobRef(IoHash hash, BlobLocator locator)
+		{
+			return BlobRef.Create(hash, CreateBlobHandle(locator));
+		}
+
+		/// <inheritdoc/>
+		public IBlobRef<T> CreateBlobRef<T>(IoHash hash, BlobLocator locator, BlobSerializerOptions options)
+		{
+			return BlobRef.Create<T>(hash, CreateBlobHandle(locator), options);
+		}
+
+		/// <inheritdoc/>
+		public IBlobWriter CreateBlobWriter(string? basePath = null, BlobSerializerOptions? options = null) => new Writer(this, basePath, options);
 
 		/// <inheritdoc/>
 		public async ValueTask<BlobData> ReadBlobAsync(BlobLocator locator, CancellationToken cancellationToken = default)
@@ -150,20 +166,21 @@ namespace EpicGames.Horde.Storage.Clients
 			IReadOnlyMemoryOwner<byte> owner = await _backend.ReadBlobAsync(locator, cancellationToken);
 			EncodedBlobData encodedData = new EncodedBlobData(owner.Memory.ToArray());
 
-			return new BlobDataWithOwner(encodedData.Type, encodedData.Payload, encodedData.Refs.Select(x => CreateBlobHandle(x)).ToArray(), owner);
+			return new BlobDataWithOwner(encodedData.Type, encodedData.Payload, encodedData.Imports.Select(x => CreateBlobHandle(x)).ToArray(), owner);
 		}
 
 		/// <inheritdoc/>
-		public async ValueTask<IBlobHandle> WriteBlobAsync(BlobType type, Stream stream, IReadOnlyList<IBlobHandle> references, string? basePath = null, CancellationToken cancellationToken = default)
+		public async ValueTask<IBlobRef> WriteBlobAsync(BlobType type, Stream stream, IReadOnlyList<IBlobHandle> imports, string? basePath = null, CancellationToken cancellationToken = default)
 		{
-			BlobLocator[] locators = references.ConvertAll(x => x.GetLocator()).ToArray();
+			BlobLocator[] locators = imports.ConvertAll(x => x.GetLocator()).ToArray();
 			byte[] payload = await stream.ReadAllBytesAsync(cancellationToken);
 			byte[] encodedData = EncodedBlobData.Create(type, locators, payload);
 
 			using ReadOnlyMemoryStream encodedStream = new ReadOnlyMemoryStream(encodedData);
-			BlobLocator output = await _backend.WriteBlobAsync(encodedStream, basePath, cancellationToken);
+			BlobLocator locator = await _backend.WriteBlobAsync(encodedStream, basePath, cancellationToken);
 
-			return CreateBlobHandle(output);
+			IoHash hash = IoHash.Compute(payload);
+			return BlobRef.Create(hash, CreateBlobHandle(locator));
 		}
 
 		#endregion
diff --git a/Engine/Source/Programs/Shared/EpicGames.Horde/Storage/EncodedBlobData.cs b/Engine/Source/Programs/Shared/EpicGames.Horde/Storage/EncodedBlobData.cs
--- a/Engine/Source/Programs/Shared/EpicGames.Horde/Storage/EncodedBlobData.cs
+++ b/Engine/Source/Programs/Shared/EpicGames.Horde/Storage/EncodedBlobData.cs
@@ -22,11 +22,11 @@ namespace EpicGames.Horde.Storage
 
 		static readonly Version s_currentVersion = Enum.GetValues(typeof(Version)).Cast<Version>().Max();
 
-		struct RefCollection : IReadOnlyList<BlobLocator>
+		struct LocatorCollection : IReadOnlyList<BlobLocator>
 		{
 			public readonly JaggedReadOnlyMemoryArray Array;
 
-			public RefCollection(JaggedReadOnlyMemoryArray array) => Array = array;
+			public LocatorCollection(JaggedReadOnlyMemoryArray array) => Array = array;
 			public BlobLocator this[int index] => new BlobLocator(new Utf8String(Array[index]));
 			public int Count => Array.Count;
 			public IEnumerator<BlobLocator> GetEnumerator() => Array.Select(x => new BlobLocator(new Utf8String(x))).GetEnumerator();
@@ -34,7 +34,7 @@ namespace EpicGames.Horde.Storage
 			IEnumerator IEnumerable.GetEnumerator() => GetEnumerator();
 		}
 
-		readonly RefCollection _refs;
+		readonly LocatorCollection _imports;
 
 		/// <summary>
 		/// Type of the blob
@@ -49,7 +49,7 @@ namespace EpicGames.Horde.Storage
 		/// <summary>
 		/// References to other blobs
 		/// </summary>
-		public IReadOnlyList<BlobLocator> Refs => _refs;
+		public IReadOnlyList<BlobLocator> Imports => _imports;
 
 		/// <summary>
 		/// Constructor
@@ -71,8 +71,8 @@ namespace EpicGames.Horde.Storage
 				Type = BlobType.Read(span.Slice(offset));
 				offset += BlobType.NumBytes;
 
-				_refs = new RefCollection(new JaggedReadOnlyMemoryArray(memory.Slice(offset)));
-				offset += _refs.Array.Data.Length;
+				_imports = new LocatorCollection(new JaggedReadOnlyMemoryArray(memory.Slice(offset)));
+				offset += _imports.Array.Data.Length;
 
 				Payload = memory.Slice(offset);
 			}
@@ -87,10 +87,10 @@ namespace EpicGames.Horde.Storage
 		/// </summary>
 		public static byte[] Create(BlobData blobData)
 		{
-			BlobLocator[] locators = new BlobLocator[blobData.Refs.Count];
-			for (int idx = 0; idx < blobData.Refs.Count; idx++)
+			BlobLocator[] locators = new BlobLocator[blobData.Imports.Count];
+			for (int idx = 0; idx < blobData.Imports.Count; idx++)
 			{
-				locators[idx] = blobData.Refs[idx].GetLocator();
+				locators[idx] = blobData.Imports[idx].GetLocator();
 			}
 
 			return Create(blobData.Type, locators, blobData.Data);
diff --git a/Engine/Source/Programs/Shared/EpicGames.Horde/Storage/IBlobHandle.cs b/Engine/Source/Programs/Shared/EpicGames.Horde/Storage/IBlobHandle.cs
--- a/Engine/Source/Programs/Shared/EpicGames.Horde/Storage/IBlobHandle.cs
+++ b/Engine/Source/Programs/Shared/EpicGames.Horde/Storage/IBlobHandle.cs
@@ -9,10 +9,15 @@ using EpicGames.Core;
 namespace EpicGames.Horde.Storage
 {
 	/// <summary>
-	/// Handle to a node. Can be used to reference nodes that have not been flushed yet.
+	/// Reference to another node in storage. This type is similar to <see cref="IBlobRef"/>, but without a hash.
 	/// </summary>
 	public interface IBlobHandle
 	{
+		/// <summary>
+		/// Accessor for the innermost import
+		/// </summary>
+		IBlobHandle Innermost { get; }
+
 		/// <summary>
 		/// Flush the referenced data to underlying storage
 		/// </summary>
@@ -33,10 +38,9 @@ namespace EpicGames.Horde.Storage
 	}
 
 	/// <summary>
-	/// Typed interface to a particular blob handle
+	/// Handle to a node. Can be used to reference nodes that have not been flushed yet.
 	/// </summary>
-	/// <typeparam name="T">Type of the deserialized blob</typeparam>
-	public interface IBlobHandle<out T> : IBlobHandle
+	public interface IBlobRef : IBlobHandle
 	{
 		/// <summary>
 		/// Hash of the target node
@@ -45,61 +49,93 @@ namespace EpicGames.Horde.Storage
 	}
 
 	/// <summary>
-	/// Extension methods for blob handles
+	/// Typed interface to a particular blob handle
 	/// </summary>
-	public static class BlobHandleExtensions
+	/// <typeparam name="T">Type of the deserialized blob</typeparam>
+	public interface IBlobRef<out T> : IBlobRef
 	{
-		interface IWrappedBlobHandle
-		{
-			public IBlobHandle Inner { get; }
-		}
+		/// <summary>
+		/// Options for deserializing the blob
+		/// </summary>
+		BlobSerializerOptions Options { get; }
+	}
 
-		class TypedBlobHandle<T> : IWrappedBlobHandle, IBlobHandle<T>
+	/// <summary>
+	/// Helper methods for creating blob handles
+	/// </summary>
+	public static class BlobRef
+	{
+		class BlobRefImpl : IBlobRef
 		{
-			readonly IBlobHandle _inner;
 			readonly IoHash _hash;
+			readonly IBlobHandle _handle;
+
+			public IBlobHandle Innermost => _handle.Innermost;
+			public IoHash Hash => _hash;
 
-			public TypedBlobHandle(IBlobHandle inner, IoHash hash)
+			public BlobRefImpl(IoHash hash, IBlobHandle handle)
 			{
-				_inner = inner;
 				_hash = hash;
+				_handle = handle;
 			}
 
-			public IBlobHandle Inner => _inner;
-			public IoHash Hash => _hash;
+			public ValueTask FlushAsync(CancellationToken cancellationToken = default)
+				=> _handle.FlushAsync(cancellationToken);
+
+			public ValueTask<BlobData> ReadBlobDataAsync(CancellationToken cancellationToken = default)
+				=> _handle.ReadBlobDataAsync(cancellationToken);
+
+			public bool TryGetLocator([NotNullWhen(true)] out BlobLocator locator)
+				=> _handle.TryGetLocator(out locator);
+		}
 
-			public ValueTask FlushAsync(CancellationToken cancellationToken = default) => _inner.FlushAsync(cancellationToken);
-			public ValueTask<BlobData> ReadBlobDataAsync(CancellationToken cancellationToken = default) => _inner.ReadBlobDataAsync(cancellationToken);
-			public bool TryGetLocator([NotNullWhen(true)] out BlobLocator locator) => _inner.TryGetLocator(out locator);
+		class BlobRefImpl<T> : BlobRefImpl, IBlobRef<T>
+		{
+			readonly BlobSerializerOptions _options;
+
+			public BlobSerializerOptions Options => _options;
+
+			public BlobRefImpl(IoHash hash, IBlobHandle handle, BlobSerializerOptions options)
+				: base(hash, handle)
+			{
+				_options = options;
+			}
 		}
 
 		/// <summary>
-		/// Creates a typed blob handle
+		/// Create an untyped blob handle
 		/// </summary>
-		public static IBlobHandle<T> ForType<T>(this IBlobHandle handle, IoHash hash) => new TypedBlobHandle<T>(handle, hash);
+		/// <param name="handle">Imported blob interface</param>
+		/// <param name="hash">Hash of the blob</param>
+		/// <returns>Handle to the blob</returns>
+		public static IBlobRef Create(IoHash hash, IBlobHandle handle)
+			=> new BlobRefImpl(hash, handle);
 
 		/// <summary>
-		/// 
+		/// Create a typed blob handle
 		/// </summary>
-		/// <param name="handle"></param>
-		/// <returns></returns>
-		public static IBlobHandle Unwrap(this IBlobHandle handle)
-		{
-			while (handle is IWrappedBlobHandle wrappedHandle)
-			{
-				handle = wrappedHandle.Inner;
-			}
-			return handle;
-		}
+		/// <typeparam name="T"></typeparam>
+		/// <param name="hash">Hash of the blob</param>
+		/// <param name="handle">Imported blob interface</param>
+		/// <param name="options">Options for deserializing the target blob</param>
+		/// <returns>Handle to the blob</returns>
+		public static IBlobRef<T> Create<T>(IoHash hash, IBlobHandle handle, BlobSerializerOptions options)
+			=> new BlobRefImpl<T>(hash, handle, options);
+	}
 
+	/// <summary>
+	/// Extension methods for <see cref="IBlobHandle"/>
+	/// </summary>
+	public static class BlobHandleExtensions
+	{
 		/// <summary>
 		/// Gets a path to this blob that can be used to describe blob references over the wire.
 		/// </summary>
-		/// <param name="handle">Handle to query</param>
-		public static BlobLocator GetLocator(this IBlobHandle handle)
+		/// <param name="import">Handle to query</param>
+		public static BlobLocator GetLocator(this IBlobHandle import)
 		{
 			BlobLocator locator;
-			if (!handle.TryGetLocator(out locator))
+			if (!import.TryGetLocator(out locator))
 			{
 				throw new InvalidOperationException("Blob has not yet been written to storage");
 			}
diff --git a/Engine/Source/Programs/Shared/EpicGames.Horde/Storage/IBlobReader.cs b/Engine/Source/Programs/Shared/EpicGames.Horde/Storage/IBlobReader.cs
--- a/Engine/Source/Programs/Shared/EpicGames.Horde/Storage/IBlobReader.cs
+++ b/Engine/Source/Programs/Shared/EpicGames.Horde/Storage/IBlobReader.cs
@@ -4,6 +4,8 @@ using System;
 using System.Collections.Generic;
 using EpicGames.Core;
 
+#pragma warning disable CA1716 // Do not use 'Imports' as identifier
+
 namespace EpicGames.Horde.Storage
 {
 	/// <summary>
@@ -24,12 +26,17 @@ namespace EpicGames.Horde.Storage
 		/// <summary>
 		/// Locations of all referenced nodes. These handles do not have valid hashes.
 		/// </summary>
-		IReadOnlyList<IBlobHandle> References { get; }
+		IReadOnlyList<IBlobHandle> Imports { get; }
+
+		/// <summary>
+		/// Gets the next serialized blob handle
+		/// </summary>
+		IBlobRef ReadBlobRef();
 
 		/// <summary>
 		/// Gets the next serialized blob handle
 		/// </summary>
-		IBlobHandle<T> ReadBlobHandle<T>();
+		IBlobRef<T> ReadBlobRef<T>();
 	}
 
 	/// <summary>
@@ -65,28 +72,40 @@ namespace EpicGames.Horde.Storage
 		/// <summary>
 		/// Locations of all referenced nodes.
 		/// </summary>
-		public IReadOnlyList<IBlobHandle> References => _blobData.Refs;
+		public IReadOnlyList<IBlobHandle> Imports => _blobData.Imports;
 
 		readonly BlobData _blobData;
-		int _refIdx;
+		readonly BlobSerializerOptions _options;
+		int _importIdx;
 
 		/// <summary>
 		/// Constructor
 		/// </summary>
-		public BlobReader(BlobData blobData)
+		public BlobReader(BlobData blobData, BlobSerializerOptions? options)
 			: base(blobData.Data)
 		{
 			_blobData = blobData;
+			_options = options ?? BlobSerializerOptions.Default;
 		}
 
 		/// <summary>
-		/// Gets the next serialized blob handle
+		/// Gets the next serialized blob reference
+		/// </summary>
+		public IBlobRef ReadBlobRef()
+		{
+			IBlobHandle import = Imports[_importIdx++];
+			IoHash hash = this.ReadIoHash();
+			return BlobRef.Create(hash, import);
+		}
+
+		/// <summary>
+		/// Gets the next serialized blob reference
 		/// </summary>
-		public IBlobHandle<T> ReadBlobHandle<T>()
+		public IBlobRef<T> ReadBlobRef<T>()
 		{
-			IBlobHandle handle = References[_refIdx++];
+			IBlobHandle import = Imports[_importIdx++];
 			IoHash hash = this.ReadIoHash();
-			return handle.ForType<T>(hash);
+			return BlobRef.Create<T>(hash, import, _options);
 		}
 	}
 }
diff --git a/Engine/Source/Programs/Shared/EpicGames.Horde/Storage/IBlobWriter.cs b/Engine/Source/Programs/Shared/EpicGames.Horde/Storage/IBlobWriter.cs
--- a/Engine/Source/Programs/Shared/EpicGames.Horde/Storage/IBlobWriter.cs
+++ b/Engine/Source/Programs/Shared/EpicGames.Horde/Storage/IBlobWriter.cs
@@ -6,6 +6,8 @@ using System.Threading;
 using System.Threading.Tasks;
 using EpicGames.Core;
 
+#pragma warning disable CA1716 // Do not use 'imports' as variable name
+
 namespace EpicGames.Horde.Storage
 {
 	/// <summary>
@@ -13,6 +15,11 @@ namespace EpicGames.Horde.Storage
 	/// </summary>
 	public interface IBlobWriter : IMemoryWriter, IAsyncDisposable
 	{
+		/// <summary>
+		/// Options for serialization
+		/// </summary>
+		BlobSerializerOptions Options { get; }
+
 		/// <summary>
 		/// Accessor for the memory written to the current blob
 		/// </summary>
@@ -44,7 +51,7 @@ namespace EpicGames.Horde.Storage
 		/// <param name="type">Type of the node that was written</param>
 		/// <param name="cancellationToken">Cancellation token for the operation</param>
 		/// <returns>Handle to the written node</returns>
-		ValueTask<IBlobHandle> CompleteAsync(BlobType type, CancellationToken cancellationToken = default);
+		ValueTask<IBlobRef> CompleteAsync(BlobType type, CancellationToken cancellationToken = default);
 
 		/// <summary>
 		/// Finish writing a blob that has been written into the output buffer.
@@ -52,13 +59,13 @@ namespace EpicGames.Horde.Storage
 		/// <param name="type">Type of the node that was written</param>
 		/// <param name="cancellationToken">Cancellation token for the operation</param>
 		/// <returns>Handle to the written node</returns>
-		ValueTask<IBlobHandle<T>> CompleteAsync<T>(BlobType type, CancellationToken cancellationToken = default);
+		ValueTask<IBlobRef<T>> CompleteAsync<T>(BlobType type, CancellationToken cancellationToken = default);
 
 		/// <summary>
 		/// Writes a reference to another blob. The blob's hash is serialized to the output stream.
 		/// </summary>
-		/// <param name="handle">Referenced blob</param>
-		void WriteBlobHandle<T>(IBlobHandle<T> handle);
+		/// <param name="blobRef">Referenced blob</param>
+		void WriteBlobRef(IBlobRef blobRef);
 	}
 
 	/// <summary>
@@ -76,21 +83,27 @@ namespace EpicGames.Horde.Storage
 	{
 		Memory<byte> _memory;
 		readonly List<AliasInfo> _aliases = new List<AliasInfo>();
-		readonly List<IBlobHandle> _refs = new List<IBlobHandle>();
+		readonly List<IBlobHandle> _imports = new List<IBlobHandle>();
+		readonly BlobSerializerOptions _options;
 		int _length;
 
-		/// <summary>
-		/// List of serialized references
-		/// </summary>
-		public IReadOnlyList<IBlobHandle> References => _refs;
-
 		/// <inheritdoc/>
 		public int Length => _length;
 
+		/// <inheritdoc/>
+		public BlobSerializerOptions Options => _options;
+
+		/// <inheritdoc/>
+		public ReadOnlyMemory<byte> WrittenMemory => _memory.Slice(0, _length);
+
 		/// <summary>
-		/// Memory that has been written
+		/// Constructor
 		/// </summary>
-		public ReadOnlyMemory<byte> WrittenMemory => _memory.Slice(0, _length);
+		/// <param name="options"></param>
+		protected BlobWriter(BlobSerializerOptions? options)
+		{
+			_options = options ?? BlobSerializerOptions.Default;
+		}
 
 		/// <summary>
 		/// Computes the hash of the written data
@@ -100,10 +113,10 @@ namespace EpicGames.Horde.Storage
 		/// <summary>
 		/// Writes a handle to another node
 		/// </summary>
-		public void WriteBlobHandle<T>(IBlobHandle<T> target)
+		public void WriteBlobRef(IBlobRef target)
 		{
 			this.WriteIoHash(target.Hash);
-			_refs.Add(target);
+			_imports.Add(target);
 		}
 
 		/// <inheritdoc/>
@@ -137,11 +150,11 @@ namespace EpicGames.Horde.Storage
 		/// </summary>
 		/// <param name="type">Type of the blob</param>
 		/// <param name="size">Size of the blob to write</param>
-		/// <param name="references">References to other blobs</param>
+		/// <param name="imports">References to other blobs</param>
 		/// <param name="aliases">Aliases for the new blob</param>
 		/// <param name="cancellationToken">Cancellation token for the operation</param>
 		/// <returns>New buffer</returns>
-		public abstract ValueTask<IBlobHandle> WriteBlobAsync(BlobType type, int size, IReadOnlyList<IBlobHandle> references, IReadOnlyList<AliasInfo> aliases, CancellationToken cancellationToken);
+		public abstract ValueTask<IBlobRef> WriteBlobAsync(BlobType type, int size, IReadOnlyList<IBlobHandle> imports, IReadOnlyList<AliasInfo> aliases, CancellationToken cancellationToken);
 
 		/// <inheritdoc/>
 		public void AddAlias(string name, int rank, ReadOnlyMemory<byte> data)
@@ -154,24 +167,24 @@ namespace EpicGames.Horde.Storage
 		public abstract IBlobWriter Fork();
 
 		/// <inheritdoc/>
-		public async ValueTask<IBlobHandle> CompleteAsync(BlobType type, CancellationToken cancellationToken = default)
+		public async ValueTask<IBlobRef> CompleteAsync(BlobType type, CancellationToken cancellationToken = default)
 		{
-			IBlobHandle handle = await WriteBlobAsync(type, _length, _refs, _aliases, cancellationToken);
+			IBlobRef blobRef = await WriteBlobAsync(type, _length, _imports, _aliases, cancellationToken);
 
 			_memory = default;
 			_length = 0;
-			_refs.Clear();
+			_imports.Clear();
 			_aliases.Clear();
 
-			return handle;
+			return blobRef;
 		}
 
 		/// <inheritdoc/>
-		public async ValueTask<IBlobHandle<T>> CompleteAsync<T>(BlobType type, CancellationToken cancellationToken = default)
+		public async ValueTask<IBlobRef<T>> CompleteAsync<T>(BlobType type, CancellationToken cancellationToken = default)
 		{
 			IoHash hash = IoHash.Compute(_memory.Span.Slice(0, _length));
-			IBlobHandle handle = await CompleteAsync(type, cancellationToken);
-			return handle.ForType<T>(hash);
+			IBlobRef blobRef = await CompleteAsync(type, cancellationToken);
+			return BlobRef.Create<T>(hash, blobRef, _options);
 		}
 
 		/// <inheritdoc/>
diff --git a/Engine/Source/Programs/Shared/EpicGames.Horde/Storage/IStorageClient.cs b/Engine/Source/Programs/Shared/EpicGames.Horde/Storage/IStorageClient.cs
--- a/Engine/Source/Programs/Shared/EpicGames.Horde/Storage/IStorageClient.cs
+++ b/Engine/Source/Programs/Shared/EpicGames.Horde/Storage/IStorageClient.cs
@@ -5,6 +5,7 @@ using System.Collections.Generic;
 using System.Linq;
 using System.Threading;
 using System.Threading.Tasks;
+using EpicGames.Core;
 using Microsoft.Extensions.Logging;
 
 namespace EpicGames.Horde.Storage
@@ -24,18 +25,36 @@ namespace EpicGames.Horde.Storage
 		#region Blobs
 
 		/// <summary>
-		/// Creates a new blob handle by parsing a blob id
+		/// Creates a new blob handle by parsing a locator
 		/// </summary>
 		/// <param name="locator">Path to the blob</param>
 		/// <returns>New handle to the blob</returns>
 		IBlobHandle CreateBlobHandle(BlobLocator locator);
 
+		/// <summary>
+		/// Creates a new blob reference from a locator and hash
+		/// </summary>
+		/// <param name="hash">Hash of the target blob</param>
+		/// <param name="locator">Path to the blob</param>
+		/// <returns>New handle to the blob</returns>
+		IBlobRef CreateBlobRef(IoHash hash, BlobLocator locator);
+
+		/// <summary>
+		/// Creates a new blob reference from a locator and hash
+		/// </summary>
+		/// <param name="hash">Hash of the target blob</param>
+		/// <param name="locator">Path to the blob</param>
+		/// <param name="options">Options for deserializing the blob</param>
+		/// <returns>New handle to the blob</returns>
+		IBlobRef<T> CreateBlobRef<T>(IoHash hash, BlobLocator locator, BlobSerializerOptions options);
+
 		/// <summary>
 		/// Creates a new writer for storage blobs
 		/// </summary>
 		/// <param name="basePath">Base path for any nodes written from the writer.</param>
+		/// <param name="serializerOptions">Options for serializing classes</param>
 		/// <returns>New writer instance. Must be disposed after use.</returns>
-		IBlobWriter CreateBlobWriter(string? basePath = null);
+		IBlobWriter CreateBlobWriter(string? basePath = null, BlobSerializerOptions? serializerOptions = null);
 
 		#endregion
 
diff --git a/Engine/Source/Programs/Shared/EpicGames.Horde/Storage/IStorageWriter.cs b/Engine/Source/Programs/Shared/EpicGames.Horde/Storage/IStorageWriter.cs
--- a/Engine/Source/Programs/Shared/EpicGames.Horde/Storage/IStorageWriter.cs
+++ b/Engine/Source/Programs/Shared/EpicGames.Horde/Storage/IStorageWriter.cs
@@ -20,16 +20,16 @@ namespace EpicGames.Horde.Storage
 		{
 			readonly int _maxKeys;
 			readonly Queue<BlobKey> _blobKeys = new Queue<BlobKey>();
-			readonly Dictionary<BlobKey, IBlobHandle> _blobKeyToHandle = new Dictionary<BlobKey, IBlobHandle>();
+			readonly Dictionary<BlobKey, IBlobRef> _blobKeyToHandle = new Dictionary<BlobKey, IBlobRef>();
 
 			public DedupeCache(int maxKeys)
 			{
 				_maxKeys = maxKeys;
 				_blobKeys = new Queue<BlobKey>(maxKeys);
-				_blobKeyToHandle = new Dictionary<BlobKey, IBlobHandle>(maxKeys);
+				_blobKeyToHandle = new Dictionary<BlobKey, IBlobRef>(maxKeys);
 			}
 
-			internal void Add(BlobKey key, IBlobHandle handle)
+			internal void Add(BlobKey key, IBlobRef handle)
 			{
 				BlobKey? prevKey;
 				if (_blobKeys.Count == _maxKeys && _blobKeys.TryDequeue(out prevKey))
@@ -39,13 +39,21 @@ namespace EpicGames.Horde.Storage
 				_blobKeyToHandle.TryAdd(key, handle);
 			}
 
-			internal bool TryGetValue(BlobKey key, [NotNullWhen(true)] out IBlobHandle? handle) => _blobKeyToHandle.TryGetValue(key, out handle);
+			internal bool TryGetValue(BlobKey key, [NotNullWhen(true)] out IBlobRef? handle) => _blobKeyToHandle.TryGetValue(key, out handle);
 		}
 
-		class WrappedHandle : IBlobHandle
+		class WrappedHandle : IBlobRef
 		{
 			public object _lockObject = new object();
-			public IBlobHandle? _inner;
+			public IBlobRef? _inner;
+
+			/// <inheritdoc/>
+			public IBlobHandle Innermost
+				=> _inner!.Innermost;
+
+			/// <inheritdoc/>
+			public IoHash Hash 
+				=> _inner!.Hash;
 
 			/// <inheritdoc/>
 			public bool TryGetLocator(out BlobLocator locator)
@@ -80,12 +88,14 @@ namespace EpicGames.Horde.Storage
 		/// <param name="inner"></param>
 		/// <param name="maxKeys"></param>
 		public DedupeStorageWriter(IBlobWriter inner, int maxKeys = DefaultMaxKeys)
+			: base(inner.Options)
 		{
 			_inner = (BlobWriter)inner;
 			_cache = new DedupeCache(maxKeys);
 		}
 
 		private DedupeStorageWriter(IBlobWriter inner, DedupeCache cache)
+			: base(inner.Options)
 		{
 			_inner = (BlobWriter)inner;
 			_cache = cache;
@@ -104,7 +114,7 @@ namespace EpicGames.Horde.Storage
 		public override Memory<byte> GetOutputBuffer(int usedSize, int desiredSize) => _inner.GetOutputBuffer(usedSize, desiredSize);
 
 		/// <inheritdoc/>
-		public override async ValueTask<IBlobHandle> WriteBlobAsync(BlobType type, int size, IReadOnlyList<IBlobHandle> references, IReadOnlyList<AliasInfo> aliases, CancellationToken cancellationToken = default)
+		public override async ValueTask<IBlobRef> WriteBlobAsync(BlobType type, int size, IReadOnlyList<IBlobHandle> imports, IReadOnlyList<AliasInfo> aliases, CancellationToken cancellationToken = default)
 		{
 			ReadOnlyMemory<byte> data = _inner.GetOutputBuffer(size, size).Slice(0, size);
 			IoHash hash = IoHash.Compute(data.Span);
@@ -113,7 +123,7 @@ namespace EpicGames.Horde.Storage
 			WrappedHandle? wrappedHandle;
 			lock (_cache)
 			{
-				IBlobHandle? handle;
+				IBlobRef? handle;
 				if (_cache.TryGetValue(key, out handle))
 				{
 					return handle;
@@ -123,7 +133,7 @@ namespace EpicGames.Horde.Storage
 				_cache.Add(key, wrappedHandle);
 			}
 
-			wrappedHandle._inner = await _inner.WriteBlobAsync(type, size, references.ConvertAll(x => ((WrappedHandle)x.Unwrap())._inner!), aliases, cancellationToken);
+			wrappedHandle._inner = await _inner.WriteBlobAsync(type, size, imports.ConvertAll(x => x.Innermost), aliases, cancellationToken);
 			return wrappedHandle;
 		}
 	}
diff --git a/Engine/Source/Programs/Shared/EpicGames.Horde/Storage/Nodes/CbNode.cs b/Engine/Source/Programs/Shared/EpicGames.Horde/Storage/Nodes/CbNode.cs
--- a/Engine/Source/Programs/Shared/EpicGames.Horde/Storage/Nodes/CbNode.cs
+++ b/Engine/Source/Programs/Shared/EpicGames.Horde/Storage/Nodes/CbNode.cs
@@ -4,6 +4,7 @@ using EpicGames.Core;
 using EpicGames.Serialization;
 using System;
 using System.Collections.Generic;
+using System.Linq;
 
 namespace EpicGames.Horde.Storage.Nodes
 {
@@ -26,65 +27,29 @@ namespace EpicGames.Horde.Storage.Nodes
 		/// <summary>
 		/// Imported nodes
 		/// </summary>
-		public IReadOnlyList<IBlobHandle<object>> References { get; }
+		public IReadOnlyList<IBlobHandle> Imports { get; }
 
 		/// <summary>
 		/// Constructor
 		/// </summary>
 		/// <param name="obj">The compact binary object</param>
-		/// <param name="references">List of references to attachments</param>
-		public CbNode(CbObject obj, IReadOnlyList<IBlobHandle<object>> references)
+		/// <param name="imports">List of imports for attachments</param>
+		public CbNode(CbObject obj, IReadOnlyList<IBlobHandle> imports)
 		{
 			Object = obj;
-			References = references;
+			Imports = imports;
 		}
 	}
 
 	class CbNodeConverter : BlobConverter<CbNode>
 	{
-		class HandleMapper
-		{
-			readonly IBlobReader _reader;
-			readonly List<IBlobHandle<object>> _refs;
-			int _refIdx;
-
-			public HandleMapper(IBlobReader reader, List<IBlobHandle<object>> refs)
-			{
-				_reader = reader;
-				_refs = refs;
-			}
-
-			public void IterateField(CbField field)
-			{
-				if (field.IsAttachment())
-				{
-					IBlobHandle handle = _reader.References[_refIdx++];
-					_refs.Add(handle.ForType<object>(field.AsAttachment()));
-				}
-				else if (field.IsArray())
-				{
-					CbArray array = field.AsArray();
-					array.IterateAttachments(IterateField);
-				}
-				else if (field.IsObject())
-				{
-					CbObject obj = field.AsObject();
-					obj.IterateAttachments(IterateField);
-				}
-			}
-		}
-
 		public static BlobType BlobType { get; } = new BlobType(CbNode.BlobTypeGuid, 1);
 
 		/// <inheritdoc/>
 		public override CbNode Read(IBlobReader reader, BlobSerializerOptions options)
 		{
 			CbObject obj = new CbObject(reader.GetMemory().ToArray());
-
-			List<IBlobHandle<object>> references = new List<IBlobHandle<object>>();
-			obj.IterateAttachments(new HandleMapper(reader, references).IterateField);
-
-			return new CbNode(obj, references);
+			return new CbNode(obj, reader.Imports.ToArray());
 		}
 
 		/// <inheritdoc/>
diff --git a/Engine/Source/Programs/Shared/EpicGames.Horde/Storage/Nodes/ChunkedDataNode.cs b/Engine/Source/Programs/Shared/EpicGames.Horde/Storage/Nodes/ChunkedDataNode.cs
--- a/Engine/Source/Programs/Shared/EpicGames.Horde/Storage/Nodes/ChunkedDataNode.cs
+++ b/Engine/Source/Programs/Shared/EpicGames.Horde/Storage/Nodes/ChunkedDataNode.cs
@@ -151,12 +151,12 @@ namespace EpicGames.Horde.Storage.Nodes
 	/// <param name="Type">Type of the referenced node</param>
 	/// <param name="Length">Length of the data stream within this node</param>
 	/// <param name="Handle">Handle to the target node</param>
-	public record class ChunkedDataNodeRef(ChunkedDataNodeType Type, long Length, IBlobHandle<ChunkedDataNode> Handle)
+	public record class ChunkedDataNodeRef(ChunkedDataNodeType Type, long Length, IBlobRef<ChunkedDataNode> Handle)
 	{
 		/// <summary>
 		/// Leaf node constructor
 		/// </summary>
-		public ChunkedDataNodeRef(long length, IBlobHandle<LeafChunkedDataNode> handle)
+		public ChunkedDataNodeRef(long length, IBlobRef<LeafChunkedDataNode> handle)
 			: this(ChunkedDataNodeType.Leaf, length, handle)
 		{
 		}
@@ -164,7 +164,7 @@ namespace EpicGames.Horde.Storage.Nodes
 		/// <summary>
 		/// Interior node constructor
 		/// </summary>
-		public ChunkedDataNodeRef(long length, IBlobHandle<InteriorChunkedDataNode> handle)
+		public ChunkedDataNodeRef(long length, IBlobRef<InteriorChunkedDataNode> handle)
 			: this(ChunkedDataNodeType.Interior, length, handle)
 		{
 		}
@@ -322,7 +322,7 @@ namespace EpicGames.Horde.Storage.Nodes
 					sizeSinceProgressUpdate = 0;
 				}
 
-				IBlobHandle<LeafChunkedDataNode> blobHandle = await writer.CompleteAsync<LeafChunkedDataNode>(LeafChunkedDataNodeConverter.BlobType, cancellationToken);
+				IBlobRef<LeafChunkedDataNode> blobHandle = await writer.CompleteAsync<LeafChunkedDataNode>(LeafChunkedDataNodeConverter.BlobType, cancellationToken);
 				leafNodeRefs.Add(new ChunkedDataNodeRef(nextLength, blobHandle));
 
 				readBuffer.Memory.Slice(nextLength, size - nextLength).CopyTo(readBuffer.Memory);
@@ -460,13 +460,12 @@ namespace EpicGames.Horde.Storage.Nodes
 		/// </summary>
 		/// <param name="leafChunkedData">List of leaf handles</param>
 		/// <param name="options">Options for splitting the tree</param>
-		/// <param name="serializerOptions">Options for serialization</param>
 		/// <param name="writer">Output writer for new interior nodes</param>
 		/// <param name="cancellationToken">Cancellation token for the operation</param>
 		/// <returns>Handle to the root node of the tree</returns>
-		public static async Task<ChunkedData> CreateTreeAsync(LeafChunkedData leafChunkedData, InteriorChunkedDataNodeOptions options, IBlobWriter writer, BlobSerializerOptions? serializerOptions, CancellationToken cancellationToken)
+		public static async Task<ChunkedData> CreateTreeAsync(LeafChunkedData leafChunkedData, InteriorChunkedDataNodeOptions options, IBlobWriter writer, CancellationToken cancellationToken)
 		{
-			ChunkedDataNodeRef rootRef = await CreateTreeAsync(leafChunkedData.LeafHandles, options, writer, serializerOptions, cancellationToken);
+			ChunkedDataNodeRef rootRef = await CreateTreeAsync(leafChunkedData.LeafHandles, options, writer, cancellationToken);
 			return new ChunkedData(leafChunkedData.Hash, rootRef);
 		}
 
@@ -476,10 +475,9 @@ namespace EpicGames.Horde.Storage.Nodes
 		/// <param name="nodeRefs">List of leaf nodes</param>
 		/// <param name="chunkingOptions">Options for splitting the tree</param>
 		/// <param name="writer">Output writer for new interior nodes</param>
-		/// <param name="serializerOptions">Options for serializing blobs</param>
 		/// <param name="cancellationToken">Cancellation token for the operation</param>
 		/// <returns>Handle to the root node of the tree</returns>
-		public static async Task<ChunkedDataNodeRef> CreateTreeAsync(List<ChunkedDataNodeRef> nodeRefs, InteriorChunkedDataNodeOptions chunkingOptions, IBlobWriter writer, BlobSerializerOptions? serializerOptions, CancellationToken cancellationToken)
+		public static async Task<ChunkedDataNodeRef> CreateTreeAsync(List<ChunkedDataNodeRef> nodeRefs, InteriorChunkedDataNodeOptions chunkingOptions, IBlobWriter writer, CancellationToken cancellationToken)
 		{
 			List<ChunkedDataNodeRef> handleBuffer = new List<ChunkedDataNodeRef>();
 
@@ -492,7 +490,7 @@ namespace EpicGames.Horde.Storage.Nodes
 				handleBuffer.Clear();
 				foreach ((InteriorChunkedDataNode interiorNode, long interiorLength) in interiorNodes)
 				{
-					IBlobHandle<InteriorChunkedDataNode> interiorHandle = await writer.WriteBlobAsync(interiorNode, serializerOptions, cancellationToken);
+					IBlobRef<InteriorChunkedDataNode> interiorHandle = await writer.WriteBlobAsync(interiorNode, cancellationToken);
 					handleBuffer.Add(new ChunkedDataNodeRef(ChunkedDataNodeType.Interior, interiorLength, interiorHandle));
 				}
 
@@ -555,10 +553,10 @@ namespace EpicGames.Horde.Storage.Nodes
 		/// <param name="cancellationToken">Cancellation token for the operation</param>
 		public static async Task CopyToStreamAsync(BlobData nodeData, Stream outputStream, CancellationToken cancellationToken)
 		{
-			BlobReader nodeReader = new BlobReader(nodeData);
+			BlobReader nodeReader = new BlobReader(nodeData, null);
 			while (nodeReader.GetMemory(0).Length > 0)
 			{
-				IBlobHandle<ChunkedDataNode> handle = nodeReader.ReadBlobHandle<ChunkedDataNode>();
+				IBlobRef<ChunkedDataNode> handle = nodeReader.ReadBlobRef<ChunkedDataNode>();
 				if (nodeReader.Version >= 2)
 				{
 					_ = nodeReader.ReadUnsignedVarInt(); // Type
@@ -602,7 +600,7 @@ namespace EpicGames.Horde.Storage.Nodes
 			List<ChunkedDataNodeRef> children = new List<ChunkedDataNodeRef>();
 			while (reader.GetMemory().Length > 0)
 			{
-				IBlobHandle<ChunkedDataNode> handle = reader.ReadBlobHandle<ChunkedDataNode>();
+				IBlobRef<ChunkedDataNode> handle = reader.ReadBlobRef<ChunkedDataNode>();
 
 				ChunkedDataNodeType type = ChunkedDataNodeType.Unknown;
 				if (reader.Version >= 2)
@@ -626,7 +624,7 @@ namespace EpicGames.Horde.Storage.Nodes
 		{
 			foreach (ChunkedDataNodeRef child in value.Children)
 			{
-				writer.WriteBlobHandle(child.Handle);
+				writer.WriteBlobRef(child.Handle);
 				if (_writeVersion >= 2)
 				{
 					writer.WriteUnsignedVarInt((int)child.Type);
diff --git a/Engine/Source/Programs/Shared/EpicGames.Horde/Storage/Nodes/ChunkedDataWriter.cs b/Engine/Source/Programs/Shared/EpicGames.Horde/Storage/Nodes/ChunkedDataWriter.cs
--- a/Engine/Source/Programs/Shared/EpicGames.Horde/Storage/Nodes/ChunkedDataWriter.cs
+++ b/Engine/Source/Programs/Shared/EpicGames.Horde/Storage/Nodes/ChunkedDataWriter.cs
@@ -338,7 +338,7 @@ namespace EpicGames.Horde.Storage.Nodes
 		async ValueTask FlushLeafNodeAsync(CancellationToken cancellationToken)
 		{
 			int leafLength = _writer.WrittenMemory.Length;
-			IBlobHandle<LeafChunkedDataNode> leafHandle = await _writer.CompleteAsync<LeafChunkedDataNode>(LeafChunkedDataNodeConverter.BlobType, cancellationToken);
+			IBlobRef<LeafChunkedDataNode> leafHandle = await _writer.CompleteAsync<LeafChunkedDataNode>(LeafChunkedDataNodeConverter.BlobType, cancellationToken);
 			_leafHandles.Add(new ChunkedDataNodeRef(leafLength, leafHandle));
 			ResetLeafState();
 		}
@@ -358,7 +358,6 @@ namespace EpicGames.Horde.Storage.Nodes
 	{
 		readonly IBlobWriter _writer;
 		readonly ChunkingOptions _chunkingOptions;
-		readonly BlobSerializerOptions _blobSerializerOptions;
 		readonly LeafChunkedDataWriter _leafWriter;
 
 		/// <summary>
@@ -369,11 +368,10 @@ namespace EpicGames.Horde.Storage.Nodes
 		/// <summary>
 		/// Constructor
 		/// </summary>
-		public ChunkedDataWriter(IBlobWriter writer, ChunkingOptions chunkingOptions, BlobSerializerOptions blobSerializerOptions)
+		public ChunkedDataWriter(IBlobWriter writer, ChunkingOptions chunkingOptions)
 		{
 			_writer = writer;
 			_chunkingOptions = chunkingOptions;
-			_blobSerializerOptions = blobSerializerOptions;
 			_leafWriter = new LeafChunkedDataWriter(writer, chunkingOptions.LeafOptions);
 		}
 		
@@ -400,7 +398,7 @@ namespace EpicGames.Horde.Storage.Nodes
 		{
 			Reset();
 			LeafChunkedData leafChunkedData = await _leafWriter.CreateAsync(fileInfo, cancellationToken);
-			return await InteriorChunkedDataNode.CreateTreeAsync(leafChunkedData, _chunkingOptions.InteriorOptions, _writer, _blobSerializerOptions, cancellationToken);
+			return await InteriorChunkedDataNode.CreateTreeAsync(leafChunkedData, _chunkingOptions.InteriorOptions, _writer, cancellationToken);
 		}
 
 		/// <summary>
@@ -433,7 +431,7 @@ namespace EpicGames.Horde.Storage.Nodes
 		public async Task<ChunkedData> CompleteAsync(CancellationToken cancellationToken = default)
 		{
 			LeafChunkedData leafChunkedData = await _leafWriter.CompleteAsync(cancellationToken);
-			return await InteriorChunkedDataNode.CreateTreeAsync(leafChunkedData, _chunkingOptions.InteriorOptions, _writer, _blobSerializerOptions, cancellationToken);
+			return await InteriorChunkedDataNode.CreateTreeAsync(leafChunkedData, _chunkingOptions.InteriorOptions, _writer, cancellationToken);
 		}
 
 		/// <summary>
diff --git a/Engine/Source/Programs/Shared/EpicGames.Horde/Storage/Nodes/CommitNode.cs b/Engine/Source/Programs/Shared/EpicGames.Horde/Storage/Nodes/CommitNode.cs
--- a/Engine/Source/Programs/Shared/EpicGames.Horde/Storage/Nodes/CommitNode.cs
+++ b/Engine/Source/Programs/Shared/EpicGames.Horde/Storage/Nodes/CommitNode.cs
@@ -25,7 +25,7 @@ namespace EpicGames.Horde.Storage.Nodes
 		/// <summary>
 		/// Reference to the parent commit
 		/// </summary>
-		public IBlobHandle<CommitNode>? Parent { get; set; }
+		public IBlobRef<CommitNode>? Parent { get; set; }
 
 		/// <summary>
 		/// Human readable name of the author of this change
@@ -65,12 +65,12 @@ namespace EpicGames.Horde.Storage.Nodes
 		/// <summary>
 		/// Metadata for this commit, keyed by arbitrary GUID
 		/// </summary>
-		public Dictionary<Guid, IBlobHandle<object>> Metadata { get; } = new Dictionary<Guid, IBlobHandle<object>>();
+		public Dictionary<Guid, IBlobRef> Metadata { get; } = new Dictionary<Guid, IBlobRef>();
 
 		/// <summary>
 		/// Constructor
 		/// </summary>
-		public CommitNode(int number, IBlobHandle<CommitNode>? parent, string author, string? authorId, string? committer, string? commiterId, string message, DateTime time, DirectoryNodeRef contents, Dictionary<Guid, IBlobHandle<object>> metadata)
+		public CommitNode(int number, IBlobRef<CommitNode>? parent, string author, string? authorId, string? committer, string? commiterId, string message, DateTime time, DirectoryNodeRef contents, Dictionary<Guid, IBlobRef> metadata)
 		{
 			Number = number;
 			Parent = parent;
@@ -93,10 +93,10 @@ namespace EpicGames.Horde.Storage.Nodes
 		{
 			int number = (int)reader.ReadUnsignedVarInt();
 
-			IBlobHandle<CommitNode>? parent;
+			IBlobRef<CommitNode>? parent;
 			if (reader.ReadBoolean())
 			{
-				parent = reader.ReadBlobHandle<CommitNode>();
+				parent = reader.ReadBlobRef<CommitNode>();
 			}
 			else
 			{
@@ -110,11 +110,11 @@ namespace EpicGames.Horde.Storage.Nodes
 			string message = reader.ReadString();
 			DateTime time = reader.ReadDateTime();
 
-			IBlobHandle<DirectoryNode> contentsNode = reader.ReadBlobHandle<DirectoryNode>();
+			IBlobRef<DirectoryNode> contentsNode = reader.ReadBlobRef<DirectoryNode>();
 			long length = (long)reader.ReadUnsignedVarInt();
 			DirectoryNodeRef contents = new DirectoryNodeRef(length, contentsNode);
 
-			Dictionary<Guid, IBlobHandle<object>> metadata = reader.ReadDictionary(() => reader.ReadGuidUnrealOrder(), () => reader.ReadBlobHandle<object>());
+			Dictionary<Guid, IBlobRef> metadata = reader.ReadDictionary(() => reader.ReadGuidUnrealOrder(), () => reader.ReadBlobRef());
 
 			return new CommitNode(number, parent, author, authorId, committer, committerId, message, time, contents, metadata);
 		}
@@ -126,7 +126,7 @@ namespace EpicGames.Horde.Storage.Nodes
 			writer.WriteBoolean(value.Parent != null);
 			if (value.Parent != null)
 			{
-				writer.WriteBlobHandle(value.Parent);
+				writer.WriteBlobRef(value.Parent);
 			}
 			writer.WriteString(value.Author);
 			writer.WriteOptionalString(value.AuthorId);
@@ -135,10 +135,10 @@ namespace EpicGames.Horde.Storage.Nodes
 			writer.WriteString(value.Message);
 			writer.WriteDateTime(value.Time);
 
-			writer.WriteBlobHandle(value.Contents.Handle);
+			writer.WriteBlobRef(value.Contents.Handle);
 			writer.WriteUnsignedVarInt((ulong)value.Contents.Length);
 
-			writer.WriteDictionary(value.Metadata, key => writer.WriteGuidUnrealOrder(key), value => writer.WriteBlobHandle(value));
+			writer.WriteDictionary(value.Metadata, key => writer.WriteGuidUnrealOrder(key), value => writer.WriteBlobRef(value));
 
 			return s_blobType;
 		}
diff --git a/Engine/Source/Programs/Shared/EpicGames.Horde/Storage/Nodes/DdcRefNode.cs b/Engine/Source/Programs/Shared/EpicGames.Horde/Storage/Nodes/DdcRefNode.cs
--- a/Engine/Source/Programs/Shared/EpicGames.Horde/Storage/Nodes/DdcRefNode.cs
+++ b/Engine/Source/Programs/Shared/EpicGames.Horde/Storage/Nodes/DdcRefNode.cs
@@ -2,7 +2,6 @@
 
 using System;
 using System.Collections.Generic;
-using System.Linq;
 using EpicGames.Core;
 using EpicGames.Horde.Storage;
 
@@ -27,7 +26,7 @@ namespace Horde.Server.Ddc
 		/// <summary>
 		/// References to attachments. We embed this in the ref node to ensure any aliased blobs have a hard reference from the root.
 		/// </summary>
-		public List<(IoHash Hash, IBlobHandle Handle)> References { get; }
+		public List<IBlobRef> References { get; }
 
 		/// <summary>
 		/// Constructor
@@ -35,7 +34,16 @@ namespace Horde.Server.Ddc
 		public DdcRefNode(IoHash rootHash)
 		{
 			RootHash = rootHash;
-			References = new List<(IoHash, IBlobHandle)>();
+			References = new List<IBlobRef>();
+		}
+
+		/// <summary>
+		/// Constructor
+		/// </summary>
+		public DdcRefNode(IoHash rootHash, IEnumerable<IBlobRef> references)
+		{
+			RootHash = rootHash;
+			References = new List<IBlobRef>(references);
 		}
 	}
 
@@ -46,17 +54,14 @@ namespace Horde.Server.Ddc
 		public override DdcRefNode Read(IBlobReader reader, BlobSerializerOptions options)
 		{
 			IoHash rootHash = reader.ReadIoHash();
-			List<IBlobHandle<object>> references = reader.ReadList(x => reader.ReadBlobHandle<object>());
-
-			DdcRefNode refNode = new DdcRefNode(rootHash);
-			refNode.References.AddRange(references.Select(x => (x.Hash, (IBlobHandle)x)));
-			return refNode;
+			List<IBlobRef> references = reader.ReadList(x => reader.ReadBlobRef());
+			return new DdcRefNode(rootHash, references);
 		}
 
 		public override BlobType Write(IBlobWriter writer, DdcRefNode value, BlobSerializerOptions options)
 		{
 			writer.WriteIoHash(value.RootHash);
-			writer.WriteList(value.References, x => writer.WriteBlobHandle(x.Handle.ForType<object>(x.Hash)));
+			writer.WriteList(value.References, x => writer.WriteBlobRef(x));
 			return s_blobType;
 		}
 	}
diff --git a/Engine/Source/Programs/Shared/EpicGames.Horde/Storage/Nodes/DirectoryEntry.cs b/Engine/Source/Programs/Shared/EpicGames.Horde/Storage/Nodes/DirectoryEntry.cs
--- a/Engine/Source/Programs/Shared/EpicGames.Horde/Storage/Nodes/DirectoryEntry.cs
+++ b/Engine/Source/Programs/Shared/EpicGames.Horde/Storage/Nodes/DirectoryEntry.cs
@@ -8,7 +8,7 @@ namespace EpicGames.Horde.Storage.Nodes
 	/// Entry for a directory within a directory node
 	/// </summary>
 	[DebuggerDisplay("{Name}")]
-	public record class DirectoryEntry(string Name, long Length, IBlobHandle<DirectoryNode> Handle) : DirectoryNodeRef(Length, Handle)
+	public record class DirectoryEntry(string Name, long Length, IBlobRef<DirectoryNode> Handle) : DirectoryNodeRef(Length, Handle)
 	{
 	}
 }
diff --git a/Engine/Source/Programs/Shared/EpicGames.Horde/Storage/Nodes/DirectoryNode.cs b/Engine/Source/Programs/Shared/EpicGames.Horde/Storage/Nodes/DirectoryNode.cs
--- a/Engine/Source/Programs/Shared/EpicGames.Horde/Storage/Nodes/DirectoryNode.cs
+++ b/Engine/Source/Programs/Shared/EpicGames.Horde/Storage/Nodes/DirectoryNode.cs
@@ -291,10 +291,9 @@ namespace EpicGames.Horde.Storage.Nodes
 		/// Attempts to get a file entry from a path
 		/// </summary>
 		/// <param name="path">Path to the directory</param>
-		/// <param name="options">Options controling serialiation</param>
 		/// <param name="cancellationToken">Cancellation token</param>
 		/// <returns>The directory with the given path, or null if it was not found</returns>
-		public async ValueTask<FileEntry?> GetFileEntryByPathAsync(string path, BlobSerializerOptions? options = null, CancellationToken cancellationToken = default)
+		public async ValueTask<FileEntry?> GetFileEntryByPathAsync(string path, CancellationToken cancellationToken = default)
 		{
 			FileEntry? fileEntry;
 
@@ -308,7 +307,7 @@ namespace EpicGames.Horde.Storage.Nodes
 			}
 			else
 			{
-				DirectoryNode? directoryNode = await GetDirectoryByPathAsync(path.Substring(0, slashIdx), options, cancellationToken);
+				DirectoryNode? directoryNode = await GetDirectoryByPathAsync(path.Substring(0, slashIdx), cancellationToken);
 				if (directoryNode == null)
 				{
 					return null;
@@ -326,12 +325,11 @@ namespace EpicGames.Horde.Storage.Nodes
 		/// Attempts to get a directory entry from a path
 		/// </summary>
 		/// <param name="path">Path to the directory</param>
-		/// <param name="options">Options controling serialiation</param>
 		/// <param name="cancellationToken">Cancellation token</param>
 		/// <returns>The directory with the given path, or null if it was not found</returns>
-		public ValueTask<DirectoryNode?> GetDirectoryByPathAsync(string path, BlobSerializerOptions? options = null, CancellationToken cancellationToken = default) => GetDirectoryByPathAsync(this, path, options, cancellationToken);
+		public ValueTask<DirectoryNode?> GetDirectoryByPathAsync(string path, CancellationToken cancellationToken = default) => GetDirectoryByPathAsync(this, path, cancellationToken);
 
-		static async ValueTask<DirectoryNode?> GetDirectoryByPathAsync(DirectoryNode directoryNode, string path, BlobSerializerOptions? options = null, CancellationToken cancellationToken = default)
+		static async ValueTask<DirectoryNode?> GetDirectoryByPathAsync(DirectoryNode directoryNode, string path, CancellationToken cancellationToken = default)
 		{
 			while (path.Length > 0)
 			{
@@ -355,7 +353,7 @@ namespace EpicGames.Horde.Storage.Nodes
 					return null;
 				}
 
-				directoryNode = await directoryEntry.Handle.ReadBlobAsync(options, cancellationToken);
+				directoryNode = await directoryEntry.Handle.ReadBlobAsync(cancellationToken);
 			}
 			return directoryNode;
 		}
@@ -364,10 +362,9 @@ namespace EpicGames.Horde.Storage.Nodes
 		/// Deletes a file with the given path
 		/// </summary>
 		/// <param name="path"></param>
-		/// <param name="options">Options controlling serialization</param>
 		/// <param name="cancellationToken"></param>
 		/// <returns></returns>
-		public async ValueTask<bool> DeleteFileByPathAsync(string path, BlobSerializerOptions? options = null, CancellationToken cancellationToken = default)
+		public async ValueTask<bool> DeleteFileByPathAsync(string path, CancellationToken cancellationToken = default)
 		{
 			string remainingPath = path;
 			for (DirectoryNode? directory = this; directory != null;)
@@ -379,7 +376,7 @@ namespace EpicGames.Horde.Storage.Nodes
 				}
 				if (length > 0)
 				{
-					directory = await directory.TryOpenDirectoryAsync(remainingPath.Substring(0, length), options, cancellationToken);
+					directory = await directory.TryOpenDirectoryAsync(remainingPath.Substring(0, length), cancellationToken);
 				}
 				remainingPath = remainingPath.Substring(length + 1);
 			}
@@ -423,12 +420,11 @@ namespace EpicGames.Horde.Storage.Nodes
 		/// Tries to get a directory with the given name
 		/// </summary>
 		/// <param name="name">Name of the new directory</param>
-		/// <param name="options">Options controling serialiation</param>
 		/// <param name="cancellationToken">Cancellation token for the operation</param>
 		/// <returns>The new directory object</returns>
-		public async ValueTask<DirectoryNode> OpenDirectoryAsync(string name, BlobSerializerOptions? options = null, CancellationToken cancellationToken = default)
+		public async ValueTask<DirectoryNode> OpenDirectoryAsync(string name, CancellationToken cancellationToken = default)
 		{
-			DirectoryNode? directoryNode = await TryOpenDirectoryAsync(name, options, cancellationToken);
+			DirectoryNode? directoryNode = await TryOpenDirectoryAsync(name, cancellationToken);
 			if (directoryNode == null)
 			{
 				throw new DirectoryNotFoundException();
@@ -440,14 +436,13 @@ namespace EpicGames.Horde.Storage.Nodes
 		/// Tries to get a directory with the given name
 		/// </summary>
 		/// <param name="name">Name of the new directory</param>
-		/// <param name="options">Options controling serialiation</param>
 		/// <param name="cancellationToken">Cancellation token for the operation</param>
 		/// <returns>The new directory object</returns>
-		public async ValueTask<DirectoryNode?> TryOpenDirectoryAsync(string name, BlobSerializerOptions? options = null, CancellationToken cancellationToken = default)
+		public async ValueTask<DirectoryNode?> TryOpenDirectoryAsync(string name, CancellationToken cancellationToken = default)
 		{
 			if (TryGetDirectoryEntry(name, out DirectoryEntry? entry))
 			{
-				return await entry.Handle.ReadBlobAsync(options, cancellationToken);
+				return await entry.Handle.ReadBlobAsync(cancellationToken);
 			}
 			else
 			{
@@ -464,22 +459,22 @@ namespace EpicGames.Horde.Storage.Nodes
 
 		#endregion
 
-		/// <inheritdoc cref="AddFilesAsync(DirectoryReference, IEnumerable{FileInfo}, IBlobWriter, ChunkingOptions?, IProgress{ICopyStats}?, BlobSerializerOptions?, CancellationToken)"/>
-		public async Task AddFilesAsync(DirectoryInfo directoryInfo, IBlobWriter writer, ChunkingOptions? options = null, IProgress<ICopyStats>? progress = null, BlobSerializerOptions? serializerOptions = null, CancellationToken cancellationToken = default)
+		/// <inheritdoc cref="AddFilesAsync(DirectoryReference, IEnumerable{FileInfo}, IBlobWriter, ChunkingOptions?, IProgress{ICopyStats}?, CancellationToken)"/>
+		public async Task AddFilesAsync(DirectoryInfo directoryInfo, IBlobWriter writer, ChunkingOptions? options = null, IProgress<ICopyStats>? progress = null, CancellationToken cancellationToken = default)
 		{
-			await AddFilesAsync(new DirectoryReference(directoryInfo), directoryInfo.EnumerateFiles("*", SearchOption.AllDirectories).ToList(), writer, options, progress, serializerOptions, cancellationToken);
+			await AddFilesAsync(new DirectoryReference(directoryInfo), directoryInfo.EnumerateFiles("*", SearchOption.AllDirectories).ToList(), writer, options, progress, cancellationToken);
 		}
 
-		/// <inheritdoc cref="AddFilesAsync(DirectoryReference, IEnumerable{FileInfo}, IBlobWriter, ChunkingOptions?, IProgress{ICopyStats}?, BlobSerializerOptions?, CancellationToken)"/>
-		public Task AddFilesAsync(DirectoryReference baseDir, IEnumerable<FileReference> files, IBlobWriter writer, ChunkingOptions? options = null, IProgress<ICopyStats>? progress = null, BlobSerializerOptions? serializerOptions = null, CancellationToken cancellationToken = default)
+		/// <inheritdoc cref="AddFilesAsync(DirectoryReference, IEnumerable{FileInfo}, IBlobWriter, ChunkingOptions?, IProgress{ICopyStats}?, CancellationToken)"/>
+		public Task AddFilesAsync(DirectoryReference baseDir, IEnumerable<FileReference> files, IBlobWriter writer, ChunkingOptions? options = null, IProgress<ICopyStats>? progress = null, CancellationToken cancellationToken = default)
 		{
-			return AddFilesAsync(baseDir, files.Select(x => x.ToFileInfo()).ToList(), writer, options, progress, serializerOptions, cancellationToken);
+			return AddFilesAsync(baseDir, files.Select(x => x.ToFileInfo()).ToList(), writer, options, progress, cancellationToken);
 		}
 
-		/// <inheritdoc cref="AddFilesAsync(DirectoryReference, IEnumerable{FileInfo}, IBlobWriter, ChunkingOptions?, IProgress{ICopyStats}?, BlobSerializerOptions?, CancellationToken)"/>
-		public Task AddFilesAsync(DirectoryInfo baseDir, IEnumerable<FileInfo> files, IBlobWriter writer, ChunkingOptions? options = null, IProgress<ICopyStats>? progress = null, BlobSerializerOptions? serializerOptions = null, CancellationToken cancellationToken = default)
+		/// <inheritdoc cref="AddFilesAsync(DirectoryReference, IEnumerable{FileInfo}, IBlobWriter, ChunkingOptions?, IProgress{ICopyStats}?, CancellationToken)"/>
+		public Task AddFilesAsync(DirectoryInfo baseDir, IEnumerable<FileInfo> files, IBlobWriter writer, ChunkingOptions? options = null, IProgress<ICopyStats>? progress = null, CancellationToken cancellationToken = default)
 		{
-			return AddFilesAsync(new DirectoryReference(baseDir), files.ToList(), writer, options, progress, serializerOptions, cancellationToken);
+			return AddFilesAsync(new DirectoryReference(baseDir), files.ToList(), writer, options, progress, cancellationToken);
 		}
 
 		/// <summary>
@@ -490,9 +485,8 @@ namespace EpicGames.Horde.Storage.Nodes
 		/// <param name="options">Options for chunking file content</param>
 		/// <param name="writer">Writer for new node data</param>
 		/// <param name="progress">Feedback interface for progress updates</param>
-		/// <param name="serializerOptions">Options controlling serialization</param>
 		/// <param name="cancellationToken">Cancellation token for the operation</param>
-		public async Task AddFilesAsync(DirectoryReference baseDir, IEnumerable<FileInfo> files, IBlobWriter writer, ChunkingOptions? options = null, IProgress<ICopyStats>? progress = null, BlobSerializerOptions? serializerOptions = null, CancellationToken cancellationToken = default)
+		public async Task AddFilesAsync(DirectoryReference baseDir, IEnumerable<FileInfo> files, IBlobWriter writer, ChunkingOptions? options = null, IProgress<ICopyStats>? progress = null, CancellationToken cancellationToken = default)
 		{
 			options ??= new ChunkingOptions();
 
@@ -558,7 +552,7 @@ namespace EpicGames.Horde.Storage.Nodes
 			}
 
 			// Add all the new entries to the tree
-			await UpdateAsync(update, writer, serializerOptions, cancellationToken);
+			await UpdateAsync(update, writer, cancellationToken);
 		}
 
 		static List<(int Index, int Count)> ComputePartitions(IReadOnlyList<FileInfo> files, long totalSize)
@@ -616,13 +610,12 @@ namespace EpicGames.Horde.Storage.Nodes
 		/// </summary>
 		/// <param name="updates">Files to add</param>
 		/// <param name="writer">Writer for new node data</param>
-		/// <param name="options">Options controling serialiation</param>
 		/// <param name="cancellationToken">Cancellation token for the operation</param>
-		public Task UpdateAsync(IEnumerable<FileUpdate> updates, IBlobWriter writer, BlobSerializerOptions? options = null, CancellationToken cancellationToken = default)
+		public Task UpdateAsync(IEnumerable<FileUpdate> updates, IBlobWriter writer, CancellationToken cancellationToken = default)
 		{
 			DirectoryUpdate update = new DirectoryUpdate();
 			update.AddFiles(updates);
-			return UpdateAsync(update, writer, options, cancellationToken);
+			return UpdateAsync(update, writer, cancellationToken);
 		}
 
 		/// <summary>
@@ -630,9 +623,8 @@ namespace EpicGames.Horde.Storage.Nodes
 		/// </summary>
 		/// <param name="update">Files to add</param>
 		/// <param name="writer">Writer for new node data</param>
-		/// <param name="options">Options controling serialiation</param>
 		/// <param name="cancellationToken">Cancellation token for the operation</param>
-		public async Task UpdateAsync(DirectoryUpdate update, IBlobWriter writer, BlobSerializerOptions? options = null, CancellationToken cancellationToken = default)
+		public async Task UpdateAsync(DirectoryUpdate update, IBlobWriter writer, CancellationToken cancellationToken = default)
 		{
 			foreach ((string name, DirectoryUpdate? directory) in update.Directories)
 			{
@@ -642,10 +634,10 @@ namespace EpicGames.Horde.Storage.Nodes
 				}
 				else
 				{
-					DirectoryNode? childNode = await TryOpenDirectoryAsync(name, options, cancellationToken);
+					DirectoryNode? childNode = await TryOpenDirectoryAsync(name, cancellationToken);
 					childNode ??= new DirectoryNode();
-					await childNode.UpdateAsync(directory, writer, options, cancellationToken);
-					IBlobHandle<DirectoryNode> handle = await writer.WriteBlobAsync<DirectoryNode>(childNode, options, cancellationToken);
+					await childNode.UpdateAsync(directory, writer, cancellationToken);
+					IBlobRef<DirectoryNode> handle = await writer.WriteBlobAsync<DirectoryNode>(childNode, cancellationToken);
 					_nameToDirectoryEntry[name] = new DirectoryEntry(name, childNode.Length, handle);
 				}
 			}
@@ -657,7 +649,7 @@ namespace EpicGames.Horde.Storage.Nodes
 				}
 				else
 				{
-					await file.WriteInteriorNodesAsync(writer, new ChunkingOptions().InteriorOptions, options, cancellationToken);
+					await file.WriteInteriorNodesAsync(writer, new ChunkingOptions().InteriorOptions, cancellationToken);
 					_nameToFileEntry[name] = new FileEntry(name, file.Flags, file.Length, file.StreamHash, file.Nodes[0], file.CustomData);
 				}
 			}
@@ -669,9 +661,8 @@ namespace EpicGames.Horde.Storage.Nodes
 		/// <param name="stream">Input stream</param>
 		/// <param name="writer">Writer for new nodes</param>
 		/// <param name="options"></param>
-		/// <param name="serializerOptions">Options controling serialiation</param>
 		/// <param name="cancellationToken">Cancellation token for the operation</param>
-		public async Task CopyFromZipStreamAsync(Stream stream, IBlobWriter writer, ChunkingOptions options, BlobSerializerOptions? serializerOptions = null, CancellationToken cancellationToken = default)
+		public async Task CopyFromZipStreamAsync(Stream stream, IBlobWriter writer, ChunkingOptions options, CancellationToken cancellationToken = default)
 		{
 			// Create all the leaf nodes
 			List<(ZipArchiveEntry, LeafChunkedData)> entries = new List<(ZipArchiveEntry, LeafChunkedData)>();
@@ -700,22 +691,21 @@ namespace EpicGames.Horde.Storage.Nodes
 					flags |= FileEntryFlags.Executable;
 				}
 
-				ChunkedData chunkedFile = await InteriorChunkedDataNode.CreateTreeAsync(leafChunkedFile, options.InteriorOptions, writer, serializerOptions, cancellationToken);
+				ChunkedData chunkedFile = await InteriorChunkedDataNode.CreateTreeAsync(leafChunkedFile, options.InteriorOptions, writer, cancellationToken);
 				updates.Add(new FileUpdate(entry.FullName, flags, entry.Length, chunkedFile));
 			}
 
 			// Update the tree
-			await UpdateAsync(updates, writer, serializerOptions, cancellationToken);
+			await UpdateAsync(updates, writer, cancellationToken);
 		}
 
 		/// <summary>
 		/// Utility function to allow extracting a packed directory to disk
 		/// </summary>
 		/// <param name="directoryInfo"></param>
-		/// <param name="options"></param>
 		/// <param name="logger"></param>
 		/// <param name="cancellationToken"></param>
-		public Task CopyToDirectoryAsync(DirectoryInfo directoryInfo, BlobSerializerOptions? options, ILogger logger, CancellationToken cancellationToken) => CopyToDirectoryAsync(directoryInfo, null, options, logger, cancellationToken);
+		public Task CopyToDirectoryAsync(DirectoryInfo directoryInfo, ILogger logger, CancellationToken cancellationToken) => CopyToDirectoryAsync(directoryInfo, null, logger, cancellationToken);
 
 		record class OutputDir(string Path, DirectoryNode Node);
 		record class OutputFile(OutputDir Directory, FileEntry FileEntry);
@@ -732,9 +722,8 @@ namespace EpicGames.Horde.Storage.Nodes
 		/// <param name="directoryInfo">Direcotry to write to</param>
 		/// <param name="progress">Sink for progress updates</param>
 		/// <param name="logger">Logger for output</param>
-		/// <param name="options">Options controlling serialization</param>
 		/// <param name="cancellationToken">Cancellation token for the operation</param>
-		public async Task CopyToDirectoryAsync(DirectoryInfo directoryInfo, IProgress<ICopyStats>? progress, BlobSerializerOptions? options, ILogger logger, CancellationToken cancellationToken)
+		public async Task CopyToDirectoryAsync(DirectoryInfo directoryInfo, IProgress<ICopyStats>? progress, ILogger logger, CancellationToken cancellationToken)
 		{
 			int numTasks = Math.Min(1 + (int)(Length / (16 * 1024 * 1024)), 16);
 			logger.LogInformation("Splitting read into {NumThreads} threads", numTasks);
@@ -766,7 +755,7 @@ namespace EpicGames.Horde.Storage.Nodes
 					}
 				}
 
-				await FindOutputChunksRootAsync(chunks.Writer, options, logger, cancellationSource.Token);
+				await FindOutputChunksRootAsync(chunks.Writer, logger, cancellationSource.Token);
 
 				List<Task> tasks = new List<Task>();
 //				tasks.Add(RunBackgroundTask(ctx => FindOutputChunksRootAsync(chunks.Writer, options, logger, ctx)));
@@ -842,13 +831,13 @@ namespace EpicGames.Horde.Storage.Nodes
 			}
 		}
 
-		async Task FindOutputChunksRootAsync(ChannelWriter<OutputChunk> chunks, BlobSerializerOptions? options, ILogger logger, CancellationToken cancellationToken)
+		async Task FindOutputChunksRootAsync(ChannelWriter<OutputChunk> chunks, ILogger logger, CancellationToken cancellationToken)
 		{
-			await FindOutputChunksAsync("", this, chunks, options, logger, cancellationToken);
+			await FindOutputChunksAsync("", this, chunks, logger, cancellationToken);
 			chunks.Complete();
 		}
 
-		static async Task FindOutputChunksAsync(string path, DirectoryNode node, ChannelWriter<OutputChunk> chunks, BlobSerializerOptions? options, ILogger logger, CancellationToken cancellationToken)
+		static async Task FindOutputChunksAsync(string path, DirectoryNode node, ChannelWriter<OutputChunk> chunks, ILogger logger, CancellationToken cancellationToken)
 		{
 			OutputDir outputDir = new OutputDir(path, node);
 
@@ -860,12 +849,12 @@ namespace EpicGames.Horde.Storage.Nodes
 
 			foreach (DirectoryEntry directoryEntry in node._nameToDirectoryEntry.Values)
 			{
-				DirectoryNode subDirectoryNode = await directoryEntry.Handle.ReadBlobAsync(options, cancellationToken);
+				DirectoryNode subDirectoryNode = await directoryEntry.Handle.ReadBlobAsync(cancellationToken);
 
 				string subPath = CombinePaths(outputDir.Path, directoryEntry.Name);
 				TraceBlobRead("Directory", subPath, directoryEntry.Handle, logger);
 
-				await FindOutputChunksAsync(subPath, subDirectoryNode, chunks, options, logger, cancellationToken);
+				await FindOutputChunksAsync(subPath, subDirectoryNode, chunks, logger, cancellationToken);
 			}
 		}
 
@@ -923,10 +912,9 @@ namespace EpicGames.Horde.Storage.Nodes
 		/// Returns a stream containing the zipped contents of this directory
 		/// </summary>
 		/// <param name="filter">Filter for files to include in the zip</param>
-		/// <param name="options">Options controlling serialization</param>
 		/// <param name="logger">Logger for diagnostic output</param>
 		/// <returns>Stream containing zipped archive data</returns>
-		public Stream AsZipStream(FileFilter? filter = null, BlobSerializerOptions? options = null, ILogger? logger = null) => new DirectoryNodeZipStream(this, filter, options, logger);
+		public Stream AsZipStream(FileFilter? filter = null, ILogger? logger = null) => new DirectoryNodeZipStream(this, filter, logger);
 	}
 
 	class DirectoryNodeConverter : BlobConverter<DirectoryNode>
@@ -944,7 +932,7 @@ namespace EpicGames.Horde.Storage.Nodes
 			int fileCount = (int)reader.ReadUnsignedVarInt();
 			for (int idx = 0; idx < fileCount; idx++)
 			{
-				IBlobHandle<ChunkedDataNode> targetHandle = reader.ReadBlobHandle<ChunkedDataNode>();
+				IBlobRef<ChunkedDataNode> targetHandle = reader.ReadBlobRef<ChunkedDataNode>();
 
 				ChunkedDataNodeType targetType = ChunkedDataNodeType.Unknown;
 				if (reader.Version >= 2)
@@ -971,7 +959,7 @@ namespace EpicGames.Horde.Storage.Nodes
 			int directoryCount = (int)reader.ReadUnsignedVarInt();
 			for (int idx = 0; idx < directoryCount; idx++)
 			{
-				IBlobHandle<DirectoryNode> directoryHandle = reader.ReadBlobHandle<DirectoryNode>();
+				IBlobRef<DirectoryNode> directoryHandle = reader.ReadBlobRef<DirectoryNode>();
 				long length = (long)reader.ReadUnsignedVarInt();
 				string name = reader.ReadString();
 
@@ -989,7 +977,7 @@ namespace EpicGames.Horde.Storage.Nodes
 			writer.WriteUnsignedVarInt(value.Files.Count);
 			foreach (FileEntry fileEntry in value.Files)
 			{
-				writer.WriteBlobHandle(fileEntry.Target.Handle);
+				writer.WriteBlobRef(fileEntry.Target.Handle);
 				writer.WriteUnsignedVarInt((int)fileEntry.Target.Type);
 
 				FileEntryFlags flags = (fileEntry.CustomData.Length > 0) ? (fileEntry.Flags | FileEntryFlags.HasCustomData) : (fileEntry.Flags & ~FileEntryFlags.HasCustomData);
@@ -1008,7 +996,7 @@ namespace EpicGames.Horde.Storage.Nodes
 			writer.WriteUnsignedVarInt(value.Directories.Count);
 			foreach (DirectoryEntry directoryEntry in value.Directories)
 			{
-				writer.WriteBlobHandle(directoryEntry.Handle);
+				writer.WriteBlobRef(directoryEntry.Handle);
 				writer.WriteUnsignedVarInt((ulong)directoryEntry.Length);
 				writer.WriteString(directoryEntry.Name);
 			}
@@ -1047,11 +1035,11 @@ namespace EpicGames.Horde.Storage.Nodes
 		/// <summary>
 		/// Writes interior node data to the given writer
 		/// </summary>
-		public async ValueTask WriteInteriorNodesAsync(IBlobWriter writer, InteriorChunkedDataNodeOptions interiorNodeOptions, BlobSerializerOptions? serializerOptions, CancellationToken cancellationToken)
+		public async ValueTask WriteInteriorNodesAsync(IBlobWriter writer, InteriorChunkedDataNodeOptions interiorNodeOptions, CancellationToken cancellationToken)
 		{
 			if (Nodes.Count > 1)
 			{
-				ChunkedDataNodeRef rootRef = await InteriorChunkedDataNode.CreateTreeAsync(Nodes, interiorNodeOptions, writer, serializerOptions, cancellationToken);
+				ChunkedDataNodeRef rootRef = await InteriorChunkedDataNode.CreateTreeAsync(Nodes, interiorNodeOptions, writer, cancellationToken);
 				Nodes.Clear();
 				Nodes.Add(rootRef);
 			}
@@ -1193,20 +1181,20 @@ namespace EpicGames.Horde.Storage.Nodes
 		/// <summary>
 		/// Writes interior node data to the given writer
 		/// </summary>
-		public async ValueTask WriteInteriorNodesAsync(IBlobWriter writer, InteriorChunkedDataNodeOptions interiorNodeOptions, BlobSerializerOptions? serializerOptions, CancellationToken cancellationToken)
+		public async ValueTask WriteInteriorNodesAsync(IBlobWriter writer, InteriorChunkedDataNodeOptions interiorNodeOptions, CancellationToken cancellationToken)
 		{
 			foreach (DirectoryUpdate? directoryUpdate in Directories.Values)
 			{
 				if (directoryUpdate != null)
 				{
-					await directoryUpdate.WriteInteriorNodesAsync(writer, interiorNodeOptions, serializerOptions, cancellationToken);
+					await directoryUpdate.WriteInteriorNodesAsync(writer, interiorNodeOptions, cancellationToken);
 				}
 			}
 			foreach (FileUpdate? fileUpdate in Files.Values)
 			{
 				if (fileUpdate != null)
 				{
-					await fileUpdate.WriteInteriorNodesAsync(writer, interiorNodeOptions, serializerOptions, cancellationToken);
+					await fileUpdate.WriteInteriorNodesAsync(writer, interiorNodeOptions, cancellationToken);
 				}
 			}
 		}
@@ -1244,12 +1232,11 @@ namespace EpicGames.Horde.Storage.Nodes
 		/// </summary>
 		/// <param name="node">Root node to copy from</param>
 		/// <param name="filter">Filter for files to include in the zip</param>
-		/// <param name="options">Options controling serialiation</param>
 		/// <param name="logger">Optional logger for debug tracing</param>
-		public DirectoryNodeZipStream(DirectoryNode node, FileFilter? filter, BlobSerializerOptions? options, ILogger? logger)
+		public DirectoryNodeZipStream(DirectoryNode node, FileFilter? filter, ILogger? logger)
 		{
 			_pipe = new Pipe();
-			_backgroundTask = BackgroundTask.StartNew(ctx => CopyToPipeAsync(node, filter, _pipe.Writer, options, logger, ctx));
+			_backgroundTask = BackgroundTask.StartNew(ctx => CopyToPipeAsync(node, filter, _pipe.Writer, logger, ctx));
 			_logger = logger;
 		}
 
@@ -1308,14 +1295,14 @@ namespace EpicGames.Horde.Storage.Nodes
 			return length;
 		}
 
-		static async Task CopyToPipeAsync(DirectoryNode node, FileFilter? filter, PipeWriter writer, BlobSerializerOptions? options, ILogger? logger, CancellationToken cancellationToken)
+		static async Task CopyToPipeAsync(DirectoryNode node, FileFilter? filter, PipeWriter writer, ILogger? logger, CancellationToken cancellationToken)
 		{
 			using Stream outputStream = writer.AsStream();
 			using ZipArchive archive = new ZipArchive(outputStream, ZipArchiveMode.Create);
-			await CopyFilesAsync(node, "", filter, archive, options, logger, cancellationToken);
+			await CopyFilesAsync(node, "", filter, archive, logger, cancellationToken);
 		}
 
-		static async Task CopyFilesAsync(DirectoryNode directory, string prefix, FileFilter? filter, ZipArchive archive, BlobSerializerOptions? options, ILogger? logger, CancellationToken cancellationToken)
+		static async Task CopyFilesAsync(DirectoryNode directory, string prefix, FileFilter? filter, ZipArchive archive, ILogger? logger, CancellationToken cancellationToken)
 		{
 			int numDirs = directory.Directories.Count;
 			int numFiles = directory.Files.Count;
@@ -1326,8 +1313,8 @@ namespace EpicGames.Horde.Storage.Nodes
 				string directoryPath = $"{prefix}{directoryEntry.Name}/";
 				if (filter == null || filter.PossiblyMatches(directoryPath))
 				{
-					DirectoryNode node = await directoryEntry.Handle.ReadBlobAsync(options, cancellationToken);
-					await CopyFilesAsync(node, directoryPath, filter, archive, options, logger, cancellationToken);
+					DirectoryNode node = await directoryEntry.Handle.ReadBlobAsync(cancellationToken);
+					await CopyFilesAsync(node, directoryPath, filter, archive, logger, cancellationToken);
 				}
 				numCopiedDirs++;
 			}
@@ -1384,30 +1371,30 @@ namespace EpicGames.Horde.Storage.Nodes
 		/// <summary>
 		/// Writes a tree of files to a storage writer
 		/// </summary>
-		public static async Task<IBlobHandle<DirectoryNode>> WriteFilesAsync(this IBlobWriter writer, DirectoryReference baseDir, ChunkingOptions? options = null, IProgress<ICopyStats>? progress = null, CancellationToken cancellationToken = default)
+		public static async Task<IBlobRef<DirectoryNode>> WriteFilesAsync(this IBlobWriter writer, DirectoryReference baseDir, ChunkingOptions? options = null, IProgress<ICopyStats>? progress = null, CancellationToken cancellationToken = default)
 		{
 			DirectoryNode outputNode = new DirectoryNode();
-			await outputNode.AddFilesAsync(baseDir, DirectoryReference.EnumerateFiles(baseDir, "*", SearchOption.AllDirectories), writer, options, progress, null, cancellationToken);
+			await outputNode.AddFilesAsync(baseDir, DirectoryReference.EnumerateFiles(baseDir, "*", SearchOption.AllDirectories), writer, options, progress, cancellationToken);
 			return await writer.WriteBlobAsync(outputNode, cancellationToken: cancellationToken);
 		}
 
 		/// <summary>
 		/// Writes a tree of files to a storage writer
 		/// </summary>
-		public static async Task<IBlobHandle<DirectoryNode>> WriteFilesAsync(this IBlobWriter writer, DirectoryInfo baseDir, IReadOnlyList<FileInfo> files, ChunkingOptions? options = null, IProgress<ICopyStats>? progress = null, CancellationToken cancellationToken = default)
+		public static async Task<IBlobRef<DirectoryNode>> WriteFilesAsync(this IBlobWriter writer, DirectoryInfo baseDir, IReadOnlyList<FileInfo> files, ChunkingOptions? options = null, IProgress<ICopyStats>? progress = null, CancellationToken cancellationToken = default)
 		{
 			DirectoryNode outputNode = new DirectoryNode();
-			await outputNode.AddFilesAsync(baseDir, files, writer, options, progress, null, cancellationToken);
+			await outputNode.AddFilesAsync(baseDir, files, writer, options, progress, cancellationToken);
 			return await writer.WriteBlobAsync(outputNode, cancellationToken: cancellationToken);
 		}
 
 		/// <summary>
 		/// Writes a tree of files to a storage writer
 		/// </summary>
-		public static async Task<IBlobHandle<DirectoryNode>> WriteFilesAsync(this IBlobWriter writer, DirectoryReference baseDir, IReadOnlyList<FileReference> files, ChunkingOptions? options = null, IProgress<ICopyStats>? progress = null, CancellationToken cancellationToken = default)
+		public static async Task<IBlobRef<DirectoryNode>> WriteFilesAsync(this IBlobWriter writer, DirectoryReference baseDir, IReadOnlyList<FileReference> files, ChunkingOptions? options = null, IProgress<ICopyStats>? progress = null, CancellationToken cancellationToken = default)
 		{
 			DirectoryNode outputNode = new DirectoryNode();
-			await outputNode.AddFilesAsync(baseDir, files, writer, options, progress, null, cancellationToken);
+			await outputNode.AddFilesAsync(baseDir, files, writer, options, progress, cancellationToken);
 			return await writer.WriteBlobAsync(outputNode, cancellationToken: cancellationToken);
 		}
 	}
diff --git a/Engine/Source/Programs/Shared/EpicGames.Horde/Storage/Nodes/DirectoryNodeRef.cs b/Engine/Source/Programs/Shared/EpicGames.Horde/Storage/Nodes/DirectoryNodeRef.cs
--- a/Engine/Source/Programs/Shared/EpicGames.Horde/Storage/Nodes/DirectoryNodeRef.cs
+++ b/Engine/Source/Programs/Shared/EpicGames.Horde/Storage/Nodes/DirectoryNodeRef.cs
@@ -7,5 +7,5 @@ namespace EpicGames.Horde.Storage.Nodes
 	/// </summary>
 	/// <param name="Length">Sum total of all the file lengths in this directory tree</param>
 	/// <param name="Handle">Handle to the target node</param>
-	public record class DirectoryNodeRef(long Length, IBlobHandle<DirectoryNode> Handle);
+	public record class DirectoryNodeRef(long Length, IBlobRef<DirectoryNode> Handle);
 }
diff --git a/Engine/Source/Programs/Shared/EpicGames.Horde/Storage/Nodes/RedirectNode.cs b/Engine/Source/Programs/Shared/EpicGames.Horde/Storage/Nodes/RedirectNode.cs
--- a/Engine/Source/Programs/Shared/EpicGames.Horde/Storage/Nodes/RedirectNode.cs
+++ b/Engine/Source/Programs/Shared/EpicGames.Horde/Storage/Nodes/RedirectNode.cs
@@ -24,13 +24,13 @@ namespace EpicGames.Horde.Storage.Nodes
 		/// <summary>
 		/// The target handle
 		/// </summary>
-		public IBlobHandle<T> Target { get; }
+		public IBlobRef<T> Target { get; }
 
 		/// <summary>
 		/// Constructor
 		/// </summary>
 		/// <param name="handle">Target node for the redirect</param>
-		public RedirectNode(IBlobHandle<T> handle) => Target = handle;
+		public RedirectNode(IBlobRef<T> handle) => Target = handle;
 	}
 
 	class RedirectNodeConverter<T> : BlobConverter<RedirectNode<T>>
@@ -40,14 +40,14 @@ namespace EpicGames.Horde.Storage.Nodes
 		/// <inheritdoc/>
 		public override RedirectNode<T> Read(IBlobReader reader, BlobSerializerOptions options)
 		{
-			IBlobHandle<T> handle = reader.ReadBlobHandle<T>();
+			IBlobRef<T> handle = reader.ReadBlobRef<T>();
 			return new RedirectNode<T>(handle);
 		}
 
 		/// <inheritdoc/>
 		public override BlobType Write(IBlobWriter writer, RedirectNode<T> value, BlobSerializerOptions options)
 		{
-			writer.WriteBlobHandle<T>(value.Target);
+			writer.WriteBlobRef(value.Target);
 			return s_blobType;
 		}
 	}
diff --git a/Engine/Source/Programs/UnrealBuildTool/Artifacts/HordeStorageArtifactCache.cs b/Engine/Source/Programs/UnrealBuildTool/Artifacts/HordeStorageArtifactCache.cs
--- a/Engine/Source/Programs/UnrealBuildTool/Artifacts/HordeStorageArtifactCache.cs
+++ b/Engine/Source/Programs/UnrealBuildTool/Artifacts/HordeStorageArtifactCache.cs
@@ -31,7 +31,7 @@ namespace UnrealBuildTool.Artifacts
 		/// Collection of output file references.  There should be exactly the same number
 		/// of file references as outputs in the action
 		/// </summary>
-		public readonly IBlobHandle<ChunkedDataNode>[] OutputRefs;
+		public readonly IBlobRef<ChunkedDataNode>[] OutputRefs;
 
 		/// <summary>
 		/// Construct a new horde artifact number
@@ -41,7 +41,7 @@ namespace UnrealBuildTool.Artifacts
 		public HordeArtifactAction(ArtifactAction artifactAction)
 		{
 			ArtifactAction = artifactAction;
-			OutputRefs = new IBlobHandle<ChunkedDataNode>[ArtifactAction.Outputs.Length];
+			OutputRefs = new IBlobRef<ChunkedDataNode>[ArtifactAction.Outputs.Length];
 		}
 
 		/// <summary>
@@ -51,7 +51,7 @@ namespace UnrealBuildTool.Artifacts
 		public HordeArtifactAction(IBlobReader reader)
 		{
 			ArtifactAction = reader.ReadArtifactAction();
-			OutputRefs = reader.ReadVariableLengthArray(() => reader.ReadBlobHandle<ChunkedDataNode>());
+			OutputRefs = reader.ReadVariableLengthArray(() => reader.ReadBlobRef<ChunkedDataNode>());
 		}
 
 		/// <summary>
@@ -61,7 +61,7 @@ namespace UnrealBuildTool.Artifacts
 		public void Serialize(IBlobWriter writer)
 		{
 			writer.WriteArtifactAction(ArtifactAction);
-			writer.WriteVariableLengthArray(OutputRefs, x => writer.WriteBlobHandle(x));
+			writer.WriteVariableLengthArray(OutputRefs, x => writer.WriteBlobRef(x));
 		}
 
 		/// <summary>
@@ -83,7 +83,7 @@ namespace UnrealBuildTool.Artifacts
 				string outputName = artifact.GetFullPath(ArtifactAction.DirectoryMapping);
 				using FileStream stream = new(outputName, FileMode.Open, FileAccess.Read, FileShare.Read);
 				LeafChunkedData leafChunkedData = await fileWriter.CreateAsync(stream, leafOptions.TargetSize, cancellationToken);
-				ChunkedData chunkedData = await InteriorChunkedDataNode.CreateTreeAsync(leafChunkedData, interiorOptions, writer, null, cancellationToken);
+				ChunkedData chunkedData = await InteriorChunkedDataNode.CreateTreeAsync(leafChunkedData, interiorOptions, writer, cancellationToken);
 				OutputRefs[index++] = chunkedData.Root.Handle;
 			}
 		}
@@ -310,7 +310,7 @@ namespace UnrealBuildTool.Artifacts
 						output[index] = true;
 
 						int refIndex = 0;
-						foreach (IBlobHandle<ChunkedDataNode> artifactRef in hordeArtifactAction.OutputRefs)
+						foreach (IBlobRef<ChunkedDataNode> artifactRef in hordeArtifactAction.OutputRefs)
 						{
 							if (artifactRef == null)
 							{
@@ -479,7 +479,7 @@ namespace UnrealBuildTool.Artifacts
 					// Save the artifact action file
 					await using IBlobWriter writer = _store!.CreateBlobWriter();
 					await hordeArtifactAction.WriteFilesAsync(writer, cancellationToken);
-					IBlobHandle<ArtifactActionCollectionNode> nodeRef = await writer.WriteBlobAsync(node);
+					IBlobRef<ArtifactActionCollectionNode> nodeRef = await writer.WriteBlobAsync(node);
 					await writer.FlushAsync();
 
 					// Save the collection
diff --git a/Engine/Source/Programs/UnrealBuildTool/Executors/Experimental/UBAHordeSession.cs b/Engine/Source/Programs/UnrealBuildTool/Executors/Experimental/UBAHordeSession.cs
--- a/Engine/Source/Programs/UnrealBuildTool/Executors/Experimental/UBAHordeSession.cs
+++ b/Engine/Source/Programs/UnrealBuildTool/Executors/Experimental/UBAHordeSession.cs
@@ -358,10 +358,10 @@ namespace UnrealBuildTool
 			BlobSerializerOptions serializerOptions = new BlobSerializerOptions();
 			serializerOptions.Converters.Add(new InteriorChunkedDataNodeConverter(2)); // Lock to v2 for now. Could change based on protocol version.
 
-			await using IBlobWriter writer = _storage.CreateBlobWriter();
+			await using IBlobWriter writer = _storage.CreateBlobWriter(serializerOptions: serializerOptions);
 			DirectoryNode sandbox = new();
-			await sandbox.AddFilesAsync(baseDir, files, writer, serializerOptions: serializerOptions, cancellationToken: cancellationToken);
-			IBlobHandle<DirectoryNode> handle = await writer.WriteBlobAsync(sandbox, cancellationToken);
+			await sandbox.AddFilesAsync(baseDir, files, writer, cancellationToken: cancellationToken);
+			IBlobRef<DirectoryNode> handle = await writer.WriteBlobAsync(sandbox, cancellationToken);
 			await writer.FlushAsync(cancellationToken);
 			return handle.GetLocator();
 		}
