EpicGames.Core: Change ListExtensions.AsReadOnly() to return an IList<T>, similar to semantics of List<T>.AsReadOnly(). New overload AsReadOnlyList() preserves the original behavior of returning an IReadOnlyList<T>, since NET's IList<T> does not derive from IReadOnlyList<T>.

diff --git a/Engine/Source/Programs/Shared/EpicGames.Core/ListExtensions.cs b/Engine/Source/Programs/Shared/EpicGames.Core/ListExtensions.cs
--- a/Engine/Source/Programs/Shared/EpicGames.Core/ListExtensions.cs
+++ b/Engine/Source/Programs/Shared/EpicGames.Core/ListExtensions.cs
@@ -3,6 +3,7 @@
 using System;
 using System.Collections;
 using System.Collections.Generic;
+using System.Diagnostics.CodeAnalysis;
 using System.Linq;
 using System.Threading.Tasks;
 
@@ -14,22 +15,38 @@ namespace EpicGames.Core
 	public static class ListExtensions
 	{
 		/// <summary>
-		/// Wrapper around a list to implement <see cref="IReadOnlyList{T}"/>
+		/// Wrapper around a list to implement an immutable <see cref="IList{T}"/> and <see cref="IReadOnlyList{T}"/>
 		/// </summary>
 		/// <typeparam name="T"></typeparam>
-		class ReadOnlyList<T> : IReadOnlyList<T>
+		class ReadOnlyList<T> : IList<T>, IReadOnlyList<T>
 		{
 			readonly IList<T> _inner;
 
 			public ReadOnlyList(IList<T> inner) => _inner = inner;
 
-			public T this[int index] => _inner[index];
+			static NotSupportedException CreateNotSupportedException()
+				=> new NotSupportedException("Collection is read-only.");
+
+			public T this[int index]
+			{
+				get => _inner[index];
+				set => throw CreateNotSupportedException();
+			}
 
 			public int Count => _inner.Count;
+			public bool IsReadOnly => _inner.IsReadOnly;
 
+			public void Add(T item) => throw CreateNotSupportedException();
+			public void Clear() => throw CreateNotSupportedException();
+			public bool Contains(T item) => _inner.Contains(item);
+			public void CopyTo(T[] array, int arrayIndex) => _inner.CopyTo(array, arrayIndex);
 			public IEnumerator<T> GetEnumerator() => _inner.GetEnumerator();
+			public int IndexOf(T item) => _inner.IndexOf(item);
+			public void Insert(int index, T item) => throw CreateNotSupportedException();
+			public bool Remove(T item) => throw CreateNotSupportedException();
+			public void RemoveAt(int index) => throw CreateNotSupportedException();
 
-			IEnumerator IEnumerable.GetEnumerator() => _inner.GetEnumerator();
+			IEnumerator IEnumerable.GetEnumerator() => ((IEnumerable)_inner).GetEnumerator();
 		}
 
 		/// <summary>
@@ -38,7 +55,15 @@ namespace EpicGames.Core
 		/// <typeparam name="T"></typeparam>
 		/// <param name="list"></param>
 		/// <returns></returns>
-		public static IReadOnlyList<T> AsReadOnly<T>(this IList<T> list) => new ReadOnlyList<T>(list);
+		public static IList<T> AsReadOnly<T>(this IList<T> list) => new ReadOnlyList<T>(list);
+
+		/// <summary>
+		/// Create a read-only wrapper around a list
+		/// </summary>
+		/// <typeparam name="T"></typeparam>
+		/// <param name="list"></param>
+		/// <returns></returns>
+		public static IReadOnlyList<T> AsReadOnlyList<T>(this IList<T> list) => new ReadOnlyList<T>(list);
 
 		/// <summary>
 		/// Sorts a list by a particular field
