Horde: Allow marking BlobSerializerOptions instances as read only, which will throw an exception if a modifying collection method is called.

diff --git a/Engine/Source/Programs/Shared/EpicGames.Horde/Storage/BlobSerializer.cs b/Engine/Source/Programs/Shared/EpicGames.Horde/Storage/BlobSerializer.cs
--- a/Engine/Source/Programs/Shared/EpicGames.Horde/Storage/BlobSerializer.cs
+++ b/Engine/Source/Programs/Shared/EpicGames.Horde/Storage/BlobSerializer.cs
@@ -1,6 +1,7 @@
 // Copyright Epic Games, Inc. All Rights Reserved.
 
 using System;
+using System.Collections;
 using System.Collections.Concurrent;
 using System.Collections.Generic;
 using System.Reflection;
@@ -48,19 +49,113 @@ namespace EpicGames.Horde.Storage
 	/// </summary>
 	public class BlobSerializerOptions
 	{
+		class FreezableList<T> : IList<T>
+		{
+			readonly BlobSerializerOptions _owner;
+			readonly List<T> _list;
+
+			public FreezableList(BlobSerializerOptions owner)
+			{
+				_owner = owner;
+				_list = new List<T>();
+			}
+
+			public T this[int index]
+			{
+				get => _list[index];
+				set
+				{
+					_owner.CheckMutable();
+					_list[index] = value;
+				}
+			}
+
+			public int Count => _list.Count;
+			public bool IsReadOnly => _owner._readOnly;
+
+			public void Add(T item)
+			{
+				_owner.CheckMutable();
+				_list.Add(item);
+			}
+
+			public void Clear()
+			{
+				_owner.CheckMutable();
+				_list.Clear();
+			}
+
+			public bool Contains(T item) => _list.Contains(item);
+			public void CopyTo(T[] array, int arrayIndex) => _list.CopyTo(array, arrayIndex);
+			public IEnumerator<T> GetEnumerator() => _list.GetEnumerator();
+			public int IndexOf(T item) => _list.IndexOf(item);
+
+			public void Insert(int index, T item)
+			{
+				_owner.CheckMutable();
+				_list.Insert(index, item);
+			}
+
+			public bool Remove(T item)
+			{
+				_owner.CheckMutable();
+				return _list.Remove(item);
+			}
+
+			public void RemoveAt(int index)
+			{
+				_owner.CheckMutable();
+				_list.RemoveAt(index);
+			}
+
+			IEnumerator IEnumerable.GetEnumerator() => _list.GetEnumerator();
+		}
+
+		readonly ConcurrentDictionary<Type, BlobConverter> _cachedConverters = new ConcurrentDictionary<Type, BlobConverter>();
+		static readonly ConcurrentDictionary<Type, BlobConverter> s_cachedDefaultConverters = new ConcurrentDictionary<Type, BlobConverter>();
+		static readonly Func<Type, BlobConverter> s_createDefaultConverter = CreateDefaultConverter;
+
+		bool _readOnly;
+		readonly FreezableList<BlobConverter> _converters;
+
+		/// <summary>
+		/// Known converter types
+		/// </summary>
+		public IList<BlobConverter> Converters => _converters;
+
 		/// <summary>
 		/// Default options instance
 		/// </summary>
-		public static BlobSerializerOptions Default { get; } = new BlobSerializerOptions();
+		public static BlobSerializerOptions Default { get; } = CreateDefaultOptions();
 
 		/// <summary>
-		/// Known converter types
+		/// Constructor
 		/// </summary>
-		public IList<BlobConverter> Converters { get; } = new List<BlobConverter>();
+		public BlobSerializerOptions()
+		{
+			_converters = new FreezableList<BlobConverter>(this);
+		}
 
-		readonly ConcurrentDictionary<Type, BlobConverter> _cachedConverters = new ConcurrentDictionary<Type, BlobConverter>();
-		static readonly ConcurrentDictionary<Type, BlobConverter> s_cachedDefaultConverters = new ConcurrentDictionary<Type, BlobConverter>();
-		static readonly Func<Type, BlobConverter> s_createDefaultConverter = CreateDefaultConverter;
+		void CheckMutable()
+		{
+			if (_readOnly)
+			{
+				throw new NotSupportedException("Options instance is read-only");
+			}
+		}
+
+		static BlobSerializerOptions CreateDefaultOptions()
+		{
+			BlobSerializerOptions options = new BlobSerializerOptions();
+			options.MakeReadOnly();
+			return options;
+		}
+
+		/// <summary>
+		/// Create a read-only version of these options
+		/// </summary>
+		/// <returns></returns>
+		public void MakeReadOnly() => _readOnly = true;
 
 		/// <summary>
 		/// Gets a converter for the given type
