Horde: Remove read-only flag from existing files before opening for write.

diff --git a/Engine/Source/Programs/Shared/EpicGames.Horde/Storage/Nodes/DirectoryNode.cs b/Engine/Source/Programs/Shared/EpicGames.Horde/Storage/Nodes/DirectoryNode.cs
--- a/Engine/Source/Programs/Shared/EpicGames.Horde/Storage/Nodes/DirectoryNode.cs
+++ b/Engine/Source/Programs/Shared/EpicGames.Horde/Storage/Nodes/DirectoryNode.cs
@@ -755,6 +755,8 @@ namespace EpicGames.Horde.Storage.Nodes
 					}
 				}
 
+				// Enumerate all output chunks first, to avoid thrashing the cache with reads. We probably want to always look ahead
+				// a certain amount in order to make smarter read pipelining anyway. 
 				await FindOutputChunksRootAsync(chunks.Writer, logger, cancellationSource.Token);
 
 				List<Task> tasks = new List<Task>();
@@ -775,10 +777,17 @@ namespace EpicGames.Horde.Storage.Nodes
 			{
 				// Open the file for the current chunk
 				OutputFile file = chunk.File;
+
 				FileReference locator = FileReference.Combine(baseDir, file.Directory.Path, file.FileEntry.Name);
 				DirectoryReference.CreateDirectory(locator.Directory);
 
-				await using FileStream stream = FileReference.Open(locator, FileMode.OpenOrCreate, FileAccess.ReadWrite, FileShare.ReadWrite);
+				FileInfo fileInfo = locator.ToFileInfo();
+				if (fileInfo.Exists && (fileInfo.Attributes & FileAttributes.ReadOnly) != 0)
+				{
+					fileInfo.Attributes &= ~FileAttributes.ReadOnly;
+				}
+
+				await using FileStream stream = fileInfo.Open(FileMode.OpenOrCreate, FileAccess.ReadWrite, FileShare.ReadWrite);
 				stream.SetLength(file.FileEntry.Length);
 
 				// If this file is empty, don't write anything and just move to the next chunk
