Horde: Reduce verbosity for tracing blob reads.

diff --git a/Engine/Source/Programs/Shared/EpicGames.Core/KnownLogEvents.cs b/Engine/Source/Programs/Shared/EpicGames.Core/KnownLogEvents.cs
--- a/Engine/Source/Programs/Shared/EpicGames.Core/KnownLogEvents.cs
+++ b/Engine/Source/Programs/Shared/EpicGames.Core/KnownLogEvents.cs
@@ -335,6 +335,11 @@ namespace EpicGames.Core
 		/// Invalid preflight change
 		/// </summary>
 		public static EventId Horde_InvalidPreflight { get; } = new EventId(1001);
+
+		/// <summary>
+		/// Information about blobs being read
+		/// </summary>
+		public static EventId Horde_BlobRead { get; } = new EventId(1002);
 	}
 #pragma warning restore CA1707 // Identifiers should not contain underscores
 }
diff --git a/Engine/Source/Programs/Shared/EpicGames.Horde/Storage/Nodes/DirectoryNode.cs b/Engine/Source/Programs/Shared/EpicGames.Horde/Storage/Nodes/DirectoryNode.cs
--- a/Engine/Source/Programs/Shared/EpicGames.Horde/Storage/Nodes/DirectoryNode.cs
+++ b/Engine/Source/Programs/Shared/EpicGames.Horde/Storage/Nodes/DirectoryNode.cs
@@ -713,7 +713,7 @@ namespace EpicGames.Horde.Storage.Nodes
 
 		static void TraceBlobRead(string type, string path, IBlobHandle handle, ILogger logger)
 		{
-			logger.LogInformation("Blob [{Type,-20}] Path=\"{Path}\", Locator={Locator}", type, path, handle.GetLocator());
+			logger.LogTrace(KnownLogEvents.Horde_BlobRead, "Blob [{Type,-20}] Path=\"{Path}\", Locator={Locator}", type, path, handle.GetLocator());
 		}
 
 		/// <summary>
