Horde: Disable log spam for blob reads.

diff --git a/Engine/Source/Programs/Shared/EpicGames.Horde/Storage/Nodes/DirectoryNode.cs b/Engine/Source/Programs/Shared/EpicGames.Horde/Storage/Nodes/DirectoryNode.cs
--- a/Engine/Source/Programs/Shared/EpicGames.Horde/Storage/Nodes/DirectoryNode.cs
+++ b/Engine/Source/Programs/Shared/EpicGames.Horde/Storage/Nodes/DirectoryNode.cs
@@ -711,10 +711,12 @@ namespace EpicGames.Horde.Storage.Nodes
 		record class OutputFile(OutputDir Directory, FileEntry FileEntry);
 		record class OutputChunk(OutputFile File, long Offset, long Length, IBlobHandle Handle);
 
+#pragma warning disable IDE0060
 		static void TraceBlobRead(string type, string path, IBlobHandle handle, ILogger logger)
 		{
-			logger.LogTrace(KnownLogEvents.Horde_BlobRead, "Blob [{Type,-20}] Path=\"{Path}\", Locator={Locator}", type, path, handle.GetLocator());
+//			logger.LogTrace(KnownLogEvents.Horde_BlobRead, "Blob [{Type,-20}] Path=\"{Path}\", Locator={Locator}", type, path, handle.GetLocator());
 		}
+#pragma warning restore IDE0060
 
 		/// <summary>
 		/// Utility function to allow extracting a packed directory to disk
