Horde: Remove support for disabling partitioned workspaces for now.

diff --git a/Engine/Source/Programs/Shared/EpicGames.Perforce.Managed/ManagedWorkspace.cs b/Engine/Source/Programs/Shared/EpicGames.Perforce.Managed/ManagedWorkspace.cs
--- a/Engine/Source/Programs/Shared/EpicGames.Perforce.Managed/ManagedWorkspace.cs
+++ b/Engine/Source/Programs/Shared/EpicGames.Perforce.Managed/ManagedWorkspace.cs
@@ -1218,7 +1218,8 @@ namespace EpicGames.Perforce.Managed
 				client.Host = _hostName;
 				client.Stream = streamName;
 
-				if (_options.Partitioned)
+// HACK: Always using partitioned workspaces for now
+//				if (_options.Partitioned)
 				{
 					// Partitioned and read-only types store their have table separately on the server, compared to normal (writeable) clients
 					// Clients that sync without updating the have table cannot submit so they're marked as read-only. 
