Horde: Set the correct permissions on files extracted from bundles.

diff --git a/Engine/Source/Programs/Shared/EpicGames.Horde/Storage/Nodes/DirectoryNode.cs b/Engine/Source/Programs/Shared/EpicGames.Horde/Storage/Nodes/DirectoryNode.cs
--- a/Engine/Source/Programs/Shared/EpicGames.Horde/Storage/Nodes/DirectoryNode.cs
+++ b/Engine/Source/Programs/Shared/EpicGames.Horde/Storage/Nodes/DirectoryNode.cs
@@ -789,41 +789,46 @@ namespace EpicGames.Horde.Storage.Nodes
 					fileInfo.Attributes &= ~FileAttributes.ReadOnly;
 				}
 
-				await using FileStream stream = fileInfo.Open(FileMode.OpenOrCreate, FileAccess.ReadWrite, FileShare.ReadWrite);
-				stream.SetLength(file.FileEntry.Length);
-
-				// If this file is empty, don't write anything and just move to the next chunk
-				if (file.FileEntry.Length == 0)
-				{
-					chunk = await ReadNextChunkAsync(chunkReader, cancellationToken);
-					continue;
-				}
-
-				// Process as many chunks as we can for this file
-				using MemoryMappedFile memoryMappedFile = MemoryMappedFile.CreateFromFile(stream, null, file.FileEntry.Length, MemoryMappedFileAccess.ReadWrite, HandleInheritability.None, false);
-				using MemoryMappedView memoryMappedView = new MemoryMappedView(memoryMappedFile, 0, file.FileEntry.Length);
-
-				while (chunk != null && chunk.File == file)
+				await using (FileStream stream = fileInfo.Open(FileMode.OpenOrCreate, FileAccess.ReadWrite, FileShare.ReadWrite))
 				{
-					// Write this chunk
-					using (BlobData data = await chunk.Handle.ReadBlobDataAsync(cancellationToken))
+					stream.SetLength(file.FileEntry.Length);
+					if (file.FileEntry.Length == 0)
 					{
-						TraceBlobRead("Leaf", CombinePaths(chunk.File.Directory.Path, chunk.File.FileEntry.Name), chunk.Handle, logger);
-						data.Data.CopyTo(memoryMappedView!.GetMemory(chunk.Offset, data.Data.Length));
+						// If this file is empty, don't write anything and just move to the next chunk
+						chunk = await ReadNextChunkAsync(chunkReader, cancellationToken);
 					}
-
-					// Update the stats
-					int numCompleteFiles = 0;
-					if (chunk.Offset + chunk.Length == file.FileEntry.Length)
+					else
 					{
-						numCompleteFiles = 1;
-					}
+						// Process as many chunks as we can for this file
+						using MemoryMappedFile memoryMappedFile = MemoryMappedFile.CreateFromFile(stream, null, file.FileEntry.Length, MemoryMappedFileAccess.ReadWrite, HandleInheritability.None, false);
+						using MemoryMappedView memoryMappedView = new MemoryMappedView(memoryMappedFile, 0, file.FileEntry.Length);
+
+						while (chunk != null && chunk.File == file)
+						{
+							// Write this chunk
+							using (BlobData data = await chunk.Handle.ReadBlobDataAsync(cancellationToken))
+							{
+								TraceBlobRead("Leaf", CombinePaths(chunk.File.Directory.Path, chunk.File.FileEntry.Name), chunk.Handle, logger);
+								data.Data.CopyTo(memoryMappedView!.GetMemory(chunk.Offset, data.Data.Length));
+							}
 
-					copyStats?.Update(numCompleteFiles, chunk.Length);
+							// Update the stats
+							int numCompleteFiles = 0;
+							if (chunk.Offset + chunk.Length == file.FileEntry.Length)
+							{
+								numCompleteFiles = 1;
+							}
 
-					// Read the next chunk
-					chunk = await ReadNextChunkAsync(chunkReader, cancellationToken);
+							copyStats?.Update(numCompleteFiles, chunk.Length);
+
+							// Read the next chunk
+							chunk = await ReadNextChunkAsync(chunkReader, cancellationToken);
+						}
+					}
 				}
+
+				// Set correct permissions on the output file
+				FileEntry.ApplyPermissions(fileInfo, file.FileEntry.Flags);
 			}
 		}
 
diff --git a/Engine/Source/Programs/Shared/EpicGames.Horde/Storage/Nodes/FileEntry.cs b/Engine/Source/Programs/Shared/EpicGames.Horde/Storage/Nodes/FileEntry.cs
--- a/Engine/Source/Programs/Shared/EpicGames.Horde/Storage/Nodes/FileEntry.cs
+++ b/Engine/Source/Programs/Shared/EpicGames.Horde/Storage/Nodes/FileEntry.cs
@@ -130,18 +130,7 @@ namespace EpicGames.Horde.Storage.Nodes
 			try
 			{
 				await ChunkedDataNode.CopyToFileAsync(Target.Handle, file, cancellationToken);
-
-				if ((Flags & FileEntryFlags.Executable) != 0)
-				{
-					if (RuntimeInformation.IsOSPlatform(OSPlatform.Linux))
-					{
-						FileUtils.SetFileMode_Linux(file.FullName, 0b_111_111_111);
-					}
-					else if (RuntimeInformation.IsOSPlatform(OSPlatform.OSX))
-					{
-						FileUtils.SetFileMode_Mac(file.FullName, 0b_111_111_111);
-					}
-				}
+				ApplyPermissions(file, Flags);
 			}
 			catch (Exception ex) when (ex is not OperationCanceledException)
 			{
@@ -149,6 +138,31 @@ namespace EpicGames.Horde.Storage.Nodes
 			}
 		}
 
+		/// <summary>
+		/// Applies the correct permissions to a file for a particular set of file entry flags
+		/// </summary>
+		/// <param name="fileInfo">File to modify</param>
+		/// <param name="flags">Flags for the file</param>
+		public static void ApplyPermissions(FileInfo fileInfo, FileEntryFlags flags)
+		{
+			if ((flags & FileEntryFlags.ReadOnly) != 0)
+			{
+				fileInfo.Attributes |= FileAttributes.ReadOnly;
+			}
+
+			if ((flags & FileEntryFlags.Executable) != 0)
+			{
+				if (RuntimeInformation.IsOSPlatform(OSPlatform.Linux))
+				{
+					FileUtils.SetFileMode_Linux(fileInfo.FullName, 0b_111_111_111);
+				}
+				else if (RuntimeInformation.IsOSPlatform(OSPlatform.OSX))
+				{
+					FileUtils.SetFileMode_Mac(fileInfo.FullName, 0b_111_111_111);
+				}
+			}
+		}
+
 		/// <inheritdoc/>
 		public override string ToString() => Name.ToString();
 	}
