Horde: Fix deadlock when discovering output chunks for bundle extraction if there are more than 128k individual leaf nodes in the queue.

diff --git a/Engine/Source/Programs/Shared/EpicGames.Horde/Storage/Nodes/DirectoryNode.cs b/Engine/Source/Programs/Shared/EpicGames.Horde/Storage/Nodes/DirectoryNode.cs
--- a/Engine/Source/Programs/Shared/EpicGames.Horde/Storage/Nodes/DirectoryNode.cs
+++ b/Engine/Source/Programs/Shared/EpicGames.Horde/Storage/Nodes/DirectoryNode.cs
@@ -757,12 +757,8 @@ namespace EpicGames.Horde.Storage.Nodes
 					}
 				}
 
-				// Enumerate all output chunks first, to avoid thrashing the cache with reads. We probably want to always look ahead
-				// a certain amount in order to make smarter read pipelining anyway. 
-				await FindOutputChunksRootAsync(chunks.Writer, logger, cancellationSource.Token);
-
 				List<Task> tasks = new List<Task>();
-//				tasks.Add(RunBackgroundTask(ctx => FindOutputChunksRootAsync(chunks.Writer, options, logger, ctx)));
+				tasks.Add(RunBackgroundTask(ctx => FindOutputChunksRootAsync(chunks.Writer, logger, ctx)));
 				for (int idx = 0; idx < numTasks; idx++)
 				{
 					tasks.Add(RunBackgroundTask(ctx => ExtractAsync(chunks.Reader, new DirectoryReference(directoryInfo), copyStats, logger, ctx)));
