Retry of requests for LibToolTelemetry

diff --git a/Engine/Source/Programs/AutomationTool/AutomationUtils/HttpClientSingleton.cs b/Engine/Source/Programs/AutomationTool/AutomationUtils/HttpClientSingleton.cs
--- a/Engine/Source/Programs/AutomationTool/AutomationUtils/HttpClientSingleton.cs
+++ b/Engine/Source/Programs/AutomationTool/AutomationUtils/HttpClientSingleton.cs
@@ -15,7 +15,11 @@ namespace AutomationUtils
 			new Lazy<HttpClient>(() =>
 			{
 				var retryPolicy = HttpPolicyExtensions
+				// 408, 5XX responses 
 				.HandleTransientHttpError()
+				// 404, 504 should have been covered by HandleTransientHttpError but adding here to be safe
+				.OrResult(msg => msg.StatusCode == System.Net.HttpStatusCode.NotFound || 
+								 msg.StatusCode == System.Net.HttpStatusCode.GatewayTimeout)
 				.WaitAndRetryAsync(3, retryAttempt => TimeSpan.FromSeconds(Math.Pow(2, retryAttempt)));
 
 				// per MS documentation, with a lifetime specified we shouldn't have socket exhaustion issues
diff --git a/Engine/Source/Programs/Shared/EpicGames.Horde/EpicGames.Horde.csproj b/Engine/Source/Programs/Shared/EpicGames.Horde/EpicGames.Horde.csproj
--- a/Engine/Source/Programs/Shared/EpicGames.Horde/EpicGames.Horde.csproj
+++ b/Engine/Source/Programs/Shared/EpicGames.Horde/EpicGames.Horde.csproj
@@ -1,4 +1,4 @@
-﻿<Project Sdk="Microsoft.NET.Sdk">
+<Project Sdk="Microsoft.NET.Sdk">
   <PropertyGroup>
     <TargetFramework>net6.0</TargetFramework>
     <Configuration Condition=" '$(Configuration)' == '' ">Development</Configuration>
@@ -46,7 +46,7 @@
     <PackageReference Include="Microsoft.Extensions.Caching.Memory" Version="6.0.1" />
     <PackageReference Include="Microsoft.Extensions.DependencyInjection" Version="6.0.0" />
     <PackageReference Include="Microsoft.Extensions.Http" Version="6.0.0" />
-    <PackageReference Include="Microsoft.Extensions.Http.Polly" Version="6.0.10" />
+    <PackageReference Include="Microsoft.Extensions.Http.Polly" Version="6.0.26" />
     <PackageReference Include="Microsoft.Extensions.Logging" Version="6.0.0" />
   </ItemGroup>
   <ItemGroup>
