Horde: Allow auth prompt by default when connecting to Horde.

diff --git a/Engine/Source/Programs/Shared/EpicGames.Horde/HordeOptions.cs b/Engine/Source/Programs/Shared/EpicGames.Horde/HordeOptions.cs
--- a/Engine/Source/Programs/Shared/EpicGames.Horde/HordeOptions.cs
+++ b/Engine/Source/Programs/Shared/EpicGames.Horde/HordeOptions.cs
@@ -25,7 +25,7 @@ namespace EpicGames.Horde
 		/// <summary>
 		/// Whether to allow opening a browser window to prompt for authentication
 		/// </summary>
-		public bool AllowAuthPrompt { get; set; }
+		public bool AllowAuthPrompt { get; set; } = true;
 
 		/// <summary>
 		/// Callback to allow configuring any HTTP client created for Horde
