Horde: Disable interactive auth prompt when requests already have an access token.

diff --git a/Engine/Source/Programs/Shared/EpicGames.Horde/HordeHttpAuthHandler.cs b/Engine/Source/Programs/Shared/EpicGames.Horde/HordeHttpAuthHandler.cs
--- a/Engine/Source/Programs/Shared/EpicGames.Horde/HordeHttpAuthHandler.cs
+++ b/Engine/Source/Programs/Shared/EpicGames.Horde/HordeHttpAuthHandler.cs
@@ -40,36 +40,41 @@ namespace EpicGames.Horde
 		/// <inheritdoc/>
 		protected override async Task<HttpResponseMessage> SendAsync(HttpRequestMessage request, CancellationToken cancellationToken)
 		{
-			if (!_options.Value.AllowAuthPrompt)
+			if (request.Headers.Authorization == null)
 			{
-				return await base.SendAsync(request, cancellationToken);
-			}
-
-			// Do not try to override the auth header if the user has specified it explicitly
-			if (request.Headers.Authorization != null)
-			{
-				return await base.SendAsync(request, cancellationToken);
-			}
-
-			// Try to use the cached auth header
-			if (_authHeader != null)
-			{
-				request.Headers.Authorization = _authHeader;
-
-				HttpResponseMessage response = await base.SendAsync(request, cancellationToken);
-				if (response.StatusCode != HttpStatusCode.Unauthorized)
+				if (_options.Value.AccessToken != null)
 				{
-					return response;
+					// If an explicit access token is specified, just use that
+					request.Headers.Authorization = new AuthenticationHeaderValue("Bearer", _options.Value.AccessToken);
+				}
+				else if (_options.Value.AllowAuthPrompt)
+				{
+					// Try to use the cached auth header
+					if (_authHeader != null)
+					{
+						request.Headers.Authorization = _authHeader;
+
+						HttpResponseMessage response = await base.SendAsync(request, cancellationToken);
+						if (response.StatusCode != HttpStatusCode.Unauthorized)
+						{
+							return response;
+						}
+
+						_authState.Invalidate(_authHeader);
+					}
+
+					// Otherwise update the auth header and try again
+					_authHeader = await _authState.TryGetAuthHeaderAsync(cancellationToken);
+					if (_authHeader != null)
+					{
+						request.Headers.Authorization = _authHeader;
+					}
+				}
+				else
+				{
+					// Use whatever cached auth header we currently have
+					request.Headers.Authorization = _authHeader;
 				}
-
-				_authState.Invalidate(_authHeader);
-			}
-
-			// Otherwise update the auth header and try again
-			_authHeader = await _authState.TryGetAuthHeaderAsync(cancellationToken);
-			if (_authHeader != null)
-			{
-				request.Headers.Authorization = _authHeader;
 			}
 			return await base.SendAsync(request, cancellationToken);
 		}
diff --git a/Engine/Source/Programs/Shared/EpicGames.Horde/HordeHttpClient.cs b/Engine/Source/Programs/Shared/EpicGames.Horde/HordeHttpClient.cs
--- a/Engine/Source/Programs/Shared/EpicGames.Horde/HordeHttpClient.cs
+++ b/Engine/Source/Programs/Shared/EpicGames.Horde/HordeHttpClient.cs
@@ -485,10 +485,6 @@ namespace EpicGames.Horde
 				{
 					httpClient.BaseAddress = options.Value.ServerUrl;
 				}
-				if (options.Value.AccessToken != null)
-				{
-					httpClient.DefaultRequestHeaders.Authorization = new AuthenticationHeaderValue("Bearer", options.Value.AccessToken);
-				}
 
 				httpClient.Timeout = TimeSpan.FromSeconds(240); // Global timeout
 
