Horde: Move DirectoryNodeZipStream into a separate file, and allow creating it from a IBlobRef<DirectoryNode>.

diff --git a/Engine/Source/Programs/Horde/Horde.Server/Artifacts/ArtifactsController.cs b/Engine/Source/Programs/Horde/Horde.Server/Artifacts/ArtifactsController.cs
--- a/Engine/Source/Programs/Horde/Horde.Server/Artifacts/ArtifactsController.cs
+++ b/Engine/Source/Programs/Horde/Horde.Server/Artifacts/ArtifactsController.cs
@@ -454,7 +454,7 @@ namespace Horde.Server.Artifacts
 			IStorageClient storageClient = _storageService.CreateClient(artifact.NamespaceId);
 			try
 			{
-				DirectoryNode directory = await storageClient.ReadRefTargetAsync<DirectoryNode>(artifact.RefName, DateTime.UtcNow.AddHours(1.0), cancellationToken: cancellationToken);
+				IBlobRef<DirectoryNode> directory = await storageClient.ReadRefAsync<DirectoryNode>(artifact.RefName, DateTime.UtcNow.AddHours(1.0), cancellationToken: cancellationToken);
 
 				Stream stream = directory.AsZipStream(filter).WrapOwnership(storageClient);
 				return new FileStreamResult(stream, "application/zip") { FileDownloadName = $"{artifact.RefName}.zip" };
diff --git a/Engine/Source/Programs/Horde/Horde.Server/Tools/ToolCollection.cs b/Engine/Source/Programs/Horde/Horde.Server/Tools/ToolCollection.cs
--- a/Engine/Source/Programs/Horde/Horde.Server/Tools/ToolCollection.cs
+++ b/Engine/Source/Programs/Horde/Horde.Server/Tools/ToolCollection.cs
@@ -456,8 +456,8 @@ namespace Horde.Server.Tools
 			IStorageClient client = CreateStorageClient(tool);
 			try
 			{
-				DirectoryNode node = await client.ReadRefTargetAsync<DirectoryNode>(deployment.RefName, DateTime.UtcNow - TimeSpan.FromDays(2.0), cancellationToken: cancellationToken);
-				return node.AsZipStream().WrapOwnership(client);
+				IBlobRef<DirectoryNode> nodeRef = await client.ReadRefAsync<DirectoryNode>(deployment.RefName, DateTime.UtcNow - TimeSpan.FromDays(2.0), cancellationToken: cancellationToken);
+				return nodeRef.AsZipStream().WrapOwnership(client);
 			}
 			catch
 			{
diff --git a/Engine/Source/Programs/Horde/Horde.Server/Tools/ToolsController.cs b/Engine/Source/Programs/Horde/Horde.Server/Tools/ToolsController.cs
--- a/Engine/Source/Programs/Horde/Horde.Server/Tools/ToolsController.cs
+++ b/Engine/Source/Programs/Horde/Horde.Server/Tools/ToolsController.cs
@@ -330,28 +330,31 @@ namespace Horde.Server.Tools
 			IStorageClient client = _toolCollection.CreateStorageClient(tool);
 			try
 			{
-				DirectoryNode node = await client.ReadRefTargetAsync<DirectoryNode>(deployment.RefName, DateTime.UtcNow - TimeSpan.FromDays(2.0), cancellationToken: cancellationToken);
+				IBlobRef<DirectoryNode> nodeRef = await client.ReadRefAsync<DirectoryNode>(deployment.RefName, DateTime.UtcNow - TimeSpan.FromDays(2.0), cancellationToken: cancellationToken);
 
-				if (node.Directories.Count == 0 && node.Files.Count == 1 && action != GetToolAction.Zip)
+				// If we weren't specifically asked for a zip, see if this download is a single file. If it is, allow downloading it directory.
+				if (action != GetToolAction.Zip)
 				{
-					FileEntry entry = node.Files.First();
-
-					string? contentType;
-					if (!new FileExtensionContentTypeProvider().TryGetContentType(entry.Name.ToString(), out contentType))
+					DirectoryNode node = await nodeRef.ReadBlobAsync(cancellationToken);
+					if (node.Directories.Count == 0 && node.Files.Count == 1)
 					{
-						contentType = "application/octet-stream";
-					}
+						FileEntry entry = node.Files.First();
 
-					Response.Headers.ContentLength = entry.Length;
+						string? contentType;
+						if (!new FileExtensionContentTypeProvider().TryGetContentType(entry.Name.ToString(), out contentType))
+						{
+							contentType = "application/octet-stream";
+						}
 
-					Stream stream = entry.OpenAsStream().WrapOwnership(client);
-					return new FileStreamResult(stream, contentType) { FileDownloadName = entry.Name.ToString() };
-				}
-				else
-				{
-					Stream stream = node.AsZipStream().WrapOwnership(client);
-					return new FileStreamResult(stream, "application/zip") { FileDownloadName = $"{tool.Id}-{deployment.Version}.zip" };
+						Response.Headers.ContentLength = entry.Length;
+
+						Stream fileStream = entry.OpenAsStream().WrapOwnership(client);
+						return new FileStreamResult(fileStream, contentType) { FileDownloadName = entry.Name.ToString() };
+					}
 				}
+
+				Stream stream = nodeRef.AsZipStream().WrapOwnership(client);
+				return new FileStreamResult(stream, "application/zip") { FileDownloadName = $"{tool.Id}-{deployment.Version}.zip" };
 			}
 			catch
 			{
diff --git a/Engine/Source/Programs/Shared/EpicGames.Horde/Storage/Nodes/DirectoryNode.cs b/Engine/Source/Programs/Shared/EpicGames.Horde/Storage/Nodes/DirectoryNode.cs
--- a/Engine/Source/Programs/Shared/EpicGames.Horde/Storage/Nodes/DirectoryNode.cs
+++ b/Engine/Source/Programs/Shared/EpicGames.Horde/Storage/Nodes/DirectoryNode.cs
@@ -1,22 +1,20 @@
 // Copyright Epic Games, Inc. All Rights Reserved.
 
-using EpicGames.Core;
-using EpicGames.Serialization;
-using Microsoft.Extensions.Logging;
 using System;
-using System.Buffers;
 using System.Collections.Generic;
 using System.Diagnostics;
 using System.Diagnostics.CodeAnalysis;
 using System.IO;
 using System.IO.Compression;
 using System.IO.MemoryMappedFiles;
-using System.IO.Pipelines;
 using System.Linq;
 using System.Runtime.InteropServices;
 using System.Threading;
 using System.Threading.Channels;
 using System.Threading.Tasks;
+using EpicGames.Core;
+using EpicGames.Serialization;
+using Microsoft.Extensions.Logging;
 
 namespace EpicGames.Horde.Storage.Nodes
 {
@@ -919,14 +917,6 @@ namespace EpicGames.Horde.Storage.Nodes
 				return nextPath;
 			}
 		}
-
-		/// <summary>
-		/// Returns a stream containing the zipped contents of this directory
-		/// </summary>
-		/// <param name="filter">Filter for files to include in the zip</param>
-		/// <param name="logger">Logger for diagnostic output</param>
-		/// <returns>Stream containing zipped archive data</returns>
-		public Stream AsZipStream(FileFilter? filter = null, ILogger? logger = null) => new DirectoryNodeZipStream(this, filter, logger);
 	}
 
 	class DirectoryNodeConverter : BlobConverter<DirectoryNode>
@@ -1212,169 +1202,6 @@ namespace EpicGames.Horde.Storage.Nodes
 		}
 	}
 
-	/// <summary>
-	/// Stream which zips a directory node tree dynamically
-	/// </summary>
-	class DirectoryNodeZipStream : Stream
-	{
-		/// <inheritdoc/>
-		public override bool CanRead => true;
-
-		/// <inheritdoc/>
-		public override bool CanSeek => false;
-
-		/// <inheritdoc/>
-		public override bool CanWrite => false;
-
-		/// <inheritdoc/>
-		public override long Length => throw new NotImplementedException();
-
-		/// <inheritdoc/>
-		public override long Position { get => _position; set => throw new NotImplementedException(); }
-
-		readonly Pipe _pipe;
-		readonly ILogger? _logger;
-		readonly BackgroundTask _backgroundTask;
-
-		long _position;
-		ReadOnlySequence<byte> _current = ReadOnlySequence<byte>.Empty;
-
-		/// <summary>
-		/// Constructor
-		/// </summary>
-		/// <param name="node">Root node to copy from</param>
-		/// <param name="filter">Filter for files to include in the zip</param>
-		/// <param name="logger">Optional logger for debug tracing</param>
-		public DirectoryNodeZipStream(DirectoryNode node, FileFilter? filter, ILogger? logger)
-		{
-			_pipe = new Pipe();
-			_backgroundTask = BackgroundTask.StartNew(ctx => CopyToPipeAsync(node, filter, _pipe.Writer, logger, ctx));
-			_logger = logger;
-		}
-
-		/// <inheritdoc/>
-		protected override void Dispose(bool disposing)
-		{
-			if (disposing)
-			{
-				_backgroundTask.DisposeAsync().AsTask().Wait();
-			}
-
-			base.Dispose(disposing);
-		}
-
-		/// <inheritdoc/>
-		public override async ValueTask DisposeAsync()
-		{
-			await _backgroundTask.DisposeAsync();
-
-			await base.DisposeAsync();
-		}
-
-		/// <inheritdoc/>
-		public override async ValueTask<int> ReadAsync(Memory<byte> buffer, CancellationToken cancellationToken = default)
-		{
-			while (_current.Length == 0)
-			{
-				ReadResult result = await _pipe.Reader.ReadAsync(cancellationToken);
-				_current = result.Buffer;
-
-				if (result.IsCompleted && _current.Length == 0)
-				{
-					// Wait for the background thread to finish; it may error/have errored, and we want to re-throw on this thread before completing the read.
-					await _backgroundTask.StopAsync(cancellationToken);
-					_logger?.LogInformation("Zip file pipe was read to end");
-					return 0;
-				}
-			}
-
-			int initialSize = buffer.Length;
-			while (buffer.Length > 0 && _current.Length > 0)
-			{
-				int copy = Math.Min(buffer.Length, _current.First.Length);
-				_current.First.Slice(0, copy).CopyTo(buffer);
-				_current = _current.Slice(copy);
-				buffer = buffer.Slice(copy);
-			}
-
-			if (_current.Length == 0)
-			{
-				_pipe.Reader.AdvanceTo(_current.End);
-			}
-
-			int length = initialSize - buffer.Length;
-			_position += length;
-			return length;
-		}
-
-		static async Task CopyToPipeAsync(DirectoryNode node, FileFilter? filter, PipeWriter writer, ILogger? logger, CancellationToken cancellationToken)
-		{
-			using Stream outputStream = writer.AsStream();
-			using ZipArchive archive = new ZipArchive(outputStream, ZipArchiveMode.Create);
-			await CopyFilesAsync(node, "", filter, archive, logger, cancellationToken);
-		}
-
-		static async Task CopyFilesAsync(DirectoryNode directory, string prefix, FileFilter? filter, ZipArchive archive, ILogger? logger, CancellationToken cancellationToken)
-		{
-			int numDirs = directory.Directories.Count;
-			int numFiles = directory.Files.Count;
-
-			int numCopiedDirs = 0;
-			foreach (DirectoryEntry directoryEntry in directory.Directories)
-			{
-				string directoryPath = $"{prefix}{directoryEntry.Name}/";
-				if (filter == null || filter.PossiblyMatches(directoryPath))
-				{
-					DirectoryNode node = await directoryEntry.Handle.ReadBlobAsync(cancellationToken);
-					await CopyFilesAsync(node, directoryPath, filter, archive, logger, cancellationToken);
-				}
-				numCopiedDirs++;
-			}
-
-			int numCopiedFiles = 0;
-			foreach (FileEntry fileEntry in directory.Files)
-			{
-				string filePath = $"{prefix}{fileEntry}";
-				if (filter == null || filter.Matches(filePath))
-				{
-					ZipArchiveEntry entry = archive.CreateEntry(filePath);
-
-					if ((fileEntry.Flags & FileEntryFlags.Executable) != 0)
-					{
-						entry.ExternalAttributes |= 0b_111_111_101 << 16; // rwx rwx r-x
-					}
-					else
-					{
-						entry.ExternalAttributes |= 0b_110_110_100 << 16; // rw- rw- r--
-					}
-
-					using Stream entryStream = entry.Open();
-					await fileEntry.CopyToStreamAsync(entryStream, cancellationToken);
-				}
-				numCopiedFiles++;
-			}
-
-			logger?.LogInformation("Zip file path {Prefix} has {NumCopiedDirectories}/{NumDirectories} directories, {NumCopiedFiles}/{NumFiles} files", prefix, numCopiedDirs, numDirs, numCopiedFiles, numFiles);
-		}
-
-		/// <inheritdoc/>
-		public override void Flush()
-		{
-		}
-
-		/// <inheritdoc/>
-		public override int Read(byte[] buffer, int offset, int count) => ReadAsync(buffer.AsMemory(offset, count)).AsTask().Result;
-
-		/// <inheritdoc/>
-		public override long Seek(long offset, SeekOrigin origin) => throw new NotSupportedException();
-
-		/// <inheritdoc/>
-		public override void SetLength(long value) => throw new NotSupportedException();
-
-		/// <inheritdoc/>
-		public override void Write(byte[] buffer, int offset, int count) => throw new NotSupportedException();
-	}
-
 	/// <summary>
 	/// Extension methods for writing directory nodes
 	/// </summary>
diff --git a/Engine/Source/Programs/Shared/EpicGames.Horde/Storage/Nodes/DirectoryNodeZipStream.cs b/Engine/Source/Programs/Shared/EpicGames.Horde/Storage/Nodes/DirectoryNodeZipStream.cs
new file mode 100644
--- /dev/null
+++ b/Engine/Source/Programs/Shared/EpicGames.Horde/Storage/Nodes/DirectoryNodeZipStream.cs
@@ -0,0 +1,194 @@
+// Copyright Epic Games, Inc. All Rights Reserved.
+
+using System;
+using System.Buffers;
+using System.IO;
+using System.IO.Compression;
+using System.IO.Pipelines;
+using System.Threading;
+using System.Threading.Tasks;
+using EpicGames.Core;
+using Microsoft.Extensions.Logging;
+
+namespace EpicGames.Horde.Storage.Nodes
+{
+	/// <summary>
+	/// Stream which zips a directory node tree dynamically
+	/// </summary>
+	class DirectoryNodeZipStream : Stream
+	{
+		readonly Pipe _pipe;
+		readonly ILogger? _logger;
+		readonly BackgroundTask _backgroundTask;
+
+		long _position;
+		ReadOnlySequence<byte> _current = ReadOnlySequence<byte>.Empty;
+
+		/// <inheritdoc/>
+		public override bool CanRead => true;
+
+		/// <inheritdoc/>
+		public override bool CanSeek => false;
+
+		/// <inheritdoc/>
+		public override bool CanWrite => false;
+
+		/// <inheritdoc/>
+		public override long Length => throw new NotImplementedException();
+
+		/// <inheritdoc/>
+		public override long Position { get => _position; set => throw new NotImplementedException(); }
+
+		/// <summary>
+		/// Constructor
+		/// </summary>
+		/// <param name="node">Root node to copy from</param>
+		/// <param name="filter">Filter for files to include in the zip</param>
+		/// <param name="logger">Optional logger for debug tracing</param>
+		public DirectoryNodeZipStream(IBlobRef<DirectoryNode> node, FileFilter? filter, ILogger? logger)
+		{
+			_pipe = new Pipe();
+			_backgroundTask = BackgroundTask.StartNew(ctx => CopyToPipeAsync(node, filter, _pipe.Writer, logger, ctx));
+			_logger = logger;
+		}
+
+		/// <inheritdoc/>
+		protected override void Dispose(bool disposing)
+		{
+			if (disposing)
+			{
+				_backgroundTask.DisposeAsync().AsTask().Wait();
+			}
+
+			base.Dispose(disposing);
+		}
+
+		/// <inheritdoc/>
+		public override async ValueTask DisposeAsync()
+		{
+			await _backgroundTask.DisposeAsync();
+
+			await base.DisposeAsync();
+		}
+
+		/// <inheritdoc/>
+		public override async ValueTask<int> ReadAsync(Memory<byte> buffer, CancellationToken cancellationToken = default)
+		{
+			while (_current.Length == 0)
+			{
+				ReadResult result = await _pipe.Reader.ReadAsync(cancellationToken);
+				_current = result.Buffer;
+
+				if (result.IsCompleted && _current.Length == 0)
+				{
+					// Wait for the background thread to finish; it may error/have errored, and we want to re-throw on this thread before completing the read.
+					await _backgroundTask.StopAsync(cancellationToken);
+					_logger?.LogInformation("Zip file pipe was read to end");
+					return 0;
+				}
+			}
+
+			int initialSize = buffer.Length;
+			while (buffer.Length > 0 && _current.Length > 0)
+			{
+				int copy = Math.Min(buffer.Length, _current.First.Length);
+				_current.First.Slice(0, copy).CopyTo(buffer);
+				_current = _current.Slice(copy);
+				buffer = buffer.Slice(copy);
+			}
+
+			if (_current.Length == 0)
+			{
+				_pipe.Reader.AdvanceTo(_current.End);
+			}
+
+			int length = initialSize - buffer.Length;
+			_position += length;
+			return length;
+		}
+
+		static async Task CopyToPipeAsync(IBlobRef<DirectoryNode> node, FileFilter? filter, PipeWriter writer, ILogger? logger, CancellationToken cancellationToken)
+		{
+			using Stream outputStream = writer.AsStream();
+			using ZipArchive archive = new ZipArchive(outputStream, ZipArchiveMode.Create);
+			await CopyFilesAsync(node, "", filter, archive, logger, cancellationToken);
+		}
+
+		static async Task CopyFilesAsync(IBlobRef<DirectoryNode> directoryRef, string prefix, FileFilter? filter, ZipArchive archive, ILogger? logger, CancellationToken cancellationToken)
+		{
+			DirectoryNode directory = await directoryRef.ReadBlobAsync(cancellationToken);
+
+			int numDirs = directory.Directories.Count;
+			int numFiles = directory.Files.Count;
+
+			int numCopiedDirs = 0;
+			foreach (DirectoryEntry directoryEntry in directory.Directories)
+			{
+				string directoryPath = $"{prefix}{directoryEntry.Name}/";
+				if (filter == null || filter.PossiblyMatches(directoryPath))
+				{
+					await CopyFilesAsync(directoryEntry.Handle, directoryPath, filter, archive, logger, cancellationToken);
+				}
+				numCopiedDirs++;
+			}
+
+			int numCopiedFiles = 0;
+			foreach (FileEntry fileEntry in directory.Files)
+			{
+				string filePath = $"{prefix}{fileEntry}";
+				if (filter == null || filter.Matches(filePath))
+				{
+					ZipArchiveEntry entry = archive.CreateEntry(filePath);
+
+					if ((fileEntry.Flags & FileEntryFlags.Executable) != 0)
+					{
+						entry.ExternalAttributes |= 0b_111_111_101 << 16; // rwx rwx r-x
+					}
+					else
+					{
+						entry.ExternalAttributes |= 0b_110_110_100 << 16; // rw- rw- r--
+					}
+
+					using Stream entryStream = entry.Open();
+					await fileEntry.CopyToStreamAsync(entryStream, cancellationToken);
+				}
+				numCopiedFiles++;
+			}
+
+			logger?.LogInformation("Zip file path {Prefix} has {NumCopiedDirectories}/{NumDirectories} directories, {NumCopiedFiles}/{NumFiles} files", prefix, numCopiedDirs, numDirs, numCopiedFiles, numFiles);
+		}
+
+		/// <inheritdoc/>
+		public override void Flush()
+		{
+		}
+
+		/// <inheritdoc/>
+		public override int Read(byte[] buffer, int offset, int count) => ReadAsync(buffer.AsMemory(offset, count)).AsTask().Result;
+
+		/// <inheritdoc/>
+		public override long Seek(long offset, SeekOrigin origin) => throw new NotSupportedException();
+
+		/// <inheritdoc/>
+		public override void SetLength(long value) => throw new NotSupportedException();
+
+		/// <inheritdoc/>
+		public override void Write(byte[] buffer, int offset, int count) => throw new NotSupportedException();
+	}
+
+	/// <summary>
+	/// Extension methods for <see cref="DirectoryNodeZipStream"/>
+	/// </summary>
+	public static class DirectoryNodeZipStreamExtensions
+	{
+		/// <summary>
+		/// Returns a stream containing the zipped contents of this directory
+		/// </summary>
+		/// <param name="directoryRef">The directory to zip</param>
+		/// <param name="filter">Filter for files to include in the zip</param>
+		/// <param name="logger">Logger for diagnostic output</param>
+		/// <returns>Stream containing zipped archive data</returns>
+		public static Stream AsZipStream(this IBlobRef<DirectoryNode> directoryRef, FileFilter? filter = null, ILogger? logger = null)
+			=> new DirectoryNodeZipStream(directoryRef, filter, logger);
+	}
+}
