Horde: Fix handling of auth headers for Horde storage upload redirects.

* Since upload redirects are handled explicitly, unlike HTTP GET redirect responses, we need to create a new HTTP client which doesn't insert the standard Horde auth header to execute them.

diff --git a/Engine/Source/Programs/Shared/EpicGames.Horde/HordeClient.cs b/Engine/Source/Programs/Shared/EpicGames.Horde/HordeClient.cs
--- a/Engine/Source/Programs/Shared/EpicGames.Horde/HordeClient.cs
+++ b/Engine/Source/Programs/Shared/EpicGames.Horde/HordeClient.cs
@@ -36,7 +36,8 @@ namespace EpicGames.Horde
 		public IStorageClient CreateStorageClient(string basePath)
 		{
 			Func<HttpClient> createClient = () => _httpClientFactory.CreateClient(HordeHttpClient.HttpClientName);
-			HttpStorageBackend httpStorageBackend = new HttpStorageBackend(basePath, createClient, _loggerFactory.CreateLogger<HttpStorageBackend>());
+			Func<HttpClient> createUploadRedirectClient = () => _httpClientFactory.CreateClient(HordeHttpClient.UploadRedirectHttpClientName);
+			HttpStorageBackend httpStorageBackend = new HttpStorageBackend(basePath, createClient, createUploadRedirectClient, _loggerFactory.CreateLogger<HttpStorageBackend>());
 			return new BundleStorageClient(httpStorageBackend, _bundleCache, _loggerFactory.CreateLogger<BundleStorageClient>());
 		}
 	}
diff --git a/Engine/Source/Programs/Shared/EpicGames.Horde/HordeHttpClient.cs b/Engine/Source/Programs/Shared/EpicGames.Horde/HordeHttpClient.cs
--- a/Engine/Source/Programs/Shared/EpicGames.Horde/HordeHttpClient.cs
+++ b/Engine/Source/Programs/Shared/EpicGames.Horde/HordeHttpClient.cs
@@ -53,6 +53,11 @@ namespace EpicGames.Horde
 		/// </summary>
 		public const string HttpClientName = "HordeHttpClient";
 
+		/// <summary>
+		/// Name of clients created from the http client factory for handling upload redirects. Should not contain Horde auth headers.
+		/// </summary>
+		public const string UploadRedirectHttpClientName = "HordeUploadRedirectHttpClient";
+
 		readonly HttpClient _httpClient;
 
 		static readonly JsonSerializerOptions s_jsonSerializerOptions = CreateJsonSerializerOptions();
@@ -519,14 +524,7 @@ namespace EpicGames.Horde
 				}
 			}
 
-			IHttpClientBuilder builder = services.AddHttpClient<HordeHttpClient>(HordeHttpClient.HttpClientName, ConfigureClientFromEnvironment)
-				.AddPolicyHandler((serviceProvider, request) => CreateDefaultTimeoutRetryPolicy(request, serviceProvider.GetRequiredService<ILogger<HttpStorageBackend>>()))
-				.AddPolicyHandler((serviceProvider, request) => CreateDefaultTransientErrorPolicy(request, serviceProvider.GetRequiredService<ILogger<HttpStorageBackend>>()));
-
-			services.AddSingleton<HordeHttpAuthHandlerState>();
-			services.AddTransient<HordeHttpAuthHandler>();
-
-			// Sets defaults from the configured HordeOptions
+			// Register the HTTP client for handling login requests
 			void ConfigureClientFromOptions(IServiceProvider serviceProvider, HttpClient httpClient)
 			{
 				IOptions<HordeOptions> options = serviceProvider.GetRequiredService<IOptions<HordeOptions>>();
@@ -535,8 +533,20 @@ namespace EpicGames.Horde
 					httpClient.BaseAddress = options.Value.ServerUrl;
 				}
 			}
+			services.AddSingleton<HordeHttpAuthHandlerState>();
+			services.AddTransient<HordeHttpAuthHandler>();
 			services.AddHttpClient(HordeHttpAuthHandlerState.HttpClientName, ConfigureClientFromOptions);
 
+			// Register the HTTP client for processing upload redirects
+			services.AddHttpClient(HordeHttpClient.UploadRedirectHttpClientName)
+				.AddPolicyHandler((serviceProvider, request) => CreateDefaultTimeoutRetryPolicy(request, serviceProvider.GetRequiredService<ILogger<HttpStorageBackend>>()))
+				.AddPolicyHandler((serviceProvider, request) => CreateDefaultTransientErrorPolicy(request, serviceProvider.GetRequiredService<ILogger<HttpStorageBackend>>()));
+
+			// Create the HTTP client for handling Horde requests
+			IHttpClientBuilder builder = services.AddHttpClient<HordeHttpClient>(HordeHttpClient.HttpClientName, ConfigureClientFromEnvironment)
+				.AddPolicyHandler((serviceProvider, request) => CreateDefaultTimeoutRetryPolicy(request, serviceProvider.GetRequiredService<ILogger<HttpStorageBackend>>()))
+				.AddPolicyHandler((serviceProvider, request) => CreateDefaultTransientErrorPolicy(request, serviceProvider.GetRequiredService<ILogger<HttpStorageBackend>>()));
+
 			builder = builder.AddHttpMessageHandler<HordeHttpAuthHandler>();
 
 			return builder;
diff --git a/Engine/Source/Programs/Shared/EpicGames.Horde/Storage/Backends/HttpStorageBackend.cs b/Engine/Source/Programs/Shared/EpicGames.Horde/Storage/Backends/HttpStorageBackend.cs
--- a/Engine/Source/Programs/Shared/EpicGames.Horde/Storage/Backends/HttpStorageBackend.cs
+++ b/Engine/Source/Programs/Shared/EpicGames.Horde/Storage/Backends/HttpStorageBackend.cs
@@ -28,6 +28,7 @@ namespace EpicGames.Horde.Storage.Backends
 
 		readonly string _basePath;
 		readonly Func<HttpClient> _createClient;
+		readonly Func<HttpClient> _createUploadRedirectClient;
 		readonly ILogger _logger;
 		bool _supportsUploadRedirects = true;
 
@@ -37,10 +38,11 @@ namespace EpicGames.Horde.Storage.Backends
 		/// <summary>
 		/// Constructor
 		/// </summary>
-		public HttpStorageBackend(string basePath, Func<HttpClient> createClient, ILogger logger)
+		public HttpStorageBackend(string basePath, Func<HttpClient> createClient, Func<HttpClient> createUploadRedirectClient, ILogger logger)
 		{
 			_basePath = basePath.TrimEnd('/');
 			_createClient = createClient;
+			_createUploadRedirectClient = createUploadRedirectClient;
 			_logger = logger;
 		}
 
@@ -100,25 +102,19 @@ namespace EpicGames.Horde.Storage.Backends
 
 			if (_supportsUploadRedirects)
 			{
-				using (HttpClient redirectHttpClient = _createClient())
+				WriteBlobResponse redirectResponse = await SendWriteRequestAsync(null, prefix, cancellationToken);
+				if (redirectResponse.UploadUrl != null)
 				{
-					// Don't send the auth header if we're following a redirect
-					redirectHttpClient.DefaultRequestHeaders.Authorization = null;
-
-					WriteBlobResponse redirectResponse = await SendWriteRequestAsync(null, prefix, cancellationToken);
-					if (redirectResponse.UploadUrl != null)
+					using HttpClient uploadRedirectClient = _createUploadRedirectClient();
+					using HttpResponseMessage uploadResponse = await uploadRedirectClient.PutAsync(redirectResponse.UploadUrl, streamContent, cancellationToken);
+					if (!uploadResponse.IsSuccessStatusCode)
 					{
-						using (HttpResponseMessage uploadResponse = await redirectHttpClient.PutAsync(redirectResponse.UploadUrl, streamContent, cancellationToken))
-						{
-							if (!uploadResponse.IsSuccessStatusCode)
-							{
-								string body = await uploadResponse.Content.ReadAsStringAsync(cancellationToken);
-								throw new StorageException($"Unable to upload data to redirected URL: {body}");
-							}
-						}
-						_logger.LogDebug("Written {Locator} (using redirect)", redirectResponse.Blob);
-						return new BlobLocator(redirectResponse.Blob);
+						string body = await uploadResponse.Content.ReadAsStringAsync(cancellationToken);
+						throw new StorageException($"Unable to upload data to redirected URL: {body}");
 					}
+
+					_logger.LogDebug("Written {Locator} (using redirect)", redirectResponse.Blob);
+					return new BlobLocator(redirectResponse.Blob);
 				}
 			}
 
@@ -352,7 +348,9 @@ namespace EpicGames.Horde.Storage.Backends
 				return httpClient;
 			}
 
-			IStorageBackend backend = new HttpStorageBackend(basePath, CreateClient, _backendLogger);
+			HttpClient CreateUploadRedirectClient() => _httpClientFactory.CreateClient(HordeHttpClient.UploadRedirectHttpClientName);
+
+			IStorageBackend backend = new HttpStorageBackend(basePath, CreateClient, CreateUploadRedirectClient, _backendLogger);
 			if (_backendCache != null && withBackendCache)
 			{
 				backend = _backendCache.CreateWrapper(basePath, backend);
