Horde: Move functionality for extracting and updating directories into separate files.

diff --git a/Engine/Source/Programs/Shared/EpicGames.Horde/Storage/Nodes/DirectoryEntry.cs b/Engine/Source/Programs/Shared/EpicGames.Horde/Storage/Nodes/DirectoryEntry.cs
deleted file mode 100644
diff --git a/Engine/Source/Programs/Shared/EpicGames.Horde/Storage/Nodes/DirectoryNode.Extract.cs b/Engine/Source/Programs/Shared/EpicGames.Horde/Storage/Nodes/DirectoryNode.Extract.cs
new file mode 100644
--- /dev/null
+++ b/Engine/Source/Programs/Shared/EpicGames.Horde/Storage/Nodes/DirectoryNode.Extract.cs
@@ -0,0 +1,359 @@
+// Copyright Epic Games, Inc. All Rights Reserved.
+
+using System;
+using System.Collections.Generic;
+using System.Diagnostics;
+using System.IO;
+using System.IO.MemoryMappedFiles;
+using System.Threading;
+using System.Threading.Channels;
+using System.Threading.Tasks;
+using EpicGames.Core;
+using Microsoft.Extensions.Logging;
+
+namespace EpicGames.Horde.Storage.Nodes
+{
+	/// <summary>
+	/// Stats reported for copy operations
+	/// </summary>
+	public interface ICopyStats
+	{
+		/// <summary>
+		/// Number of files that have been copied
+		/// </summary>
+		int Count { get; }
+
+		/// <summary>
+		/// Total size of data to be copied
+		/// </summary>
+		long Size { get; }
+
+		/// <summary>
+		/// Processing speed, in bytes per second
+		/// </summary>
+		double Rate { get; }
+	}
+
+	/// <summary>
+	/// Reports progress info back to callers
+	/// </summary>
+	class CopyStats : ICopyStats
+	{
+		readonly object _lockObject = new object();
+		readonly Stopwatch _timer = Stopwatch.StartNew();
+		readonly IProgress<ICopyStats> _progress;
+		long _lastTotalSize;
+
+		public int Count { get; set; }
+		public long Size { get; set; }
+		public double Rate { get; set; }
+
+		public CopyStats(IProgress<ICopyStats> progress)
+		{
+			_progress = progress;
+		}
+
+		public void Update(int count, long size)
+		{
+			lock (_lockObject)
+			{
+				Count += count;
+				Size += size;
+				if (_timer.Elapsed > TimeSpan.FromSeconds(10.0))
+				{
+					Rate = (Size - _lastTotalSize) / _timer.Elapsed.TotalSeconds;
+					_lastTotalSize = Size;
+
+					_progress.Report(this);
+					_timer.Restart();
+				}
+			}
+		}
+
+		public void Flush()
+		{
+			lock (_lockObject)
+			{
+				_progress.Report(this);
+				_timer.Restart();
+			}
+		}
+	}
+
+	/// <summary>
+	/// Progress logger for writing copy stats
+	/// </summary>
+	public class CopyStatsLogger : IProgress<ICopyStats>
+	{
+		readonly int _totalCount;
+		readonly long _totalSize;
+		readonly ILogger _logger;
+
+		/// <summary>
+		/// Constructor
+		/// </summary>
+		public CopyStatsLogger(ILogger logger)
+			=> _logger = logger;
+
+		/// <summary>
+		/// Constructor
+		/// </summary>
+		public CopyStatsLogger(int totalCount, long totalSize, ILogger logger)
+		{
+			_totalCount = totalCount;
+			_totalSize = totalSize;
+			_logger = logger;
+		}
+
+		/// <inheritdoc/>
+		public void Report(ICopyStats stats)
+		{
+			if (_totalCount > 0 && _totalSize > 0)
+			{
+				_logger.LogInformation("Copied {NumFiles:n0}/{TotalFiles:n0} files ({Size:n1}/{TotalSize:n1}mb, {Rate:n1}mb/s, {Pct}%)", stats.Count, _totalCount, stats.Size / (1024.0 * 1024.0), _totalSize / (1024.0 * 1024.0), stats.Rate / (1024.0 * 1024.0), (int)((Math.Max(stats.Size, 1) * 100) / Math.Max(_totalSize, 1)));
+			}
+			else if (_totalCount > 0)
+			{
+				_logger.LogInformation("Copied {NumFiles:n0}/{TotalFiles:n0} files ({Size:n1}mb, {Rate:n1}mb/s)", stats.Count, _totalCount, stats.Size / (1024.0 * 1024.0), stats.Rate / (1024.0 * 1024.0));
+			}
+			else if (_totalSize > 0)
+			{
+				_logger.LogInformation("Copied {NumFiles:n0} files ({Size:n1}/{TotalSize:n1}mb, {Rate:n1}mb/s, {Pct}%)", stats.Count, stats.Size / (1024.0 * 1024.0), _totalSize / (1024.0 * 1024.0), stats.Rate / (1024.0 * 1024.0), (int)((Math.Max(stats.Size, 1) * 100) / Math.Max(_totalSize, 1)));
+			}
+			else
+			{
+				_logger.LogInformation("Copied {NumFiles:n0} files ({Size:n1}mb, {Rate:n1}mb/s)", stats.Count, stats.Size / (1024.0 * 1024.0), stats.Rate / (1024.0 * 1024.0));
+			}
+		}
+	}
+
+	/// <summary>
+	/// Extension methods for extracting data from directory nodes
+	/// </summary>
+	public static class DirectoryNodeExtract
+	{
+		/// <summary>
+		/// Utility function to allow extracting a packed directory to disk
+		/// </summary>
+		/// <param name="directoryNode">Directory to update</param>
+		/// <param name="directoryInfo"></param>
+		/// <param name="logger"></param>
+		/// <param name="cancellationToken"></param>
+		public static Task CopyToDirectoryAsync(this DirectoryNode directoryNode, DirectoryInfo directoryInfo, ILogger logger, CancellationToken cancellationToken) => CopyToDirectoryAsync(directoryNode, directoryInfo, null, logger, cancellationToken);
+
+		record class OutputDir(string Path, DirectoryNode Node);
+		record class OutputFile(OutputDir Directory, FileEntry FileEntry);
+		record class OutputChunk(OutputFile File, long Offset, long Length, IBlobHandle Handle);
+
+#pragma warning disable IDE0060
+		static void TraceBlobRead(string type, string path, IBlobHandle handle, ILogger logger)
+		{
+			//			logger.LogTrace(KnownLogEvents.Horde_BlobRead, "Blob [{Type,-20}] Path=\"{Path}\", Locator={Locator}", type, path, handle.GetLocator());
+		}
+#pragma warning restore IDE0060
+
+		/// <summary>
+		/// Utility function to allow extracting a packed directory to disk
+		/// </summary>
+		/// <param name="directoryNode">Directory to update</param>
+		/// <param name="directoryInfo">Direcotry to write to</param>
+		/// <param name="progress">Sink for progress updates</param>
+		/// <param name="logger">Logger for output</param>
+		/// <param name="cancellationToken">Cancellation token for the operation</param>
+		public static async Task CopyToDirectoryAsync(this DirectoryNode directoryNode, DirectoryInfo directoryInfo, IProgress<ICopyStats>? progress, ILogger logger, CancellationToken cancellationToken)
+		{
+			int numTasks = Math.Min(1 + (int)(directoryNode.Length / (16 * 1024 * 1024)), 16);
+			logger.LogInformation("Splitting read into {NumThreads} threads", numTasks);
+
+			CopyStats? copyStats = null;
+			if (progress != null)
+			{
+				copyStats = new CopyStats(progress);
+			}
+
+			Channel<OutputChunk> chunks = Channel.CreateBounded<OutputChunk>(new BoundedChannelOptions(128 * 1024) { FullMode = BoundedChannelFullMode.Wait });
+			using (CancellationTokenSource cancellationSource = CancellationTokenSource.CreateLinkedTokenSource(cancellationToken))
+			{
+				// Helper method to run a background task and set a cancellation source on error
+				async Task RunBackgroundTask(Func<CancellationToken, Task> taskFunc)
+				{
+					try
+					{
+						await taskFunc(cancellationSource.Token);
+					}
+					catch (OperationCanceledException)
+					{
+						// Ignore
+					}
+					catch (Exception ex)
+					{
+						logger.LogError(ex, "Error while extracting data: {Message}", ex.Message);
+						cancellationSource.Cancel();
+					}
+				}
+
+				List<Task> tasks = new List<Task>();
+				tasks.Add(RunBackgroundTask(ctx => FindOutputChunksRootAsync(directoryNode, chunks.Writer, logger, ctx)));
+				for (int idx = 0; idx < numTasks; idx++)
+				{
+					tasks.Add(RunBackgroundTask(ctx => ExtractAsync(chunks.Reader, new DirectoryReference(directoryInfo), copyStats, logger, ctx)));
+				}
+
+				await Task.WhenAll(tasks);
+			}
+		}
+
+		static async Task ExtractAsync(ChannelReader<OutputChunk> chunkReader, DirectoryReference baseDir, CopyStats? copyStats, ILogger logger, CancellationToken cancellationToken)
+		{
+			OutputChunk? chunk = await ReadNextChunkAsync(chunkReader, cancellationToken);
+			while (chunk != null)
+			{
+				// Open the file for the current chunk
+				OutputFile file = chunk.File;
+
+				FileReference locator = FileReference.Combine(baseDir, file.Directory.Path, file.FileEntry.Name);
+				DirectoryReference.CreateDirectory(locator.Directory);
+
+				FileInfo fileInfo = locator.ToFileInfo();
+				if (fileInfo.Exists && (fileInfo.Attributes & FileAttributes.ReadOnly) != 0)
+				{
+					fileInfo.Attributes &= ~FileAttributes.ReadOnly;
+				}
+
+				await using (FileStream stream = fileInfo.Open(FileMode.OpenOrCreate, FileAccess.ReadWrite, FileShare.ReadWrite))
+				{
+					stream.SetLength(file.FileEntry.Length);
+					if (file.FileEntry.Length == 0)
+					{
+						// If this file is empty, don't write anything and just move to the next chunk
+						chunk = await ReadNextChunkAsync(chunkReader, cancellationToken);
+					}
+					else
+					{
+						// Process as many chunks as we can for this file
+						using MemoryMappedFile memoryMappedFile = MemoryMappedFile.CreateFromFile(stream, null, file.FileEntry.Length, MemoryMappedFileAccess.ReadWrite, HandleInheritability.None, false);
+						using MemoryMappedView memoryMappedView = new MemoryMappedView(memoryMappedFile, 0, file.FileEntry.Length);
+
+						while (chunk != null && chunk.File == file)
+						{
+							// Write this chunk
+							using (BlobData data = await chunk.Handle.ReadBlobDataAsync(cancellationToken))
+							{
+								TraceBlobRead("Leaf", CombinePaths(chunk.File.Directory.Path, chunk.File.FileEntry.Name), chunk.Handle, logger);
+								data.Data.CopyTo(memoryMappedView!.GetMemory(chunk.Offset, data.Data.Length));
+							}
+
+							// Update the stats
+							int numCompleteFiles = 0;
+							if (chunk.Offset + chunk.Length == file.FileEntry.Length)
+							{
+								numCompleteFiles = 1;
+							}
+
+							copyStats?.Update(numCompleteFiles, chunk.Length);
+
+							// Read the next chunk
+							chunk = await ReadNextChunkAsync(chunkReader, cancellationToken);
+						}
+					}
+				}
+
+				// Set correct permissions on the output file
+				FileEntry.SetPermissions(fileInfo, file.FileEntry.Flags);
+			}
+		}
+
+		static async ValueTask<OutputChunk?> ReadNextChunkAsync(ChannelReader<OutputChunk> chunkReader, CancellationToken cancellationToken)
+		{
+			await chunkReader.WaitToReadAsync(cancellationToken);
+
+			OutputChunk? chunk;
+			if (chunkReader.TryRead(out chunk))
+			{
+				return chunk;
+			}
+			else
+			{
+				return null;
+			}
+		}
+
+		static async Task FindOutputChunksRootAsync(DirectoryNode node, ChannelWriter<OutputChunk> chunks, ILogger logger, CancellationToken cancellationToken)
+		{
+			await FindOutputChunksAsync("", node, chunks, logger, cancellationToken);
+			chunks.Complete();
+		}
+
+		static async Task FindOutputChunksAsync(string path, DirectoryNode node, ChannelWriter<OutputChunk> chunks, ILogger logger, CancellationToken cancellationToken)
+		{
+			OutputDir outputDir = new OutputDir(path, node);
+
+			foreach (FileEntry fileEntry in node.Files)
+			{
+				OutputFile outputFile = new OutputFile(outputDir, fileEntry);
+				await FindOutputChunksAsync(outputFile, 0, fileEntry.Target, chunks, logger, cancellationToken);
+			}
+
+			foreach (DirectoryEntry directoryEntry in node.Directories)
+			{
+				DirectoryNode subDirectoryNode = await directoryEntry.Handle.ReadBlobAsync(cancellationToken);
+
+				string subPath = CombinePaths(outputDir.Path, directoryEntry.Name);
+				TraceBlobRead("Directory", subPath, directoryEntry.Handle, logger);
+
+				await FindOutputChunksAsync(subPath, subDirectoryNode, chunks, logger, cancellationToken);
+			}
+		}
+
+		static async Task<long> FindOutputChunksAsync(OutputFile outputFile, long offset, ChunkedDataNodeRef dataRef, ChannelWriter<OutputChunk> chunks, ILogger logger, CancellationToken cancellationToken)
+		{
+			if (dataRef.Type == ChunkedDataNodeType.Leaf)
+			{
+				await chunks.WriteAsync(new OutputChunk(outputFile, offset, dataRef.Length, dataRef.Handle), cancellationToken);
+				if (dataRef.Length < 0)
+				{
+					// Backwards compatibility hack for v2 format
+					using BlobData data = await dataRef.Handle.ReadBlobDataAsync(cancellationToken);
+					return data.Data.Length;
+				}
+				return dataRef.Length;
+			}
+			else
+			{
+				using BlobData data = await dataRef.Handle.ReadBlobDataAsync(cancellationToken);
+				TraceBlobRead("Interior", CombinePaths(outputFile.Directory.Path, outputFile.FileEntry.Name), dataRef.Handle, logger);
+
+				if (data.Type.Guid == LeafChunkedDataNodeConverter.BlobType.Guid)
+				{
+					await chunks.WriteAsync(new OutputChunk(outputFile, offset, dataRef.Length, dataRef.Handle), cancellationToken);
+					return data.Data.Length;
+				}
+				else
+				{
+					long length = 0;
+
+					InteriorChunkedDataNode interiorNode = BlobSerializer.Deserialize<InteriorChunkedDataNode>(data);
+					foreach (ChunkedDataNodeRef childRef in interiorNode.Children)
+					{
+						length += await FindOutputChunksAsync(outputFile, offset + length, childRef, chunks, logger, cancellationToken);
+					}
+
+					return length;
+				}
+			}
+		}
+
+		static string CombinePaths(string basePath, string nextPath)
+		{
+			if (basePath.Length > 0)
+			{
+				return $"{basePath}/{nextPath}";
+			}
+			else
+			{
+				return nextPath;
+			}
+		}
+	}
+}
+
diff --git a/Engine/Source/Programs/Shared/EpicGames.Horde/Storage/Nodes/DirectoryNode.Update.cs b/Engine/Source/Programs/Shared/EpicGames.Horde/Storage/Nodes/DirectoryNode.Update.cs
new file mode 100644
--- /dev/null
+++ b/Engine/Source/Programs/Shared/EpicGames.Horde/Storage/Nodes/DirectoryNode.Update.cs
@@ -0,0 +1,469 @@
+// Copyright Epic Games, Inc. All Rights Reserved.
+
+using System;
+using System.Collections.Generic;
+using System.IO;
+using System.IO.Compression;
+using System.Linq;
+using System.Threading;
+using System.Threading.Tasks;
+using EpicGames.Core;
+
+namespace EpicGames.Horde.Storage.Nodes
+{
+	/// <summary>
+	/// Describes an update to a file in a directory tree
+	/// </summary>
+	/// <param name="Path">Path to the file</param>
+	/// <param name="Length">Length of the file data</param>
+	/// <param name="Flags">Flags for the new file entry</param>
+	/// <param name="StreamHash">Hash of the entire stream</param>
+	/// <param name="Nodes">Chunked data for the file</param>
+	/// <param name="CustomData"></param>
+	public record class FileUpdate(string Path, FileEntryFlags Flags, long Length, IoHash StreamHash, List<ChunkedDataNodeRef> Nodes, ReadOnlyMemory<byte> CustomData = default)
+	{
+		/// <summary>
+		/// Constructor
+		/// </summary>
+		public FileUpdate(string path, FileEntryFlags flags, long length, ChunkedData chunkedData, ReadOnlyMemory<byte> customData = default)
+			: this(path, flags, length, chunkedData.StreamHash, new List<ChunkedDataNodeRef> { chunkedData.Root }, customData)
+		{
+		}
+
+		/// <summary>
+		/// Constructor
+		/// </summary>
+		public FileUpdate(string path, FileEntryFlags flags, long length, LeafChunkedData chunkedData, ReadOnlyMemory<byte> customData = default)
+			: this(path, flags, length, chunkedData.Hash, chunkedData.LeafHandles, customData)
+		{
+		}
+
+		/// <summary>
+		/// Writes interior node data to the given writer
+		/// </summary>
+		public async ValueTask WriteInteriorNodesAsync(IBlobWriter writer, InteriorChunkedDataNodeOptions interiorNodeOptions, CancellationToken cancellationToken)
+		{
+			if (Nodes.Count > 1)
+			{
+				ChunkedDataNodeRef rootRef = await InteriorChunkedDataNode.CreateTreeAsync(Nodes, interiorNodeOptions, writer, cancellationToken);
+				Nodes.Clear();
+				Nodes.Add(rootRef);
+			}
+		}
+	}
+
+	/// <summary>
+	/// Describes an update to a directory node
+	/// </summary>
+	public class DirectoryUpdate
+	{
+		/// <summary>
+		/// Directories to be updated
+		/// </summary>
+		public SortedDictionary<string, DirectoryUpdate?> Directories { get; } = new SortedDictionary<string, DirectoryUpdate?>(StringComparer.Ordinal);
+
+		/// <summary>
+		/// Files to be updated
+		/// </summary>
+		public SortedDictionary<string, FileUpdate?> Files { get; } = new SortedDictionary<string, FileUpdate?>(StringComparer.Ordinal);
+
+		/// <summary>
+		/// Reset this instance
+		/// </summary>
+		public void Clear()
+		{
+			Directories.Clear();
+			Files.Clear();
+		}
+
+		/// <summary>
+		/// Adds a file by path to this object
+		/// </summary>
+		/// <param name="path">Path to add to</param>
+		/// <param name="fileUpdate">Content for the file</param>
+		public void AddFile(string path, FileUpdate? fileUpdate)
+		{
+			string[] fragments = path.Split(new char[] { '/', '\\' });
+
+			DirectoryUpdate lastDir = this;
+			for (int idx = 0; idx + 1 < fragments.Length; idx++)
+			{
+				DirectoryUpdate? nextDir;
+				if (!lastDir.Directories.TryGetValue(fragments[idx], out nextDir) || nextDir == null)
+				{
+					if (fileUpdate == null)
+					{
+						return;
+					}
+
+					nextDir = new DirectoryUpdate();
+					lastDir.Directories.Add(fragments[idx], nextDir);
+				}
+				lastDir = nextDir;
+			}
+
+			lastDir.Files[fragments[^1]] = fileUpdate;
+		}
+
+		/// <summary>
+		/// Adds a file to the tree
+		/// </summary>
+		/// <param name="path">Path to the file</param>
+		/// <param name="flags">Flags for the new file entry</param>
+		/// <param name="length">Length of the file</param>
+		/// <param name="chunkedData">Chunked data instance</param>
+		/// <param name="customData"></param>
+		public FileUpdate AddFile(string path, FileEntryFlags flags, long length, LeafChunkedData chunkedData, ReadOnlyMemory<byte> customData)
+		{
+			string name = path;
+			for (int idx = path.Length - 1; idx >= 0; idx--)
+			{
+				if (path[idx] == Path.DirectorySeparatorChar || path[idx] == Path.AltDirectorySeparatorChar)
+				{
+					name = path.Substring(idx + 1);
+					break;
+				}
+			}
+
+			FileUpdate update = new FileUpdate(name, flags, length, chunkedData, customData);
+			AddFile(path, update);
+			return update;
+		}
+
+		/// <summary>
+		/// Adds a filtered list of files from disk
+		/// </summary>
+		/// <param name="files">Files to add</param>
+		public void AddFiles(IEnumerable<FileUpdate> files)
+		{
+			IEnumerator<FileUpdate> enumerator = files.GetEnumerator();
+			if (enumerator.MoveNext())
+			{
+				AddFiles(String.Empty, enumerator);
+			}
+		}
+
+		bool AddFiles(ReadOnlySpan<char> prefix, IEnumerator<FileUpdate> files)
+		{
+			for (; ; )
+			{
+				FileUpdate fileUpdate = files.Current;
+				if (fileUpdate.Path.Length < prefix.Length || !fileUpdate.Path.AsSpan().StartsWith(prefix, FileReference.Comparison))
+				{
+					return true;
+				}
+
+				int nextDirLength = fileUpdate.Path.AsSpan(prefix.Length).IndexOfAny(Path.DirectorySeparatorChar, Path.AltDirectorySeparatorChar);
+				if (nextDirLength == -1)
+				{
+					string fileName = fileUpdate.Path.Substring(prefix.Length);
+					Files[fileName] = fileUpdate;
+
+					if (!files.MoveNext())
+					{
+						return false;
+					}
+				}
+				else
+				{
+					string dirName = fileUpdate.Path.Substring(prefix.Length, nextDirLength);
+
+					DirectoryUpdate? nextTree;
+					if (!Directories.TryGetValue(dirName, out nextTree) || nextTree == null)
+					{
+						nextTree = new DirectoryUpdate();
+						Directories[dirName] = nextTree;
+					}
+
+					ReadOnlySpan<char> nextPrefix = fileUpdate.Path.AsSpan(0, prefix.Length + nextDirLength + 1);
+					if (!nextTree.AddFiles(nextPrefix, files))
+					{
+						return false;
+					}
+				}
+			}
+		}
+
+		/// <summary>
+		/// Writes interior node data to the given writer
+		/// </summary>
+		public async ValueTask WriteInteriorNodesAsync(IBlobWriter writer, InteriorChunkedDataNodeOptions interiorNodeOptions, CancellationToken cancellationToken)
+		{
+			foreach (DirectoryUpdate? directoryUpdate in Directories.Values)
+			{
+				if (directoryUpdate != null)
+				{
+					await directoryUpdate.WriteInteriorNodesAsync(writer, interiorNodeOptions, cancellationToken);
+				}
+			}
+			foreach (FileUpdate? fileUpdate in Files.Values)
+			{
+				if (fileUpdate != null)
+				{
+					await fileUpdate.WriteInteriorNodesAsync(writer, interiorNodeOptions, cancellationToken);
+				}
+			}
+		}
+	}
+
+	/// <summary>
+	/// Extension methods for directory node updates
+	/// </summary>
+	public static class DirectoryNodeUpdateExtensions
+	{
+		/// <summary>
+		/// Writes a tree of files to a storage writer
+		/// </summary>
+		public static async Task<IBlobRef<DirectoryNode>> WriteFilesAsync(this IBlobWriter writer, DirectoryReference baseDir, ChunkingOptions? options = null, IProgress<ICopyStats>? progress = null, CancellationToken cancellationToken = default)
+		{
+			DirectoryNode outputNode = new DirectoryNode();
+			await outputNode.AddFilesAsync(baseDir, DirectoryReference.EnumerateFiles(baseDir, "*", SearchOption.AllDirectories), writer, options, progress, cancellationToken);
+			return await writer.WriteBlobAsync(outputNode, cancellationToken: cancellationToken);
+		}
+
+		/// <summary>
+		/// Writes a tree of files to a storage writer
+		/// </summary>
+		public static async Task<IBlobRef<DirectoryNode>> WriteFilesAsync(this IBlobWriter writer, DirectoryInfo baseDir, IReadOnlyList<FileInfo> files, ChunkingOptions? options = null, IProgress<ICopyStats>? progress = null, CancellationToken cancellationToken = default)
+		{
+			DirectoryNode outputNode = new DirectoryNode();
+			await outputNode.AddFilesAsync(baseDir, files, writer, options, progress, cancellationToken);
+			return await writer.WriteBlobAsync(outputNode, cancellationToken: cancellationToken);
+		}
+
+		/// <summary>
+		/// Writes a tree of files to a storage writer
+		/// </summary>
+		public static async Task<IBlobRef<DirectoryNode>> WriteFilesAsync(this IBlobWriter writer, DirectoryReference baseDir, IReadOnlyList<FileReference> files, ChunkingOptions? options = null, IProgress<ICopyStats>? progress = null, CancellationToken cancellationToken = default)
+		{
+			DirectoryNode outputNode = new DirectoryNode();
+			await outputNode.AddFilesAsync(baseDir, files, writer, options, progress, cancellationToken);
+			return await writer.WriteBlobAsync(outputNode, cancellationToken: cancellationToken);
+		}
+
+		/// <inheritdoc cref="AddFilesAsync(DirectoryNode, DirectoryReference, IEnumerable{FileInfo}, IBlobWriter, ChunkingOptions?, IProgress{ICopyStats}?, CancellationToken)"/>
+		public static async Task AddFilesAsync(this DirectoryNode directoryNode, DirectoryInfo directoryInfo, IBlobWriter writer, ChunkingOptions? options = null, IProgress<ICopyStats>? progress = null, CancellationToken cancellationToken = default)
+		{
+			await directoryNode.AddFilesAsync(new DirectoryReference(directoryInfo), directoryInfo.EnumerateFiles("*", SearchOption.AllDirectories).ToList(), writer, options, progress, cancellationToken);
+		}
+
+		/// <inheritdoc cref="AddFilesAsync(DirectoryNode, DirectoryReference, IEnumerable{FileInfo}, IBlobWriter, ChunkingOptions?, IProgress{ICopyStats}?, CancellationToken)"/>
+		public static Task AddFilesAsync(this DirectoryNode directoryNode, DirectoryReference baseDir, IEnumerable<FileReference> files, IBlobWriter writer, ChunkingOptions? options = null, IProgress<ICopyStats>? progress = null, CancellationToken cancellationToken = default)
+		{
+			return directoryNode.AddFilesAsync(baseDir, files.Select(x => x.ToFileInfo()).ToList(), writer, options, progress, cancellationToken);
+		}
+
+		/// <inheritdoc cref="AddFilesAsync(DirectoryNode, DirectoryReference, IEnumerable{FileInfo}, IBlobWriter, ChunkingOptions?, IProgress{ICopyStats}?, CancellationToken)"/>
+		public static Task AddFilesAsync(this DirectoryNode directoryNode, DirectoryInfo baseDir, IEnumerable<FileInfo> files, IBlobWriter writer, ChunkingOptions? options = null, IProgress<ICopyStats>? progress = null, CancellationToken cancellationToken = default)
+		{
+			return directoryNode.AddFilesAsync(new DirectoryReference(baseDir), files.ToList(), writer, options, progress, cancellationToken);
+		}
+
+		/// <summary>
+		/// Updates this tree of directory objects
+		/// </summary>
+		/// <param name="directoryNode">Directory to update</param>
+		/// <param name="updates">Files to add</param>
+		/// <param name="writer">Writer for new node data</param>
+		/// <param name="cancellationToken">Cancellation token for the operation</param>
+		public static Task UpdateAsync(this DirectoryNode directoryNode, IEnumerable<FileUpdate> updates, IBlobWriter writer, CancellationToken cancellationToken = default)
+		{
+			DirectoryUpdate update = new DirectoryUpdate();
+			update.AddFiles(updates);
+			return directoryNode.UpdateAsync(update, writer, cancellationToken);
+		}
+
+		/// <summary>
+		/// Updates this tree of directory objects
+		/// </summary>
+		/// <param name="directoryNode">Directory to update</param>
+		/// <param name="update">Files to add</param>
+		/// <param name="writer">Writer for new node data</param>
+		/// <param name="cancellationToken">Cancellation token for the operation</param>
+		public static async Task UpdateAsync(this DirectoryNode directoryNode, DirectoryUpdate update, IBlobWriter writer, CancellationToken cancellationToken = default)
+		{
+			foreach ((string name, DirectoryUpdate? directory) in update.Directories)
+			{
+				if (directory == null)
+				{
+					directoryNode.DeleteDirectory(name);
+				}
+				else
+				{
+					DirectoryNode? childNode = await directoryNode.TryOpenDirectoryAsync(name, cancellationToken);
+					childNode ??= new DirectoryNode();
+					await childNode.UpdateAsync(directory, writer, cancellationToken);
+					IBlobRef<DirectoryNode> handle = await writer.WriteBlobAsync<DirectoryNode>(childNode, cancellationToken);
+					directoryNode.DeleteDirectory(name);
+					directoryNode.AddDirectory(new DirectoryEntry(name, childNode.Length, handle));
+				}
+			}
+			foreach ((string name, FileUpdate? file) in update.Files)
+			{
+				directoryNode.DeleteFile(name);
+				if (file != null)
+				{
+					await file.WriteInteriorNodesAsync(writer, new ChunkingOptions().InteriorOptions, cancellationToken);
+					directoryNode.AddFile(new FileEntry(name, file.Flags, file.Length, file.StreamHash, file.Nodes[0], file.CustomData));
+				}
+			}
+		}
+
+		/// <summary>
+		/// Adds files from a directory to the storage
+		/// </summary>
+		/// <param name="directoryNode">Directory to add to</param>
+		/// <param name="baseDir">Base directory to base paths relative to</param>
+		/// <param name="files">Files to add</param>
+		/// <param name="options">Options for chunking file content</param>
+		/// <param name="writer">Writer for new node data</param>
+		/// <param name="progress">Feedback interface for progress updates</param>
+		/// <param name="cancellationToken">Cancellation token for the operation</param>
+		public static async Task AddFilesAsync(this DirectoryNode directoryNode, DirectoryReference baseDir, IEnumerable<FileInfo> files, IBlobWriter writer, ChunkingOptions? options = null, IProgress<ICopyStats>? progress = null, CancellationToken cancellationToken = default)
+		{
+			options ??= new ChunkingOptions();
+
+			CopyStats? copyStats = null;
+			if (progress != null)
+			{
+				copyStats = new CopyStats(progress);
+			}
+
+			DirectoryUpdate update = new DirectoryUpdate();
+
+			// Process the input sequence in 2gb batches
+			const long MaxBatchSize = 2 * 1024 * 1024 * 1024L;
+			using (IEnumerator<FileInfo> fileEnumerator = files.GetEnumerator())
+			{
+				List<FileInfo> batch = new List<FileInfo>();
+				for (bool moreData = fileEnumerator.MoveNext(); moreData;)
+				{
+					batch.Clear();
+
+					// Take the next batch of files
+					long batchSize = 0;
+					while (batchSize < MaxBatchSize && moreData)
+					{
+						FileInfo file = fileEnumerator.Current;
+						batch.Add(file);
+						batchSize += file.Length;
+						moreData = fileEnumerator.MoveNext();
+					}
+
+					// Partition them up into parallel writers
+					List<(int Start, int Count)> partitions = ComputePartitions(batch, batchSize);
+					LeafChunkedData[] leafChunkedFiles = new LeafChunkedData[batch.Count];
+					await Parallel.ForEachAsync(partitions, cancellationToken, (filePartition, ctx) => CreateLeafChunkNodesAsync(writer, batch, leafChunkedFiles, filePartition.Start, filePartition.Count, copyStats, options, cancellationToken));
+
+					// Write all the interior nodes and generate the directory update
+					for (int idx = 0; idx < batch.Count; idx++)
+					{
+						FileInfo file = batch[idx];
+
+						FileEntryFlags flags = FileEntry.GetPermissions(file);
+						FileUpdate entry = new FileUpdate(file.Name, flags, file.Length, leafChunkedFiles[idx]);
+						update.AddFile(new FileReference(file).MakeRelativeTo(baseDir), entry);
+					}
+				}
+			}
+
+			// Add all the new entries to the tree
+			await directoryNode.UpdateAsync(update, writer, cancellationToken);
+		}
+
+		static List<(int Index, int Count)> ComputePartitions(IReadOnlyList<FileInfo> files, long totalSize)
+		{
+			// Maximum number of streams to write in parallel
+			const int MaxPartitions = 16;
+
+			// Minimum size of output payload for each writer
+			const long MinSizePerPartition = 1024 * 1024;
+
+			// Calculate the number of writers
+			int numPartitions = 1 + (int)Math.Min(totalSize / MinSizePerPartition, MaxPartitions);
+
+			// Create the partitions
+			List<(int Index, int Count)> partitions = new List<(int, int)>();
+
+			long writtenSize = 0;
+			long remainingSizeForPartition = 0;
+
+			int startIndex = 0;
+			for (int index = 0; index < files.Count; index++)
+			{
+				FileInfo file = files[index];
+				if (file.Length >= remainingSizeForPartition && partitions.Count < numPartitions)
+				{
+					remainingSizeForPartition = (totalSize - writtenSize) / (numPartitions - partitions.Count);
+					partitions.Add((startIndex, index - startIndex));
+					startIndex = index;
+				}
+
+				writtenSize += file.Length;
+				remainingSizeForPartition -= file.Length;
+			}
+
+			partitions.Add((startIndex, files.Count - startIndex));
+			return partitions;
+		}
+
+		static async ValueTask CreateLeafChunkNodesAsync(IBlobWriter writer, IReadOnlyList<FileInfo> files, LeafChunkedData[] leafChunks, int start, int count, CopyStats? copyStats, ChunkingOptions options, CancellationToken cancellationToken)
+		{
+			await using IBlobWriter writerFork = writer.Fork();
+			for (int idx = start; idx < start + count; idx++)
+			{
+				FileInfo file = files[idx];
+				using (Stream stream = file.OpenRead())
+				{
+					leafChunks[idx] = await LeafChunkedDataNode.CreateFromStreamAsync(writerFork, stream, options.LeafOptions, copyStats, cancellationToken);
+				}
+			}
+			await writerFork.FlushAsync(cancellationToken);
+		}
+
+		/// <summary>
+		/// Copies entries from a zip file
+		/// </summary>
+		/// <param name="directoryNode">Directory to update</param>
+		/// <param name="stream">Input stream</param>
+		/// <param name="writer">Writer for new nodes</param>
+		/// <param name="options"></param>
+		/// <param name="cancellationToken">Cancellation token for the operation</param>
+		public static async Task CopyFromZipStreamAsync(this DirectoryNode directoryNode, Stream stream, IBlobWriter writer, ChunkingOptions options, CancellationToken cancellationToken = default)
+		{
+			// Create all the leaf nodes
+			List<(ZipArchiveEntry, LeafChunkedData)> entries = new List<(ZipArchiveEntry, LeafChunkedData)>();
+			using (ZipArchive archive = new ZipArchive(stream, ZipArchiveMode.Read, true))
+			{
+				foreach (ZipArchiveEntry entry in archive.Entries)
+				{
+					if (entry.Name.Length > 0)
+					{
+						using (Stream entryStream = entry.Open())
+						{
+							LeafChunkedData leafChunkedData = await LeafChunkedDataNode.CreateFromStreamAsync(writer, entryStream, options.LeafOptions, cancellationToken);
+							entries.Add((entry, leafChunkedData));
+						}
+					}
+				}
+			}
+
+			// Create all the interior nodes
+			List<FileUpdate> updates = new List<FileUpdate>();
+			foreach ((ZipArchiveEntry entry, LeafChunkedData leafChunkedFile) in entries)
+			{
+				FileEntryFlags flags = FileEntryFlags.None;
+				if ((entry.ExternalAttributes & (0b_001_001_001 << 16)) != 0)
+				{
+					flags |= FileEntryFlags.Executable;
+				}
+
+				ChunkedData chunkedFile = await InteriorChunkedDataNode.CreateTreeAsync(leafChunkedFile, options.InteriorOptions, writer, cancellationToken);
+				updates.Add(new FileUpdate(entry.FullName, flags, entry.Length, chunkedFile));
+			}
+
+			// Update the tree
+			await directoryNode.UpdateAsync(updates, writer, cancellationToken);
+		}
+	}
+}
diff --git a/Engine/Source/Programs/Shared/EpicGames.Horde/Storage/Nodes/DirectoryNode.ZipStream.cs b/Engine/Source/Programs/Shared/EpicGames.Horde/Storage/Nodes/DirectoryNode.ZipStream.cs
new file mode 100644
--- /dev/null
+++ b/Engine/Source/Programs/Shared/EpicGames.Horde/Storage/Nodes/DirectoryNode.ZipStream.cs
@@ -0,0 +1,194 @@
+// Copyright Epic Games, Inc. All Rights Reserved.
+
+using System;
+using System.Buffers;
+using System.IO;
+using System.IO.Compression;
+using System.IO.Pipelines;
+using System.Threading;
+using System.Threading.Tasks;
+using EpicGames.Core;
+using Microsoft.Extensions.Logging;
+
+namespace EpicGames.Horde.Storage.Nodes
+{
+	/// <summary>
+	/// Stream which zips a directory node tree dynamically
+	/// </summary>
+	class DirectoryNodeZipStream : Stream
+	{
+		readonly Pipe _pipe;
+		readonly ILogger? _logger;
+		readonly BackgroundTask _backgroundTask;
+
+		long _position;
+		ReadOnlySequence<byte> _current = ReadOnlySequence<byte>.Empty;
+
+		/// <inheritdoc/>
+		public override bool CanRead => true;
+
+		/// <inheritdoc/>
+		public override bool CanSeek => false;
+
+		/// <inheritdoc/>
+		public override bool CanWrite => false;
+
+		/// <inheritdoc/>
+		public override long Length => throw new NotImplementedException();
+
+		/// <inheritdoc/>
+		public override long Position { get => _position; set => throw new NotImplementedException(); }
+
+		/// <summary>
+		/// Constructor
+		/// </summary>
+		/// <param name="node">Root node to copy from</param>
+		/// <param name="filter">Filter for files to include in the zip</param>
+		/// <param name="logger">Optional logger for debug tracing</param>
+		public DirectoryNodeZipStream(IBlobRef<DirectoryNode> node, FileFilter? filter, ILogger? logger)
+		{
+			_pipe = new Pipe();
+			_backgroundTask = BackgroundTask.StartNew(ctx => CopyToPipeAsync(node, filter, _pipe.Writer, logger, ctx));
+			_logger = logger;
+		}
+
+		/// <inheritdoc/>
+		protected override void Dispose(bool disposing)
+		{
+			if (disposing)
+			{
+				_backgroundTask.DisposeAsync().AsTask().Wait();
+			}
+
+			base.Dispose(disposing);
+		}
+
+		/// <inheritdoc/>
+		public override async ValueTask DisposeAsync()
+		{
+			await _backgroundTask.DisposeAsync();
+
+			await base.DisposeAsync();
+		}
+
+		/// <inheritdoc/>
+		public override async ValueTask<int> ReadAsync(Memory<byte> buffer, CancellationToken cancellationToken = default)
+		{
+			while (_current.Length == 0)
+			{
+				ReadResult result = await _pipe.Reader.ReadAsync(cancellationToken);
+				_current = result.Buffer;
+
+				if (result.IsCompleted && _current.Length == 0)
+				{
+					// Wait for the background thread to finish; it may error/have errored, and we want to re-throw on this thread before completing the read.
+					await _backgroundTask.StopAsync(cancellationToken);
+					_logger?.LogInformation("Zip file pipe was read to end");
+					return 0;
+				}
+			}
+
+			int initialSize = buffer.Length;
+			while (buffer.Length > 0 && _current.Length > 0)
+			{
+				int copy = Math.Min(buffer.Length, _current.First.Length);
+				_current.First.Slice(0, copy).CopyTo(buffer);
+				_current = _current.Slice(copy);
+				buffer = buffer.Slice(copy);
+			}
+
+			if (_current.Length == 0)
+			{
+				_pipe.Reader.AdvanceTo(_current.End);
+			}
+
+			int length = initialSize - buffer.Length;
+			_position += length;
+			return length;
+		}
+
+		static async Task CopyToPipeAsync(IBlobRef<DirectoryNode> node, FileFilter? filter, PipeWriter writer, ILogger? logger, CancellationToken cancellationToken)
+		{
+			using Stream outputStream = writer.AsStream();
+			using ZipArchive archive = new ZipArchive(outputStream, ZipArchiveMode.Create);
+			await CopyFilesAsync(node, "", filter, archive, logger, cancellationToken);
+		}
+
+		static async Task CopyFilesAsync(IBlobRef<DirectoryNode> directoryRef, string prefix, FileFilter? filter, ZipArchive archive, ILogger? logger, CancellationToken cancellationToken)
+		{
+			DirectoryNode directory = await directoryRef.ReadBlobAsync(cancellationToken);
+
+			int numDirs = directory.Directories.Count;
+			int numFiles = directory.Files.Count;
+
+			int numCopiedDirs = 0;
+			foreach (DirectoryEntry directoryEntry in directory.Directories)
+			{
+				string directoryPath = $"{prefix}{directoryEntry.Name}/";
+				if (filter == null || filter.PossiblyMatches(directoryPath))
+				{
+					await CopyFilesAsync(directoryEntry.Handle, directoryPath, filter, archive, logger, cancellationToken);
+				}
+				numCopiedDirs++;
+			}
+
+			int numCopiedFiles = 0;
+			foreach (FileEntry fileEntry in directory.Files)
+			{
+				string filePath = $"{prefix}{fileEntry}";
+				if (filter == null || filter.Matches(filePath))
+				{
+					ZipArchiveEntry entry = archive.CreateEntry(filePath);
+
+					if ((fileEntry.Flags & FileEntryFlags.Executable) != 0)
+					{
+						entry.ExternalAttributes |= 0b_111_111_101 << 16; // rwx rwx r-x
+					}
+					else
+					{
+						entry.ExternalAttributes |= 0b_110_110_100 << 16; // rw- rw- r--
+					}
+
+					using Stream entryStream = entry.Open();
+					await fileEntry.CopyToStreamAsync(entryStream, cancellationToken);
+				}
+				numCopiedFiles++;
+			}
+
+			logger?.LogInformation("Zip file path {Prefix} has {NumCopiedDirectories}/{NumDirectories} directories, {NumCopiedFiles}/{NumFiles} files", prefix, numCopiedDirs, numDirs, numCopiedFiles, numFiles);
+		}
+
+		/// <inheritdoc/>
+		public override void Flush()
+		{
+		}
+
+		/// <inheritdoc/>
+		public override int Read(byte[] buffer, int offset, int count) => ReadAsync(buffer.AsMemory(offset, count)).AsTask().Result;
+
+		/// <inheritdoc/>
+		public override long Seek(long offset, SeekOrigin origin) => throw new NotSupportedException();
+
+		/// <inheritdoc/>
+		public override void SetLength(long value) => throw new NotSupportedException();
+
+		/// <inheritdoc/>
+		public override void Write(byte[] buffer, int offset, int count) => throw new NotSupportedException();
+	}
+
+	/// <summary>
+	/// Extension methods for <see cref="DirectoryNodeZipStream"/>
+	/// </summary>
+	public static class DirectoryNodeZipStreamExtensions
+	{
+		/// <summary>
+		/// Returns a stream containing the zipped contents of this directory
+		/// </summary>
+		/// <param name="directoryRef">The directory to zip</param>
+		/// <param name="filter">Filter for files to include in the zip</param>
+		/// <param name="logger">Logger for diagnostic output</param>
+		/// <returns>Stream containing zipped archive data</returns>
+		public static Stream AsZipStream(this IBlobRef<DirectoryNode> directoryRef, FileFilter? filter = null, ILogger? logger = null)
+			=> new DirectoryNodeZipStream(directoryRef, filter, logger);
+	}
+}
diff --git a/Engine/Source/Programs/Shared/EpicGames.Horde/Storage/Nodes/DirectoryNode.cs b/Engine/Source/Programs/Shared/EpicGames.Horde/Storage/Nodes/DirectoryNode.cs
--- a/Engine/Source/Programs/Shared/EpicGames.Horde/Storage/Nodes/DirectoryNode.cs
+++ b/Engine/Source/Programs/Shared/EpicGames.Horde/Storage/Nodes/DirectoryNode.cs
@@ -5,133 +5,14 @@ using System.Collections.Generic;
 using System.Diagnostics;
 using System.Diagnostics.CodeAnalysis;
 using System.IO;
-using System.IO.Compression;
-using System.IO.MemoryMappedFiles;
 using System.Linq;
-using System.Runtime.InteropServices;
 using System.Threading;
-using System.Threading.Channels;
 using System.Threading.Tasks;
 using EpicGames.Core;
 using EpicGames.Serialization;
-using Microsoft.Extensions.Logging;
 
 namespace EpicGames.Horde.Storage.Nodes
 {
-	/// <summary>
-	/// Stats reported for copy operations
-	/// </summary>
-	public interface ICopyStats
-	{
-		/// <summary>
-		/// Number of files that have been copied
-		/// </summary>
-		int Count { get; }
-
-		/// <summary>
-		/// Total size of data to be copied
-		/// </summary>
-		long Size { get; }
-
-		/// <summary>
-		/// Processing speed, in bytes per second
-		/// </summary>
-		double Rate { get; }
-	}
-
-	/// <summary>
-	/// Reports progress info back to callers
-	/// </summary>
-	class CopyStats : ICopyStats
-	{
-		readonly object _lockObject = new object();
-		readonly Stopwatch _timer = Stopwatch.StartNew();
-		readonly IProgress<ICopyStats> _progress;
-		long _lastTotalSize;
-
-		public int Count { get; set; }
-		public long Size { get; set; }
-		public double Rate { get; set; }
-
-		public CopyStats(IProgress<ICopyStats> progress)
-		{
-			_progress = progress;
-		}
-
-		public void Update(int count, long size)
-		{
-			lock (_lockObject)
-			{
-				Count += count;
-				Size += size;
-				if (_timer.Elapsed > TimeSpan.FromSeconds(10.0))
-				{
-					Rate = (Size - _lastTotalSize) / _timer.Elapsed.TotalSeconds;
-					_lastTotalSize = Size;
-
-					_progress.Report(this);
-					_timer.Restart();
-				}
-			}
-		}
-
-		public void Flush()
-		{
-			lock (_lockObject)
-			{
-				_progress.Report(this);
-				_timer.Restart();
-			}
-		}
-	}
-
-	/// <summary>
-	/// Progress logger for writing copy stats
-	/// </summary>
-	public class CopyStatsLogger : IProgress<ICopyStats>
-	{
-		readonly int _totalCount;
-		readonly long _totalSize;
-		readonly ILogger _logger;
-
-		/// <summary>
-		/// Constructor
-		/// </summary>
-		public CopyStatsLogger(ILogger logger)
-			=> _logger = logger;
-
-		/// <summary>
-		/// Constructor
-		/// </summary>
-		public CopyStatsLogger(int totalCount, long totalSize, ILogger logger)
-		{
-			_totalCount = totalCount;
-			_totalSize = totalSize;
-			_logger = logger;
-		}
-
-		/// <inheritdoc/>
-		public void Report(ICopyStats stats)
-		{
-			if (_totalCount > 0 && _totalSize > 0)
-			{
-				_logger.LogInformation("Copied {NumFiles:n0}/{TotalFiles:n0} files ({Size:n1}/{TotalSize:n1}mb, {Rate:n1}mb/s, {Pct}%)", stats.Count, _totalCount, stats.Size / (1024.0 * 1024.0), _totalSize / (1024.0 * 1024.0), stats.Rate / (1024.0 * 1024.0), (int)((Math.Max(stats.Size, 1) * 100) / Math.Max(_totalSize, 1)));
-			}
-			else if (_totalCount > 0)
-			{
-				_logger.LogInformation("Copied {NumFiles:n0}/{TotalFiles:n0} files ({Size:n1}mb, {Rate:n1}mb/s)", stats.Count, _totalCount, stats.Size / (1024.0 * 1024.0), stats.Rate / (1024.0 * 1024.0));
-			}
-			else if (_totalSize > 0)
-			{
-				_logger.LogInformation("Copied {NumFiles:n0} files ({Size:n1}/{TotalSize:n1}mb, {Rate:n1}mb/s, {Pct}%)", stats.Count, stats.Size / (1024.0 * 1024.0), _totalSize / (1024.0 * 1024.0), stats.Rate / (1024.0 * 1024.0), (int)((Math.Max(stats.Size, 1) * 100) / Math.Max(_totalSize, 1)));
-			}
-			else
-			{
-				_logger.LogInformation("Copied {NumFiles:n0} files ({Size:n1}mb, {Rate:n1}mb/s)", stats.Count, stats.Size / (1024.0 * 1024.0), stats.Rate / (1024.0 * 1024.0));
-			}
-		}
-	}
-
 	/// <summary>
 	/// Flags for a directory node
 	/// </summary>
@@ -456,467 +337,6 @@ namespace EpicGames.Horde.Storage.Nodes
 		public bool DeleteDirectory(string name) => _nameToDirectoryEntry.Remove(name);
 
 		#endregion
-
-		/// <inheritdoc cref="AddFilesAsync(DirectoryReference, IEnumerable{FileInfo}, IBlobWriter, ChunkingOptions?, IProgress{ICopyStats}?, CancellationToken)"/>
-		public async Task AddFilesAsync(DirectoryInfo directoryInfo, IBlobWriter writer, ChunkingOptions? options = null, IProgress<ICopyStats>? progress = null, CancellationToken cancellationToken = default)
-		{
-			await AddFilesAsync(new DirectoryReference(directoryInfo), directoryInfo.EnumerateFiles("*", SearchOption.AllDirectories).ToList(), writer, options, progress, cancellationToken);
-		}
-
-		/// <inheritdoc cref="AddFilesAsync(DirectoryReference, IEnumerable{FileInfo}, IBlobWriter, ChunkingOptions?, IProgress{ICopyStats}?, CancellationToken)"/>
-		public Task AddFilesAsync(DirectoryReference baseDir, IEnumerable<FileReference> files, IBlobWriter writer, ChunkingOptions? options = null, IProgress<ICopyStats>? progress = null, CancellationToken cancellationToken = default)
-		{
-			return AddFilesAsync(baseDir, files.Select(x => x.ToFileInfo()).ToList(), writer, options, progress, cancellationToken);
-		}
-
-		/// <inheritdoc cref="AddFilesAsync(DirectoryReference, IEnumerable{FileInfo}, IBlobWriter, ChunkingOptions?, IProgress{ICopyStats}?, CancellationToken)"/>
-		public Task AddFilesAsync(DirectoryInfo baseDir, IEnumerable<FileInfo> files, IBlobWriter writer, ChunkingOptions? options = null, IProgress<ICopyStats>? progress = null, CancellationToken cancellationToken = default)
-		{
-			return AddFilesAsync(new DirectoryReference(baseDir), files.ToList(), writer, options, progress, cancellationToken);
-		}
-
-		/// <summary>
-		/// Adds files from a directory to the storage
-		/// </summary>
-		/// <param name="baseDir">Base directory to base paths relative to</param>
-		/// <param name="files">Files to add</param>
-		/// <param name="options">Options for chunking file content</param>
-		/// <param name="writer">Writer for new node data</param>
-		/// <param name="progress">Feedback interface for progress updates</param>
-		/// <param name="cancellationToken">Cancellation token for the operation</param>
-		public async Task AddFilesAsync(DirectoryReference baseDir, IEnumerable<FileInfo> files, IBlobWriter writer, ChunkingOptions? options = null, IProgress<ICopyStats>? progress = null, CancellationToken cancellationToken = default)
-		{
-			options ??= new ChunkingOptions();
-
-			CopyStats? copyStats = null;
-			if (progress != null)
-			{
-				copyStats = new CopyStats(progress);
-			}
-
-			DirectoryUpdate update = new DirectoryUpdate();
-
-			// Process the input sequence in 2gb batches
-			const long MaxBatchSize = 2 * 1024 * 1024 * 1024L;
-			using (IEnumerator<FileInfo> fileEnumerator = files.GetEnumerator())
-			{
-				List<FileInfo> batch = new List<FileInfo>();
-				for (bool moreData = fileEnumerator.MoveNext(); moreData; )
-				{
-					batch.Clear();
-
-					// Take the next batch of files
-					long batchSize = 0;
-					while (batchSize < MaxBatchSize && moreData)
-					{
-						FileInfo file = fileEnumerator.Current;
-						batch.Add(file);
-						batchSize += file.Length;
-						moreData = fileEnumerator.MoveNext();
-					}
-
-					// Partition them up into parallel writers
-					List<(int Start, int Count)> partitions = ComputePartitions(batch, batchSize);
-					LeafChunkedData[] leafChunkedFiles = new LeafChunkedData[batch.Count];
-					await Parallel.ForEachAsync(partitions, cancellationToken, (filePartition, ctx) => CreateLeafChunkNodesAsync(writer, batch, leafChunkedFiles, filePartition.Start, filePartition.Count, copyStats, options, cancellationToken));
-
-					// Write all the interior nodes and generate the directory update
-					for (int idx = 0; idx < batch.Count; idx++)
-					{
-						FileInfo file = batch[idx];
-
-						FileEntryFlags flags = FileEntryFlags.None;
-						if (RuntimeInformation.IsOSPlatform(OSPlatform.Linux))
-						{
-							int mode = FileUtils.GetFileMode_Linux(file.FullName);
-							if ((mode & ((1 << 0) | (1 << 3) | (1 << 6))) != 0)
-							{
-								flags |= FileEntryFlags.Executable;
-							}
-						}
-						else if (RuntimeInformation.IsOSPlatform(OSPlatform.OSX))
-						{
-							int mode = FileUtils.GetFileMode_Mac(file.FullName);
-							if ((mode & ((1 << 0) | (1 << 3) | (1 << 6))) != 0)
-							{
-								flags |= FileEntryFlags.Executable;
-							}
-						}
-
-						FileUpdate entry = new FileUpdate(file.Name, flags, file.Length, leafChunkedFiles[idx]);
-						update.AddFile(new FileReference(file).MakeRelativeTo(baseDir), entry);
-					}
-				}
-			}
-
-			// Add all the new entries to the tree
-			await UpdateAsync(update, writer, cancellationToken);
-		}
-
-		static List<(int Index, int Count)> ComputePartitions(IReadOnlyList<FileInfo> files, long totalSize)
-		{
-			// Maximum number of streams to write in parallel
-			const int MaxPartitions = 16;
-
-			// Minimum size of output payload for each writer
-			const long MinSizePerPartition = 1024 * 1024;
-
-			// Calculate the number of writers
-			int numPartitions = 1 + (int)Math.Min(totalSize / MinSizePerPartition, MaxPartitions);
-
-			// Create the partitions
-			List<(int Index, int Count)> partitions = new List<(int, int)>();
-
-			long writtenSize = 0;
-			long remainingSizeForPartition = 0;
-
-			int startIndex = 0;
-			for (int index = 0; index < files.Count; index++)
-			{
-				FileInfo file = files[index];
-				if (file.Length >= remainingSizeForPartition && partitions.Count < numPartitions)
-				{
-					remainingSizeForPartition = (totalSize - writtenSize) / (numPartitions - partitions.Count);
-					partitions.Add((startIndex, index - startIndex));
-					startIndex = index;
-				}
-
-				writtenSize += file.Length;
-				remainingSizeForPartition -= file.Length;
-			}
-
-			partitions.Add((startIndex, files.Count - startIndex));
-			return partitions;
-		}
-
-		static async ValueTask CreateLeafChunkNodesAsync(IBlobWriter writer, IReadOnlyList<FileInfo> files, LeafChunkedData[] leafChunks, int start, int count, CopyStats? copyStats, ChunkingOptions options, CancellationToken cancellationToken)
-		{
-			await using IBlobWriter writerFork = writer.Fork();
-			for (int idx = start; idx < start + count; idx++)
-			{
-				FileInfo file = files[idx];
-				using (Stream stream = file.OpenRead())
-				{
-					leafChunks[idx] = await LeafChunkedDataNode.CreateFromStreamAsync(writerFork, stream, options.LeafOptions, copyStats, cancellationToken);
-				}
-			}
-			await writerFork.FlushAsync(cancellationToken);
-		}
-
-		/// <summary>
-		/// Updates this tree of directory objects
-		/// </summary>
-		/// <param name="updates">Files to add</param>
-		/// <param name="writer">Writer for new node data</param>
-		/// <param name="cancellationToken">Cancellation token for the operation</param>
-		public Task UpdateAsync(IEnumerable<FileUpdate> updates, IBlobWriter writer, CancellationToken cancellationToken = default)
-		{
-			DirectoryUpdate update = new DirectoryUpdate();
-			update.AddFiles(updates);
-			return UpdateAsync(update, writer, cancellationToken);
-		}
-
-		/// <summary>
-		/// Updates this tree of directory objects
-		/// </summary>
-		/// <param name="update">Files to add</param>
-		/// <param name="writer">Writer for new node data</param>
-		/// <param name="cancellationToken">Cancellation token for the operation</param>
-		public async Task UpdateAsync(DirectoryUpdate update, IBlobWriter writer, CancellationToken cancellationToken = default)
-		{
-			foreach ((string name, DirectoryUpdate? directory) in update.Directories)
-			{
-				if (directory == null)
-				{
-					_nameToDirectoryEntry.Remove(name);
-				}
-				else
-				{
-					DirectoryNode? childNode = await TryOpenDirectoryAsync(name, cancellationToken);
-					childNode ??= new DirectoryNode();
-					await childNode.UpdateAsync(directory, writer, cancellationToken);
-					IBlobRef<DirectoryNode> handle = await writer.WriteBlobAsync<DirectoryNode>(childNode, cancellationToken);
-					_nameToDirectoryEntry[name] = new DirectoryEntry(name, childNode.Length, handle);
-				}
-			}
-			foreach ((string name, FileUpdate? file) in update.Files)
-			{
-				if (file == null)
-				{
-					_nameToFileEntry.Remove(name);
-				}
-				else
-				{
-					await file.WriteInteriorNodesAsync(writer, new ChunkingOptions().InteriorOptions, cancellationToken);
-					_nameToFileEntry[name] = new FileEntry(name, file.Flags, file.Length, file.StreamHash, file.Nodes[0], file.CustomData);
-				}
-			}
-		}
-
-		/// <summary>
-		/// Copies entries from a zip file
-		/// </summary>
-		/// <param name="stream">Input stream</param>
-		/// <param name="writer">Writer for new nodes</param>
-		/// <param name="options"></param>
-		/// <param name="cancellationToken">Cancellation token for the operation</param>
-		public async Task CopyFromZipStreamAsync(Stream stream, IBlobWriter writer, ChunkingOptions options, CancellationToken cancellationToken = default)
-		{
-			// Create all the leaf nodes
-			List<(ZipArchiveEntry, LeafChunkedData)> entries = new List<(ZipArchiveEntry, LeafChunkedData)>();
-			using (ZipArchive archive = new ZipArchive(stream, ZipArchiveMode.Read, true))
-			{
-				foreach (ZipArchiveEntry entry in archive.Entries)
-				{
-					if (entry.Name.Length > 0)
-					{
-						using (Stream entryStream = entry.Open())
-						{
-							LeafChunkedData leafChunkedData = await LeafChunkedDataNode.CreateFromStreamAsync(writer, entryStream, options.LeafOptions, cancellationToken);
-							entries.Add((entry, leafChunkedData));
-						}
-					}
-				}
-			}
-
-			// Create all the interior nodes
-			List<FileUpdate> updates = new List<FileUpdate>();
-			foreach ((ZipArchiveEntry entry, LeafChunkedData leafChunkedFile) in entries)
-			{
-				FileEntryFlags flags = FileEntryFlags.None;
-				if ((entry.ExternalAttributes & (0b_001_001_001 << 16)) != 0)
-				{
-					flags |= FileEntryFlags.Executable;
-				}
-
-				ChunkedData chunkedFile = await InteriorChunkedDataNode.CreateTreeAsync(leafChunkedFile, options.InteriorOptions, writer, cancellationToken);
-				updates.Add(new FileUpdate(entry.FullName, flags, entry.Length, chunkedFile));
-			}
-
-			// Update the tree
-			await UpdateAsync(updates, writer, cancellationToken);
-		}
-
-		/// <summary>
-		/// Utility function to allow extracting a packed directory to disk
-		/// </summary>
-		/// <param name="directoryInfo"></param>
-		/// <param name="logger"></param>
-		/// <param name="cancellationToken"></param>
-		public Task CopyToDirectoryAsync(DirectoryInfo directoryInfo, ILogger logger, CancellationToken cancellationToken) => CopyToDirectoryAsync(directoryInfo, null, logger, cancellationToken);
-
-		record class OutputDir(string Path, DirectoryNode Node);
-		record class OutputFile(OutputDir Directory, FileEntry FileEntry);
-		record class OutputChunk(OutputFile File, long Offset, long Length, IBlobHandle Handle);
-
-#pragma warning disable IDE0060
-		static void TraceBlobRead(string type, string path, IBlobHandle handle, ILogger logger)
-		{
-//			logger.LogTrace(KnownLogEvents.Horde_BlobRead, "Blob [{Type,-20}] Path=\"{Path}\", Locator={Locator}", type, path, handle.GetLocator());
-		}
-#pragma warning restore IDE0060
-
-		/// <summary>
-		/// Utility function to allow extracting a packed directory to disk
-		/// </summary>
-		/// <param name="directoryInfo">Direcotry to write to</param>
-		/// <param name="progress">Sink for progress updates</param>
-		/// <param name="logger">Logger for output</param>
-		/// <param name="cancellationToken">Cancellation token for the operation</param>
-		public async Task CopyToDirectoryAsync(DirectoryInfo directoryInfo, IProgress<ICopyStats>? progress, ILogger logger, CancellationToken cancellationToken)
-		{
-			int numTasks = Math.Min(1 + (int)(Length / (16 * 1024 * 1024)), 16);
-			logger.LogInformation("Splitting read into {NumThreads} threads", numTasks);
-
-			CopyStats? copyStats = null;
-			if (progress != null)
-			{
-				copyStats = new CopyStats(progress);
-			}
-
-			Channel<OutputChunk> chunks = Channel.CreateBounded<OutputChunk>(new BoundedChannelOptions(128 * 1024) { FullMode = BoundedChannelFullMode.Wait });
-			using (CancellationTokenSource cancellationSource = CancellationTokenSource.CreateLinkedTokenSource(cancellationToken))
-			{
-				// Helper method to run a background task and set a cancellation source on error
-				async Task RunBackgroundTask(Func<CancellationToken, Task> taskFunc)
-				{
-					try
-					{
-						await taskFunc(cancellationSource.Token);
-					}
-					catch (OperationCanceledException)
-					{
-						// Ignore
-					}
-					catch (Exception ex)
-					{
-						logger.LogError(ex, "Error while extracting data: {Message}", ex.Message);
-						cancellationSource.Cancel();
-					}
-				}
-
-				List<Task> tasks = new List<Task>();
-				tasks.Add(RunBackgroundTask(ctx => FindOutputChunksRootAsync(chunks.Writer, logger, ctx)));
-				for (int idx = 0; idx < numTasks; idx++)
-				{
-					tasks.Add(RunBackgroundTask(ctx => ExtractAsync(chunks.Reader, new DirectoryReference(directoryInfo), copyStats, logger, ctx)));
-				}
-
-				await Task.WhenAll(tasks);
-			}
-		}
-
-		static async Task ExtractAsync(ChannelReader<OutputChunk> chunkReader, DirectoryReference baseDir, CopyStats? copyStats, ILogger logger, CancellationToken cancellationToken)
-		{
-			OutputChunk? chunk = await ReadNextChunkAsync(chunkReader, cancellationToken);
-			while (chunk != null)
-			{
-				// Open the file for the current chunk
-				OutputFile file = chunk.File;
-
-				FileReference locator = FileReference.Combine(baseDir, file.Directory.Path, file.FileEntry.Name);
-				DirectoryReference.CreateDirectory(locator.Directory);
-
-				FileInfo fileInfo = locator.ToFileInfo();
-				if (fileInfo.Exists && (fileInfo.Attributes & FileAttributes.ReadOnly) != 0)
-				{
-					fileInfo.Attributes &= ~FileAttributes.ReadOnly;
-				}
-
-				await using (FileStream stream = fileInfo.Open(FileMode.OpenOrCreate, FileAccess.ReadWrite, FileShare.ReadWrite))
-				{
-					stream.SetLength(file.FileEntry.Length);
-					if (file.FileEntry.Length == 0)
-					{
-						// If this file is empty, don't write anything and just move to the next chunk
-						chunk = await ReadNextChunkAsync(chunkReader, cancellationToken);
-					}
-					else
-					{
-						// Process as many chunks as we can for this file
-						using MemoryMappedFile memoryMappedFile = MemoryMappedFile.CreateFromFile(stream, null, file.FileEntry.Length, MemoryMappedFileAccess.ReadWrite, HandleInheritability.None, false);
-						using MemoryMappedView memoryMappedView = new MemoryMappedView(memoryMappedFile, 0, file.FileEntry.Length);
-
-						while (chunk != null && chunk.File == file)
-						{
-							// Write this chunk
-							using (BlobData data = await chunk.Handle.ReadBlobDataAsync(cancellationToken))
-							{
-								TraceBlobRead("Leaf", CombinePaths(chunk.File.Directory.Path, chunk.File.FileEntry.Name), chunk.Handle, logger);
-								data.Data.CopyTo(memoryMappedView!.GetMemory(chunk.Offset, data.Data.Length));
-							}
-
-							// Update the stats
-							int numCompleteFiles = 0;
-							if (chunk.Offset + chunk.Length == file.FileEntry.Length)
-							{
-								numCompleteFiles = 1;
-							}
-
-							copyStats?.Update(numCompleteFiles, chunk.Length);
-
-							// Read the next chunk
-							chunk = await ReadNextChunkAsync(chunkReader, cancellationToken);
-						}
-					}
-				}
-
-				// Set correct permissions on the output file
-				FileEntry.ApplyPermissions(fileInfo, file.FileEntry.Flags);
-			}
-		}
-
-		static async ValueTask<OutputChunk?> ReadNextChunkAsync(ChannelReader<OutputChunk> chunkReader, CancellationToken cancellationToken)
-		{
-			await chunkReader.WaitToReadAsync(cancellationToken);
-
-			OutputChunk? chunk;
-			if (chunkReader.TryRead(out chunk))
-			{
-				return chunk;
-			}
-			else
-			{
-				return null;
-			}
-		}
-
-		async Task FindOutputChunksRootAsync(ChannelWriter<OutputChunk> chunks, ILogger logger, CancellationToken cancellationToken)
-		{
-			await FindOutputChunksAsync("", this, chunks, logger, cancellationToken);
-			chunks.Complete();
-		}
-
-		static async Task FindOutputChunksAsync(string path, DirectoryNode node, ChannelWriter<OutputChunk> chunks, ILogger logger, CancellationToken cancellationToken)
-		{
-			OutputDir outputDir = new OutputDir(path, node);
-
-			foreach (FileEntry fileEntry in node._nameToFileEntry.Values)
-			{
-				OutputFile outputFile = new OutputFile(outputDir, fileEntry);
-				await FindOutputChunksAsync(outputFile, 0, fileEntry.Target, chunks, logger, cancellationToken);
-			}
-
-			foreach (DirectoryEntry directoryEntry in node._nameToDirectoryEntry.Values)
-			{
-				DirectoryNode subDirectoryNode = await directoryEntry.Handle.ReadBlobAsync(cancellationToken);
-
-				string subPath = CombinePaths(outputDir.Path, directoryEntry.Name);
-				TraceBlobRead("Directory", subPath, directoryEntry.Handle, logger);
-
-				await FindOutputChunksAsync(subPath, subDirectoryNode, chunks, logger, cancellationToken);
-			}
-		}
-
-		static async Task<long> FindOutputChunksAsync(OutputFile outputFile, long offset, ChunkedDataNodeRef dataRef, ChannelWriter<OutputChunk> chunks, ILogger logger, CancellationToken cancellationToken)
-		{
-			if (dataRef.Type == ChunkedDataNodeType.Leaf)
-			{
-				await chunks.WriteAsync(new OutputChunk(outputFile, offset, dataRef.Length, dataRef.Handle), cancellationToken);
-				if (dataRef.Length < 0)
-				{
-					// Backwards compatibility hack for v2 format
-					using BlobData data = await dataRef.Handle.ReadBlobDataAsync(cancellationToken);
-					return data.Data.Length;
-				}
-				return dataRef.Length;
-			}
-			else
-			{
-				using BlobData data = await dataRef.Handle.ReadBlobDataAsync(cancellationToken);
-				TraceBlobRead("Interior", CombinePaths(outputFile.Directory.Path, outputFile.FileEntry.Name), dataRef.Handle, logger);
-
-				if (data.Type.Guid == LeafChunkedDataNodeConverter.BlobType.Guid)
-				{
-					await chunks.WriteAsync(new OutputChunk(outputFile, offset, dataRef.Length, dataRef.Handle), cancellationToken);
-					return data.Data.Length;
-				}
-				else
-				{
-					long length = 0;
-
-					InteriorChunkedDataNode interiorNode = BlobSerializer.Deserialize<InteriorChunkedDataNode>(data);
-					foreach (ChunkedDataNodeRef childRef in interiorNode.Children)
-					{
-						length += await FindOutputChunksAsync(outputFile, offset + length, childRef, chunks, logger, cancellationToken);
-					}
-
-					return length;
-				}
-			}
-		}
-
-		static string CombinePaths(string basePath, string nextPath)
-		{
-			if (basePath.Length > 0)
-			{
-				return $"{basePath}/{nextPath}";
-			}
-			else
-			{
-				return nextPath;
-			}
-		}
 	}
 
 	class DirectoryNodeConverter : BlobConverter<DirectoryNode>
@@ -1008,233 +428,17 @@ namespace EpicGames.Horde.Storage.Nodes
 	}
 
 	/// <summary>
-	/// Describes an update to a file in a directory tree
+	/// Reference to a directory node, including the target hash and length
 	/// </summary>
-	/// <param name="Path">Path to the file</param>
-	/// <param name="Length">Length of the file data</param>
-	/// <param name="Flags">Flags for the new file entry</param>
-	/// <param name="StreamHash">Hash of the entire stream</param>
-	/// <param name="Nodes">Chunked data for the file</param>
-	/// <param name="CustomData"></param>
-	public record class FileUpdate(string Path, FileEntryFlags Flags, long Length, IoHash StreamHash, List<ChunkedDataNodeRef> Nodes, ReadOnlyMemory<byte> CustomData = default)
-	{
-		/// <summary>
-		/// Constructor
-		/// </summary>
-		public FileUpdate(string path, FileEntryFlags flags, long length, ChunkedData chunkedData, ReadOnlyMemory<byte> customData = default)
-			: this(path, flags, length, chunkedData.StreamHash, new List<ChunkedDataNodeRef> { chunkedData.Root }, customData)
-		{
-		}
-
-		/// <summary>
-		/// Constructor
-		/// </summary>
-		public FileUpdate(string path, FileEntryFlags flags, long length, LeafChunkedData chunkedData, ReadOnlyMemory<byte> customData = default)
-			: this(path, flags, length, chunkedData.Hash, chunkedData.LeafHandles, customData)
-		{
-		}
-
-		/// <summary>
-		/// Writes interior node data to the given writer
-		/// </summary>
-		public async ValueTask WriteInteriorNodesAsync(IBlobWriter writer, InteriorChunkedDataNodeOptions interiorNodeOptions, CancellationToken cancellationToken)
-		{
-			if (Nodes.Count > 1)
-			{
-				ChunkedDataNodeRef rootRef = await InteriorChunkedDataNode.CreateTreeAsync(Nodes, interiorNodeOptions, writer, cancellationToken);
-				Nodes.Clear();
-				Nodes.Add(rootRef);
-			}
-		}
-	}
+	/// <param name="Length">Sum total of all the file lengths in this directory tree</param>
+	/// <param name="Handle">Handle to the target node</param>
+	public record class DirectoryNodeRef(long Length, IBlobRef<DirectoryNode> Handle);
 
 	/// <summary>
-	/// Describes an update to a directory node
+	/// Entry for a directory within a directory node
 	/// </summary>
-	public class DirectoryUpdate
+	[DebuggerDisplay("{Name}")]
+	public record class DirectoryEntry(string Name, long Length, IBlobRef<DirectoryNode> Handle) : DirectoryNodeRef(Length, Handle)
 	{
-		/// <summary>
-		/// Directories to be updated
-		/// </summary>
-		public SortedDictionary<string, DirectoryUpdate?> Directories { get; } = new SortedDictionary<string, DirectoryUpdate?>(StringComparer.Ordinal);
-
-		/// <summary>
-		/// Files to be updated
-		/// </summary>
-		public SortedDictionary<string, FileUpdate?> Files { get; } = new SortedDictionary<string, FileUpdate?>(StringComparer.Ordinal);
-
-		/// <summary>
-		/// Reset this instance
-		/// </summary>
-		public void Clear()
-		{
-			Directories.Clear();
-			Files.Clear();
-		}
-
-		/// <summary>
-		/// Adds a file by path to this object
-		/// </summary>
-		/// <param name="path">Path to add to</param>
-		/// <param name="fileUpdate">Content for the file</param>
-		public void AddFile(string path, FileUpdate? fileUpdate)
-		{
-			string[] fragments = path.Split(new char[] { '/', '\\' });
-
-			DirectoryUpdate lastDir = this;
-			for (int idx = 0; idx + 1 < fragments.Length; idx++)
-			{
-				DirectoryUpdate? nextDir;
-				if (!lastDir.Directories.TryGetValue(fragments[idx], out nextDir) || nextDir == null)
-				{
-					if (fileUpdate == null)
-					{
-						return;
-					}
-
-					nextDir = new DirectoryUpdate();
-					lastDir.Directories.Add(fragments[idx], nextDir);
-				}
-				lastDir = nextDir;
-			}
-
-			lastDir.Files[fragments[^1]] = fileUpdate;
-		}
-
-		/// <summary>
-		/// Adds a file to the tree
-		/// </summary>
-		/// <param name="path">Path to the file</param>
-		/// <param name="flags">Flags for the new file entry</param>
-		/// <param name="length">Length of the file</param>
-		/// <param name="chunkedData">Chunked data instance</param>
-		/// <param name="customData"></param>
-		public FileUpdate AddFile(string path, FileEntryFlags flags, long length, LeafChunkedData chunkedData, ReadOnlyMemory<byte> customData)
-		{
-			string name = path;
-			for (int idx = path.Length - 1; idx >= 0; idx--)
-			{
-				if (path[idx] == Path.DirectorySeparatorChar || path[idx] == Path.AltDirectorySeparatorChar)
-				{
-					name = path.Substring(idx + 1);
-					break;
-				}
-			}
-
-			FileUpdate update = new FileUpdate(name, flags, length, chunkedData, customData);
-			AddFile(path, update);
-			return update;
-		}
-
-		/// <summary>
-		/// Adds a filtered list of files from disk
-		/// </summary>
-		/// <param name="files">Files to add</param>
-		public void AddFiles(IEnumerable<FileUpdate> files)
-		{
-			IEnumerator<FileUpdate> enumerator = files.GetEnumerator();
-			if (enumerator.MoveNext())
-			{
-				AddFiles(String.Empty, enumerator);
-			}
-		}
-
-		bool AddFiles(ReadOnlySpan<char> prefix, IEnumerator<FileUpdate> files)
-		{
-			for (; ; )
-			{
-				FileUpdate fileUpdate = files.Current;
-				if (fileUpdate.Path.Length < prefix.Length || !fileUpdate.Path.AsSpan().StartsWith(prefix, FileReference.Comparison))
-				{
-					return true;
-				}
-
-				int nextDirLength = fileUpdate.Path.AsSpan(prefix.Length).IndexOfAny(Path.DirectorySeparatorChar, Path.AltDirectorySeparatorChar);
-				if (nextDirLength == -1)
-				{
-					string fileName = fileUpdate.Path.Substring(prefix.Length);
-					Files[fileName] = fileUpdate;
-
-					if (!files.MoveNext())
-					{
-						return false;
-					}
-				}
-				else
-				{
-					string dirName = fileUpdate.Path.Substring(prefix.Length, nextDirLength);
-
-					DirectoryUpdate? nextTree;
-					if (!Directories.TryGetValue(dirName, out nextTree) || nextTree == null)
-					{
-						nextTree = new DirectoryUpdate();
-						Directories[dirName] = nextTree;
-					}
-
-					ReadOnlySpan<char> nextPrefix = fileUpdate.Path.AsSpan(0, prefix.Length + nextDirLength + 1);
-					if (!nextTree.AddFiles(nextPrefix, files))
-					{
-						return false;
-					}
-				}
-			}
-		}
-
-		/// <summary>
-		/// Writes interior node data to the given writer
-		/// </summary>
-		public async ValueTask WriteInteriorNodesAsync(IBlobWriter writer, InteriorChunkedDataNodeOptions interiorNodeOptions, CancellationToken cancellationToken)
-		{
-			foreach (DirectoryUpdate? directoryUpdate in Directories.Values)
-			{
-				if (directoryUpdate != null)
-				{
-					await directoryUpdate.WriteInteriorNodesAsync(writer, interiorNodeOptions, cancellationToken);
-				}
-			}
-			foreach (FileUpdate? fileUpdate in Files.Values)
-			{
-				if (fileUpdate != null)
-				{
-					await fileUpdate.WriteInteriorNodesAsync(writer, interiorNodeOptions, cancellationToken);
-				}
-			}
-		}
-	}
-
-	/// <summary>
-	/// Extension methods for writing directory nodes
-	/// </summary>
-	public static class DirectoryNodeExtensions
-	{
-		/// <summary>
-		/// Writes a tree of files to a storage writer
-		/// </summary>
-		public static async Task<IBlobRef<DirectoryNode>> WriteFilesAsync(this IBlobWriter writer, DirectoryReference baseDir, ChunkingOptions? options = null, IProgress<ICopyStats>? progress = null, CancellationToken cancellationToken = default)
-		{
-			DirectoryNode outputNode = new DirectoryNode();
-			await outputNode.AddFilesAsync(baseDir, DirectoryReference.EnumerateFiles(baseDir, "*", SearchOption.AllDirectories), writer, options, progress, cancellationToken);
-			return await writer.WriteBlobAsync(outputNode, cancellationToken: cancellationToken);
-		}
-
-		/// <summary>
-		/// Writes a tree of files to a storage writer
-		/// </summary>
-		public static async Task<IBlobRef<DirectoryNode>> WriteFilesAsync(this IBlobWriter writer, DirectoryInfo baseDir, IReadOnlyList<FileInfo> files, ChunkingOptions? options = null, IProgress<ICopyStats>? progress = null, CancellationToken cancellationToken = default)
-		{
-			DirectoryNode outputNode = new DirectoryNode();
-			await outputNode.AddFilesAsync(baseDir, files, writer, options, progress, cancellationToken);
-			return await writer.WriteBlobAsync(outputNode, cancellationToken: cancellationToken);
-		}
-
-		/// <summary>
-		/// Writes a tree of files to a storage writer
-		/// </summary>
-		public static async Task<IBlobRef<DirectoryNode>> WriteFilesAsync(this IBlobWriter writer, DirectoryReference baseDir, IReadOnlyList<FileReference> files, ChunkingOptions? options = null, IProgress<ICopyStats>? progress = null, CancellationToken cancellationToken = default)
-		{
-			DirectoryNode outputNode = new DirectoryNode();
-			await outputNode.AddFilesAsync(baseDir, files, writer, options, progress, cancellationToken);
-			return await writer.WriteBlobAsync(outputNode, cancellationToken: cancellationToken);
-		}
 	}
 }
diff --git a/Engine/Source/Programs/Shared/EpicGames.Horde/Storage/Nodes/DirectoryNodeRef.cs b/Engine/Source/Programs/Shared/EpicGames.Horde/Storage/Nodes/DirectoryNodeRef.cs
deleted file mode 100644
diff --git a/Engine/Source/Programs/Shared/EpicGames.Horde/Storage/Nodes/DirectoryNodeZipStream.cs b/Engine/Source/Programs/Shared/EpicGames.Horde/Storage/Nodes/DirectoryNodeZipStream.cs
deleted file mode 100644
diff --git a/Engine/Source/Programs/Shared/EpicGames.Horde/Storage/Nodes/FileEntry.cs b/Engine/Source/Programs/Shared/EpicGames.Horde/Storage/Nodes/FileEntry.cs
--- a/Engine/Source/Programs/Shared/EpicGames.Horde/Storage/Nodes/FileEntry.cs
+++ b/Engine/Source/Programs/Shared/EpicGames.Horde/Storage/Nodes/FileEntry.cs
@@ -130,7 +130,7 @@ namespace EpicGames.Horde.Storage.Nodes
 			try
 			{
 				await ChunkedDataNode.CopyToFileAsync(Target.Handle, file, cancellationToken);
-				ApplyPermissions(file, Flags);
+				SetPermissions(file, Flags);
 			}
 			catch (Exception ex) when (ex is not OperationCanceledException)
 			{
@@ -138,12 +138,44 @@ namespace EpicGames.Horde.Storage.Nodes
 			}
 		}
 
+		/// <summary>
+		/// Get the permission flags from a file on disk
+		/// </summary>
+		/// <param name="fileInfo">File to check</param>
+		/// <returns>Permission flags for the given file</returns>
+		public static FileEntryFlags GetPermissions(FileInfo fileInfo)
+		{
+			FileEntryFlags flags = FileEntryFlags.None;
+			if((fileInfo.Attributes & FileAttributes.ReadOnly) != 0)
+			{
+				flags |= FileEntryFlags.ReadOnly;
+			}
+
+			if (RuntimeInformation.IsOSPlatform(OSPlatform.Linux))
+			{
+				int mode = FileUtils.GetFileMode_Linux(fileInfo.FullName);
+				if ((mode & ((1 << 0) | (1 << 3) | (1 << 6))) != 0)
+				{
+					flags |= FileEntryFlags.Executable;
+				}
+			}
+			else if (RuntimeInformation.IsOSPlatform(OSPlatform.OSX))
+			{
+				int mode = FileUtils.GetFileMode_Mac(fileInfo.FullName);
+				if ((mode & ((1 << 0) | (1 << 3) | (1 << 6))) != 0)
+				{
+					flags |= FileEntryFlags.Executable;
+				}
+			}
+			return flags;
+		}
+
 		/// <summary>
 		/// Applies the correct permissions to a file for a particular set of file entry flags
 		/// </summary>
 		/// <param name="fileInfo">File to modify</param>
 		/// <param name="flags">Flags for the file</param>
-		public static void ApplyPermissions(FileInfo fileInfo, FileEntryFlags flags)
+		public static void SetPermissions(FileInfo fileInfo, FileEntryFlags flags)
 		{
 			if ((flags & FileEntryFlags.ReadOnly) != 0)
 			{
