Horde: Increase Redis timeout when a debugger is attached.

diff --git a/Engine/Source/Programs/Shared/EpicGames.Redis/RedisConnectionPool.cs b/Engine/Source/Programs/Shared/EpicGames.Redis/RedisConnectionPool.cs
--- a/Engine/Source/Programs/Shared/EpicGames.Redis/RedisConnectionPool.cs
+++ b/Engine/Source/Programs/Shared/EpicGames.Redis/RedisConnectionPool.cs
@@ -29,7 +29,15 @@ namespace EpicGames.Redis
 			_connections = new Lazy<IConnectionMultiplexer>[poolSize];
 			for (int i = 0; i < poolSize; i++)
 			{
-				_connections[i] = new Lazy<IConnectionMultiplexer>(() => ConnectionMultiplexer.Connect(redisConfString));
+				static void ConfigureOptions(ConfigurationOptions options)
+				{
+					if (Debugger.IsAttached)
+					{
+						options.SyncTimeout = 1_000_000;
+					}
+				}
+
+				_connections[i] = new Lazy<IConnectionMultiplexer>(() => ConnectionMultiplexer.Connect(redisConfString, ConfigureOptions));
 			}
 		}
 
