Horde: Fix missing parameter from log message.

diff --git a/Engine/Source/Programs/Shared/EpicGames.Perforce.Managed/ManagedWorkspace.cs b/Engine/Source/Programs/Shared/EpicGames.Perforce.Managed/ManagedWorkspace.cs
--- a/Engine/Source/Programs/Shared/EpicGames.Perforce.Managed/ManagedWorkspace.cs
+++ b/Engine/Source/Programs/Shared/EpicGames.Perforce.Managed/ManagedWorkspace.cs
@@ -1220,7 +1220,7 @@ namespace EpicGames.Perforce.Managed
 				client.Type = _options.UseHaveTable ? "partitioned" : "readonly";
 			}
 
-			_logger.LogInformation("Updating client {ClientName} (Host: {HostName}, Stream: {StreamName}, Type: {Type})", client.Host, client.Stream, client.Type ?? "(default)");
+			_logger.LogInformation("Updating client {ClientName} (Host: {HostName}, Stream: {StreamName}, Type: {Type})", client.Name, client.Host, client.Stream, client.Type ?? "(default)");
 
 			using (Trace("UpdateClient"))
 			using (ILoggerProgress status = _logger.BeginProgressScope("Updating client..."))
