Horde: Don't catch and re-throw cancellation exceptions when reading data from bundles; allow the cancellation exception to propagate to the caller.

diff --git a/Engine/Source/Programs/Shared/EpicGames.Horde/Storage/Bundles/V2/PacketHandle.cs b/Engine/Source/Programs/Shared/EpicGames.Horde/Storage/Bundles/V2/PacketHandle.cs
--- a/Engine/Source/Programs/Shared/EpicGames.Horde/Storage/Bundles/V2/PacketHandle.cs
+++ b/Engine/Source/Programs/Shared/EpicGames.Horde/Storage/Bundles/V2/PacketHandle.cs
@@ -159,6 +159,10 @@ namespace EpicGames.Horde.Storage.Bundles.V2
 				using IRefCountedHandle<PacketReader> packetReaderHandle = await GetPacketReaderAsync(cancellationToken);
 				return packetReaderHandle.Target.ReadExport(exportIdx);
 			}
+			catch (OperationCanceledException)
+			{
+				throw;
+			}
 			catch (Exception ex)
 			{
 				Utf8String locator = GetLocator();
