Horde: Speculative fix for Linux error stating that size of memory mapped file is less than the size of the filestream (suspect caused by a symlink in the tree).

diff --git a/Engine/Source/Programs/Shared/EpicGames.Horde/Storage/Nodes/DirectoryNode.Extract.cs b/Engine/Source/Programs/Shared/EpicGames.Horde/Storage/Nodes/DirectoryNode.Extract.cs
--- a/Engine/Source/Programs/Shared/EpicGames.Horde/Storage/Nodes/DirectoryNode.Extract.cs
+++ b/Engine/Source/Programs/Shared/EpicGames.Horde/Storage/Nodes/DirectoryNode.Extract.cs
@@ -211,55 +211,73 @@ namespace EpicGames.Horde.Storage.Nodes
 				// Open the file for the current chunk
 				OutputFile file = chunk.File;
 
-				FileReference locator = FileReference.Combine(baseDir, file.Directory.Path, file.FileEntry.Name);
-				DirectoryReference.CreateDirectory(locator.Directory);
-
-				FileInfo fileInfo = locator.ToFileInfo();
-				if (fileInfo.Exists && (fileInfo.Attributes & FileAttributes.ReadOnly) != 0)
+				FileReference location = FileReference.Combine(baseDir, file.Directory.Path, file.FileEntry.Name);
+				try
 				{
-					fileInfo.Attributes &= ~FileAttributes.ReadOnly;
-				}
+					DirectoryReference.CreateDirectory(location.Directory);
 
-				await using (FileStream stream = fileInfo.Open(FileMode.OpenOrCreate, FileAccess.ReadWrite, FileShare.ReadWrite))
-				{
-					stream.SetLength(file.FileEntry.Length);
-					if (file.FileEntry.Length == 0)
+					FileInfo fileInfo = location.ToFileInfo();
+					if (fileInfo.Exists)
 					{
-						// If this file is empty, don't write anything and just move to the next chunk
-						chunk = await ReadNextChunkAsync(chunkReader, cancellationToken);
+						if ((fileInfo.Attributes & FileAttributes.ReadOnly) != 0)
+						{
+							fileInfo.Attributes &= ~FileAttributes.ReadOnly;
+						}
+						if (fileInfo.LinkTarget != null)
+						{
+							fileInfo.Delete();
+						}
 					}
-					else
-					{
-						// Process as many chunks as we can for this file
-						using MemoryMappedFile memoryMappedFile = MemoryMappedFile.CreateFromFile(stream, null, file.FileEntry.Length, MemoryMappedFileAccess.ReadWrite, HandleInheritability.None, false);
-						using MemoryMappedView memoryMappedView = new MemoryMappedView(memoryMappedFile, 0, file.FileEntry.Length);
 
-						while (chunk != null && chunk.File == file)
+					await using (FileStream stream = fileInfo.Open(FileMode.OpenOrCreate, FileAccess.ReadWrite, FileShare.ReadWrite))
+					{
+						stream.SetLength(file.FileEntry.Length);
+						if (file.FileEntry.Length == 0)
 						{
-							// Write this chunk
-							using (BlobData data = await chunk.Handle.ReadBlobDataAsync(cancellationToken))
-							{
-								TraceBlobRead("Leaf", CombinePaths(chunk.File.Directory.Path, chunk.File.FileEntry.Name), chunk.Handle, logger);
-								data.Data.CopyTo(memoryMappedView!.GetMemory(chunk.Offset, data.Data.Length));
-							}
+							// If this file is empty, don't write anything and just move to the next chunk
+							chunk = await ReadNextChunkAsync(chunkReader, cancellationToken);
+						}
+						else
+						{
+							// Process as many chunks as we can for this file
+							using MemoryMappedFile memoryMappedFile = MemoryMappedFile.CreateFromFile(stream, null, file.FileEntry.Length, MemoryMappedFileAccess.ReadWrite, HandleInheritability.None, false);
+							using MemoryMappedView memoryMappedView = new MemoryMappedView(memoryMappedFile, 0, file.FileEntry.Length);
 
-							// Update the stats
-							int numCompleteFiles = 0;
-							if (chunk.Offset + chunk.Length == file.FileEntry.Length)
+							while (chunk != null && chunk.File == file)
 							{
-								numCompleteFiles = 1;
+								// Write this chunk
+								using (BlobData data = await chunk.Handle.ReadBlobDataAsync(cancellationToken))
+								{
+									TraceBlobRead("Leaf", CombinePaths(chunk.File.Directory.Path, chunk.File.FileEntry.Name), chunk.Handle, logger);
+									data.Data.CopyTo(memoryMappedView!.GetMemory(chunk.Offset, data.Data.Length));
+								}
+
+								// Update the stats
+								int numCompleteFiles = 0;
+								if (chunk.Offset + chunk.Length == file.FileEntry.Length)
+								{
+									numCompleteFiles = 1;
+								}
+
+								copyStats?.Update(numCompleteFiles, chunk.Length);
+
+								// Read the next chunk
+								chunk = await ReadNextChunkAsync(chunkReader, cancellationToken);
 							}
-
-							copyStats?.Update(numCompleteFiles, chunk.Length);
-
-							// Read the next chunk
-							chunk = await ReadNextChunkAsync(chunkReader, cancellationToken);
 						}
 					}
-				}
 
-				// Set correct permissions on the output file
-				FileEntry.SetPermissions(fileInfo, file.FileEntry.Flags);
+					// Set correct permissions on the output file
+					FileEntry.SetPermissions(fileInfo, file.FileEntry.Flags);
+				}
+				catch (OperationCanceledException)
+				{
+					throw;
+				}
+				catch (Exception ex)
+				{
+					throw new StorageException($"Unable to extract {location}: {ex.Message}", ex); 
+				}
 			}
 		}
 
