Horde: Include the root directory when logging info about managed workspace use.

diff --git a/Engine/Source/Programs/Shared/EpicGames.Perforce.Managed/ManagedWorkspace.cs b/Engine/Source/Programs/Shared/EpicGames.Perforce.Managed/ManagedWorkspace.cs
--- a/Engine/Source/Programs/Shared/EpicGames.Perforce.Managed/ManagedWorkspace.cs
+++ b/Engine/Source/Programs/Shared/EpicGames.Perforce.Managed/ManagedWorkspace.cs
@@ -1220,7 +1220,7 @@ namespace EpicGames.Perforce.Managed
 				client.Type = _options.UseHaveTable ? "partitioned" : "readonly";
 			}
 
-			_logger.LogInformation("Updating client {ClientName} (Host: {HostName}, Stream: {StreamName}, Type: {Type})", client.Name, client.Host, client.Stream, client.Type ?? "(default)");
+			_logger.LogInformation("Using client {ClientName} (Host: {HostName}, Stream: {StreamName}, Type: {Type}, Root: {Path})", client.Name, client.Host, client.Stream, client.Type ?? "full", client.Root);
 
 			using (Trace("UpdateClient"))
 			using (ILoggerProgress status = _logger.BeginProgressScope("Updating client..."))
