Horde: Fix race condition between multiple threads writing to a file and setting file attributes.

A shared counter for the number of remaining chunks of a file to be written is stored for each output file, and the attributes are only applied after the last write is complete.

diff --git a/Engine/Source/Programs/Shared/EpicGames.Horde/Storage/Nodes/DirectoryNode.Extract.cs b/Engine/Source/Programs/Shared/EpicGames.Horde/Storage/Nodes/DirectoryNode.Extract.cs
--- a/Engine/Source/Programs/Shared/EpicGames.Horde/Storage/Nodes/DirectoryNode.Extract.cs
+++ b/Engine/Source/Programs/Shared/EpicGames.Horde/Storage/Nodes/DirectoryNode.Extract.cs
@@ -141,9 +141,17 @@ namespace EpicGames.Horde.Storage.Nodes
 		/// <param name="cancellationToken"></param>
 		public static Task CopyToDirectoryAsync(this DirectoryNode directoryNode, DirectoryInfo directoryInfo, ILogger logger, CancellationToken cancellationToken) => CopyToDirectoryAsync(directoryNode, directoryInfo, null, logger, cancellationToken);
 
+		class RemainingChunkCounter
+		{
+			public int _value;
+
+			public int Increment() => Interlocked.Increment(ref _value);
+			public int Decrement() => Interlocked.Decrement(ref _value);
+		}
+
 		record class OutputDir(string Path, DirectoryNode Node);
-		record class OutputFile(OutputDir Directory, FileEntry FileEntry);
-		record class OutputChunk(OutputFile File, long Offset, long Length, IBlobHandle Handle);
+		record class OutputFile(OutputDir Directory, FileEntry FileEntry, RemainingChunkCounter RemainingChunks);
+		record class OutputChunk(OutputFile File, long Offset, long Length, IBlobHandle? Handle);
 
 #pragma warning disable IDE0060
 		static void TraceBlobRead(string type, string path, IBlobHandle handle, ILogger logger)
@@ -229,6 +237,7 @@ namespace EpicGames.Horde.Storage.Nodes
 						}
 					}
 
+					int remainingChunks = 0;
 					await using (FileStream stream = fileInfo.Open(FileMode.OpenOrCreate, FileAccess.ReadWrite, FileShare.ReadWrite))
 					{
 						stream.SetLength(file.FileEntry.Length);
@@ -236,6 +245,7 @@ namespace EpicGames.Horde.Storage.Nodes
 						{
 							// If this file is empty, don't write anything and just move to the next chunk
 							chunk = await ReadNextChunkAsync(chunkReader, cancellationToken);
+							remainingChunks = file.RemainingChunks.Decrement();
 						}
 						else
 						{
@@ -246,20 +256,18 @@ namespace EpicGames.Horde.Storage.Nodes
 							while (chunk != null && chunk.File == file)
 							{
 								// Write this chunk
-								using (BlobData data = await chunk.Handle.ReadBlobDataAsync(cancellationToken))
+								if (chunk.Handle != null)
 								{
-									TraceBlobRead("Leaf", CombinePaths(chunk.File.Directory.Path, chunk.File.FileEntry.Name), chunk.Handle, logger);
-									data.Data.CopyTo(memoryMappedView!.GetMemory(chunk.Offset, data.Data.Length));
+									using (BlobData data = await chunk.Handle.ReadBlobDataAsync(cancellationToken))
+									{
+										TraceBlobRead("Leaf", CombinePaths(chunk.File.Directory.Path, chunk.File.FileEntry.Name), chunk.Handle, logger);
+										data.Data.CopyTo(memoryMappedView!.GetMemory(chunk.Offset, data.Data.Length));
+									}
 								}
 
 								// Update the stats
-								int numCompleteFiles = 0;
-								if (chunk.Offset + chunk.Length == file.FileEntry.Length)
-								{
-									numCompleteFiles = 1;
-								}
-
-								copyStats?.Update(numCompleteFiles, chunk.Length);
+								remainingChunks = file.RemainingChunks.Decrement();
+								copyStats?.Update(0, chunk.Length);
 
 								// Read the next chunk
 								chunk = await ReadNextChunkAsync(chunkReader, cancellationToken);
@@ -268,7 +276,11 @@ namespace EpicGames.Horde.Storage.Nodes
 					}
 
 					// Set correct permissions on the output file
-					FileEntry.SetPermissions(fileInfo, file.FileEntry.Flags);
+					if (remainingChunks == 0)
+					{
+						FileEntry.SetPermissions(fileInfo, file.FileEntry.Flags);
+						copyStats?.Update(1, 0);
+					}
 				}
 				catch (OperationCanceledException)
 				{
@@ -308,8 +320,21 @@ namespace EpicGames.Horde.Storage.Nodes
 
 			foreach (FileEntry fileEntry in node.Files)
 			{
-				OutputFile outputFile = new OutputFile(outputDir, fileEntry);
-				await FindOutputChunksAsync(outputFile, 0, fileEntry.Target, chunks, logger, cancellationToken);
+				OutputFile outputFile = new OutputFile(outputDir, fileEntry, new RemainingChunkCounter());
+
+				// Ensure that no worker finalizes the file until we've finished adding chunks
+				outputFile.RemainingChunks.Increment();
+
+				if (fileEntry.Target.Type == ChunkedDataNodeType.Leaf && fileEntry.Target.Length >= 0)
+				{
+					// If there's only one chunk for this file, queue the chunk and let the worker finalize the file in one operation
+					await chunks.WriteAsync(new OutputChunk(outputFile, 0, fileEntry.Target.Length, fileEntry.Target.Handle), cancellationToken);
+				}
+				else
+				{
+					await FindOutputChunksAsync(outputFile, 0, fileEntry.Target, chunks, logger, cancellationToken);
+					await chunks.WriteAsync(new OutputChunk(outputFile, 0, 0, null));
+				}
 			}
 
 			foreach (DirectoryEntry directoryEntry in node.Directories)
@@ -328,6 +353,8 @@ namespace EpicGames.Horde.Storage.Nodes
 			if (dataRef.Type == ChunkedDataNodeType.Leaf)
 			{
 				await chunks.WriteAsync(new OutputChunk(outputFile, offset, dataRef.Length, dataRef.Handle), cancellationToken);
+				outputFile.RemainingChunks.Increment();
+
 				if (dataRef.Length < 0)
 				{
 					// Backwards compatibility hack for v2 format
@@ -344,6 +371,7 @@ namespace EpicGames.Horde.Storage.Nodes
 				if (data.Type.Guid == LeafChunkedDataNodeConverter.BlobType.Guid)
 				{
 					await chunks.WriteAsync(new OutputChunk(outputFile, offset, dataRef.Length, dataRef.Handle), cancellationToken);
+					outputFile.RemainingChunks.Increment();
 					return data.Data.Length;
 				}
 				else
