Horde: Pass cancellation token when extracting blobs.

diff --git a/Engine/Source/Programs/Shared/EpicGames.Horde/Storage/Nodes/DirectoryNode.Extract.cs b/Engine/Source/Programs/Shared/EpicGames.Horde/Storage/Nodes/DirectoryNode.Extract.cs
--- a/Engine/Source/Programs/Shared/EpicGames.Horde/Storage/Nodes/DirectoryNode.Extract.cs
+++ b/Engine/Source/Programs/Shared/EpicGames.Horde/Storage/Nodes/DirectoryNode.Extract.cs
@@ -333,7 +333,7 @@ namespace EpicGames.Horde.Storage.Nodes
 				else
 				{
 					await FindOutputChunksAsync(outputFile, 0, fileEntry.Target, chunks, logger, cancellationToken);
-					await chunks.WriteAsync(new OutputChunk(outputFile, 0, 0, null));
+					await chunks.WriteAsync(new OutputChunk(outputFile, 0, 0, null), cancellationToken);
 				}
 			}
 
