Horde: Fix file modes for Mac/Linux when applying extracted file permissions.

diff --git a/Engine/Source/Programs/Shared/EpicGames.Horde/Storage/Nodes/FileEntry.cs b/Engine/Source/Programs/Shared/EpicGames.Horde/Storage/Nodes/FileEntry.cs
--- a/Engine/Source/Programs/Shared/EpicGames.Horde/Storage/Nodes/FileEntry.cs
+++ b/Engine/Source/Programs/Shared/EpicGames.Horde/Storage/Nodes/FileEntry.cs
@@ -177,24 +177,52 @@ namespace EpicGames.Horde.Storage.Nodes
 		/// <param name="flags">Flags for the file</param>
 		public static void SetPermissions(FileInfo fileInfo, FileEntryFlags flags)
 		{
-			if ((flags & FileEntryFlags.ReadOnly) != 0)
+			if (RuntimeInformation.IsOSPlatform(OSPlatform.Linux))
 			{
-				fileInfo.Attributes |= FileAttributes.ReadOnly;
-			}
+				int mode = FileUtils.GetFileMode_Linux(fileInfo.FullName);
 
-			if ((flags & FileEntryFlags.Executable) != 0)
+				int newMode = UpdateFileMode(mode, flags);
+				if (mode != newMode)
+				{
+					FileUtils.SetFileMode_Linux(fileInfo.FullName, (ushort)newMode);
+					fileInfo.Refresh();
+				}
+			}
+			else if (RuntimeInformation.IsOSPlatform(OSPlatform.OSX))
 			{
-				if (RuntimeInformation.IsOSPlatform(OSPlatform.Linux))
+				int mode = FileUtils.GetFileMode_Mac(fileInfo.FullName);
+
+				int newMode = UpdateFileMode(mode, flags);
+				if (mode != newMode)
 				{
-					FileUtils.SetFileMode_Linux(fileInfo.FullName, 0b_111_111_111);
+					FileUtils.SetFileMode_Mac(fileInfo.FullName, (ushort)newMode);
+					fileInfo.Refresh();
 				}
-				else if (RuntimeInformation.IsOSPlatform(OSPlatform.OSX))
+			}
+			else
+			{
+				if ((flags & FileEntryFlags.ReadOnly) != 0)
 				{
-					FileUtils.SetFileMode_Mac(fileInfo.FullName, 0b_111_111_111);
+					fileInfo.Attributes |= FileAttributes.ReadOnly;
 				}
 			}
 		}
 
+		static int UpdateFileMode(int mode, FileEntryFlags flags)
+		{
+			if ((flags & FileEntryFlags.ReadOnly) != 0)
+			{
+				const int WriteMask = 0b_010_010_010;
+				mode &= ~WriteMask;
+			}
+			if ((flags & FileEntryFlags.Executable) != 0)
+			{
+				const int ExecuteMask = 0b_001_001_001;
+				mode |= ExecuteMask;
+			}
+			return mode;
+		}
+
 		/// <inheritdoc/>
 		public override string ToString() => Name.ToString();
 	}
