EpicGames.Core: Set the TaskCreationOptions.RunContinuationsAsynchronously flag by default for all task completion sources.

diff --git a/Engine/Source/Programs/Shared/EpicGames.Core/AsyncEvent.cs b/Engine/Source/Programs/Shared/EpicGames.Core/AsyncEvent.cs
--- a/Engine/Source/Programs/Shared/EpicGames.Core/AsyncEvent.cs
+++ b/Engine/Source/Programs/Shared/EpicGames.Core/AsyncEvent.cs
@@ -13,7 +13,7 @@ namespace EpicGames.Core
 		/// <summary>
 		/// Completion source to wait on
 		/// </summary>
-		TaskCompletionSource<bool> _source = new TaskCompletionSource<bool>();
+		TaskCompletionSource<bool> _source = new TaskCompletionSource<bool>(TaskCreationOptions.RunContinuationsAsynchronously);
 
 		/// <summary>
 		/// Legacy name for <see cref="Pulse"/>.
@@ -33,7 +33,7 @@ namespace EpicGames.Core
 					// The task source is latched or has been set from another thread. Either way behaves identically wrt the contract of this method.
 					break;
 				}
-				if (Interlocked.CompareExchange(ref _source, new TaskCompletionSource<bool>(), prevSource) == prevSource)
+				if (Interlocked.CompareExchange(ref _source, new TaskCompletionSource<bool>(TaskCreationOptions.RunContinuationsAsynchronously), prevSource) == prevSource)
 				{
 					prevSource.TrySetResult(true);
 					break;
@@ -53,7 +53,7 @@ namespace EpicGames.Core
 				{
 					break;
 				}
-				else if (Interlocked.CompareExchange(ref _source, new TaskCompletionSource<bool>(), prevSource) == prevSource)
+				else if (Interlocked.CompareExchange(ref _source, new TaskCompletionSource<bool>(TaskCreationOptions.RunContinuationsAsynchronously), prevSource) == prevSource)
 				{
 					break;
 				}
diff --git a/Engine/Source/Programs/Shared/EpicGames.Core/CancellationTask.cs b/Engine/Source/Programs/Shared/EpicGames.Core/CancellationTask.cs
--- a/Engine/Source/Programs/Shared/EpicGames.Core/CancellationTask.cs
+++ b/Engine/Source/Programs/Shared/EpicGames.Core/CancellationTask.cs
@@ -27,7 +27,7 @@ namespace EpicGames.Core
 		/// <param name="token">The cancellation token to register with</param>
 		public CancellationTask(CancellationToken token)
 		{
-			_completionSource = new TaskCompletionSource<bool>();
+			_completionSource = new TaskCompletionSource<bool>(TaskCreationOptions.RunContinuationsAsynchronously);
 			_registration = token.Register(() => _completionSource.TrySetResult(true));
 		}
 
diff --git a/Engine/Source/Programs/Shared/EpicGames.Core/ManagedProcess.cs b/Engine/Source/Programs/Shared/EpicGames.Core/ManagedProcess.cs
--- a/Engine/Source/Programs/Shared/EpicGames.Core/ManagedProcess.cs
+++ b/Engine/Source/Programs/Shared/EpicGames.Core/ManagedProcess.cs
@@ -1092,7 +1092,7 @@ namespace EpicGames.Core
 		/// <returns></returns>
 		public async Task CopyToAsync(Func<byte[], int, int, CancellationToken, Task> writeOutputAsync, int bufferSize, CancellationToken cancellationToken)
 		{
-			TaskCompletionSource<bool> taskCompletionSource = new TaskCompletionSource<bool>();
+			TaskCompletionSource<bool> taskCompletionSource = new TaskCompletionSource<bool>(TaskCreationOptions.RunContinuationsAsynchronously);
 			using (CancellationTokenRegistration registration = cancellationToken.Register(() => taskCompletionSource.SetResult(false)))
 			{
 				byte[] buffer = new byte[bufferSize];
diff --git a/Engine/Source/Programs/Shared/EpicGames.Core/SingleInstanceMutex.cs b/Engine/Source/Programs/Shared/EpicGames.Core/SingleInstanceMutex.cs
--- a/Engine/Source/Programs/Shared/EpicGames.Core/SingleInstanceMutex.cs
+++ b/Engine/Source/Programs/Shared/EpicGames.Core/SingleInstanceMutex.cs
@@ -15,7 +15,7 @@ namespace EpicGames.Core
 		class SingleInstanceMutexImpl : IDisposable
 		{
 			readonly string? _name;
-			readonly TaskCompletionSource _readyTcs = new TaskCompletionSource();
+			readonly TaskCompletionSource _readyTcs = new TaskCompletionSource(TaskCreationOptions.RunContinuationsAsynchronously);
 			readonly ManualResetEvent _disposing = new ManualResetEvent(false);
 			Thread? _thread;
 
