Horde: Fix test failure; side effect due to asynchronous continuations.

Also fix a few static analyzer warnings.

diff --git a/Engine/Source/Programs/Horde/Horde.Server.Tests/Agents/Relay/AgentRelayServiceTests.cs b/Engine/Source/Programs/Horde/Horde.Server.Tests/Agents/Relay/AgentRelayServiceTests.cs
--- a/Engine/Source/Programs/Horde/Horde.Server.Tests/Agents/Relay/AgentRelayServiceTests.cs
+++ b/Engine/Source/Programs/Horde/Horde.Server.Tests/Agents/Relay/AgentRelayServiceTests.cs
@@ -60,6 +60,13 @@ public class AgentRelayServiceTests : TestSetup
 		_service = new AgentRelayService(GetRedisServiceSingleton(), Clock, NullLogger<AgentRelayService>.Instance);
 	}
 
+	public override async ValueTask DisposeAsync()
+	{
+		await _service.DisposeAsync();
+
+		await base.DisposeAsync();
+	}
+
 	[TestInitialize]
 	public async Task SetupAsync()
 	{
diff --git a/Engine/Source/Programs/Horde/Horde.Server.Tests/Utilities/AsyncTaskQueueTests.cs b/Engine/Source/Programs/Horde/Horde.Server.Tests/Utilities/AsyncTaskQueueTests.cs
--- a/Engine/Source/Programs/Horde/Horde.Server.Tests/Utilities/AsyncTaskQueueTests.cs
+++ b/Engine/Source/Programs/Horde/Horde.Server.Tests/Utilities/AsyncTaskQueueTests.cs
@@ -1,6 +1,5 @@
 // Copyright Epic Games, Inc. All Rights Reserved.
 
-using System;
 using System.Collections.Generic;
 using System.Linq;
 using System.Text;
@@ -20,10 +19,18 @@ namespace Horde.Server.Tests.Utilities
 			await using AsyncTaskQueue queue = new AsyncTaskQueue(NullLogger.Instance);
 
 			bool executed1 = false;
-			queue.Enqueue(async _ => { await Task.Yield(); executed1 = true; });
+			queue.Enqueue(async _ => 
+			{ 
+				await Task.Yield(); 
+				executed1 = true; 
+			});
 
 			bool executed2 = false;
-			queue.Enqueue(async _ => { await Task.Yield(); executed2 = true; });
+			queue.Enqueue(async _ => 
+			{ 
+				await Task.Yield(); 
+				executed2 = true; 
+			});
 
 			await queue.FlushAsync();
 
diff --git a/Engine/Source/Programs/Horde/Horde.Server/Configuration/ConfigService.cs b/Engine/Source/Programs/Horde/Horde.Server/Configuration/ConfigService.cs
--- a/Engine/Source/Programs/Horde/Horde.Server/Configuration/ConfigService.cs
+++ b/Engine/Source/Programs/Horde/Horde.Server/Configuration/ConfigService.cs
@@ -327,7 +327,9 @@ namespace Horde.Server.Configuration
 		/// <inheritdoc/>
 		public GlobalConfig Create(string name)
 		{
+#pragma warning disable VSTHRD002 // Synchronous wait
 			return _stateTask.Result.GlobalConfig;
+#pragma warning restore VSTHRD002 // Synchronous wait
 		}
 
 		/// <inheritdoc/>
diff --git a/Engine/Source/Programs/Shared/EpicGames.Horde.Tests/BufferTests.cs b/Engine/Source/Programs/Shared/EpicGames.Horde.Tests/BufferTests.cs
--- a/Engine/Source/Programs/Shared/EpicGames.Horde.Tests/BufferTests.cs
+++ b/Engine/Source/Programs/Shared/EpicGames.Horde.Tests/BufferTests.cs
@@ -3,6 +3,7 @@
 using System;
 using System.IO.Pipelines;
 using System.Linq;
+using System.Runtime.CompilerServices;
 using System.Runtime.InteropServices;
 using System.Security.Cryptography;
 using System.Threading;
@@ -34,7 +35,7 @@ namespace EpicGames.Horde.Tests
 		}
 
 		[TestMethod]
-		public void TestOverflow()
+		public async Task TestOverflow()
 		{
 			using PooledBuffer buffer = new PooledBuffer(2, 20);
 			using ComputeBufferReader bufferReader = buffer.CreateReader();
@@ -73,7 +74,7 @@ namespace EpicGames.Horde.Tests
 			Assert.IsTrue(waitToReadTask.IsCompleted);
 
 			Assert.AreEqual(20, bufferReader.GetReadBuffer().Length);
-			Assert.IsTrue(waitToWriteTask.IsCompleted);
+			await waitToWriteTask;
 
 			// Make sure both reader and writer have something to work with
 			Assert.AreEqual(20, bufferReader.GetReadBuffer().Length);
