Horde: Fix compiler warnings

diff --git a/Engine/Source/Programs/Horde/Horde.Server.Tests/Agents/Relay/AgentRelayServiceTests.cs b/Engine/Source/Programs/Horde/Horde.Server.Tests/Agents/Relay/AgentRelayServiceTests.cs
--- a/Engine/Source/Programs/Horde/Horde.Server.Tests/Agents/Relay/AgentRelayServiceTests.cs
+++ b/Engine/Source/Programs/Horde/Horde.Server.Tests/Agents/Relay/AgentRelayServiceTests.cs
@@ -63,8 +63,8 @@ public class AgentRelayServiceTests : TestSetup
 	public override async ValueTask DisposeAsync()
 	{
 		await _service.DisposeAsync();
-
 		await base.DisposeAsync();
+		GC.SuppressFinalize(this);
 	}
 
 	[TestInitialize]
diff --git a/Engine/Source/Programs/Horde/Horde.Server.Tests/Utilities/AsyncTaskQueueTests.cs b/Engine/Source/Programs/Horde/Horde.Server.Tests/Utilities/AsyncTaskQueueTests.cs
--- a/Engine/Source/Programs/Horde/Horde.Server.Tests/Utilities/AsyncTaskQueueTests.cs
+++ b/Engine/Source/Programs/Horde/Horde.Server.Tests/Utilities/AsyncTaskQueueTests.cs
@@ -1,8 +1,5 @@
 // Copyright Epic Games, Inc. All Rights Reserved.
 
-using System.Collections.Generic;
-using System.Linq;
-using System.Text;
 using System.Threading.Tasks;
 using Horde.Server.Utilities;
 using Microsoft.Extensions.Logging.Abstractions;
diff --git a/Engine/Source/Programs/Shared/EpicGames.Horde.Tests/BufferTests.cs b/Engine/Source/Programs/Shared/EpicGames.Horde.Tests/BufferTests.cs
--- a/Engine/Source/Programs/Shared/EpicGames.Horde.Tests/BufferTests.cs
+++ b/Engine/Source/Programs/Shared/EpicGames.Horde.Tests/BufferTests.cs
@@ -3,7 +3,6 @@
 using System;
 using System.IO.Pipelines;
 using System.Linq;
-using System.Runtime.CompilerServices;
 using System.Runtime.InteropServices;
 using System.Security.Cryptography;
 using System.Threading;
@@ -35,7 +34,7 @@ namespace EpicGames.Horde.Tests
 		}
 
 		[TestMethod]
-		public async Task TestOverflow()
+		public async Task TestOverflowAsync()
 		{
 			using PooledBuffer buffer = new PooledBuffer(2, 20);
 			using ComputeBufferReader bufferReader = buffer.CreateReader();
