Horde: Include more information when an upload fails.

diff --git a/Engine/Source/Programs/Shared/EpicGames.Horde/Storage/Backends/HttpStorageBackend.cs b/Engine/Source/Programs/Shared/EpicGames.Horde/Storage/Backends/HttpStorageBackend.cs
--- a/Engine/Source/Programs/Shared/EpicGames.Horde/Storage/Backends/HttpStorageBackend.cs
+++ b/Engine/Source/Programs/Shared/EpicGames.Horde/Storage/Backends/HttpStorageBackend.cs
@@ -142,7 +142,12 @@ namespace EpicGames.Horde.Storage.Backends
 					request.Content = form;
 					using (HttpResponseMessage response = await httpClient.SendAsync(request, cancellationToken))
 					{
-						response.EnsureSuccessStatusCode();
+						if (!response.IsSuccessStatusCode)
+						{
+							string responseText = await response.Content.ReadAsStringAsync();
+							throw new StorageException($"Upload to {request.RequestUri} failed ({response.StatusCode}). Response: {responseText}");
+						}
+
 						WriteBlobResponse? data = await response.Content.ReadFromJsonAsync<WriteBlobResponse>(cancellationToken: cancellationToken);
 						return data!;
 					}
