Horde: Fix static analysis warning.

diff --git a/Engine/Source/Programs/Shared/EpicGames.Horde/Storage/Backends/HttpStorageBackend.cs b/Engine/Source/Programs/Shared/EpicGames.Horde/Storage/Backends/HttpStorageBackend.cs
--- a/Engine/Source/Programs/Shared/EpicGames.Horde/Storage/Backends/HttpStorageBackend.cs
+++ b/Engine/Source/Programs/Shared/EpicGames.Horde/Storage/Backends/HttpStorageBackend.cs
@@ -144,7 +144,7 @@ namespace EpicGames.Horde.Storage.Backends
 					{
 						if (!response.IsSuccessStatusCode)
 						{
-							string responseText = await response.Content.ReadAsStringAsync();
+							string responseText = await response.Content.ReadAsStringAsync(cancellationToken);
 							throw new StorageException($"Upload to {request.RequestUri} failed ({response.StatusCode}). Response: {responseText}");
 						}
 
