Horde: Fix access token from the environment overriding any access token manually set in HordeOptions.

diff --git a/Engine/Source/Programs/Shared/EpicGames.Horde/HordeHttpAuthHandler.cs b/Engine/Source/Programs/Shared/EpicGames.Horde/HordeHttpAuthHandler.cs
--- a/Engine/Source/Programs/Shared/EpicGames.Horde/HordeHttpAuthHandler.cs
+++ b/Engine/Source/Programs/Shared/EpicGames.Horde/HordeHttpAuthHandler.cs
@@ -47,6 +47,11 @@ namespace EpicGames.Horde
 					// If an explicit access token is specified, just use that
 					request.Headers.Authorization = new AuthenticationHeaderValue("Bearer", _options.Value.AccessToken);
 				}
+				else if (request.RequestUri != null && TryGetAccessTokenFromEnvironment(request.RequestUri, out string? accessToken))
+				{
+					// Use the access token specified in the environment
+					request.Headers.Authorization = new AuthenticationHeaderValue("Bearer", accessToken);
+				}
 				else if (_options.Value.AllowAuthPrompt)
 				{
 					// Try to use the cached auth header
@@ -78,6 +83,28 @@ namespace EpicGames.Horde
 			}
 			return await base.SendAsync(request, cancellationToken);
 		}
+
+		static bool TryGetAccessTokenFromEnvironment(Uri requestUri, out string? accessToken)
+		{
+			// Only use the token from the environment if the configured base address is missing or matches the one configured in the environment
+			string? hordeUrlEnvVar = Environment.GetEnvironmentVariable(HordeHttpClient.HordeUrlEnvVarName);
+			if (!String.IsNullOrEmpty(hordeUrlEnvVar))
+			{
+				Uri hordeUrl = new Uri(hordeUrlEnvVar);
+				if (String.Equals(requestUri.Host, hordeUrl.Host, StringComparison.OrdinalIgnoreCase))
+				{
+					string? hordeToken = Environment.GetEnvironmentVariable(HordeHttpClient.HordeTokenEnvVarName);
+					if (!String.IsNullOrEmpty(hordeToken))
+					{
+						accessToken = hordeToken;
+						return true;
+					}
+				}
+			}
+
+			accessToken = null;
+			return false;
+		}
 	}
 
 	/// <summary>
diff --git a/Engine/Source/Programs/Shared/EpicGames.Horde/HordeHttpClient.cs b/Engine/Source/Programs/Shared/EpicGames.Horde/HordeHttpClient.cs
--- a/Engine/Source/Programs/Shared/EpicGames.Horde/HordeHttpClient.cs
+++ b/Engine/Source/Programs/Shared/EpicGames.Horde/HordeHttpClient.cs
@@ -497,23 +497,13 @@ namespace EpicGames.Horde
 				options.Value.ConfigureHttpClient?.Invoke(httpClient);
 				configureClient(serviceProvider, httpClient);
 
-				// Only use the token from the environment if the configured base address is missing or matches the one configured in the environment
-				string? hordeUrlEnvVar = Environment.GetEnvironmentVariable(HordeHttpClient.HordeUrlEnvVarName);
-				if (!String.IsNullOrEmpty(hordeUrlEnvVar))
+				// If the server URL isn't set, take it from the environment
+				if (httpClient.BaseAddress == null)
 				{
-					Uri hordeUrl = new Uri(hordeUrlEnvVar);
-					if (httpClient.BaseAddress == null || String.Equals(httpClient.BaseAddress.Host, hordeUrl.Host, StringComparison.OrdinalIgnoreCase))
+					string? hordeUrlEnvVar = Environment.GetEnvironmentVariable(HordeHttpClient.HordeUrlEnvVarName);
+					if (!String.IsNullOrEmpty(hordeUrlEnvVar))
 					{
-						httpClient.BaseAddress ??= hordeUrl;
-
-						if (httpClient.DefaultRequestHeaders.Authorization == null)
-						{
-							string? hordeToken = Environment.GetEnvironmentVariable(HordeHttpClient.HordeTokenEnvVarName);
-							if (!String.IsNullOrEmpty(hordeToken))
-							{
-								httpClient.DefaultRequestHeaders.Authorization = new AuthenticationHeaderValue("Bearer", hordeToken);
-							}
-						}
+						httpClient.BaseAddress = new Uri(hordeUrlEnvVar);
 					}
 				}
 
