Horde: Fix static analysis warning.

diff --git a/Engine/Source/Programs/Shared/EpicGames.Horde/HordeHttpClient.cs b/Engine/Source/Programs/Shared/EpicGames.Horde/HordeHttpClient.cs
--- a/Engine/Source/Programs/Shared/EpicGames.Horde/HordeHttpClient.cs
+++ b/Engine/Source/Programs/Shared/EpicGames.Horde/HordeHttpClient.cs
@@ -4,7 +4,6 @@ using System;
 using System.Collections.Generic;
 using System.IO;
 using System.Net.Http;
-using System.Net.Http.Headers;
 using System.Net.Http.Json;
 using System.Text.Json;
 using System.Text.Json.Serialization;
