Horde: Include parent tool ids in summary responses.

diff --git a/Engine/Source/Programs/Horde/Horde.Server/Tools/ToolsController.cs b/Engine/Source/Programs/Horde/Horde.Server/Tools/ToolsController.cs
--- a/Engine/Source/Programs/Horde/Horde.Server/Tools/ToolsController.cs
+++ b/Engine/Source/Programs/Horde/Horde.Server/Tools/ToolsController.cs
@@ -214,7 +214,7 @@ namespace Horde.Server.Tools
 		static GetToolSummaryResponse CreateGetToolSummaryResponse(ITool tool)
 		{
 			IToolDeployment? deployment = (tool.Deployments.Count == 0) ? null : tool.Deployments[^1];
-			return new GetToolSummaryResponse(tool.Id, tool.Config.Name, tool.Config.Description, deployment?.Version, deployment?.Id, tool.Config.ShowInUgs);
+			return new GetToolSummaryResponse(tool.Id, tool.Config.Name, tool.Config.Description, tool.Config.ParentId, deployment?.Version, deployment?.Id, tool.Config.ShowInUgs);
 		}
 
 		/// <summary>
diff --git a/Engine/Source/Programs/Shared/EpicGames.Horde/Tools/ToolMessages.cs b/Engine/Source/Programs/Shared/EpicGames.Horde/Tools/ToolMessages.cs
--- a/Engine/Source/Programs/Shared/EpicGames.Horde/Tools/ToolMessages.cs
+++ b/Engine/Source/Programs/Shared/EpicGames.Horde/Tools/ToolMessages.cs
@@ -25,10 +25,11 @@ namespace EpicGames.Horde.Tools
 	/// <param name="Id">Unique identifier for the tool</param>
 	/// <param name="Name">Name of the tool</param>
 	/// <param name="Description">Description for the tool</param>
+	/// <param name="ParentId">The parent tool id</param>
 	/// <param name="Version">Version number of the current deployment of this tool</param>
 	/// <param name="DeploymentId">Identifier for the current deployment</param>
 	/// <param name="ShowInUgs">Whether to show this tool for download inside UGS</param>
-	public record class GetToolSummaryResponse(ToolId Id, string Name, string Description, string? Version, ToolDeploymentId? DeploymentId, bool ShowInUgs);
+	public record class GetToolSummaryResponse(ToolId Id, string Name, string Description, ToolId? ParentId, string? Version, ToolDeploymentId? DeploymentId, bool ShowInUgs);
 
 	/// <summary>
 	/// Response when querying all tools
