Horde: Enable dedupe for tool data bundles when building installer.

diff --git a/Engine/Source/Programs/Horde/Horde.Server/Commands/Generate/ToolDataCommand.cs b/Engine/Source/Programs/Horde/Horde.Server/Commands/Generate/ToolDataCommand.cs
--- a/Engine/Source/Programs/Horde/Horde.Server/Commands/Generate/ToolDataCommand.cs
+++ b/Engine/Source/Programs/Horde/Horde.Server/Commands/Generate/ToolDataCommand.cs
@@ -2,10 +2,13 @@
 
 using System;
 using System.ComponentModel;
+using System.IO;
 using System.Text.Json.Nodes;
+using System.Threading;
 using System.Threading.Tasks;
 using EpicGames.Core;
 using EpicGames.Horde.Storage;
+using EpicGames.Horde.Storage.Backends;
 using EpicGames.Horde.Storage.Bundles;
 using EpicGames.Horde.Storage.Nodes;
 using Horde.Server.Utilities;
@@ -45,20 +48,30 @@ namespace Horde.Server.Commands.Generate
 		public override async Task<int> ExecuteAsync(ILogger logger)
 		{
 			// Create the local agent bundle
-			DirectoryReference bundleDir = DirectoryReference.Combine(ServerDir, "Tools", Id);
+			DirectoryReference bundleDir = DirectoryReference.Combine(ServerDir, "Tools");
 			DirectoryReference.CreateDirectory(bundleDir);
 
-			RefName refName = new RefName("latest");
+			RefName refName = new RefName(Id);
 			await using (BundleCache bundleCache = new BundleCache())
 			{
 				using (IStorageClient client = BundleStorageClient.CreateFromDirectory(bundleDir, bundleCache, logger))
 				{
 					IBlobRef<DirectoryNode> dirNodeRef;
-					await using (IBlobWriter writer = client.CreateBlobWriter(refName))
+					await using (DedupeBlobWriter writer = client.CreateBlobWriter(refName).WithDedupe())
 					{
+						logger.LogInformation("Populating cache with existing refs...");
+						await PopulateCacheAsync(client, writer, bundleDir, CancellationToken.None);
+
+						logger.LogInformation("");
+						logger.LogInformation("Writing tool data for {ToolId}", refName);
 						DirectoryNode dirNode = new DirectoryNode();
 						await dirNode.AddFilesAsync(InputDir.ToDirectoryInfo(), writer);
 						dirNodeRef = await writer.WriteBlobAsync(dirNode);
+
+						logger.LogInformation("");
+						writer.GetStats().Print(logger);
+
+						logger.LogInformation("");
 					}
 					await client.WriteRefAsync(refName, dirNodeRef);
 				}
@@ -91,5 +104,22 @@ namespace Horde.Server.Commands.Generate
 
 			return 0;
 		}
+
+		static async Task PopulateCacheAsync(IStorageClient client, DedupeBlobWriter writer, DirectoryReference searchDir, CancellationToken cancellationToken)
+		{
+			foreach (RefName refName in FileStorageBackend.EnumerateRefs(searchDir))
+			{
+				IBlobRef? blobRef = await client.TryReadRefAsync(refName, cancellationToken: cancellationToken);
+				if (blobRef != null)
+				{
+					BlobData blobData = await blobRef.ReadBlobDataAsync(cancellationToken);
+					if (blobData.Type.Guid == DirectoryNode.BlobTypeGuid)
+					{
+						IBlobRef<DirectoryNode> directoryRef = BlobRef.Create<DirectoryNode>(blobRef);
+						await writer.AddToCacheAsync(directoryRef, cancellationToken);
+					}
+				}
+			}
+		}
 	}
 }
diff --git a/Engine/Source/Programs/Shared/EpicGames.Horde.Tests/BundleStoreTests.cs b/Engine/Source/Programs/Shared/EpicGames.Horde.Tests/BundleStoreTests.cs
--- a/Engine/Source/Programs/Shared/EpicGames.Horde.Tests/BundleStoreTests.cs
+++ b/Engine/Source/Programs/Shared/EpicGames.Horde.Tests/BundleStoreTests.cs
@@ -385,7 +385,7 @@ namespace EpicGames.Horde.Tests
 			// Generate a tree
 			DirectoryNode root;
 			{
-				await using DedupeStorageWriter writer = new DedupeStorageWriter(store.CreateBlobWriter(bundleOptions: new BundleOptions { MaxBlobSize = 1024 }));
+				await using DedupeBlobWriter writer = new DedupeBlobWriter(store.CreateBlobWriter(bundleOptions: new BundleOptions { MaxBlobSize = 1024 }));
 
 				ChunkingOptions options = new ChunkingOptions();
 				options.LeafOptions = new LeafChunkedDataNodeOptions(128, 256, 64 * 1024);
diff --git a/Engine/Source/Programs/Shared/EpicGames.Horde.Tests/BundleTests.cs b/Engine/Source/Programs/Shared/EpicGames.Horde.Tests/BundleTests.cs
--- a/Engine/Source/Programs/Shared/EpicGames.Horde.Tests/BundleTests.cs
+++ b/Engine/Source/Programs/Shared/EpicGames.Horde.Tests/BundleTests.cs
@@ -253,7 +253,7 @@ namespace EpicGames.Horde.Tests
 			BundleOptions bundleOptions = new BundleOptions();
 			bundleOptions.MaxBlobSize = 1;
 
-			await using (IBlobWriter writer = new DedupeStorageWriter(_storage.CreateBlobWriter(bundleOptions: bundleOptions)))
+			await using (IBlobWriter writer = new DedupeBlobWriter(_storage.CreateBlobWriter(bundleOptions: bundleOptions)))
 			{
 				DirectoryNode root = new DirectoryNode();
 				for (int idx = 1; idx <= 3; idx++)
diff --git a/Engine/Source/Programs/Shared/EpicGames.Horde/Storage/Backends/FileStorageBackend.cs b/Engine/Source/Programs/Shared/EpicGames.Horde/Storage/Backends/FileStorageBackend.cs
--- a/Engine/Source/Programs/Shared/EpicGames.Horde/Storage/Backends/FileStorageBackend.cs
+++ b/Engine/Source/Programs/Shared/EpicGames.Horde/Storage/Backends/FileStorageBackend.cs
@@ -1,6 +1,8 @@
 // Copyright Epic Games, Inc. All Rights Reserved.
 
 using System;
+using System.Collections;
+using System.Collections.Generic;
 using System.IO;
 using System.Threading;
 using System.Threading.Tasks;
@@ -94,6 +96,19 @@ namespace EpicGames.Horde.Storage.Backends
 
 		FileReference GetRefFile(RefName name) => FileReference.Combine(_rootDir, name.ToString() + ".ref");
 
+		/// <summary>
+		/// Enumerate all the refs in a file store
+		/// </summary>
+		/// <param name="rootDir">Root directory to search</param>
+		public static IEnumerable<RefName> EnumerateRefs(DirectoryReference rootDir)
+		{
+			foreach (FileReference refFile in DirectoryReference.EnumerateFiles(rootDir, "*.ref", SearchOption.AllDirectories))
+			{
+				string name = Path.GetFileNameWithoutExtension(refFile.MakeRelativeTo(rootDir)).Replace(Path.DirectorySeparatorChar, '/');
+				yield return new RefName(name);
+			}
+		}
+
 		#region Aliases
 
 		/// <inheritdoc/>
diff --git a/Engine/Source/Programs/Shared/EpicGames.Horde/Storage/BlobType.cs b/Engine/Source/Programs/Shared/EpicGames.Horde/Storage/BlobType.cs
--- a/Engine/Source/Programs/Shared/EpicGames.Horde/Storage/BlobType.cs
+++ b/Engine/Source/Programs/Shared/EpicGames.Horde/Storage/BlobType.cs
@@ -18,11 +18,6 @@ namespace EpicGames.Horde.Storage
 		/// </summary>
 		public const int NumBytes = 20;
 
-		/// <summary>
-		/// Blob type used to indicate a leaf node. Data is an opaque blob and has no references.
-		/// </summary>
-		public static BlobType Leaf { get; } = new BlobType("{1080B643-4A4D-8015-FBC2-BBAD818D6C58}", 1);
-
 		/// <summary>
 		/// Constructor
 		/// </summary>
diff --git a/Engine/Source/Programs/Shared/EpicGames.Horde/Storage/DedupeBlobWriter.cs b/Engine/Source/Programs/Shared/EpicGames.Horde/Storage/DedupeBlobWriter.cs
new file mode 100644
--- /dev/null
+++ b/Engine/Source/Programs/Shared/EpicGames.Horde/Storage/DedupeBlobWriter.cs
@@ -0,0 +1,239 @@
+// Copyright Epic Games, Inc. All Rights Reserved.
+
+using System;
+using System.Collections.Generic;
+using System.Diagnostics.CodeAnalysis;
+using System.Linq;
+using System.Threading;
+using System.Threading.Tasks;
+using EpicGames.Core;
+using EpicGames.Horde.Storage.Nodes;
+
+namespace EpicGames.Horde.Storage
+{
+	/// <summary>
+	/// Index of known nodes that can be used for deduplication.
+	/// </summary>
+	public sealed class DedupeBlobWriter : BlobWriter
+	{
+		record BlobKey(IoHash Hash, BlobType Type);
+
+		class DedupeCache
+		{
+			readonly int _maxKeys;
+			readonly Queue<BlobKey> _blobKeys = new Queue<BlobKey>();
+			readonly Dictionary<BlobKey, IBlobRef> _blobKeyToHandle = new Dictionary<BlobKey, IBlobRef>();
+
+			public DedupeCache(int maxKeys)
+			{
+				_maxKeys = maxKeys;
+				_blobKeys = new Queue<BlobKey>(maxKeys);
+				_blobKeyToHandle = new Dictionary<BlobKey, IBlobRef>(maxKeys);
+			}
+
+			internal void Add(BlobKey key, IBlobRef handle)
+			{
+				BlobKey? prevKey;
+				if (_blobKeys.Count == _maxKeys && _blobKeys.TryDequeue(out prevKey))
+				{
+					_blobKeyToHandle.Remove(prevKey);
+				}
+				_blobKeyToHandle.TryAdd(key, handle);
+			}
+
+			internal bool TryGetValue(BlobKey key, [NotNullWhen(true)] out IBlobRef? handle) => _blobKeyToHandle.TryGetValue(key, out handle);
+		}
+
+		class WrappedHandle : IBlobRef
+		{
+			public object _lockObject = new object();
+			public IBlobRef? _inner;
+
+			/// <inheritdoc/>
+			public IBlobHandle Innermost
+				=> _inner!.Innermost;
+
+			/// <inheritdoc/>
+			public IoHash Hash 
+				=> _inner!.Hash;
+
+			/// <inheritdoc/>
+			public bool TryGetLocator(out BlobLocator locator)
+				=> _inner!.TryGetLocator(out locator);
+
+			/// <inheritdoc/>
+			public ValueTask FlushAsync(CancellationToken cancellationToken)
+				=> _inner!.FlushAsync(cancellationToken);
+
+			/// <inheritdoc/>
+			public ValueTask<BlobData> ReadBlobDataAsync(CancellationToken cancellationToken = default)
+				=> _inner!.ReadBlobDataAsync(cancellationToken);
+
+			/// <inheritdoc/>
+			public override bool Equals(object? obj) => _inner is not null && obj is WrappedHandle other && _inner == other._inner;
+
+			/// <inheritdoc/>
+			public override int GetHashCode() => HashCode.Combine((_inner is null) ? 0 : _inner.GetHashCode(), 1);
+		}
+
+		/// <summary>
+		/// Default value for maximum number of keys
+		/// </summary>
+		public const int DefaultMaxKeys = 512 * 1024;
+
+		readonly BlobWriter _inner;
+		readonly DedupeCache _cache;
+
+		int _numHits;
+		long _totalHitsSize;
+		int _numMisses;
+		long _totalMissesSize;
+		int _numCacheAdds;
+
+		/// <summary>
+		/// Constructor
+		/// </summary>
+		/// <param name="inner"></param>
+		/// <param name="maxKeys"></param>
+		public DedupeBlobWriter(IBlobWriter inner, int maxKeys = DefaultMaxKeys)
+			: base(inner.Options)
+		{
+			_inner = (BlobWriter)inner;
+			_cache = new DedupeCache(maxKeys);
+		}
+
+		private DedupeBlobWriter(IBlobWriter inner, DedupeCache cache)
+			: base(inner.Options)
+		{
+			_inner = (BlobWriter)inner;
+			_cache = cache;
+		}
+
+		/// <inheritdoc/>
+		public override ValueTask DisposeAsync() => _inner.DisposeAsync();
+
+		/// <inheritdoc/>
+		public override Task FlushAsync(CancellationToken cancellationToken = default) => _inner.FlushAsync(cancellationToken);
+
+		/// <inheritdoc/>
+		public override IBlobWriter Fork() => new DedupeBlobWriter(_inner.Fork(), _cache);
+
+		/// <inheritdoc/>
+		public override Memory<byte> GetOutputBuffer(int usedSize, int desiredSize) => _inner.GetOutputBuffer(usedSize, desiredSize);
+
+		/// <summary>
+		/// Add a blob to the cache
+		/// </summary>
+		/// <param name="type">Type of the blob</param>
+		/// <param name="blobRef">Reference to the blob data</param>
+		public void AddToCache(BlobType type, IBlobRef blobRef)
+		{
+			_cache.Add(new BlobKey(blobRef.Hash, type), blobRef);
+			Interlocked.Increment(ref _numCacheAdds);
+		}
+
+		/// <inheritdoc/>
+		public override async ValueTask<IBlobRef> WriteBlobAsync(BlobType type, int size, IReadOnlyList<IBlobHandle> imports, IReadOnlyList<AliasInfo> aliases, CancellationToken cancellationToken = default)
+		{
+			ReadOnlyMemory<byte> data = _inner.GetOutputBuffer(size, size).Slice(0, size);
+			IoHash hash = IoHash.Compute(data.Span);
+			BlobKey key = new BlobKey(hash, type);
+
+			WrappedHandle? wrappedHandle;
+			lock (_cache)
+			{
+				IBlobRef? handle;
+				if (_cache.TryGetValue(key, out handle))
+				{
+					_numHits++;
+					_totalHitsSize += size;
+					return handle;
+				}
+				else
+				{
+					_numMisses++;
+					_totalMissesSize += size;
+
+					wrappedHandle = new WrappedHandle();
+					_cache.Add(key, wrappedHandle);
+				}
+			}
+
+			wrappedHandle._inner = await _inner.WriteBlobAsync(type, size, imports.ConvertAll(x => x.Innermost), aliases, cancellationToken);
+			return wrappedHandle;
+		}
+
+		/// <summary>
+		/// Gets stats for the copy operation
+		/// </summary>
+		public StorageStats GetStats()
+		{
+			StorageStats stats = new StorageStats();
+			stats.Add("Cache hits", _numHits);
+			stats.Add("Cache hits size", _totalHitsSize);
+			stats.Add("Cache misses", _numMisses);
+			stats.Add("Cache misses size", _totalMissesSize);
+			stats.Add("Cache adds", _numCacheAdds);
+			return stats;
+		}
+	}
+
+	/// <summary>
+	/// Extension methods for <see cref="IBlobWriter"/>
+	/// </summary>
+	public static class DedupeBlobWriterExtensions
+	{
+		/// <summary>
+		/// Wraps a <see cref="IBlobWriter"/> with a <see cref="DedupeBlobWriter"/>
+		/// </summary>
+		public static DedupeBlobWriter WithDedupe(this IBlobWriter writer, int maxKeys = DedupeBlobWriter.DefaultMaxKeys) => new DedupeBlobWriter(writer, maxKeys);
+
+		/// <summary>
+		/// Adds a directory tree to the cache
+		/// </summary>
+		/// <param name="dedupeWriter">Dedupe writer to operate on</param>
+		/// <param name="directoryNodeRef">Reference to the directory to add</param>
+		/// <param name="cancellationToken">Cancellation token for the operation</param>
+		public static async Task AddToCacheAsync(this DedupeBlobWriter dedupeWriter, IBlobRef<DirectoryNode> directoryNodeRef, CancellationToken cancellationToken = default)
+		{
+			using BlobData blobData = await directoryNodeRef.ReadBlobDataAsync(cancellationToken);
+			dedupeWriter.AddToCache(blobData.Type, directoryNodeRef);
+
+			DirectoryNode directoryNode = BlobSerializer.Deserialize<DirectoryNode>(blobData);
+			foreach (DirectoryEntry directoryEntry in directoryNode.Directories)
+			{
+				await AddToCacheAsync(dedupeWriter, directoryEntry.Handle, cancellationToken);
+			}
+			foreach (FileEntry fileEntry in directoryNode.Files)
+			{
+				await AddToCacheAsync(dedupeWriter, fileEntry.Target, cancellationToken);
+			}
+		}
+
+		/// <summary>
+		/// Adds a chunked data stream to the cache
+		/// </summary>
+		/// <param name="dedupeWriter">Dedupe writer to operate on</param>
+		/// <param name="dataNodeRef">Reference to a data node in the stream</param>
+		/// <param name="cancellationToken">Cancellation token for the operation</param>
+		public static async Task AddToCacheAsync(this DedupeBlobWriter dedupeWriter, ChunkedDataNodeRef dataNodeRef, CancellationToken cancellationToken = default)
+		{
+			if (dataNodeRef.Type == ChunkedDataNodeType.Leaf)
+			{
+				BlobType blobType = LeafChunkedDataNodeConverter.BlobType;
+				dedupeWriter.AddToCache(blobType, dataNodeRef.Handle);
+			}
+			else if (dataNodeRef.Type == ChunkedDataNodeType.Interior)
+			{
+				using BlobData blobData = await dataNodeRef.Handle.ReadBlobDataAsync(cancellationToken);
+				dedupeWriter.AddToCache(blobData.Type, dataNodeRef.Handle);
+
+				InteriorChunkedDataNode interiorNode = BlobSerializer.Deserialize<InteriorChunkedDataNode>(blobData);
+				foreach (ChunkedDataNodeRef childNodeRef in interiorNode.Children)
+				{
+					await AddToCacheAsync(dedupeWriter, childNodeRef, cancellationToken);
+				}
+			}
+		}
+	}
+}
diff --git a/Engine/Source/Programs/Shared/EpicGames.Horde/Storage/IBlobHandle.cs b/Engine/Source/Programs/Shared/EpicGames.Horde/Storage/IBlobHandle.cs
--- a/Engine/Source/Programs/Shared/EpicGames.Horde/Storage/IBlobHandle.cs
+++ b/Engine/Source/Programs/Shared/EpicGames.Horde/Storage/IBlobHandle.cs
@@ -111,6 +111,16 @@ namespace EpicGames.Horde.Storage
 		public static IBlobRef Create(IoHash hash, IBlobHandle handle)
 			=> new BlobRefImpl(hash, handle);
 
+		/// <summary>
+		/// Create a typed blob handle
+		/// </summary>
+		/// <typeparam name="T"></typeparam>
+		/// <param name="blobRef">Existing blob reference</param>
+		/// <param name="options">Options for deserializing the target blob</param>
+		/// <returns>Handle to the blob</returns>
+		public static IBlobRef<T> Create<T>(IBlobRef blobRef, BlobSerializerOptions? options = null)
+			=> new BlobRefImpl<T>(blobRef.Hash, blobRef, options ?? BlobSerializerOptions.Default);
+
 		/// <summary>
 		/// Create a typed blob handle
 		/// </summary>
@@ -119,7 +129,7 @@ namespace EpicGames.Horde.Storage
 		/// <param name="handle">Imported blob interface</param>
 		/// <param name="options">Options for deserializing the target blob</param>
 		/// <returns>Handle to the blob</returns>
-		public static IBlobRef<T> Create<T>(IoHash hash, IBlobHandle handle, BlobSerializerOptions? options)
+		public static IBlobRef<T> Create<T>(IoHash hash, IBlobHandle handle, BlobSerializerOptions? options = null)
 			=> new BlobRefImpl<T>(hash, handle, options ?? BlobSerializerOptions.Default);
 	}
 
diff --git a/Engine/Source/Programs/Shared/EpicGames.Horde/Storage/IStorageWriter.cs b/Engine/Source/Programs/Shared/EpicGames.Horde/Storage/IStorageWriter.cs
deleted file mode 100644
