Improved error message when identifiers contain more than one namespace.

diff --git a/Engine/Source/Programs/Shared/EpicGames.UHT/Utils/UhtSession.cs b/Engine/Source/Programs/Shared/EpicGames.UHT/Utils/UhtSession.cs
--- a/Engine/Source/Programs/Shared/EpicGames.UHT/Utils/UhtSession.cs
+++ b/Engine/Source/Programs/Shared/EpicGames.UHT/Utils/UhtSession.cs
@@ -1275,7 +1275,9 @@ namespace EpicGames.UHT.Utils
 			{
 				if (messageSite != null)
 				{
-					messageSite.LogError(lineNumber, "UnrealHeaderTool only supports C++ identifiers of two or less identifiers");
+					//messageSite.LogError(lineNumber, "UnrealHeaderTool only supports C++ identifiers of two or less identifiers");
+					string fullIdentifier = identifiers.Join("::");
+					FindTypeError(messageSite, lineNumber, options, fullIdentifier);
 					return null;
 				}
 			}
@@ -1319,7 +1321,9 @@ namespace EpicGames.UHT.Utils
 			{
 				if (messageSite != null)
 				{
-					messageSite.LogError(lineNumber, "UnrealHeaderTool only supports C++ identifiers of two or less identifiers");
+					//messageSite.LogError(lineNumber, "UnrealHeaderTool only supports C++ identifiers of two or less identifiers");
+					string fullIdentifier = String.Join("::", identifiers);
+					FindTypeError(messageSite, lineNumber, options, fullIdentifier);
 					return null;
 				}
 			}
