Horde: Create and resize output files for extracting bundles under a lock, to ensure it only happens once.

Speculative fix for intermittent error extracting files on Linux where multiple threads may be opening/resizing output files at once.

diff --git a/Engine/Source/Programs/Shared/EpicGames.Horde/Storage/Nodes/DirectoryNode.Extract.cs b/Engine/Source/Programs/Shared/EpicGames.Horde/Storage/Nodes/DirectoryNode.Extract.cs
--- a/Engine/Source/Programs/Shared/EpicGames.Horde/Storage/Nodes/DirectoryNode.Extract.cs
+++ b/Engine/Source/Programs/Shared/EpicGames.Horde/Storage/Nodes/DirectoryNode.Extract.cs
@@ -141,17 +141,105 @@ namespace EpicGames.Horde.Storage.Nodes
 		/// <param name="cancellationToken"></param>
 		public static Task CopyToDirectoryAsync(this DirectoryNode directoryNode, DirectoryInfo directoryInfo, ILogger logger, CancellationToken cancellationToken) => CopyToDirectoryAsync(directoryNode, directoryInfo, null, logger, cancellationToken);
 
-		class RemainingChunkCounter
+		class OutputFile
 		{
-			public int _value;
+			public string Path { get; }
+			public FileInfo FileInfo { get; }
+			public FileEntry FileEntry { get; }
 
-			public int Increment() => Interlocked.Increment(ref _value);
-			public int Decrement() => Interlocked.Decrement(ref _value);
+			bool _createdFile;
+			int _remainingChunks;
+
+			public OutputFile(string path, FileInfo fileInfo, FileEntry fileEntry)
+			{
+				Path = path;
+				FileInfo = fileInfo;
+				FileEntry = fileEntry;
+			}
+
+			public int IncrementRemaining() => Interlocked.Increment(ref _remainingChunks);
+			public int DecrementRemaining() => Interlocked.Decrement(ref _remainingChunks);
+
+			public FileStream OpenStream()
+			{
+				lock (FileEntry)
+				{
+					if (!_createdFile)
+					{
+						if (FileInfo.Exists)
+						{
+							if ((FileInfo.Attributes & FileAttributes.ReadOnly) != 0)
+							{
+								FileInfo.Attributes &= ~FileAttributes.ReadOnly;
+							}
+							if (FileInfo.LinkTarget != null)
+							{
+								FileInfo.Delete();
+							}
+						}
+						else
+						{
+							FileInfo.Directory?.Create();
+						}
+					}
+
+					FileStream? stream = null;
+					try
+					{
+						stream = FileInfo.Open(FileMode.OpenOrCreate, FileAccess.ReadWrite, FileShare.ReadWrite);
+						if (!_createdFile)
+						{
+							stream.SetLength(FileEntry.Length);
+							_createdFile = true;
+						}
+						return stream;
+					}
+					catch
+					{
+						stream?.Dispose();
+						throw;
+					}
+				}
+			}
 		}
 
-		record class OutputDir(string Path, DirectoryNode Node);
-		record class OutputFile(OutputDir Directory, FileEntry FileEntry, RemainingChunkCounter RemainingChunks);
-		record class OutputChunk(OutputFile File, long Offset, long Length, IBlobHandle? Handle);
+		record class OutputChunk(OutputFile File, long Offset, long Length, IBlobHandle Handle);
+
+		// Writes output chunks to a channel. Buffers one chunk until FlushAsync() is called to ensure
+		// the remaining chunk reference count doesn't reach zero until the last chunk has been processed.
+		class OutputChunkWriter
+		{
+			public OutputFile OutputFile { get; }
+
+			readonly ChannelWriter<OutputChunk> _chunkWriter;
+			OutputChunk? _bufferedChunk;
+
+			public OutputChunkWriter(OutputFile file, ChannelWriter<OutputChunk> chunkWriter)
+			{
+				OutputFile = file;
+				_chunkWriter = chunkWriter;
+			}
+
+			public async Task WriteAsync(long offset, long length, IBlobHandle handle, CancellationToken cancellationToken)
+			{
+				if (_bufferedChunk != null)
+				{
+					await _chunkWriter.WriteAsync(_bufferedChunk, cancellationToken);
+				}
+
+				OutputFile.IncrementRemaining();
+				_bufferedChunk = new OutputChunk(OutputFile, offset, length, handle);
+			}
+
+			public async Task FlushAsync(CancellationToken cancellationToken)
+			{
+				if (_bufferedChunk != null)
+				{
+					await _chunkWriter.WriteAsync(_bufferedChunk, cancellationToken);
+					_bufferedChunk = null;
+				}
+			}
+		}
 
 #pragma warning disable IDE0060
 		static void TraceBlobRead(string type, string path, IBlobHandle handle, ILogger logger)
@@ -201,51 +289,33 @@ namespace EpicGames.Horde.Storage.Nodes
 				}
 
 				List<Task> tasks = new List<Task>();
-				tasks.Add(RunBackgroundTask(ctx => FindOutputChunksRootAsync(directoryNode, chunks.Writer, logger, ctx)));
+				tasks.Add(RunBackgroundTask(ctx => FindOutputChunksRootAsync(directoryInfo, directoryNode, chunks.Writer, logger, ctx)));
 				for (int idx = 0; idx < numTasks; idx++)
 				{
-					tasks.Add(RunBackgroundTask(ctx => ExtractAsync(chunks.Reader, new DirectoryReference(directoryInfo), copyStats, logger, ctx)));
+					tasks.Add(RunBackgroundTask(ctx => ExtractAsync(chunks.Reader, copyStats, logger, ctx)));
 				}
 
 				await Task.WhenAll(tasks);
 			}
 		}
 
-		static async Task ExtractAsync(ChannelReader<OutputChunk> chunkReader, DirectoryReference baseDir, CopyStats? copyStats, ILogger logger, CancellationToken cancellationToken)
+		static async Task ExtractAsync(ChannelReader<OutputChunk> chunkReader, CopyStats? copyStats, ILogger logger, CancellationToken cancellationToken)
 		{
 			OutputChunk? chunk = await ReadNextChunkAsync(chunkReader, cancellationToken);
 			while (chunk != null)
 			{
 				// Open the file for the current chunk
 				OutputFile file = chunk.File;
-
-				FileReference location = FileReference.Combine(baseDir, file.Directory.Path, file.FileEntry.Name);
 				try
 				{
-					DirectoryReference.CreateDirectory(location.Directory);
-
-					FileInfo fileInfo = location.ToFileInfo();
-					if (fileInfo.Exists)
-					{
-						if ((fileInfo.Attributes & FileAttributes.ReadOnly) != 0)
-						{
-							fileInfo.Attributes &= ~FileAttributes.ReadOnly;
-						}
-						if (fileInfo.LinkTarget != null)
-						{
-							fileInfo.Delete();
-						}
-					}
-
 					int remainingChunks = 0;
-					await using (FileStream stream = fileInfo.Open(FileMode.OpenOrCreate, FileAccess.ReadWrite, FileShare.ReadWrite))
+					await using (FileStream stream = file.OpenStream())
 					{
-						stream.SetLength(file.FileEntry.Length);
 						if (file.FileEntry.Length == 0)
 						{
 							// If this file is empty, don't write anything and just move to the next chunk
 							chunk = await ReadNextChunkAsync(chunkReader, cancellationToken);
-							remainingChunks = file.RemainingChunks.Decrement();
+							remainingChunks = file.DecrementRemaining();
 						}
 						else
 						{
@@ -256,17 +326,14 @@ namespace EpicGames.Horde.Storage.Nodes
 							while (chunk != null && chunk.File == file)
 							{
 								// Write this chunk
-								if (chunk.Handle != null)
+								using (BlobData data = await chunk.Handle.ReadBlobDataAsync(cancellationToken))
 								{
-									using (BlobData data = await chunk.Handle.ReadBlobDataAsync(cancellationToken))
-									{
-										TraceBlobRead("Leaf", CombinePaths(chunk.File.Directory.Path, chunk.File.FileEntry.Name), chunk.Handle, logger);
-										data.Data.CopyTo(memoryMappedView!.GetMemory(chunk.Offset, data.Data.Length));
-									}
+									TraceBlobRead("Leaf", chunk.File.Path, chunk.Handle, logger);
+									data.Data.CopyTo(memoryMappedView!.GetMemory(chunk.Offset, data.Data.Length));
 								}
 
 								// Update the stats
-								remainingChunks = file.RemainingChunks.Decrement();
+								remainingChunks = file.DecrementRemaining();
 								copyStats?.Update(0, chunk.Length);
 
 								// Read the next chunk
@@ -278,7 +345,7 @@ namespace EpicGames.Horde.Storage.Nodes
 					// Set correct permissions on the output file
 					if (remainingChunks == 0)
 					{
-						FileEntry.SetPermissions(fileInfo, file.FileEntry.Flags);
+						FileEntry.SetPermissions(file.FileInfo!, file.FileEntry.Flags);
 						copyStats?.Update(1, 0);
 					}
 				}
@@ -288,7 +355,7 @@ namespace EpicGames.Horde.Storage.Nodes
 				}
 				catch (Exception ex)
 				{
-					throw new StorageException($"Unable to extract {location}: {ex.Message}", ex); 
+					throw new StorageException($"Unable to extract {file?.FileInfo?.FullName}: {ex.Message}", ex); 
 				}
 			}
 		}
@@ -308,53 +375,45 @@ namespace EpicGames.Horde.Storage.Nodes
 			}
 		}
 
-		static async Task FindOutputChunksRootAsync(DirectoryNode node, ChannelWriter<OutputChunk> chunks, ILogger logger, CancellationToken cancellationToken)
+		static async Task FindOutputChunksRootAsync(DirectoryInfo rootDir, DirectoryNode node, ChannelWriter<OutputChunk> chunks, ILogger logger, CancellationToken cancellationToken)
 		{
-			await FindOutputChunksAsync("", node, chunks, logger, cancellationToken);
+			await FindOutputChunksForDirectoryAsync(rootDir, "", node, chunks, logger, cancellationToken);
 			chunks.Complete();
 		}
 
-		static async Task FindOutputChunksAsync(string path, DirectoryNode node, ChannelWriter<OutputChunk> chunks, ILogger logger, CancellationToken cancellationToken)
+		static async Task FindOutputChunksForDirectoryAsync(DirectoryInfo rootDir, string path, DirectoryNode node, ChannelWriter<OutputChunk> chunks, ILogger logger, CancellationToken cancellationToken)
 		{
-			OutputDir outputDir = new OutputDir(path, node);
-
 			foreach (FileEntry fileEntry in node.Files)
 			{
-				OutputFile outputFile = new OutputFile(outputDir, fileEntry, new RemainingChunkCounter());
-
-				// Ensure that no worker finalizes the file until we've finished adding chunks
-				outputFile.RemainingChunks.Increment();
+				string filePath = CombinePaths(path, fileEntry.Name);
+				FileInfo fileInfo = new FileInfo(Path.Combine(rootDir.FullName, filePath));
+				OutputFile outputFile = new OutputFile(filePath, fileInfo, fileEntry);
 
-				if (fileEntry.Target.Type == ChunkedDataNodeType.Leaf && fileEntry.Target.Length >= 0)
-				{
-					// If there's only one chunk for this file, queue the chunk and let the worker finalize the file in one operation
-					await chunks.WriteAsync(new OutputChunk(outputFile, 0, fileEntry.Target.Length, fileEntry.Target.Handle), cancellationToken);
-				}
-				else
-				{
-					await FindOutputChunksAsync(outputFile, 0, fileEntry.Target, chunks, logger, cancellationToken);
-					await chunks.WriteAsync(new OutputChunk(outputFile, 0, 0, null), cancellationToken);
-				}
+				await FindOutputChunksForFileAsync(outputFile, chunks, logger, cancellationToken);
 			}
 
 			foreach (DirectoryEntry directoryEntry in node.Directories)
 			{
-				DirectoryNode subDirectoryNode = await directoryEntry.Handle.ReadBlobAsync(cancellationToken);
-
-				string subPath = CombinePaths(outputDir.Path, directoryEntry.Name);
+				string subPath = CombinePaths(path, directoryEntry.Name);
 				TraceBlobRead("Directory", subPath, directoryEntry.Handle, logger);
+				DirectoryNode subDirectoryNode = await directoryEntry.Handle.ReadBlobAsync(cancellationToken);
 
-				await FindOutputChunksAsync(subPath, subDirectoryNode, chunks, logger, cancellationToken);
+				await FindOutputChunksForDirectoryAsync(rootDir, subPath, subDirectoryNode, chunks, logger, cancellationToken);
 			}
 		}
 
-		static async Task<long> FindOutputChunksAsync(OutputFile outputFile, long offset, ChunkedDataNodeRef dataRef, ChannelWriter<OutputChunk> chunks, ILogger logger, CancellationToken cancellationToken)
+		static async Task FindOutputChunksForFileAsync(OutputFile outputFile, ChannelWriter<OutputChunk> chunks, ILogger logger, CancellationToken cancellationToken)
+		{
+			OutputChunkWriter outputWriter = new OutputChunkWriter(outputFile, chunks);
+			await FindOutputChunksAsync(outputWriter, 0, outputFile.FileEntry.Target, logger, cancellationToken);
+			await outputWriter.FlushAsync(cancellationToken);
+		}
+
+		static async Task<long> FindOutputChunksAsync(OutputChunkWriter chunkWriter, long offset, ChunkedDataNodeRef dataRef, ILogger logger, CancellationToken cancellationToken)
 		{
 			if (dataRef.Type == ChunkedDataNodeType.Leaf)
 			{
-				await chunks.WriteAsync(new OutputChunk(outputFile, offset, dataRef.Length, dataRef.Handle), cancellationToken);
-				outputFile.RemainingChunks.Increment();
-
+				await chunkWriter.WriteAsync(offset, dataRef.Length, dataRef.Handle, cancellationToken);
 				if (dataRef.Length < 0)
 				{
 					// Backwards compatibility hack for v2 format
@@ -366,12 +425,11 @@ namespace EpicGames.Horde.Storage.Nodes
 			else
 			{
 				using BlobData data = await dataRef.Handle.ReadBlobDataAsync(cancellationToken);
-				TraceBlobRead("Interior", CombinePaths(outputFile.Directory.Path, outputFile.FileEntry.Name), dataRef.Handle, logger);
+				TraceBlobRead("Interior", chunkWriter.OutputFile.Path, dataRef.Handle, logger);
 
 				if (data.Type.Guid == LeafChunkedDataNodeConverter.BlobType.Guid)
 				{
-					await chunks.WriteAsync(new OutputChunk(outputFile, offset, dataRef.Length, dataRef.Handle), cancellationToken);
-					outputFile.RemainingChunks.Increment();
+					await chunkWriter.WriteAsync(offset, dataRef.Length, dataRef.Handle, cancellationToken);
 					return data.Data.Length;
 				}
 				else
@@ -381,7 +439,7 @@ namespace EpicGames.Horde.Storage.Nodes
 					InteriorChunkedDataNode interiorNode = BlobSerializer.Deserialize<InteriorChunkedDataNode>(data);
 					foreach (ChunkedDataNodeRef childRef in interiorNode.Children)
 					{
-						length += await FindOutputChunksAsync(outputFile, offset + length, childRef, chunks, logger, cancellationToken);
+						length += await FindOutputChunksAsync(chunkWriter, offset + length, childRef, logger, cancellationToken);
 					}
 
 					return length;
