Horde: Group read requests into batches before passing to worker threads.

diff --git a/Engine/Source/Programs/Shared/EpicGames.Horde/Storage/Nodes/DirectoryNode.Extract.cs b/Engine/Source/Programs/Shared/EpicGames.Horde/Storage/Nodes/DirectoryNode.Extract.cs
--- a/Engine/Source/Programs/Shared/EpicGames.Horde/Storage/Nodes/DirectoryNode.Extract.cs
+++ b/Engine/Source/Programs/Shared/EpicGames.Horde/Storage/Nodes/DirectoryNode.Extract.cs
@@ -205,6 +205,8 @@ namespace EpicGames.Horde.Storage.Nodes
 
 		record class OutputChunk(OutputFile File, long Offset, long Length, IBlobHandle Handle);
 
+		record class OutputBatch(List<OutputChunk> Chunks);
+
 		// Writes output chunks to a channel. Buffers one chunk until FlushAsync() is called to ensure
 		// the remaining chunk reference count doesn't reach zero until the last chunk has been processed.
 		class OutputChunkWriter
@@ -267,7 +269,6 @@ namespace EpicGames.Horde.Storage.Nodes
 				copyStats = new CopyStats(progress);
 			}
 
-			Channel<OutputChunk> chunks = Channel.CreateBounded<OutputChunk>(new BoundedChannelOptions(128 * 1024) { FullMode = BoundedChannelFullMode.Wait });
 			using (CancellationTokenSource cancellationSource = CancellationTokenSource.CreateLinkedTokenSource(cancellationToken))
 			{
 				// Helper method to run a background task and set a cancellation source on error
@@ -289,89 +290,98 @@ namespace EpicGames.Horde.Storage.Nodes
 				}
 
 				List<Task> tasks = new List<Task>();
+
+				Channel<OutputChunk> chunks = Channel.CreateBounded<OutputChunk>(new BoundedChannelOptions(128 * 1024) { FullMode = BoundedChannelFullMode.Wait });
 				tasks.Add(RunBackgroundTask(ctx => FindOutputChunksRootAsync(directoryInfo, directoryNode, chunks.Writer, logger, ctx)));
+
+				Channel<OutputBatch> batches = Channel.CreateUnbounded<OutputBatch>();
+				tasks.Add(RunBackgroundTask(ctx => BatchReadRequestsAsync(chunks.Reader, batches.Writer, ctx)));
+
 				for (int idx = 0; idx < numTasks; idx++)
 				{
-					tasks.Add(RunBackgroundTask(ctx => ExtractAsync(chunks.Reader, copyStats, logger, ctx)));
+					tasks.Add(RunBackgroundTask(ctx => ExtractAsync(batches.Reader, copyStats, logger, ctx)));
 				}
 
 				await Task.WhenAll(tasks);
 			}
 		}
 
-		static async Task ExtractAsync(ChannelReader<OutputChunk> chunkReader, CopyStats? copyStats, ILogger logger, CancellationToken cancellationToken)
+		static async Task ExtractAsync(ChannelReader<OutputBatch> batchReader, CopyStats? copyStats, ILogger logger, CancellationToken cancellationToken)
 		{
-			OutputChunk? chunk = await ReadNextChunkAsync(chunkReader, cancellationToken);
-			while (chunk != null)
+			while (await batchReader.WaitToReadAsync(cancellationToken))
 			{
-				// Open the file for the current chunk
-				OutputFile file = chunk.File;
-				try
+				OutputBatch? batch;
+				while (batchReader.TryRead(out batch))
 				{
-					int remainingChunks = 0;
-					await using (FileStream stream = file.OpenStream())
+					for (int chunkIdx = 0; chunkIdx < batch.Chunks.Count; )
 					{
-						if (file.FileEntry.Length == 0)
+						OutputFile file = batch.Chunks[chunkIdx].File;
+
+						int maxChunkIdx = chunkIdx + 1;
+						while (maxChunkIdx < batch.Chunks.Count && batch.Chunks[maxChunkIdx].File == file)
 						{
-							// If this file is empty, don't write anything and just move to the next chunk
-							chunk = await ReadNextChunkAsync(chunkReader, cancellationToken);
-							remainingChunks = file.DecrementRemaining();
+							maxChunkIdx++;
 						}
-						else
-						{
-							// Process as many chunks as we can for this file
-							using MemoryMappedFile memoryMappedFile = MemoryMappedFile.CreateFromFile(stream, null, file.FileEntry.Length, MemoryMappedFileAccess.ReadWrite, HandleInheritability.None, false);
-							using MemoryMappedView memoryMappedView = new MemoryMappedView(memoryMappedFile, 0, file.FileEntry.Length);
 
-							while (chunk != null && chunk.File == file)
-							{
-								// Write this chunk
-								using (BlobData data = await chunk.Handle.ReadBlobDataAsync(cancellationToken))
-								{
-									TraceBlobRead("Leaf", chunk.File.Path, chunk.Handle, logger);
-									data.Data.CopyTo(memoryMappedView!.GetMemory(chunk.Offset, data.Data.Length));
-								}
-
-								// Update the stats
-								remainingChunks = file.DecrementRemaining();
-								copyStats?.Update(0, chunk.Length);
-
-								// Read the next chunk
-								chunk = await ReadNextChunkAsync(chunkReader, cancellationToken);
-							}
+						try
+						{
+							await ExtractChunksToFileAsync(file, batch.Chunks.Slice(chunkIdx, maxChunkIdx - chunkIdx), copyStats, logger, cancellationToken);
+						}
+						catch (OperationCanceledException)
+						{
+							throw;
+						}
+						catch (Exception ex)
+						{
+							throw new StorageException($"Unable to extract {file?.FileInfo?.FullName}: {ex.Message}", ex);
 						}
-					}
 
-					// Set correct permissions on the output file
-					if (remainingChunks == 0)
-					{
-						FileEntry.SetPermissions(file.FileInfo!, file.FileEntry.Flags);
-						copyStats?.Update(1, 0);
+						chunkIdx = maxChunkIdx;
 					}
 				}
-				catch (OperationCanceledException)
-				{
-					throw;
-				}
-				catch (Exception ex)
-				{
-					throw new StorageException($"Unable to extract {file?.FileInfo?.FullName}: {ex.Message}", ex); 
-				}
 			}
 		}
 
-		static async ValueTask<OutputChunk?> ReadNextChunkAsync(ChannelReader<OutputChunk> chunkReader, CancellationToken cancellationToken)
+		static async Task ExtractChunksToFileAsync(OutputFile file, ListSegment<OutputChunk> chunks, CopyStats? copyStats, ILogger logger, CancellationToken cancellationToken)
 		{
-			await chunkReader.WaitToReadAsync(cancellationToken);
-
-			OutputChunk? chunk;
-			if (chunkReader.TryRead(out chunk))
+			// Open the file for the current chunk
+			int remainingChunks = 0;
+			await using (FileStream stream = file.OpenStream())
 			{
-				return chunk;
+				if (file.FileEntry.Length == 0)
+				{
+					// If this file is empty, don't write anything and just move to the next chunk
+					remainingChunks = file.DecrementRemaining();
+				}
+				else
+				{
+					// Process as many chunks as we can for this file
+					using MemoryMappedFile memoryMappedFile = MemoryMappedFile.CreateFromFile(stream, null, file.FileEntry.Length, MemoryMappedFileAccess.ReadWrite, HandleInheritability.None, false);
+					using MemoryMappedView memoryMappedView = new MemoryMappedView(memoryMappedFile, 0, file.FileEntry.Length);
+
+					for(int chunkIdx = 0; chunkIdx < chunks.Count; chunkIdx++)
+					{
+						OutputChunk chunk = chunks[chunkIdx];
+
+						// Write this chunk
+						using (BlobData data = await chunk.Handle.ReadBlobDataAsync(cancellationToken))
+						{
+							TraceBlobRead("Leaf", chunk.File.Path, chunk.Handle, logger);
+							data.Data.CopyTo(memoryMappedView!.GetMemory(chunk.Offset, data.Data.Length));
+						}
+
+						// Update the stats
+						remainingChunks = file.DecrementRemaining();
+						copyStats?.Update(0, chunk.Length);
+					}
+				}
 			}
-			else
+
+			// Set correct permissions on the output file
+			if (remainingChunks == 0)
 			{
-				return null;
+				FileEntry.SetPermissions(file.FileInfo!, file.FileEntry.Flags);
+				copyStats?.Update(1, 0);
 			}
 		}
 
@@ -458,6 +468,28 @@ namespace EpicGames.Horde.Storage.Nodes
 				return nextPath;
 			}
 		}
+
+		static async Task BatchReadRequestsAsync(ChannelReader<OutputChunk> chunkReader, ChannelWriter<OutputBatch> batchWriter, CancellationToken cancellationToken)
+		{
+			List<OutputChunk> batch = new List<OutputChunk>();
+			while (await chunkReader.WaitToReadAsync(cancellationToken))
+			{
+				OutputChunk? chunk;
+				while (chunkReader.TryRead(out chunk))
+				{
+					if (batch.Count > 0 && chunk.File != batch[0].File)
+					{
+						await batchWriter.WriteAsync(new OutputBatch(batch), cancellationToken);
+						batch = new List<OutputChunk>();
+					}
+					batch.Add(chunk);
+				}
+			}
+			if (batch.Count > 0)
+			{
+				await batchWriter.WriteAsync(new OutputBatch(batch), cancellationToken);
+			}
+		}
 	}
 }
 
