Horde: Add steps to the bundle extract pipleine to buffer and prefetch bundles before writing to output files.

diff --git a/Engine/Source/Programs/Shared/EpicGames.Horde/Storage/Bundles/V2/PacketHandle.cs b/Engine/Source/Programs/Shared/EpicGames.Horde/Storage/Bundles/V2/PacketHandle.cs
--- a/Engine/Source/Programs/Shared/EpicGames.Horde/Storage/Bundles/V2/PacketHandle.cs
+++ b/Engine/Source/Programs/Shared/EpicGames.Horde/Storage/Bundles/V2/PacketHandle.cs
@@ -66,6 +66,11 @@ namespace EpicGames.Horde.Storage.Bundles.V2
 		/// <inheritdoc/>
 		public override BundleHandle Bundle => _outer;
 
+		/// <summary>
+		/// Offset of the packet within the bundle
+		/// </summary>
+		public int PacketOffset => _packetOffset;
+
 		/// <summary>
 		/// Constructor
 		/// </summary>
diff --git a/Engine/Source/Programs/Shared/EpicGames.Horde/Storage/Nodes/DirectoryNode.Extract.cs b/Engine/Source/Programs/Shared/EpicGames.Horde/Storage/Nodes/DirectoryNode.Extract.cs
--- a/Engine/Source/Programs/Shared/EpicGames.Horde/Storage/Nodes/DirectoryNode.Extract.cs
+++ b/Engine/Source/Programs/Shared/EpicGames.Horde/Storage/Nodes/DirectoryNode.Extract.cs
@@ -3,12 +3,16 @@
 using System;
 using System.Collections.Generic;
 using System.Diagnostics;
+using System.Diagnostics.CodeAnalysis;
 using System.IO;
 using System.IO.MemoryMappedFiles;
+using System.Linq;
 using System.Threading;
 using System.Threading.Channels;
 using System.Threading.Tasks;
 using EpicGames.Core;
+using EpicGames.Horde.Storage.Bundles;
+using EpicGames.Horde.Storage.Bundles.V2;
 using Microsoft.Extensions.Logging;
 
 namespace EpicGames.Horde.Storage.Nodes
@@ -295,11 +299,17 @@ namespace EpicGames.Horde.Storage.Nodes
 				tasks.Add(RunBackgroundTask(ctx => FindOutputChunksRootAsync(directoryInfo, directoryNode, chunks.Writer, logger, ctx)));
 
 				Channel<OutputBatch> batches = Channel.CreateUnbounded<OutputBatch>();
-				tasks.Add(RunBackgroundTask(ctx => BatchReadRequestsAsync(chunks.Reader, batches.Writer, ctx)));
+				tasks.Add(RunBackgroundTask(ctx => ReadBatchesAsync(chunks.Reader, batches.Writer, ctx)));
 
+				Channel<OutputBatch> prefectBatches = Channel.CreateBounded<OutputBatch>(new BoundedChannelOptions(128) { FullMode = BoundedChannelFullMode.Wait });
 				for (int idx = 0; idx < numTasks; idx++)
 				{
-					tasks.Add(RunBackgroundTask(ctx => ExtractAsync(batches.Reader, copyStats, logger, ctx)));
+					tasks.Add(RunBackgroundTask(ctx => PrefetchAsync(batches.Reader, prefectBatches.Writer, ctx)));
+				}
+
+				for (int idx = 0; idx < numTasks; idx++)
+				{
+					tasks.Add(RunBackgroundTask(ctx => ExtractAsync(prefectBatches.Reader, copyStats, logger, ctx)));
 				}
 
 				await Task.WhenAll(tasks);
@@ -469,27 +479,136 @@ namespace EpicGames.Horde.Storage.Nodes
 			}
 		}
 
-		static async Task BatchReadRequestsAsync(ChannelReader<OutputChunk> chunkReader, ChannelWriter<OutputBatch> batchWriter, CancellationToken cancellationToken)
+		static async Task PrefetchAsync(ChannelReader<OutputBatch> batchReader, ChannelWriter<OutputBatch> batchWriter, CancellationToken cancellationToken)
 		{
-			List<OutputChunk> batch = new List<OutputChunk>();
-			while (await chunkReader.WaitToReadAsync(cancellationToken))
+			while (await batchReader.WaitToReadAsync(cancellationToken))
 			{
-				OutputChunk? chunk;
-				while (chunkReader.TryRead(out chunk))
+				OutputBatch? batch;
+				if (batchReader.TryRead(out batch))
+				{
+					await batch.Chunks[0].Handle.ReadBlobDataAsync(cancellationToken);
+					await batchWriter.WriteAsync(batch, cancellationToken);
+				}
+			}
+		}
+
+		record class OutputExport(BundleHandle BundleHandle, int PacketOffset, int ExportIdx, OutputChunk Chunk);
+
+		class BundleReadQueue
+		{
+			public BundleHandle BundleHandle { get; }
+			public List<OutputExport> Requests { get; set; } = new List<OutputExport>();
+			public long LastUsedTime { get; set; }
+
+			public BundleReadQueue(BundleHandle bundleHandle)
+			{
+				BundleHandle = bundleHandle;
+			}
+		}
+
+		static async Task ReadBatchesAsync(ChannelReader<OutputChunk> chunkReader, ChannelWriter<OutputBatch> batchWriter, CancellationToken cancellationToken)
+		{
+			const int MinQueueLength = 100;
+			const int MaxQueueLength = 200;
+
+			const int MaxActiveQueues = 5;
+
+			int queueLength = 0;
+			Dictionary<BundleHandle, BundleReadQueue> bundleHandleToQueue = new Dictionary<BundleHandle, BundleReadQueue>();
+
+			for (; ; )
+			{
+				// Bring more data into the queue once it drops below a threshold
+				if (queueLength < MinQueueLength)
+				{
+					// Fill the queue up to the max length
+					while (queueLength < MaxQueueLength && await chunkReader.WaitToReadAsync(cancellationToken))
+					{
+						OutputChunk? chunk;
+						if (chunkReader.TryRead(out chunk))
+						{
+							OutputExport? outputExport;
+							if (TryGetOutputExport(chunk, out outputExport))
+							{
+								BundleHandle bundleHandle = outputExport.BundleHandle;
+								if (!bundleHandleToQueue.TryGetValue(bundleHandle, out BundleReadQueue? bundleQueue))
+								{
+									bundleQueue = new BundleReadQueue(bundleHandle);
+									bundleHandleToQueue.Add(bundleHandle, bundleQueue);
+								}
+
+								bundleQueue.Requests.Add(outputExport);
+								queueLength++;
+							}
+							else
+							{
+								OutputBatch batch = new OutputBatch(new List<OutputChunk> { chunk });
+								await batchWriter.WriteAsync(batch, cancellationToken);
+							}
+						}
+					}
+
+					// Exit once we've processed everything and can't get any more items to read.
+					if (queueLength == 0)
+					{
+						batchWriter.TryComplete();
+						break;
+					}
+
+					// If there are any running queues that have new pending reads, start new tasks for them
+					foreach (BundleReadQueue bundleQueue in bundleHandleToQueue.Values)
+					{
+						if (bundleQueue.Requests.Count > 0)
+						{
+							queueLength -= bundleQueue.Requests.Count;
+							await FlushQueueAsync(bundleQueue, batchWriter, cancellationToken);
+						}
+					}
+				}
+				else
 				{
-					if (batch.Count > 0 && chunk.File != batch[0].File)
+					// Find the longest queue that we can start a new read for
+					BundleReadQueue longestPendingQueue = bundleHandleToQueue.Values.MaxBy(x => x.Requests.Count);
+
+					// Wait for a queue to finish
+					while (bundleHandleToQueue.Count >= MaxActiveQueues)
 					{
-						await batchWriter.WriteAsync(new OutputBatch(batch), cancellationToken);
-						batch = new List<OutputChunk>();
+						BundleReadQueue? removeQueue = bundleHandleToQueue.Values.MinBy(x => x.LastUsedTime);
+						bundleHandleToQueue.Remove(removeQueue!.BundleHandle);
 					}
-					batch.Add(chunk);
+
+					// Start the task for a new queue
+					queueLength -= longestPendingQueue.Requests.Count;
+					await FlushQueueAsync(longestPendingQueue, batchWriter, cancellationToken);
 				}
 			}
-			if (batch.Count > 0)
+		}
+
+		static bool TryGetOutputExport(OutputChunk chunk, [NotNullWhen(true)] out OutputExport? export)
+		{
+			if (chunk.Handle is ExportHandle exportHandle && exportHandle.Packet is FlushedPacketHandle packetHandle)
 			{
-				await batchWriter.WriteAsync(new OutputBatch(batch), cancellationToken);
+				export = new OutputExport(packetHandle.Bundle, packetHandle.PacketOffset, exportHandle.ExportIdx, chunk);
+				return true;
 			}
+			else
+			{
+				export = null;
+				return false;
+			}
+		}
+
+		static async Task FlushQueueAsync(BundleReadQueue bundleQueue, ChannelWriter<OutputBatch> batchWriter, CancellationToken cancellationToken)
+		{
+			List<OutputChunk> chunks = bundleQueue.Requests
+				.OrderBy(x => x.PacketOffset)
+				.ThenBy(x => x.ExportIdx)
+				.Select(x => x.Chunk)
+				.ToList();
+
+			await batchWriter.WriteAsync(new OutputBatch(chunks), cancellationToken);
+			bundleQueue.LastUsedTime = Stopwatch.GetTimestamp();
+			bundleQueue.Requests.Clear();
 		}
 	}
 }
-
