Horde: Fix bundle fetch performance. Batching was not unwrapping the bundle handle, resulting in every request going down the slow path for non-bundle requests.

diff --git a/Engine/Source/Programs/Shared/EpicGames.Horde/Storage/Nodes/DirectoryNode.Extract.cs b/Engine/Source/Programs/Shared/EpicGames.Horde/Storage/Nodes/DirectoryNode.Extract.cs
--- a/Engine/Source/Programs/Shared/EpicGames.Horde/Storage/Nodes/DirectoryNode.Extract.cs
+++ b/Engine/Source/Programs/Shared/EpicGames.Horde/Storage/Nodes/DirectoryNode.Extract.cs
@@ -302,10 +302,13 @@ namespace EpicGames.Horde.Storage.Nodes
 				tasks.Add(RunBackgroundTask(ctx => ReadBatchesAsync(chunks.Reader, batches.Writer, ctx)));
 
 				Channel<OutputBatch> prefectBatches = Channel.CreateBounded<OutputBatch>(new BoundedChannelOptions(128) { FullMode = BoundedChannelFullMode.Wait });
+
+				List<Task> prefetchTasks = new List<Task>();
 				for (int idx = 0; idx < numTasks; idx++)
 				{
-					tasks.Add(RunBackgroundTask(ctx => PrefetchAsync(batches.Reader, prefectBatches.Writer, ctx)));
+					prefetchTasks.Add(RunBackgroundTask(ctx => PrefetchAsync(batches.Reader, prefectBatches.Writer, ctx)));
 				}
+				tasks.Add(Task.WhenAll(prefetchTasks).ContinueWith(_ => prefectBatches.Writer.Complete()));
 
 				for (int idx = 0; idx < numTasks; idx++)
 				{
@@ -586,7 +589,7 @@ namespace EpicGames.Horde.Storage.Nodes
 
 		static bool TryGetOutputExport(OutputChunk chunk, [NotNullWhen(true)] out OutputExport? export)
 		{
-			if (chunk.Handle is ExportHandle exportHandle && exportHandle.Packet is FlushedPacketHandle packetHandle)
+			if (chunk.Handle.Innermost is ExportHandle exportHandle && exportHandle.Packet is FlushedPacketHandle packetHandle)
 			{
 				export = new OutputExport(packetHandle.Bundle, packetHandle.PacketOffset, exportHandle.ExportIdx, chunk);
 				return true;
