Horde: Improve cancellation handling when prefetching data.

diff --git a/Engine/Source/Programs/Shared/EpicGames.Horde/Storage/Nodes/DirectoryNode.Extract.cs b/Engine/Source/Programs/Shared/EpicGames.Horde/Storage/Nodes/DirectoryNode.Extract.cs
--- a/Engine/Source/Programs/Shared/EpicGames.Horde/Storage/Nodes/DirectoryNode.Extract.cs
+++ b/Engine/Source/Programs/Shared/EpicGames.Horde/Storage/Nodes/DirectoryNode.Extract.cs
@@ -301,18 +301,12 @@ namespace EpicGames.Horde.Storage.Nodes
 				Channel<OutputBatch> batches = Channel.CreateUnbounded<OutputBatch>();
 				tasks.Add(RunBackgroundTask(ctx => ReadBatchesAsync(chunks.Reader, batches.Writer, ctx)));
 
-				Channel<OutputBatch> prefectBatches = Channel.CreateBounded<OutputBatch>(new BoundedChannelOptions(128) { FullMode = BoundedChannelFullMode.Wait });
+				Channel<OutputBatch> prefetchBatches = Channel.CreateBounded<OutputBatch>(new BoundedChannelOptions(128) { FullMode = BoundedChannelFullMode.Wait });
+				tasks.Add(RunBackgroundTask(ctx => PrefetchAsync(batches.Reader, prefetchBatches.Writer, numTasks, ctx)));
 
-				List<Task> prefetchTasks = new List<Task>();
 				for (int idx = 0; idx < numTasks; idx++)
 				{
-					prefetchTasks.Add(RunBackgroundTask(ctx => PrefetchAsync(batches.Reader, prefectBatches.Writer, ctx)));
-				}
-				tasks.Add(Task.WhenAll(prefetchTasks).ContinueWith(_ => prefectBatches.Writer.Complete()));
-
-				for (int idx = 0; idx < numTasks; idx++)
-				{
-					tasks.Add(RunBackgroundTask(ctx => ExtractAsync(prefectBatches.Reader, copyStats, logger, ctx)));
+					tasks.Add(RunBackgroundTask(ctx => ExtractAsync(prefetchBatches.Reader, copyStats, logger, ctx)));
 				}
 
 				await Task.WhenAll(tasks);
@@ -482,17 +476,35 @@ namespace EpicGames.Horde.Storage.Nodes
 			}
 		}
 
-		static async Task PrefetchAsync(ChannelReader<OutputBatch> batchReader, ChannelWriter<OutputBatch> batchWriter, CancellationToken cancellationToken)
+		static async Task PrefetchAsync(ChannelReader<OutputBatch> batchReader, ChannelWriter<OutputBatch> batchWriter, int numParallel, CancellationToken cancellationToken)
 		{
-			while (await batchReader.WaitToReadAsync(cancellationToken))
+			async Task RunAsync()
 			{
-				OutputBatch? batch;
-				if (batchReader.TryRead(out batch))
+				while (await batchReader.WaitToReadAsync(cancellationToken))
 				{
-					await batch.Chunks[0].Handle.ReadBlobDataAsync(cancellationToken);
-					await batchWriter.WriteAsync(batch, cancellationToken);
+					OutputBatch? batch;
+					if (batchReader.TryRead(out batch))
+					{
+						await batch.Chunks[0].Handle.ReadBlobDataAsync(cancellationToken);
+						await batchWriter.WriteAsync(batch, cancellationToken);
+					}
 				}
 			}
+
+			List<Task> tasks = new List<Task>();
+			for (int idx = 0; idx < numParallel; idx++)
+			{
+				tasks.Add(RunAsync());
+			}
+
+			try
+			{
+				await Task.WhenAll(tasks);
+			}
+			finally
+			{
+				batchWriter.Complete();
+			}
 		}
 
 		record class OutputExport(BundleHandle BundleHandle, int PacketOffset, int ExportIdx, OutputChunk Chunk);
