Horde: Fix file object store failing to catch exceptions when moving temp files into final locations.

diff --git a/Engine/Source/Programs/Shared/EpicGames.Horde/Storage/ObjectStores/FileObjectStore.cs b/Engine/Source/Programs/Shared/EpicGames.Horde/Storage/ObjectStores/FileObjectStore.cs
--- a/Engine/Source/Programs/Shared/EpicGames.Horde/Storage/ObjectStores/FileObjectStore.cs
+++ b/Engine/Source/Programs/Shared/EpicGames.Horde/Storage/ObjectStores/FileObjectStore.cs
@@ -89,7 +89,7 @@ namespace EpicGames.Horde.Storage.ObjectStores
 			{
 				FileReference.Move(tempLocation, finalLocation, true);
 			}
-			catch (IOException) // Already exists
+			catch // Already exists
 			{
 				if (FileReference.Exists(finalLocation))
 				{
