Horde: Partition output writes into separate tasks.

diff --git a/Engine/Source/Programs/Shared/EpicGames.Horde/Storage/Nodes/DirectoryNode.Extract.cs b/Engine/Source/Programs/Shared/EpicGames.Horde/Storage/Nodes/DirectoryNode.Extract.cs
--- a/Engine/Source/Programs/Shared/EpicGames.Horde/Storage/Nodes/DirectoryNode.Extract.cs
+++ b/Engine/Source/Programs/Shared/EpicGames.Horde/Storage/Nodes/DirectoryNode.Extract.cs
@@ -315,41 +315,52 @@ namespace EpicGames.Horde.Storage.Nodes
 
 		static async Task ExtractAsync(ChannelReader<OutputBatch> batchReader, CopyStats? copyStats, ILogger logger, CancellationToken cancellationToken)
 		{
+			const int WriteBatchSize = 64;
 			while (await batchReader.WaitToReadAsync(cancellationToken))
 			{
 				OutputBatch? batch;
 				while (batchReader.TryRead(out batch))
 				{
-					for (int chunkIdx = 0; chunkIdx < batch.Chunks.Count; )
+					List<Task> tasks = new List<Task>();
+					foreach (IReadOnlyList<OutputChunk> group in batch.Chunks.Batch(WriteBatchSize))
 					{
-						OutputFile file = batch.Chunks[chunkIdx].File;
+						tasks.Add(ExtractChunksAsync(group.ToArray(), copyStats, logger, cancellationToken));
+					}
+					await Task.WhenAll(tasks);
+				}
+			}
+		}
 
-						int maxChunkIdx = chunkIdx + 1;
-						while (maxChunkIdx < batch.Chunks.Count && batch.Chunks[maxChunkIdx].File == file)
-						{
-							maxChunkIdx++;
-						}
+		static async Task ExtractChunksAsync(ArraySegment<OutputChunk> chunks, CopyStats? copyStats, ILogger logger, CancellationToken cancellationToken)
+		{
+			for (int chunkIdx = 0; chunkIdx < chunks.Count; )
+			{
+				OutputFile file = chunks[chunkIdx].File;
 
-						try
-						{
-							await ExtractChunksToFileAsync(file, batch.Chunks.Slice(chunkIdx, maxChunkIdx - chunkIdx), copyStats, logger, cancellationToken);
-						}
-						catch (OperationCanceledException)
-						{
-							throw;
-						}
-						catch (Exception ex)
-						{
-							throw new StorageException($"Unable to extract {file?.FileInfo?.FullName}: {ex.Message}", ex);
-						}
+				int maxChunkIdx = chunkIdx + 1;
+				while (maxChunkIdx < chunks.Count && chunks[maxChunkIdx].File == file)
+				{
+					maxChunkIdx++;
+				}
 
-						chunkIdx = maxChunkIdx;
-					}
+				try
+				{
+					await ExtractChunksToFileAsync(file, chunks.Slice(chunkIdx, maxChunkIdx - chunkIdx), copyStats, logger, cancellationToken);
+				}
+				catch (OperationCanceledException)
+				{
+					throw;
 				}
+				catch (Exception ex)
+				{
+					throw new StorageException($"Unable to extract {file?.FileInfo?.FullName}: {ex.Message}", ex);
+				}
+
+				chunkIdx = maxChunkIdx;
 			}
 		}
 
-		static async Task ExtractChunksToFileAsync(OutputFile file, ListSegment<OutputChunk> chunks, CopyStats? copyStats, ILogger logger, CancellationToken cancellationToken)
+		static async Task ExtractChunksToFileAsync(OutputFile file, ArraySegment<OutputChunk> chunks, CopyStats? copyStats, ILogger logger, CancellationToken cancellationToken)
 		{
 			// Open the file for the current chunk
 			int remainingChunks = 0;
@@ -478,23 +489,10 @@ namespace EpicGames.Horde.Storage.Nodes
 
 		static async Task PrefetchAsync(ChannelReader<OutputBatch> batchReader, ChannelWriter<OutputBatch> batchWriter, int numParallel, CancellationToken cancellationToken)
 		{
-			async Task RunAsync()
-			{
-				while (await batchReader.WaitToReadAsync(cancellationToken))
-				{
-					OutputBatch? batch;
-					if (batchReader.TryRead(out batch))
-					{
-						await batch.Chunks[0].Handle.ReadBlobDataAsync(cancellationToken);
-						await batchWriter.WriteAsync(batch, cancellationToken);
-					}
-				}
-			}
-
 			List<Task> tasks = new List<Task>();
 			for (int idx = 0; idx < numParallel; idx++)
 			{
-				tasks.Add(RunAsync());
+				tasks.Add(PrefetchWorkerAsync(batchReader, batchWriter, cancellationToken));
 			}
 
 			try
@@ -507,6 +505,19 @@ namespace EpicGames.Horde.Storage.Nodes
 			}
 		}
 
+		static async Task PrefetchWorkerAsync(ChannelReader<OutputBatch> batchReader, ChannelWriter<OutputBatch> batchWriter, CancellationToken cancellationToken)
+		{
+			while (await batchReader.WaitToReadAsync(cancellationToken))
+			{
+				OutputBatch? batch;
+				if (batchReader.TryRead(out batch))
+				{
+					await batch.Chunks[0].Handle.ReadBlobDataAsync(cancellationToken);
+					await batchWriter.WriteAsync(batch, cancellationToken);
+				}
+			}
+		}
+
 		record class OutputExport(BundleHandle BundleHandle, int PacketOffset, int ExportIdx, OutputChunk Chunk);
 
 		class BundleReadQueue
