Horde: Improve batching for bundle read requests.

- Was incorrectly flushing any queue with entries assigned to it, rather than only those which were already active.
- Repeatedly taking the largest batch would degrade the queue to a state where buffer was full of individual requests.
- Now just operates as a FIFO queue, but adds new requests to existing buffers before flushing.

diff --git a/Engine/Source/Programs/Shared/EpicGames.Horde/Storage/Nodes/DirectoryNode.Extract.cs b/Engine/Source/Programs/Shared/EpicGames.Horde/Storage/Nodes/DirectoryNode.Extract.cs
--- a/Engine/Source/Programs/Shared/EpicGames.Horde/Storage/Nodes/DirectoryNode.Extract.cs
+++ b/Engine/Source/Programs/Shared/EpicGames.Horde/Storage/Nodes/DirectoryNode.Extract.cs
@@ -519,94 +519,67 @@ namespace EpicGames.Horde.Storage.Nodes
 		}
 
 		record class OutputExport(BundleHandle BundleHandle, int PacketOffset, int ExportIdx, OutputChunk Chunk);
-
-		class BundleReadQueue
-		{
-			public BundleHandle BundleHandle { get; }
-			public List<OutputExport> Requests { get; set; } = new List<OutputExport>();
-			public long LastUsedTime { get; set; }
-
-			public BundleReadQueue(BundleHandle bundleHandle)
-			{
-				BundleHandle = bundleHandle;
-			}
-		}
+		record class OutputExportBatch(BundleHandle BundleHandle, List<OutputExport> Exports);
 
 		static async Task ReadBatchesAsync(ChannelReader<OutputChunk> chunkReader, ChannelWriter<OutputBatch> batchWriter, CancellationToken cancellationToken)
 		{
-			const int MinQueueLength = 100;
 			const int MaxQueueLength = 200;
 
-			const int MaxActiveQueues = 5;
-
 			int queueLength = 0;
-			Dictionary<BundleHandle, BundleReadQueue> bundleHandleToQueue = new Dictionary<BundleHandle, BundleReadQueue>();
+			Queue<OutputExportBatch> exportBatchQueue = new Queue<OutputExportBatch>();
+			Dictionary<BundleHandle, OutputExportBatch> bundleHandleToExportBatch = new Dictionary<BundleHandle, OutputExportBatch>();
 
 			for (; ; )
 			{
-				// Bring more data into the queue once it drops below a threshold
-				if (queueLength < MinQueueLength)
+				// Fill the queue up to the max length
+				while (queueLength < MaxQueueLength)
 				{
-					// Fill the queue up to the max length
-					while (queueLength < MaxQueueLength && await chunkReader.WaitToReadAsync(cancellationToken))
+					OutputChunk? chunk;
+					if (!chunkReader.TryRead(out chunk))
 					{
-						OutputChunk? chunk;
-						if (chunkReader.TryRead(out chunk))
+						if (!await chunkReader.WaitToReadAsync(cancellationToken))
 						{
-							OutputExport? outputExport;
-							if (TryGetOutputExport(chunk, out outputExport))
-							{
-								BundleHandle bundleHandle = outputExport.BundleHandle;
-								if (!bundleHandleToQueue.TryGetValue(bundleHandle, out BundleReadQueue? bundleQueue))
-								{
-									bundleQueue = new BundleReadQueue(bundleHandle);
-									bundleHandleToQueue.Add(bundleHandle, bundleQueue);
-								}
-
-								bundleQueue.Requests.Add(outputExport);
-								queueLength++;
-							}
-							else
-							{
-								OutputBatch batch = new OutputBatch(new List<OutputChunk> { chunk });
-								await batchWriter.WriteAsync(batch, cancellationToken);
-							}
+							break;
 						}
 					}
-
-					// Exit once we've processed everything and can't get any more items to read.
-					if (queueLength == 0)
-					{
-						batchWriter.TryComplete();
-						break;
-					}
-
-					// If there are any running queues that have new pending reads, start new tasks for them
-					foreach (BundleReadQueue bundleQueue in bundleHandleToQueue.Values)
+					else
 					{
-						if (bundleQueue.Requests.Count > 0)
+						OutputExport? outputExport;
+						if (!TryGetOutputExport(chunk, out outputExport))
 						{
-							queueLength -= bundleQueue.Requests.Count;
-							await FlushQueueAsync(bundleQueue, batchWriter, cancellationToken);
+							OutputBatch batch = new OutputBatch(new List<OutputChunk> { chunk });
+							await batchWriter.WriteAsync(batch, cancellationToken);
+						}
+						else
+						{
+							BundleHandle bundleHandle = outputExport.BundleHandle;
+							if (!bundleHandleToExportBatch.TryGetValue(bundleHandle, out OutputExportBatch? existingExportBatch))
+							{
+								existingExportBatch = new OutputExportBatch(bundleHandle, new List<OutputExport>());
+								exportBatchQueue.Enqueue(existingExportBatch);
+								bundleHandleToExportBatch.Add(bundleHandle, existingExportBatch);
+							}
+
+							existingExportBatch.Exports.Add(outputExport);
+							queueLength++;
 						}
 					}
 				}
-				else
+
+				// Exit once we've processed everything and can't get any more items to read.
+				if (queueLength == 0)
 				{
-					// Find the longest queue that we can start a new read for
-					BundleReadQueue longestPendingQueue = bundleHandleToQueue.Values.MaxBy(x => x.Requests.Count);
+					batchWriter.TryComplete();
+					break;
+				}
 
-					// Wait for a queue to finish
-					while (bundleHandleToQueue.Count >= MaxActiveQueues)
-					{
-						BundleReadQueue? removeQueue = bundleHandleToQueue.Values.MinBy(x => x.LastUsedTime);
-						bundleHandleToQueue.Remove(removeQueue!.BundleHandle);
-					}
+				// Flush the first queue
+				OutputExportBatch exportBatch = exportBatchQueue.Dequeue();
+				queueLength -= exportBatch.Exports.Count;
+				bundleHandleToExportBatch.Remove(exportBatch.BundleHandle);
 
-					// Start the task for a new queue
-					queueLength -= longestPendingQueue.Requests.Count;
-					await FlushQueueAsync(longestPendingQueue, batchWriter, cancellationToken);
-				}
+				List<OutputChunk> chunkBatch = exportBatch.Exports.OrderBy(x => x.PacketOffset).ThenBy(x => x.ExportIdx).Select(x => x.Chunk).ToList();
+				await batchWriter.WriteAsync(new OutputBatch(chunkBatch), cancellationToken);
 			}
 		}
 
@@ -623,18 +596,5 @@ namespace EpicGames.Horde.Storage.Nodes
 				return false;
 			}
 		}
-
-		static async Task FlushQueueAsync(BundleReadQueue bundleQueue, ChannelWriter<OutputBatch> batchWriter, CancellationToken cancellationToken)
-		{
-			List<OutputChunk> chunks = bundleQueue.Requests
-				.OrderBy(x => x.PacketOffset)
-				.ThenBy(x => x.ExportIdx)
-				.Select(x => x.Chunk)
-				.ToList();
-
-			await batchWriter.WriteAsync(new OutputBatch(chunks), cancellationToken);
-			bundleQueue.LastUsedTime = Stopwatch.GetTimestamp();
-			bundleQueue.Requests.Clear();
-		}
 	}
 }
