Horde: Prefetch all blob data in a separate pipeline step to writing it. Removes any waits on the writer thread apart from actually outputting data.

diff --git a/Engine/Source/Programs/Shared/EpicGames.Core/EnumerableExtensions.cs b/Engine/Source/Programs/Shared/EpicGames.Core/EnumerableExtensions.cs
--- a/Engine/Source/Programs/Shared/EpicGames.Core/EnumerableExtensions.cs
+++ b/Engine/Source/Programs/Shared/EpicGames.Core/EnumerableExtensions.cs
@@ -10,6 +10,33 @@ namespace EpicGames.Core
 	/// </summary>
 	public static class EnumerableExtensions
 	{
+		/// <summary>
+		/// Dispose all elements in a sequence
+		/// </summary>
+		/// <typeparam name="T">The element type</typeparam>
+		/// <param name="sequence">Sequence of elements to dispose</param>
+		public static void DisposeElements<T>(this IEnumerable<T> sequence) where T : IDisposable
+		{
+			List<Exception>? exceptions = null;
+			foreach (T element in sequence)
+			{
+				try
+				{
+					element.Dispose();
+				}
+				catch (Exception ex)
+				{
+					exceptions ??= new List<Exception>();
+					exceptions.Add(ex);
+				}
+			}
+
+			if (exceptions != null)
+			{
+				throw new AggregateException(exceptions);
+			}
+		}
+
 		/// <summary>
 		/// Split the sequence into batches of at most the given size
 		/// </summary>
diff --git a/Engine/Source/Programs/Shared/EpicGames.Horde/Storage/Nodes/DirectoryNode.Extract.cs b/Engine/Source/Programs/Shared/EpicGames.Horde/Storage/Nodes/DirectoryNode.Extract.cs
--- a/Engine/Source/Programs/Shared/EpicGames.Horde/Storage/Nodes/DirectoryNode.Extract.cs
+++ b/Engine/Source/Programs/Shared/EpicGames.Horde/Storage/Nodes/DirectoryNode.Extract.cs
@@ -301,37 +301,54 @@ namespace EpicGames.Horde.Storage.Nodes
 				Channel<OutputBatch> batches = Channel.CreateUnbounded<OutputBatch>();
 				tasks.Add(RunBackgroundTask(ctx => ReadBatchesAsync(chunks.Reader, batches.Writer, ctx)));
 
-				Channel<OutputBatch> prefetchBatches = Channel.CreateBounded<OutputBatch>(new BoundedChannelOptions(128) { FullMode = BoundedChannelFullMode.Wait });
-				tasks.Add(RunBackgroundTask(ctx => PrefetchAsync(batches.Reader, prefetchBatches.Writer, numTasks, ctx)));
+				Channel<OutputFetchedBatch> prefetchBatches = Channel.CreateBounded<OutputFetchedBatch>(new BoundedChannelOptions(128) { FullMode = BoundedChannelFullMode.Wait });
+				try
+				{
+					tasks.Add(RunBackgroundTask(ctx => PrefetchAsync(batches.Reader, prefetchBatches.Writer, numTasks, ctx)));
+
+					for (int idx = 0; idx < numTasks; idx++)
+					{
+						tasks.Add(RunBackgroundTask(ctx => ExtractAsync(prefetchBatches.Reader, copyStats, logger, ctx)));
+					}
 
-				for (int idx = 0; idx < numTasks; idx++)
+					await Task.WhenAll(tasks);
+				}
+				finally
 				{
-					tasks.Add(RunBackgroundTask(ctx => ExtractAsync(prefetchBatches.Reader, copyStats, logger, ctx)));
+					while (prefetchBatches.Reader.TryRead(out OutputFetchedBatch? fetchedBatch))
+					{
+						fetchedBatch.Dispose();
+					}
 				}
-
-				await Task.WhenAll(tasks);
 			}
 		}
 
-		static async Task ExtractAsync(ChannelReader<OutputBatch> batchReader, CopyStats? copyStats, ILogger logger, CancellationToken cancellationToken)
+		static async Task ExtractAsync(ChannelReader<OutputFetchedBatch> batchReader, CopyStats? copyStats, ILogger logger, CancellationToken cancellationToken)
 		{
 			const int WriteBatchSize = 64;
 			while (await batchReader.WaitToReadAsync(cancellationToken))
 			{
-				OutputBatch? batch;
+				OutputFetchedBatch? batch;
 				while (batchReader.TryRead(out batch))
 				{
-					List<Task> tasks = new List<Task>();
-					foreach (IReadOnlyList<OutputChunk> group in batch.Chunks.Batch(WriteBatchSize))
+					try
 					{
-						tasks.Add(ExtractChunksAsync(group.ToArray(), copyStats, logger, cancellationToken));
+						List<Task> tasks = new List<Task>();
+						foreach (IReadOnlyList<OutputFetchedChunk> group in batch.Chunks.Batch(WriteBatchSize))
+						{
+							tasks.Add(ExtractChunksAsync(group.ToArray(), copyStats, logger, cancellationToken));
+						}
+						await Task.WhenAll(tasks);
+					}
+					finally
+					{
+						batch.Dispose();
 					}
-					await Task.WhenAll(tasks);
 				}
 			}
 		}
 
-		static async Task ExtractChunksAsync(ArraySegment<OutputChunk> chunks, CopyStats? copyStats, ILogger logger, CancellationToken cancellationToken)
+		static async Task ExtractChunksAsync(ArraySegment<OutputFetchedChunk> chunks, CopyStats? copyStats, ILogger logger, CancellationToken cancellationToken)
 		{
 			for (int chunkIdx = 0; chunkIdx < chunks.Count; )
 			{
@@ -360,7 +377,7 @@ namespace EpicGames.Horde.Storage.Nodes
 			}
 		}
 
-		static async Task ExtractChunksToFileAsync(OutputFile file, ArraySegment<OutputChunk> chunks, CopyStats? copyStats, ILogger logger, CancellationToken cancellationToken)
+		static async Task ExtractChunksToFileAsync(OutputFile file, ArraySegment<OutputFetchedChunk> chunks, CopyStats? copyStats, ILogger logger, CancellationToken cancellationToken)
 		{
 			// Open the file for the current chunk
 			int remainingChunks = 0;
@@ -379,14 +396,12 @@ namespace EpicGames.Horde.Storage.Nodes
 
 					for(int chunkIdx = 0; chunkIdx < chunks.Count; chunkIdx++)
 					{
-						OutputChunk chunk = chunks[chunkIdx];
+						OutputFetchedChunk chunk = chunks[chunkIdx];
+						cancellationToken.ThrowIfCancellationRequested();
 
 						// Write this chunk
-						using (BlobData data = await chunk.Handle.ReadBlobDataAsync(cancellationToken))
-						{
-							TraceBlobRead("Leaf", chunk.File.Path, chunk.Handle, logger);
-							data.Data.CopyTo(memoryMappedView!.GetMemory(chunk.Offset, data.Data.Length));
-						}
+						TraceBlobRead("Leaf", chunk.File.Path, chunk.Handle, logger);
+						chunk.BlobData.Data.CopyTo(memoryMappedView!.GetMemory(chunk.Offset, chunk.BlobData.Data.Length));
 
 						// Update the stats
 						remainingChunks = file.DecrementRemaining();
@@ -487,7 +502,18 @@ namespace EpicGames.Horde.Storage.Nodes
 			}
 		}
 
-		static async Task PrefetchAsync(ChannelReader<OutputBatch> batchReader, ChannelWriter<OutputBatch> batchWriter, int numParallel, CancellationToken cancellationToken)
+		record class OutputFetchedChunk(OutputChunk Chunk, BlobData BlobData) : OutputChunk(Chunk), IDisposable
+		{
+			public void Dispose() => BlobData.Dispose();
+		}
+
+		record class OutputFetchedBatch(List<OutputFetchedChunk> Chunks) : IDisposable
+		{
+			public OutputFetchedBatch() : this(new List<OutputFetchedChunk>()) { }
+			public void Dispose() => Chunks.DisposeElements();
+		}
+
+		static async Task PrefetchAsync(ChannelReader<OutputBatch> batchReader, ChannelWriter<OutputFetchedBatch> batchWriter, int numParallel, CancellationToken cancellationToken)
 		{
 			List<Task> tasks = new List<Task>();
 			for (int idx = 0; idx < numParallel; idx++)
@@ -505,15 +531,29 @@ namespace EpicGames.Horde.Storage.Nodes
 			}
 		}
 
-		static async Task PrefetchWorkerAsync(ChannelReader<OutputBatch> batchReader, ChannelWriter<OutputBatch> batchWriter, CancellationToken cancellationToken)
+		static async Task PrefetchWorkerAsync(ChannelReader<OutputBatch> batchReader, ChannelWriter<OutputFetchedBatch> batchWriter, CancellationToken cancellationToken)
 		{
 			while (await batchReader.WaitToReadAsync(cancellationToken))
 			{
 				OutputBatch? batch;
 				if (batchReader.TryRead(out batch))
 				{
-					await batch.Chunks[0].Handle.ReadBlobDataAsync(cancellationToken);
-					await batchWriter.WriteAsync(batch, cancellationToken);
+#pragma warning disable CA2000 // fetchedBatch may be pushed onto the output channel, which assumes its ownership.
+					OutputFetchedBatch fetchedBatch = new OutputFetchedBatch();
+					try
+					{
+						foreach (OutputChunk chunk in batch.Chunks)
+						{
+							BlobData blobData = await chunk.Handle.ReadBlobDataAsync(cancellationToken);
+							fetchedBatch.Chunks.Add(new OutputFetchedChunk(chunk, blobData));
+						}
+						await batchWriter.WriteAsync(fetchedBatch, cancellationToken);
+					}
+					catch
+					{
+						fetchedBatch.Dispose();
+					}
+#pragma warning restore CA2000
 				}
 			}
 		}
