Horde: Minor refactoring.

* Add better stats for cache and http backend, and name all the stats using dotted identifiers to make grouping easier.
* Fix static analysis warning in EnumerableExtensions.DisposeAll().
* Reorder methods in DirectoryNode.Extract to match order in pipeline.

diff --git a/Engine/Source/Programs/Shared/EpicGames.Core/EnumerableExtensions.cs b/Engine/Source/Programs/Shared/EpicGames.Core/EnumerableExtensions.cs
--- a/Engine/Source/Programs/Shared/EpicGames.Core/EnumerableExtensions.cs
+++ b/Engine/Source/Programs/Shared/EpicGames.Core/EnumerableExtensions.cs
@@ -13,27 +13,12 @@ namespace EpicGames.Core
 		/// <summary>
 		/// Dispose all elements in a sequence
 		/// </summary>
-		/// <typeparam name="T">The element type</typeparam>
 		/// <param name="sequence">Sequence of elements to dispose</param>
-		public static void DisposeElements<T>(this IEnumerable<T> sequence) where T : IDisposable
+		public static void DisposeAll(this IEnumerable<IDisposable> sequence)
 		{
-			List<Exception>? exceptions = null;
-			foreach (T element in sequence)
+			foreach (IDisposable element in sequence)
 			{
-				try
-				{
-					element.Dispose();
-				}
-				catch (Exception ex)
-				{
-					exceptions ??= new List<Exception>();
-					exceptions.Add(ex);
-				}
-			}
-
-			if (exceptions != null)
-			{
-				throw new AggregateException(exceptions);
+				element.Dispose();
 			}
 		}
 
diff --git a/Engine/Source/Programs/Shared/EpicGames.Horde/Storage/Backends/HttpStorageBackend.cs b/Engine/Source/Programs/Shared/EpicGames.Horde/Storage/Backends/HttpStorageBackend.cs
--- a/Engine/Source/Programs/Shared/EpicGames.Horde/Storage/Backends/HttpStorageBackend.cs
+++ b/Engine/Source/Programs/Shared/EpicGames.Horde/Storage/Backends/HttpStorageBackend.cs
@@ -1,6 +1,7 @@
 // Copyright Epic Games, Inc. All Rights Reserved.
 
 using System;
+using System.Diagnostics;
 using System.IO;
 using System.Net;
 using System.Net.Http;
@@ -32,6 +33,11 @@ namespace EpicGames.Horde.Storage.Backends
 		readonly ILogger _logger;
 		bool _supportsUploadRedirects = true;
 
+		long _numBytes;
+		int _numActive;
+		TimeSpan _sequentialReadTime;
+		readonly Stopwatch _readTime = new Stopwatch();
+
 		/// <inheritdoc/>
 		public bool SupportsRedirects => _supportsUploadRedirects;
 
@@ -51,36 +57,70 @@ namespace EpicGames.Horde.Storage.Backends
 		/// <inheritdoc/>
 		public async Task<Stream> OpenBlobAsync(BlobLocator locator, int offset, int? length, CancellationToken cancellationToken = default)
 		{
-			if (offset == 0 && length == null)
-			{
-				_logger.LogDebug("Reading {Locator}", locator);
-			}
-			else if (length == null)
+			Stopwatch blobTimer = Stopwatch.StartNew();
+			try
 			{
-				_logger.LogDebug("Reading {Locator} ({Offset}..)", locator, offset);
-			}
-			else
-			{
-				_logger.LogDebug("Reading {Locator} ({Offset}+{Length})", locator, offset, length);
-			}
+				lock (_readTime)
+				{
+					if (++_numActive == 1)
+					{
+						_readTime.Start();
+					}
+				}
 
-			if (length.HasValue && length.Value == 0)
-			{
-				return new MemoryStream(Array.Empty<byte>());
-			}
+				if (offset == 0 && length == null)
+				{
+					_logger.LogDebug("Reading {Locator}", locator);
+				}
+				else if (length == null)
+				{
+					_logger.LogDebug("Reading {Locator} ({Offset}..)", locator, offset);
+				}
+				else
+				{
+					_logger.LogDebug("Reading {Locator} ({Offset}+{Length})", locator, offset, length);
+				}
 
-			using (HttpClient httpClient = _createClient())
+				if (length.HasValue && length.Value == 0)
+				{
+					return new MemoryStream(Array.Empty<byte>());
+				}
+
+				using (HttpClient httpClient = _createClient())
+				{
+					using (HttpRequestMessage request = new HttpRequestMessage(HttpMethod.Get, $"{_basePath}/blobs/{locator}"))
+					{
+						if (offset != 0 || length != null)
+						{
+							request.Headers.Range = new RangeHeaderValue(offset, (length == null) ? null : (offset + (length - 1)));
+						}
+
+						HttpResponseMessage response = await httpClient.SendAsync(request, cancellationToken);
+						response.EnsureSuccessStatusCode();
+
+						Stream stream = await response.Content.ReadAsStreamAsync(cancellationToken);
+
+						try
+						{
+							Interlocked.Add(ref _numBytes, response.Content.Headers.ContentLength ?? stream.Length);
+						}
+						catch
+						{
+						}
+
+						return stream;
+					}
+				}
+			}
+			finally
 			{
-				using (HttpRequestMessage request = new HttpRequestMessage(HttpMethod.Get, $"{_basePath}/blobs/{locator}"))
+				lock (_readTime)
 				{
-					if (offset != 0 || length != null)
+					_sequentialReadTime += blobTimer.Elapsed;
+					if (--_numActive == 0)
 					{
-						request.Headers.Range = new RangeHeaderValue(offset, (length == null) ? null : (offset + (length - 1)));
+						_readTime.Stop();
 					}
-
-					HttpResponseMessage response = await httpClient.SendAsync(request, cancellationToken);
-					response.EnsureSuccessStatusCode();
-					return await response.Content.ReadAsStreamAsync(cancellationToken);
 				}
 			}
 		}
@@ -313,7 +353,16 @@ namespace EpicGames.Horde.Storage.Backends
 		#endregion
 
 		/// <inheritdoc/>
-		public void GetStats(StorageStats stats) { }
+		public void GetStats(StorageStats stats)
+		{
+			stats.Add("backend.http.wall_time_secs", (long)_readTime.Elapsed.TotalSeconds);
+			stats.Add("backend.http.num_bytes", _numBytes);
+			if (_readTime.Elapsed > TimeSpan.Zero)
+			{
+				stats.Add("backend.http.speed_mb_sec", (long)(_numBytes / (1024.0 * 1024.0 * _readTime.Elapsed.TotalSeconds)));
+				stats.Add("backend.http.concurrency_ratio", (long)(_sequentialReadTime.TotalSeconds * 100.0 / _readTime.Elapsed.TotalSeconds));
+			}
+		}
 	}
 
 	/// <summary>
diff --git a/Engine/Source/Programs/Shared/EpicGames.Horde/Storage/Bundles/BundleCache.cs b/Engine/Source/Programs/Shared/EpicGames.Horde/Storage/Bundles/BundleCache.cs
--- a/Engine/Source/Programs/Shared/EpicGames.Horde/Storage/Bundles/BundleCache.cs
+++ b/Engine/Source/Programs/Shared/EpicGames.Horde/Storage/Bundles/BundleCache.cs
@@ -181,6 +181,7 @@ namespace EpicGames.Horde.Storage
 		readonly MemoryAllocator _allocator;
 
 		int _numCacheHits;
+		int _numCacheMisses;
 		long _currentSize;
 
 		readonly MemoryCache? _headerCache;
@@ -358,6 +359,7 @@ namespace EpicGames.Horde.Storage
 					}
 					else
 					{
+						Interlocked.Increment(ref _numCacheMisses);
 						item = new CacheValue(key, async () => await createAsync(key, _cancellationSource.Token));
 						_itemLookup.Add(key, item);
 					}
@@ -406,8 +408,10 @@ namespace EpicGames.Horde.Storage
 		/// </summary>
 		public void GetStats(StorageStats stats)
 		{
-			stats.Add("Num cache hits", _numCacheHits);
-			stats.Add("Cache size", _currentSize);
+			stats.Add("bundle.cache.hits", _numCacheHits);
+			stats.Add("bundle.cache.misses", _numCacheMisses);
+			stats.Add("bundle.cache.size_count", _items.Count);
+			stats.Add("bundle.cache.size_bytes", _currentSize);
 		}
 
 		#region V1
diff --git a/Engine/Source/Programs/Shared/EpicGames.Horde/Storage/Bundles/V2/PacketReader.cs b/Engine/Source/Programs/Shared/EpicGames.Horde/Storage/Bundles/V2/PacketReader.cs
--- a/Engine/Source/Programs/Shared/EpicGames.Horde/Storage/Bundles/V2/PacketReader.cs
+++ b/Engine/Source/Programs/Shared/EpicGames.Horde/Storage/Bundles/V2/PacketReader.cs
@@ -18,10 +18,10 @@ namespace EpicGames.Horde.Storage.Bundles.V2
 
 		public void GetStats(StorageStats stats)
 		{
-			stats.Add("Num reads (v2)", _numReads);
-			stats.Add("Num packets read (v2)", _numPacketsRead);
-			stats.Add("Num bytes read (v2)", _numBytesRead);
-			stats.Add("Num wasted bytes read (v2)", _numWastedBytesRead);
+			stats.Add("bundle.packet_reader.num_reads", _numReads);
+			stats.Add("bundle.packet_reader.num_packets", _numPacketsRead);
+			stats.Add("bundle.packet_reader.num_bytes", _numBytesRead);
+			stats.Add("bundle.packet_reader.num_wasted_bytes", _numWastedBytesRead);
 		}
 	}
 
diff --git a/Engine/Source/Programs/Shared/EpicGames.Horde/Storage/Nodes/DirectoryNode.Extract.cs b/Engine/Source/Programs/Shared/EpicGames.Horde/Storage/Nodes/DirectoryNode.Extract.cs
--- a/Engine/Source/Programs/Shared/EpicGames.Horde/Storage/Nodes/DirectoryNode.Extract.cs
+++ b/Engine/Source/Programs/Shared/EpicGames.Horde/Storage/Nodes/DirectoryNode.Extract.cs
@@ -304,7 +304,7 @@
 				Channel<OutputFetchedBatch> prefetchBatches = Channel.CreateBounded<OutputFetchedBatch>(new BoundedChannelOptions(128) { FullMode = BoundedChannelFullMode.Wait });
 				try
 				{
-					tasks.Add(RunBackgroundTask(ctx => PrefetchAsync(batches.Reader, prefetchBatches.Writer, numTasks, ctx)));
+					tasks.Add(RunBackgroundTask(ctx => FetchAsync(batches.Reader, prefetchBatches.Writer, numTasks, ctx)));
 
 					for (int idx = 0; idx < numTasks; idx++)
 					{
@@ -323,100 +323,7 @@
 			}
 		}
 
-		static async Task ExtractAsync(ChannelReader<OutputFetchedBatch> batchReader, CopyStats? copyStats, ILogger logger, CancellationToken cancellationToken)
-		{
-			const int WriteBatchSize = 64;
-			while (await batchReader.WaitToReadAsync(cancellationToken))
-			{
-				OutputFetchedBatch? batch;
-				while (batchReader.TryRead(out batch))
-				{
-					try
-					{
-						List<Task> tasks = new List<Task>();
-						foreach (IReadOnlyList<OutputFetchedChunk> group in batch.Chunks.Batch(WriteBatchSize))
-						{
-							tasks.Add(ExtractChunksAsync(group.ToArray(), copyStats, logger, cancellationToken));
-						}
-						await Task.WhenAll(tasks);
-					}
-					finally
-					{
-						batch.Dispose();
-					}
-				}
-			}
-		}
+		#region Enumerate chunks
 
-		static async Task ExtractChunksAsync(ArraySegment<OutputFetchedChunk> chunks, CopyStats? copyStats, ILogger logger, CancellationToken cancellationToken)
-		{
-			for (int chunkIdx = 0; chunkIdx < chunks.Count; )
-			{
-				OutputFile file = chunks[chunkIdx].File;
-
-				int maxChunkIdx = chunkIdx + 1;
-				while (maxChunkIdx < chunks.Count && chunks[maxChunkIdx].File == file)
-				{
-					maxChunkIdx++;
-				}
-
-				try
-				{
-					await ExtractChunksToFileAsync(file, chunks.Slice(chunkIdx, maxChunkIdx - chunkIdx), copyStats, logger, cancellationToken);
-				}
-				catch (OperationCanceledException)
-				{
-					throw;
-				}
-				catch (Exception ex)
-				{
-					throw new StorageException($"Unable to extract {file?.FileInfo?.FullName}: {ex.Message}", ex);
-				}
-
-				chunkIdx = maxChunkIdx;
-			}
-		}
-
-		static async Task ExtractChunksToFileAsync(OutputFile file, ArraySegment<OutputFetchedChunk> chunks, CopyStats? copyStats, ILogger logger, CancellationToken cancellationToken)
-		{
-			// Open the file for the current chunk
-			int remainingChunks = 0;
-			await using (FileStream stream = file.OpenStream())
-			{
-				if (file.FileEntry.Length == 0)
-				{
-					// If this file is empty, don't write anything and just move to the next chunk
-					remainingChunks = file.DecrementRemaining();
-				}
-				else
-				{
-					// Process as many chunks as we can for this file
-					using MemoryMappedFile memoryMappedFile = MemoryMappedFile.CreateFromFile(stream, null, file.FileEntry.Length, MemoryMappedFileAccess.ReadWrite, HandleInheritability.None, false);
-					using MemoryMappedView memoryMappedView = new MemoryMappedView(memoryMappedFile, 0, file.FileEntry.Length);
-
-					for(int chunkIdx = 0; chunkIdx < chunks.Count; chunkIdx++)
-					{
-						OutputFetchedChunk chunk = chunks[chunkIdx];
-						cancellationToken.ThrowIfCancellationRequested();
-
-						// Write this chunk
-						TraceBlobRead("Leaf", chunk.File.Path, chunk.Handle, logger);
-						chunk.BlobData.Data.CopyTo(memoryMappedView!.GetMemory(chunk.Offset, chunk.BlobData.Data.Length));
-
-						// Update the stats
-						remainingChunks = file.DecrementRemaining();
-						copyStats?.Update(0, chunk.Length);
-					}
-				}
-			}
-
-			// Set correct permissions on the output file
-			if (remainingChunks == 0)
-			{
-				FileEntry.SetPermissions(file.FileInfo!, file.FileEntry.Flags);
-				copyStats?.Update(1, 0);
-			}
-		}
-
 		static async Task FindOutputChunksRootAsync(DirectoryInfo rootDir, DirectoryNode node, ChannelWriter<OutputChunk> chunks, ILogger logger, CancellationToken cancellationToken)
 		{
@@ -490,18 +397,93 @@
 			}
 		}
 
-		static string CombinePaths(string basePath, string nextPath)
+		#endregion
+
+		#region Group requests by bundles
+
+		record class OutputExport(BundleHandle BundleHandle, int PacketOffset, int ExportIdx, OutputChunk Chunk);
+		record class OutputExportBatch(BundleHandle BundleHandle, List<OutputExport> Exports);
+
+		static async Task ReadBatchesAsync(ChannelReader<OutputChunk> chunkReader, ChannelWriter<OutputBatch> batchWriter, CancellationToken cancellationToken)
 		{
-			if (basePath.Length > 0)
+			const int MaxQueueLength = 200;
+
+			int queueLength = 0;
+			Queue<OutputExportBatch> exportBatchQueue = new Queue<OutputExportBatch>();
+			Dictionary<BundleHandle, OutputExportBatch> bundleHandleToExportBatch = new Dictionary<BundleHandle, OutputExportBatch>();
+
+			for (; ; )
 			{
-				return $"{basePath}/{nextPath}";
+				// Fill the queue up to the max length
+				while (queueLength < MaxQueueLength)
+				{
+					OutputChunk? chunk;
+					if (!chunkReader.TryRead(out chunk))
+					{
+						if (!await chunkReader.WaitToReadAsync(cancellationToken))
+						{
+							break;
+						}
+					}
+					else
+					{
+						OutputExport? outputExport;
+						if (!TryGetOutputExport(chunk, out outputExport))
+						{
+							OutputBatch batch = new OutputBatch(new List<OutputChunk> { chunk });
+							await batchWriter.WriteAsync(batch, cancellationToken);
+						}
+						else
+						{
+							BundleHandle bundleHandle = outputExport.BundleHandle;
+							if (!bundleHandleToExportBatch.TryGetValue(bundleHandle, out OutputExportBatch? existingExportBatch))
+							{
+								existingExportBatch = new OutputExportBatch(bundleHandle, new List<OutputExport>());
+								exportBatchQueue.Enqueue(existingExportBatch);
+								bundleHandleToExportBatch.Add(bundleHandle, existingExportBatch);
+							}
+
+							existingExportBatch.Exports.Add(outputExport);
+							queueLength++;
+						}
+					}
+				}
+
+				// Exit once we've processed everything and can't get any more items to read.
+				if (queueLength == 0)
+				{
+					batchWriter.TryComplete();
+					break;
+				}
+
+				// Flush the first queue
+				OutputExportBatch exportBatch = exportBatchQueue.Dequeue();
+				queueLength -= exportBatch.Exports.Count;
+				bundleHandleToExportBatch.Remove(exportBatch.BundleHandle);
+
+				List<OutputChunk> chunkBatch = exportBatch.Exports.OrderBy(x => x.PacketOffset).ThenBy(x => x.ExportIdx).Select(x => x.Chunk).ToList();
+				await batchWriter.WriteAsync(new OutputBatch(chunkBatch), cancellationToken);
 			}
+		}
+
+		static bool TryGetOutputExport(OutputChunk chunk, [NotNullWhen(true)] out OutputExport? export)
+		{
+			if (chunk.Handle.Innermost is ExportHandle exportHandle && exportHandle.Packet is FlushedPacketHandle packetHandle)
+			{
+				export = new OutputExport(packetHandle.Bundle, packetHandle.PacketOffset, exportHandle.ExportIdx, chunk);
+				return true;
+			}
 			else
 			{
-				return nextPath;
+				export = null;
+				return false;
 			}
 		}
 
+		#endregion
+
+		#region Fetch data
+
 		record class OutputFetchedChunk(OutputChunk Chunk, BlobData BlobData) : OutputChunk(Chunk), IDisposable
 		{
 			public void Dispose() => BlobData.Dispose();
@@ -510,15 +492,15 @@
 		record class OutputFetchedBatch(List<OutputFetchedChunk> Chunks) : IDisposable
 		{
 			public OutputFetchedBatch() : this(new List<OutputFetchedChunk>()) { }
-			public void Dispose() => Chunks.DisposeElements();
+			public void Dispose() => Chunks.DisposeAll();
 		}
 
-		static async Task PrefetchAsync(ChannelReader<OutputBatch> batchReader, ChannelWriter<OutputFetchedBatch> batchWriter, int numParallel, CancellationToken cancellationToken)
+		static async Task FetchAsync(ChannelReader<OutputBatch> batchReader, ChannelWriter<OutputFetchedBatch> batchWriter, int numParallel, CancellationToken cancellationToken)
 		{
 			List<Task> tasks = new List<Task>();
 			for (int idx = 0; idx < numParallel; idx++)
 			{
-				tasks.Add(PrefetchWorkerAsync(batchReader, batchWriter, cancellationToken));
+				tasks.Add(FetchWorkerAsync(batchReader, batchWriter, cancellationToken));
 			}
 
 			try
@@ -531,7 +513,7 @@
 			}
 		}
 
-		static async Task PrefetchWorkerAsync(ChannelReader<OutputBatch> batchReader, ChannelWriter<OutputFetchedBatch> batchWriter, CancellationToken cancellationToken)
+		static async Task FetchWorkerAsync(ChannelReader<OutputBatch> batchReader, ChannelWriter<OutputFetchedBatch> batchWriter, CancellationToken cancellationToken)
 		{
 			while (await batchReader.WaitToReadAsync(cancellationToken))
 			{
@@ -558,82 +540,143 @@
 			}
 		}
 
-		record class OutputExport(BundleHandle BundleHandle, int PacketOffset, int ExportIdx, OutputChunk Chunk);
-		record class OutputExportBatch(BundleHandle BundleHandle, List<OutputExport> Exports);
+		#endregion
 
-		static async Task ReadBatchesAsync(ChannelReader<OutputChunk> chunkReader, ChannelWriter<OutputBatch> batchWriter, CancellationToken cancellationToken)
+		#region Write to disk
+
+		static async Task ExtractAsync(ChannelReader<OutputFetchedBatch> batchReader, CopyStats? copyStats, ILogger logger, CancellationToken cancellationToken)
 		{
-			const int MaxQueueLength = 200;
-
-			int queueLength = 0;
-			Queue<OutputExportBatch> exportBatchQueue = new Queue<OutputExportBatch>();
-			Dictionary<BundleHandle, OutputExportBatch> bundleHandleToExportBatch = new Dictionary<BundleHandle, OutputExportBatch>();
-
-			for (; ; )
+			const int WriteBatchSize = 64;
+			while (await batchReader.WaitToReadAsync(cancellationToken))
 			{
-				// Fill the queue up to the max length
-				while (queueLength < MaxQueueLength)
+				OutputFetchedBatch? batch;
+				while (batchReader.TryRead(out batch))
 				{
-					OutputChunk? chunk;
-					if (!chunkReader.TryRead(out chunk))
+					try
 					{
-						if (!await chunkReader.WaitToReadAsync(cancellationToken))
+						List<Task> tasks = new List<Task>();
+						foreach (IReadOnlyList<OutputFetchedChunk> group in batch.Chunks.Batch(WriteBatchSize))
 						{
-							break;
+							tasks.Add(ExtractChunksAsync(group.ToArray(), copyStats, logger, cancellationToken));
 						}
+						await Task.WhenAll(tasks);
 					}
-					else
+					finally
 					{
-						OutputExport? outputExport;
-						if (!TryGetOutputExport(chunk, out outputExport))
-						{
-							OutputBatch batch = new OutputBatch(new List<OutputChunk> { chunk });
-							await batchWriter.WriteAsync(batch, cancellationToken);
-						}
-						else
-						{
-							BundleHandle bundleHandle = outputExport.BundleHandle;
-							if (!bundleHandleToExportBatch.TryGetValue(bundleHandle, out OutputExportBatch? existingExportBatch))
-							{
-								existingExportBatch = new OutputExportBatch(bundleHandle, new List<OutputExport>());
-								exportBatchQueue.Enqueue(existingExportBatch);
-								bundleHandleToExportBatch.Add(bundleHandle, existingExportBatch);
-							}
-
-							existingExportBatch.Exports.Add(outputExport);
-							queueLength++;
-						}
+						batch.Dispose();
 					}
 				}
-
-				// Exit once we've processed everything and can't get any more items to read.
-				if (queueLength == 0)
-				{
-					batchWriter.TryComplete();
-					break;
-				}
-
-				// Flush the first queue
-				OutputExportBatch exportBatch = exportBatchQueue.Dequeue();
-				queueLength -= exportBatch.Exports.Count;
-				bundleHandleToExportBatch.Remove(exportBatch.BundleHandle);
-
-				List<OutputChunk> chunkBatch = exportBatch.Exports.OrderBy(x => x.PacketOffset).ThenBy(x => x.ExportIdx).Select(x => x.Chunk).ToList();
-				await batchWriter.WriteAsync(new OutputBatch(chunkBatch), cancellationToken);
 			}
 		}
 
-		static bool TryGetOutputExport(OutputChunk chunk, [NotNullWhen(true)] out OutputExport? export)
+		static async Task ExtractChunksAsync(ArraySegment<OutputFetchedChunk> chunks, CopyStats? copyStats, ILogger logger, CancellationToken cancellationToken)
 		{
-			if (chunk.Handle.Innermost is ExportHandle exportHandle && exportHandle.Packet is FlushedPacketHandle packetHandle)
+			for (int chunkIdx = 0; chunkIdx < chunks.Count;)
 			{
-				export = new OutputExport(packetHandle.Bundle, packetHandle.PacketOffset, exportHandle.ExportIdx, chunk);
-				return true;
+				OutputFile file = chunks[chunkIdx].File;
+
+				int maxChunkIdx = chunkIdx + 1;
+				while (maxChunkIdx < chunks.Count && chunks[maxChunkIdx].File == file)
+				{
+					maxChunkIdx++;
+				}
+
+				try
+				{
+					await ExtractChunksToFileAsync(file, chunks.Slice(chunkIdx, maxChunkIdx - chunkIdx), copyStats, logger, cancellationToken);
+					//await ExtractChunksToNullAsync(file, chunks.Slice(chunkIdx, maxChunkIdx - chunkIdx), copyStats, logger, cancellationToken);
+				}
+				catch (OperationCanceledException)
+				{
+					throw;
+				}
+				catch (Exception ex)
+				{
+					throw new StorageException($"Unable to extract {file?.FileInfo?.FullName}: {ex.Message}", ex);
+				}
+
+				chunkIdx = maxChunkIdx;
 			}
+		}
+
+		static async Task ExtractChunksToFileAsync(OutputFile file, ArraySegment<OutputFetchedChunk> chunks, CopyStats? copyStats, ILogger logger, CancellationToken cancellationToken)
+		{
+			// Open the file for the current chunk
+			int remainingChunks = 0;
+			await using (FileStream stream = file.OpenStream())
+			{
+				if (file.FileEntry.Length == 0)
+				{
+					// If this file is empty, don't write anything and just move to the next chunk
+					for (int idx = 0; idx < chunks.Count; idx++)
+					{
+						remainingChunks = file.DecrementRemaining();
+					}
+				}
+				else
+				{
+					// Process as many chunks as we can for this file
+					using MemoryMappedFile memoryMappedFile = MemoryMappedFile.CreateFromFile(stream, null, file.FileEntry.Length, MemoryMappedFileAccess.ReadWrite, HandleInheritability.None, false);
+					using MemoryMappedView memoryMappedView = new MemoryMappedView(memoryMappedFile, 0, file.FileEntry.Length);
+
+					for (int chunkIdx = 0; chunkIdx < chunks.Count; chunkIdx++)
+					{
+						OutputFetchedChunk chunk = chunks[chunkIdx];
+						cancellationToken.ThrowIfCancellationRequested();
+
+						// Write this chunk
+						TraceBlobRead("Leaf", chunk.File.Path, chunk.Handle, logger);
+						chunk.BlobData.Data.CopyTo(memoryMappedView!.GetMemory(chunk.Offset, chunk.BlobData.Data.Length));
+
+						// Update the stats
+						remainingChunks = file.DecrementRemaining();
+						copyStats?.Update(0, chunk.Length);
+					}
+				}
+			}
+
+			// Set correct permissions on the output file
+			if (remainingChunks == 0)
+			{
+				FileEntry.SetPermissions(file.FileInfo!, file.FileEntry.Flags);
+				copyStats?.Update(1, 0);
+			}
+		}
+
+#pragma warning disable IDE0051
+		// Update counters for extracting chunks without writing any data. Useful for profiling bottlenecks in other stages of the pipeline.
+		static Task ExtractChunksToNullAsync(OutputFile file, ArraySegment<OutputFetchedChunk> chunks, CopyStats? copyStats, ILogger logger, CancellationToken cancellationToken)
+		{
+			_ = logger;
+			_ = cancellationToken;
+
+			foreach (OutputFetchedChunk chunk in chunks)
+			{
+				int remainingChunks = file.DecrementRemaining();
+				if (remainingChunks == 0)
+				{
+					copyStats?.Update(1, chunk.Length);
+				}
+				else
+				{
+					copyStats?.Update(0, chunk.Length);
+				}
+			}
+			return Task.CompletedTask;
+		}
+#pragma warning restore IDE0051
+
+		#endregion
+
+		static string CombinePaths(string basePath, string nextPath)
+		{
+			if (basePath.Length > 0)
+			{
+				return $"{basePath}/{nextPath}";
+			}
 			else
 			{
-				export = null;
-				return false;
+				return nextPath;
 			}
 		}
 	}
diff --git a/Engine/Source/Programs/Shared/EpicGames.Horde/Storage/StorageBackendCache.cs b/Engine/Source/Programs/Shared/EpicGames.Horde/Storage/StorageBackendCache.cs
--- a/Engine/Source/Programs/Shared/EpicGames.Horde/Storage/StorageBackendCache.cs
+++ b/Engine/Source/Programs/Shared/EpicGames.Horde/Storage/StorageBackendCache.cs
@@ -364,11 +364,11 @@ namespace EpicGames.Horde.Storage
 		/// </summary>
 		public void GetStats(StorageStats stats)
 		{
-			stats.Add("Cache clean count", _cleanCount);
-			stats.Add("Cache fetch time (ms)", (_fetchTimeTicks * 1000) / Stopwatch.Frequency);
-			stats.Add("Cache clean time (ms)", (_cleanTimeTicks * 1000) / Stopwatch.Frequency);
-			stats.Add("Cache write time (ms)", (_writeTimeTicks * 1000) / Stopwatch.Frequency);
-			stats.Add("Cache fetch bytes", _fetchBytes);
+			stats.Add("backend.cache.clean_count", _cleanCount);
+			stats.Add("backend.cache.fetch_time_ms", (_fetchTimeTicks * 1000) / Stopwatch.Frequency);
+			stats.Add("backend.cache.clean_time_ms", (_cleanTimeTicks * 1000) / Stopwatch.Frequency);
+			stats.Add("backend.cache.write_time_ms", (_writeTimeTicks * 1000) / Stopwatch.Frequency);
+			stats.Add("backend.cache.fetch_bytes", _fetchBytes);
 		}
 	}
 }
