Horde: Allow setting a tag for unmanaged memory allocations that helps identify it for tracking down memory leaks.

diff --git a/Engine/Source/Programs/Shared/EpicGames.Core.Tests/MemoryAllocatorTests.cs b/Engine/Source/Programs/Shared/EpicGames.Core.Tests/MemoryAllocatorTests.cs
--- a/Engine/Source/Programs/Shared/EpicGames.Core.Tests/MemoryAllocatorTests.cs
+++ b/Engine/Source/Programs/Shared/EpicGames.Core.Tests/MemoryAllocatorTests.cs
@@ -13,7 +13,7 @@ namespace EpicGames.Core.Tests
 		{
 			GlobalHeapAllocator allocator = new GlobalHeapAllocator();
 
-			using IMemoryOwner<byte> handle = allocator.Alloc(1024);
+			using IMemoryOwner<byte> handle = allocator.Alloc(1024, null);
 			handle.Memory.Span[0] = 123;
 		}
 
@@ -22,7 +22,7 @@ namespace EpicGames.Core.Tests
 		{
 			VirtualMemoryAllocator allocator = new VirtualMemoryAllocator();
 
-			using IMemoryOwner<byte> handle = allocator.Alloc(1024);
+			using IMemoryOwner<byte> handle = allocator.Alloc(1024, null);
 			handle.Memory.Span[0] = 123;
 		}
 	}
diff --git a/Engine/Source/Programs/Shared/EpicGames.Core.Tests/MemoryWriterTests.cs b/Engine/Source/Programs/Shared/EpicGames.Core.Tests/MemoryWriterTests.cs
--- a/Engine/Source/Programs/Shared/EpicGames.Core.Tests/MemoryWriterTests.cs
+++ b/Engine/Source/Programs/Shared/EpicGames.Core.Tests/MemoryWriterTests.cs
@@ -43,7 +43,7 @@ namespace EpicGames.Core.Tests
 			public int AllocatedSize => _allocatedSize;
 			public int AllocatedCount => _allocatedCount;
 
-			public IMemoryOwner<byte> Alloc(int size)
+			public IMemoryOwner<byte> Alloc(int size, object? tag = null)
 			{
 				Interlocked.Increment(ref _allocatedCount);
 				Interlocked.Add(ref _allocatedSize, size);
diff --git a/Engine/Source/Programs/Shared/EpicGames.Core/IMemoryAllocator.cs b/Engine/Source/Programs/Shared/EpicGames.Core/IMemoryAllocator.cs
--- a/Engine/Source/Programs/Shared/EpicGames.Core/IMemoryAllocator.cs
+++ b/Engine/Source/Programs/Shared/EpicGames.Core/IMemoryAllocator.cs
@@ -2,6 +2,7 @@
 
 using System;
 using System.Buffers;
+using System.Collections.Generic;
 using System.IO.MemoryMappedFiles;
 using System.Runtime.InteropServices;
 using System.Threading;
@@ -18,8 +19,9 @@ namespace EpicGames.Core
 		/// Allocate a block of memory of the given minimum size
 		/// </summary>
 		/// <param name="minSize">Minimum size for the allocated data</param>
+		/// <param name="tag">Tag for tracking this allocation</param>
 		/// <returns>Reference counted block of memory</returns>
-		IMemoryOwner<T> Alloc(int minSize);
+		IMemoryOwner<T> Alloc(int minSize, object? tag);
 	}
 
 	/// <summary>
@@ -65,7 +67,7 @@ namespace EpicGames.Core
 		public long AllocatedSize => _allocatedSize;
 
 		/// <inheritdoc/>
-		public IMemoryOwner<byte> Alloc(int minSize) => new MemoryOwner(this, minSize);
+		public IMemoryOwner<byte> Alloc(int minSize, object? tag) => new MemoryOwner(this, minSize);
 	}
 
 	/// <summary>
@@ -86,7 +88,7 @@ namespace EpicGames.Core
 		public PoolAllocator(MemoryPool<byte> pool) => _pool = pool;
 
 		/// <inheritdoc/>
-		public IMemoryOwner<byte> Alloc(int minSize) => _pool.Rent(minSize);
+		public IMemoryOwner<byte> Alloc(int minSize, object? tag) => _pool.Rent(minSize);
 	}
 
 	/// <summary>
@@ -148,7 +150,7 @@ namespace EpicGames.Core
 		public long AllocatedSize => _allocatedSize;
 
 		/// <inheritdoc/>
-		public IMemoryOwner<byte> Alloc(int minSize) => new Allocation(this, minSize);
+		public IMemoryOwner<byte> Alloc(int minSize, object? tag) => new Allocation(this, minSize);
 	}
 	
 	/// <summary>
@@ -205,6 +207,73 @@ namespace EpicGames.Core
 		}
 
 		/// <inheritdoc/>
-		public IMemoryOwner<byte> Alloc(int minSize) => new Allocation(minSize);
+		public IMemoryOwner<byte> Alloc(int minSize, object? tag) => new Allocation(minSize);
+	}
+
+	/// <summary>
+	/// Implementation of <see cref="IMemoryAllocator{Byte}"/> which collects stats on the allocations performed.
+	/// </summary>
+	public class TrackingMemoryAllocator : IMemoryAllocator<byte>
+	{
+		// Tracks an owned blocks of memory against the cache budged.
+		sealed class MemoryAllocation : IMemoryOwner<byte>
+		{
+			readonly TrackingMemoryAllocator _allocator;
+			readonly LinkedListNode<MemoryAllocation> _node;
+			readonly object? _tag;
+			IMemoryOwner<byte> _allocation;
+
+			public Memory<byte> Memory => _allocation.Memory;
+
+			public MemoryAllocation(TrackingMemoryAllocator allocator, IMemoryOwner<byte> allocation, object? tag)
+			{
+				_allocator = allocator;
+				_allocation = allocation;
+				_tag = tag;
+
+				lock (allocator._lockObject)
+				{
+					_allocator._allocatedSize += allocation.Memory.Length;
+					_node = allocator._allocations.AddLast(this);
+				}
+			}
+
+			/// <inheritdoc/>
+			public void Dispose()
+			{
+				if (_allocation != null)
+				{
+					_allocator._allocatedSize -= _allocation.Memory.Length;
+					_allocator._allocations.Remove(_node);
+
+					_allocation.Dispose();
+					_allocation = null!;
+				}
+			}
+
+			/// <inheritdoc/>
+			public override string ToString()
+				=> $"{_tag ?? "???"} ({_allocation?.Memory.Length ?? 0} bytes)";
+		}
+
+		readonly object _lockObject = new object();
+		readonly IMemoryAllocator<byte> _inner;
+		long _allocatedSize;
+		readonly LinkedList<MemoryAllocation> _allocations = new LinkedList<MemoryAllocation>();
+
+		/// <summary>
+		/// Total size of allocated memory
+		/// </summary>
+		public long AllocatedSize => _allocatedSize;
+
+		/// <summary>
+		/// Constructor
+		/// </summary>
+		public TrackingMemoryAllocator(IMemoryAllocator<byte> inner)
+			=> _inner = inner;
+
+		/// <inheritdoc/>
+		public IMemoryOwner<byte> Alloc(int minSize, object? tag)
+			=> new MemoryAllocation(this, _inner.Alloc(minSize, tag), tag);
 	}
 }
diff --git a/Engine/Source/Programs/Shared/EpicGames.Core/MemoryWriter.cs b/Engine/Source/Programs/Shared/EpicGames.Core/MemoryWriter.cs
--- a/Engine/Source/Programs/Shared/EpicGames.Core/MemoryWriter.cs
+++ b/Engine/Source/Programs/Shared/EpicGames.Core/MemoryWriter.cs
@@ -535,13 +535,13 @@ namespace EpicGames.Core
 		/// <param name="initialSize">Size of the initial chunk</param>
 		/// <param name="chunkSize">Default size for subsequent chunks</param>
 		public ChunkedMemoryWriter(IMemoryAllocator<byte> allocator, int initialSize = 4096, int chunkSize = 4096)
-			: base(new PooledChunk(0, allocator.Alloc(initialSize)), chunkSize)
+			: base(new PooledChunk(0, allocator.Alloc(initialSize, null)), chunkSize)
 		{
 			_allocator = allocator;
 		}
 
 		/// <inheritdoc/>
-		protected override Chunk CreateChunk(int runningIndex, int size) => new PooledChunk(runningIndex, _allocator.Alloc(size));
+		protected override Chunk CreateChunk(int runningIndex, int size) => new PooledChunk(runningIndex, _allocator.Alloc(size, null));
 
 		/// <inheritdoc/>
 		public void Dispose()
@@ -587,14 +587,16 @@ namespace EpicGames.Core
 		}
 
 		readonly IMemoryAllocator<byte> _allocator;
+		readonly object? _allocationTag;
 
 		/// <summary>
 		/// Constructor
 		/// </summary>
 		/// <param name="allocator">Allocator for buffers</param>
 		/// <param name="chunkSize"></param>
-		public RefCountedMemoryWriter(IMemoryAllocator<byte> allocator, int chunkSize)
-			: this(allocator, chunkSize, chunkSize)
+		/// <param name="allocationTag">Tag for allocated memory</param>
+		public RefCountedMemoryWriter(IMemoryAllocator<byte> allocator, int chunkSize, object? allocationTag = null)
+			: this(allocator, chunkSize, chunkSize, allocationTag)
 		{ }
 
 		/// <summary>
@@ -603,14 +605,16 @@ namespace EpicGames.Core
 		/// <param name="allocator">Allocator for buffers</param>
 		/// <param name="initialSize">Size of the initial chunk</param>
 		/// <param name="chunkSize">Default size for subsequent chunks</param>
-		public RefCountedMemoryWriter(IMemoryAllocator<byte> allocator, int initialSize, int chunkSize)
-			: base(new BufferChunk(0, RefCountedHandle.Create(allocator.Alloc(initialSize))), chunkSize)
+		/// <param name="allocationTag">Tag for allocated memory</param>
+		public RefCountedMemoryWriter(IMemoryAllocator<byte> allocator, int initialSize, int chunkSize, object? allocationTag)
+			: base(new BufferChunk(0, RefCountedHandle.Create(allocator.Alloc(initialSize, allocationTag))), chunkSize)
 		{
 			_allocator = allocator;
+			_allocationTag = allocationTag;
 		}
 
 		/// <inheritdoc/>
-		protected override Chunk CreateChunk(int runningIndex, int size) => new BufferChunk(runningIndex, RefCountedHandle.Create(_allocator.Alloc(size)));
+		protected override Chunk CreateChunk(int runningIndex, int size) => new BufferChunk(runningIndex, RefCountedHandle.Create(_allocator.Alloc(size, _allocationTag)));
 
 		/// <inheritdoc/>
 		public void Dispose()
@@ -641,7 +645,7 @@ namespace EpicGames.Core
 
 			try
 			{
-				IRefCountedHandle<Memory<byte>> allocation = RefCountedHandle.Create(_allocator.Alloc(length));
+				IRefCountedHandle<Memory<byte>> allocation = RefCountedHandle.Create(_allocator.Alloc(length, _allocationTag));
 				sequence.Target.CopyTo(allocation.Target.Span);
 				return RefCountedHandle.Create<ReadOnlyMemory<byte>>(allocation.Target.Slice(0, length), allocation);
 			}
diff --git a/Engine/Source/Programs/Shared/EpicGames.Horde.Tests/BundleCacheTests.cs b/Engine/Source/Programs/Shared/EpicGames.Horde.Tests/BundleCacheTests.cs
--- a/Engine/Source/Programs/Shared/EpicGames.Horde.Tests/BundleCacheTests.cs
+++ b/Engine/Source/Programs/Shared/EpicGames.Horde.Tests/BundleCacheTests.cs
@@ -18,12 +18,12 @@ namespace EpicGames.Horde.Tests
 			await using BundleCache cache = new BundleCache(new BundleCacheOptions { MaxSize = 4096 });
 
 			Assert.AreEqual(0, cache.CurrentSize);
-			using (IMemoryOwner<byte> owner = cache.Allocator.Alloc(1024))
+			using (IMemoryOwner<byte> owner = cache.Allocator.Alloc(1024, null))
 			{
 				Assert.AreEqual(1024, cache.CurrentSize);
 			}
 			Assert.AreEqual(0, cache.CurrentSize);
-			using (IMemoryOwner<byte> owner = cache.Allocator.Alloc(1024))
+			using (IMemoryOwner<byte> owner = cache.Allocator.Alloc(1024, null))
 			{
 				Assert.AreEqual(1024, cache.CurrentSize);
 			}
@@ -45,7 +45,7 @@ namespace EpicGames.Horde.Tests
 
 			Assert.AreEqual(0, cache.CurrentSize);
 
-			IRefCountedHandle<TestData> result = await cache.FindOrAddAsync("test", (key, ctx) => Task.FromResult(new TestData(cache.Allocator.Alloc(1000))));
+			IRefCountedHandle<TestData> result = await cache.FindOrAddAsync("test", (key, ctx) => Task.FromResult(new TestData(cache.Allocator.Alloc(1000, null))));
 			Assert.AreEqual(1000, cache.CurrentSize);
 			Assert.AreEqual(2, result.RefCount); // cache and result
 
@@ -68,12 +68,12 @@ namespace EpicGames.Horde.Tests
 			Assert.IsNull(cache.Find<string, TestData>("other"));
 
 			// Allocate above the 1024 byte budget and check it's released
-			IMemoryOwner<byte> owner = cache.Allocator.Alloc(25);
+			IMemoryOwner<byte> owner = cache.Allocator.Alloc(25, null);
 			Assert.AreEqual(25, cache.CurrentSize);
 			Assert.IsNull(cache.Find<string, TestData>("test"));
 
 			// Allow allocating more than the limit if we don't have any cache values to free
-			IMemoryOwner<byte> owner2 = cache.Allocator.Alloc(2000);
+			IMemoryOwner<byte> owner2 = cache.Allocator.Alloc(2000, null);
 			Assert.AreEqual(2025, cache.CurrentSize);
 
 			owner2.Dispose();
@@ -88,11 +88,11 @@ namespace EpicGames.Horde.Tests
 
 			Assert.AreEqual(0, cache.CurrentSize);
 
-			IRefCountedHandle<TestData> result = await cache.FindOrAddAsync("test", (key, ctx) => Task.FromResult(new TestData(cache.Allocator.Alloc(1000))));
+			IRefCountedHandle<TestData> result = await cache.FindOrAddAsync("test", (key, ctx) => Task.FromResult(new TestData(cache.Allocator.Alloc(1000, null))));
 			Assert.AreEqual(1000, cache.CurrentSize);
 			result.Dispose();
 
-			IMemoryOwner<byte> owner = cache.Allocator.Alloc(10);
+			IMemoryOwner<byte> owner = cache.Allocator.Alloc(10, null);
 
 			Assert.AreEqual(1010, cache.CurrentSize);
 
@@ -111,11 +111,11 @@ namespace EpicGames.Horde.Tests
 			Assert.AreEqual(0, cache.CurrentSize);
 
 #pragma warning disable CA2000
-			Assert.IsTrue(cache.TryAdd("test", new TestData(cache.Allocator.Alloc(20))));
+			Assert.IsTrue(cache.TryAdd("test", new TestData(cache.Allocator.Alloc(20, null))));
 #pragma warning restore CA2000
 			Assert.AreEqual(20, cache.CurrentSize);
 
-			IRefCountedHandle<TestData> result = await cache.FindOrAddAsync("test", (key, ctx) => Task.FromResult(new TestData(cache.Allocator.Alloc(20))));
+			IRefCountedHandle<TestData> result = await cache.FindOrAddAsync("test", (key, ctx) => Task.FromResult(new TestData(cache.Allocator.Alloc(20, null))));
 			Assert.AreEqual(20, cache.CurrentSize);
 
 			cache.Trim(); 
@@ -133,10 +133,10 @@ namespace EpicGames.Horde.Tests
 
 			Assert.AreEqual(0, cache.CurrentSize);
 
-			IRefCountedHandle<TestData> result = await cache.FindOrAddAsync("test", (key, ctx) => Task.FromResult(new TestData(cache.Allocator.Alloc(20))));
+			IRefCountedHandle<TestData> result = await cache.FindOrAddAsync("test", (key, ctx) => Task.FromResult(new TestData(cache.Allocator.Alloc(20, null))));
 			Assert.AreEqual(20, cache.CurrentSize);
 
-			TestData testData = new TestData(cache.Allocator.Alloc(20));
+			TestData testData = new TestData(cache.Allocator.Alloc(20, null));
 			Assert.IsFalse(cache.TryAdd("test", testData));
 			Assert.AreEqual(40, cache.CurrentSize);
 			testData.Dispose();
diff --git a/Engine/Source/Programs/Shared/EpicGames.Horde/Storage/Bundles/BundleCache.cs b/Engine/Source/Programs/Shared/EpicGames.Horde/Storage/Bundles/BundleCache.cs
--- a/Engine/Source/Programs/Shared/EpicGames.Horde/Storage/Bundles/BundleCache.cs
+++ b/Engine/Source/Programs/Shared/EpicGames.Horde/Storage/Bundles/BundleCache.cs
@@ -127,14 +127,14 @@ namespace EpicGames.Horde.Storage
 		// Tracks an owned blocks of memory against the cache budged.
 		sealed class MemoryAllocation : IMemoryOwner<byte>
 		{
-			readonly BundleCache _outer;
+			readonly BundleCache _bundleCache;
 			IMemoryOwner<byte> _owner;
 
 			public Memory<byte> Memory => _owner.Memory;
 
-			public MemoryAllocation(BundleCache outer, IMemoryOwner<byte> owner)
+			public MemoryAllocation(BundleCache bundleCache, IMemoryOwner<byte> owner)
 			{
-				_outer = outer;
+				_bundleCache = bundleCache;
 				_owner = owner;
 			}
 
@@ -145,7 +145,7 @@ namespace EpicGames.Horde.Storage
 				{
 					long size = _owner.Memory.Length;
 					_owner.Dispose();
-					_outer.ReleaseSpace(size);
+					_bundleCache.ReleaseSpace(size);
 					_owner = null!;
 				}
 			}
@@ -155,21 +155,21 @@ namespace EpicGames.Horde.Storage
 		// disposed to create space for new allocations.
 		class MemoryAllocator : IMemoryAllocator<byte>
 		{
-			readonly BundleCache _outer;
-			readonly IMemoryAllocator<byte> _inner;
+			readonly BundleCache _bundleCache;
+			readonly IMemoryAllocator<byte> _allocator;
 
-			public MemoryAllocator(BundleCache outer, IMemoryAllocator<byte> inner)
+			public MemoryAllocator(BundleCache bundleCache, IMemoryAllocator<byte> allocator)
 			{
-				_outer = outer;
-				_inner = inner;
+				_bundleCache = bundleCache;
+				_allocator = allocator;
 			}
 
-			public IMemoryOwner<byte> Alloc(int minSize)
+			public IMemoryOwner<byte> Alloc(int minSize, object? tag)
 			{
-				_outer.CreateSpace(minSize);
-				IMemoryOwner<byte> owner = _inner.Alloc(minSize);
-				_outer.CreateSpace(owner.Memory.Length - minSize);
-				return new MemoryAllocation(_outer, owner);
+				_bundleCache.CreateSpace(minSize);
+				IMemoryOwner<byte> owner = _allocator.Alloc(minSize, tag);
+				_bundleCache.CreateSpace(owner.Memory.Length - minSize);
+				return new MemoryAllocation(_bundleCache, owner);
 			}
 		}
 
@@ -242,6 +242,9 @@ namespace EpicGames.Horde.Storage
 		public BundleCache(BundleCacheOptions options, IMemoryAllocator<byte> innerAllocator)
 		{
 			_options = options;
+#if DEBUG
+			innerAllocator = new TrackingMemoryAllocator(innerAllocator);
+#endif
 			_allocator = new MemoryAllocator(this, innerAllocator);
 
 			if (options.HeaderCacheSize > 0)
diff --git a/Engine/Source/Programs/Shared/EpicGames.Horde/Storage/Bundles/BundleHandle.cs b/Engine/Source/Programs/Shared/EpicGames.Horde/Storage/Bundles/BundleHandle.cs
--- a/Engine/Source/Programs/Shared/EpicGames.Horde/Storage/Bundles/BundleHandle.cs
+++ b/Engine/Source/Programs/Shared/EpicGames.Horde/Storage/Bundles/BundleHandle.cs
@@ -73,5 +73,8 @@ namespace EpicGames.Horde.Storage.Bundles
 
 		/// <inheritdoc/>
 		public override int GetHashCode() => _locator.GetHashCode();
+
+		/// <inheritdoc/>
+		public override string ToString() => _locator.ToString();
 	}
 }
diff --git a/Engine/Source/Programs/Shared/EpicGames.Horde/Storage/Bundles/BundleStorageClient.cs b/Engine/Source/Programs/Shared/EpicGames.Horde/Storage/Bundles/BundleStorageClient.cs
--- a/Engine/Source/Programs/Shared/EpicGames.Horde/Storage/Bundles/BundleStorageClient.cs
+++ b/Engine/Source/Programs/Shared/EpicGames.Horde/Storage/Bundles/BundleStorageClient.cs
@@ -126,7 +126,7 @@ namespace EpicGames.Horde.Storage.Bundles
 			{
 				BundleSignature signature = BundleSignature.Read(data.Span);
 
-				using IRefCountedHandle<Bundles.V2.Packet> packet = Bundles.V2.Packet.Decode(data, Cache.Allocator);
+				using IRefCountedHandle<Bundles.V2.Packet> packet = Bundles.V2.Packet.Decode(data, Cache.Allocator, nameof(ReadImportsFromDataV2));
 				for (int idx = 0; idx < packet.Target.GetImportCount(); idx++)
 				{
 					PacketImport import = packet.Target.GetImport(idx);
diff --git a/Engine/Source/Programs/Shared/EpicGames.Horde/Storage/Bundles/V2/BundleWriter.cs b/Engine/Source/Programs/Shared/EpicGames.Horde/Storage/Bundles/V2/BundleWriter.cs
--- a/Engine/Source/Programs/Shared/EpicGames.Horde/Storage/Bundles/V2/BundleWriter.cs
+++ b/Engine/Source/Programs/Shared/EpicGames.Horde/Storage/Bundles/V2/BundleWriter.cs
@@ -87,7 +87,7 @@ namespace EpicGames.Horde.Storage.Bundles.V2
 				_cache = cache;
 				_bundleOptions = bundleOptions;
 
-				_encodedPacketWriter = new RefCountedMemoryWriter(_cache.Allocator, 65536);
+				_encodedPacketWriter = new RefCountedMemoryWriter(_cache.Allocator, 65536, nameof(PendingBundleHandle));
 
 				StartPacket();
 			}
diff --git a/Engine/Source/Programs/Shared/EpicGames.Horde/Storage/Bundles/V2/Packet.cs b/Engine/Source/Programs/Shared/EpicGames.Horde/Storage/Bundles/V2/Packet.cs
--- a/Engine/Source/Programs/Shared/EpicGames.Horde/Storage/Bundles/V2/Packet.cs
+++ b/Engine/Source/Programs/Shared/EpicGames.Horde/Storage/Bundles/V2/Packet.cs
@@ -204,8 +204,7 @@ namespace EpicGames.Horde.Storage.Bundles.V2
 		/// <summary>
 		/// Decodes this packet
 		/// </summary>
-		/// <returns></returns>
-		public static IRefCountedHandle<Packet> Decode(ReadOnlyMemory<byte> data, IMemoryAllocator<byte> allocator)
+		public static IRefCountedHandle<Packet> Decode(ReadOnlyMemory<byte> data, IMemoryAllocator<byte> allocator, object? allocationTag)
 		{
 			BundleSignature signature = BundleSignature.Read(data.Span);
 			if (signature.Version <= BundleVersion.LatestV1 || signature.Version > BundleVersion.LatestV2)
@@ -222,7 +221,7 @@ namespace EpicGames.Horde.Storage.Bundles.V2
 			BundleCompressionFormat format = (BundleCompressionFormat)span[0];
 			span = span[1..];
 
-			IMemoryOwner<byte> owner = allocator.Alloc(decodedLength);
+			IMemoryOwner<byte> owner = allocator.Alloc(decodedLength, allocationTag);
 			Memory<byte> memory = owner.Memory.Slice(0, decodedLength);
 
 			BundleData.Decompress(format, data.Slice(data.Length - span.Length), memory);
diff --git a/Engine/Source/Programs/Shared/EpicGames.Horde/Storage/Bundles/V2/PacketHandle.cs b/Engine/Source/Programs/Shared/EpicGames.Horde/Storage/Bundles/V2/PacketHandle.cs
--- a/Engine/Source/Programs/Shared/EpicGames.Horde/Storage/Bundles/V2/PacketHandle.cs
+++ b/Engine/Source/Programs/Shared/EpicGames.Horde/Storage/Bundles/V2/PacketHandle.cs
@@ -203,8 +203,17 @@ namespace EpicGames.Horde.Storage.Bundles.V2
 
 		#region Packet reader access
 
-		record struct PacketReaderCacheKey(BundleHandle Bundle, int Offset);
-		record struct EncodedPacketCacheKey(BundleHandle Bundle, int Offset);
+		record struct PacketReaderCacheKey(BundleHandle Bundle, int Offset)
+		{
+			public override string ToString()
+				=> $"packet-reader:{Bundle}@{Offset}";
+		}
+
+		record struct EncodedPacketCacheKey(BundleHandle Bundle, int Offset)
+		{
+			public override string ToString()
+				=> $"encoded-packet:{Bundle}@{Offset}";
+		}
 
 		async ValueTask<IRefCountedHandle<PacketReader>> GetPacketReaderAsync(CancellationToken cancellationToken = default)
 		{
@@ -215,7 +224,7 @@ namespace EpicGames.Horde.Storage.Bundles.V2
 		async Task<PacketReader> CreatePacketReaderAsync(PacketReaderCacheKey cacheKey, CancellationToken cancellationToken)
 		{
 			using IRefCountedHandle<IReadOnlyMemoryOwner<byte>> encodedData = await ReadEncodedPacketAsync(cancellationToken);
-			IRefCountedHandle<Packet> packet = Packet.Decode(encodedData.Target.Memory, _cache.Allocator);
+			IRefCountedHandle<Packet> packet = Packet.Decode(encodedData.Target.Memory, _cache.Allocator, cacheKey);
 			return new PacketReader(_storageClient, _cache, _outer, this, packet.Target, packet);
 		}
 
@@ -236,7 +245,7 @@ namespace EpicGames.Horde.Storage.Bundles.V2
 				using Stream stream = await _outer.OpenAsync(_packetOffset, readLength, cancellationToken);
 
 				// Read the first packet
-				leadingPacket = _cache.Allocator.Alloc(_packetLength);
+				leadingPacket = _cache.Allocator.Alloc(_packetLength, key);
 				Memory<byte> memory = leadingPacket.Memory.Slice(0, _packetLength);
 				await stream.ReadFixedLengthBytesAsync(memory, cancellationToken);
 
@@ -263,7 +272,7 @@ namespace EpicGames.Horde.Storage.Bundles.V2
 						break;
 					}
 
-					trailingPacket = _cache.Allocator.Alloc(signature.HeaderLength);
+					trailingPacket = _cache.Allocator.Alloc(signature.HeaderLength, key);
 					memory = trailingPacket.Memory.Slice(0, signature.HeaderLength);
 					header.CopyTo(memory);
 					await stream.ReadFixedLengthBytesAsync(memory.Slice(BundleSignature.NumBytes), cancellationToken);
diff --git a/Engine/Source/Programs/Shared/EpicGames.Horde/Storage/Bundles/V2/PacketWriter.cs b/Engine/Source/Programs/Shared/EpicGames.Horde/Storage/Bundles/V2/PacketWriter.cs
--- a/Engine/Source/Programs/Shared/EpicGames.Horde/Storage/Bundles/V2/PacketWriter.cs
+++ b/Engine/Source/Programs/Shared/EpicGames.Horde/Storage/Bundles/V2/PacketWriter.cs
@@ -14,6 +14,8 @@ namespace EpicGames.Horde.Storage.Bundles.V2
 	/// </summary>
 	public sealed class PacketWriter : IDisposable
 	{
+		const string AllocationTag = nameof(PacketWriter);
+
 		readonly BundleHandle _bundleHandle;
 		readonly PacketHandle _packetHandle;
 		readonly IMemoryAllocator<byte> _allocator;
@@ -56,7 +58,7 @@ namespace EpicGames.Horde.Storage.Bundles.V2
 			_importMap[_bundleHandle] = PacketImport.CurrentBundleBaseIdx;
 			_importMap[_packetHandle] = PacketImport.CurrentPacketBaseIdx;
 
-			_bufferHandle = RefCountedHandle.Create(_allocator.Alloc(1024));
+			_bufferHandle = RefCountedHandle.Create(_allocator.Alloc(1024, AllocationTag));
 			_buffer = _bufferHandle.Target;
 
 			BundleSignature signature = new BundleSignature(BundleVersion.LatestV2, 0);
@@ -284,7 +286,7 @@ namespace EpicGames.Horde.Storage.Bundles.V2
 				{
 					int newSize = (_length + desiredSize + 4096 + 16384) & ~16384;
 
-					IRefCountedHandle<Memory<byte>> newBufferHandle = RefCountedHandle.Create(_allocator.Alloc(newSize));
+					IRefCountedHandle<Memory<byte>> newBufferHandle = RefCountedHandle.Create(_allocator.Alloc(newSize, AllocationTag));
 					_buffer.Slice(0, _length + usedSize).CopyTo(newBufferHandle.Target);
 					_buffer = newBufferHandle.Target;
 					_bufferHandle.Dispose();
