Horde: Fix bundle cache not correctly freeing items.

* Removing items from the cache would cause the search to free space to terminate immediately, since the linked list node is updated.
* Calling LinkedList.Remove(T) performs a linear search for the item (!). Passing the node object does not.

diff --git a/Engine/Source/Programs/Shared/EpicGames.Horde/Storage/Bundles/BundleCache.cs b/Engine/Source/Programs/Shared/EpicGames.Horde/Storage/Bundles/BundleCache.cs
--- a/Engine/Source/Programs/Shared/EpicGames.Horde/Storage/Bundles/BundleCache.cs
+++ b/Engine/Source/Programs/Shared/EpicGames.Horde/Storage/Bundles/BundleCache.cs
@@ -387,12 +387,17 @@ namespace EpicGames.Horde.Storage
 			{
 				lock (_lockObject)
 				{
-					for (LinkedListNode<CacheValue>? lastNode = _items.Last; lastNode != null && Interlocked.CompareExchange(ref _currentSize, 0, 0) > _options.MaxSize; lastNode = lastNode.Previous)
+					LinkedListNode<CacheValue>? node = _items.Last;
+					while (node != null && Interlocked.CompareExchange(ref _currentSize, 0, 0) > _options.MaxSize)
 					{
-						CacheValue lastItem = lastNode.Value;
+						// Get this item and move to the next item in the list before we consider disposing it
+						LinkedListNode<CacheValue> lastNode = node;
+						CacheValue lastItem = node.Value;
+						node = node.Previous;
+
 						if (lastItem.RefCount == 1)
 						{
-							_items.Remove(lastItem);
+							_items.Remove(lastNode);
 							_itemLookup.Remove(lastItem.Key);
 							lastItem.Dispose();
 						}
