Horde: Cache fixed size chunks read from bundles, rather than creating separate packet entries in the cache.

Some packets compress very well, so we can add thousands of packets through a single bundle read. This can result in more time spent updating the cache than actually doing useful work. Now reads pages aligned to strict 1mb boundaries and caches them in that form.

Also move purgeable items to a separate list for tracking in the cache, to avoid having to scan until we find a purgeable entry.

diff --git a/Engine/Source/Programs/Shared/EpicGames.Core/IReadOnlyMemoryOwner.cs b/Engine/Source/Programs/Shared/EpicGames.Core/IReadOnlyMemoryOwner.cs
--- a/Engine/Source/Programs/Shared/EpicGames.Core/IReadOnlyMemoryOwner.cs
+++ b/Engine/Source/Programs/Shared/EpicGames.Core/IReadOnlyMemoryOwner.cs
@@ -23,16 +23,26 @@ namespace EpicGames.Core
 	{
 		class DefaultReadOnlyMemoryOwner<T> : IReadOnlyMemoryOwner<T>
 		{
-			public ReadOnlyMemory<T> Memory { get; }
-			readonly IDisposable? _owner;
+			ReadOnlyMemory<T>? _memory;
+			IDisposable? _owner;
+
+			public ReadOnlyMemory<T> Memory => _memory ?? throw new ObjectDisposedException(typeof(DefaultReadOnlyMemoryOwner<T>).FullName);
 
 			public DefaultReadOnlyMemoryOwner(ReadOnlyMemory<T> memory, IDisposable? owner)
 			{
-				Memory = memory;
+				_memory = memory;
 				_owner = owner;
 			}
 
-			public void Dispose() => _owner?.Dispose();
+			public void Dispose()
+			{
+				_memory = null;
+				if (_owner != null)
+				{
+					_owner?.Dispose();
+					_owner = null;
+				}
+			}
 		}
 
 		/// <summary>
diff --git a/Engine/Source/Programs/Shared/EpicGames.Horde/Storage/Bundles/BundleCache.cs b/Engine/Source/Programs/Shared/EpicGames.Horde/Storage/Bundles/BundleCache.cs
--- a/Engine/Source/Programs/Shared/EpicGames.Horde/Storage/Bundles/BundleCache.cs
+++ b/Engine/Source/Programs/Shared/EpicGames.Horde/Storage/Bundles/BundleCache.cs
@@ -3,6 +3,7 @@
 using System;
 using System.Buffers;
 using System.Collections.Generic;
+using System.Diagnostics;
 using System.Diagnostics.CodeAnalysis;
 using System.Threading;
 using System.Threading.Tasks;
@@ -38,29 +39,38 @@ namespace EpicGames.Horde.Storage
 	/// </summary>
 	public sealed class BundleCache : IAsyncDisposable
 	{
+		[DebuggerDisplay("{Key}")]
 		class CacheValue : IDisposable
 		{
 			public object Key { get; }
 			public LinkedListNode<CacheValue> Node { get; }
 			public Task<IDisposable> InitTask { get; }
+
+			readonly BundleCache _bundleCache;
 			int _refCount;
 
 			public int RefCount => Interlocked.CompareExchange(ref _refCount, 0, 0);
 
-			public CacheValue(object key, IDisposable value)
+			public CacheValue(BundleCache bundleCache, object key, IDisposable value)
 			{
+				_bundleCache = bundleCache;
 				_refCount = 1;
 				Key = key;
 				Node = new LinkedListNode<CacheValue>(this);
 				InitTask = Task.FromResult<IDisposable>(value);
+
+				_bundleCache._purgeableItems.AddFirst(Node);
 			}
 
-			public CacheValue(object key, Func<Task<IDisposable>> initTask)
+			public CacheValue(BundleCache bundleCache, object key, Func<Task<IDisposable>> initTask)
 			{
+				_bundleCache = bundleCache;
 				_refCount = 2; // Will be released by RunAndUnlock
 				Key = key;
 				Node = new LinkedListNode<CacheValue>(this);
 				InitTask = RunAndUnlockAsync(initTask);
+
+				_bundleCache._items.AddFirst(Node);
 			}
 
 			public void Dispose()
@@ -83,8 +93,31 @@ namespace EpicGames.Horde.Storage
 				}
 			}
 
-			public void AddRef() => Interlocked.Increment(ref _refCount);
-			public void Release() => Interlocked.Decrement(ref _refCount);
+			public void AddRef()
+			{
+				lock (_bundleCache._lockObject)
+				{
+					if (_refCount == 1)
+					{
+						_bundleCache._purgeableItems.Remove(Node);
+						_bundleCache._items.AddFirst(Node);
+					}
+					_refCount++;
+				}
+			}
+
+			public void Release()
+			{
+				lock (_bundleCache._lockObject)
+				{
+					_refCount--;
+					if (_refCount == 1)
+					{
+						_bundleCache._items.Remove(Node);
+						_bundleCache._purgeableItems.AddFirst(Node);
+					}
+				}
+			}
 		}
 
 		class CacheValueHandle<T> : IRefCountedHandle<T> where T : class, IDisposable
@@ -177,12 +210,16 @@ namespace EpicGames.Horde.Storage
 		readonly BundleCacheOptions _options;
 		readonly LinkedList<CacheValue> _items = new LinkedList<CacheValue>();
 		readonly Dictionary<object, CacheValue> _itemLookup = new Dictionary<object, CacheValue>();
+		readonly LinkedList<CacheValue> _purgeableItems = new LinkedList<CacheValue>();
 		readonly CancellationTokenSource _cancellationSource = new CancellationTokenSource();
 		readonly MemoryAllocator _allocator;
 
+		int _numCacheAdds;
 		int _numCacheHits;
 		int _numCacheMisses;
 		long _currentSize;
+		int _numPurges;
+		readonly Stopwatch _createSpaceTimer = new Stopwatch();
 
 		readonly MemoryCache? _headerCache;
 		readonly MemoryCache? _packetCache;
@@ -329,8 +366,8 @@ namespace EpicGames.Horde.Storage
 			{
 				if (!_itemLookup.ContainsKey(key))
 				{
-					CacheValue item = new CacheValue(key, value);
-					_items.AddFirst(item);
+					Interlocked.Increment(ref _numCacheAdds);
+					CacheValue item = new CacheValue(this, key, value);
 					_itemLookup.Add(key, item);
 					return true;
 				}
@@ -363,11 +400,9 @@ namespace EpicGames.Horde.Storage
 					else
 					{
 						Interlocked.Increment(ref _numCacheMisses);
-						item = new CacheValue(key, async () => await createAsync(key, _cancellationSource.Token));
+						item = new CacheValue(this, key, async () => await createAsync(key, _cancellationSource.Token));
 						_itemLookup.Add(key, item);
 					}
-
-					_items.AddFirst(item);
 					item.AddRef(); // Don't allow the item to be freed while we wait for it
 				}
 
@@ -387,21 +422,20 @@ namespace EpicGames.Horde.Storage
 			{
 				lock (_lockObject)
 				{
-					LinkedListNode<CacheValue>? node = _items.Last;
-					while (node != null && Interlocked.CompareExchange(ref _currentSize, 0, 0) > _options.MaxSize)
+					Interlocked.Increment(ref _numPurges);
+					_createSpaceTimer.Start();
+
+					while (_purgeableItems.Count > 0 && Interlocked.CompareExchange(ref _currentSize, 0, 0) > _options.MaxSize)
 					{
-						// Get this item and move to the next item in the list before we consider disposing it
-						LinkedListNode<CacheValue> lastNode = node;
-						CacheValue lastItem = node.Value;
-						node = node.Previous;
-
-						if (lastItem.RefCount == 1)
-						{
-							_items.Remove(lastNode);
-							_itemLookup.Remove(lastItem.Key);
-							lastItem.Dispose();
-						}
+						LinkedListNode<CacheValue> lastNode = _purgeableItems.Last!;
+						CacheValue lastItem = lastNode.Value;
+
+						_purgeableItems.Remove(lastNode);
+						_itemLookup.Remove(lastItem.Key);
+						lastItem.Dispose();
 					}
+
+					_createSpaceTimer.Stop();
 				}
 			}
 		}
@@ -416,10 +450,13 @@ namespace EpicGames.Horde.Storage
 		/// </summary>
 		public void GetStats(StorageStats stats)
 		{
+			stats.Add("bundle.cache.adds", _numCacheAdds);
 			stats.Add("bundle.cache.hits", _numCacheHits);
 			stats.Add("bundle.cache.misses", _numCacheMisses);
 			stats.Add("bundle.cache.size_count", _items.Count);
 			stats.Add("bundle.cache.size_bytes", _currentSize);
+			stats.Add("bundle.cache.purge.count", _numPurges);
+			stats.Add("bundle.cache.purge.time_ms", _createSpaceTimer.ElapsedMilliseconds);
 		}
 
 		#region V1
diff --git a/Engine/Source/Programs/Shared/EpicGames.Horde/Storage/Bundles/V2/PacketHandle.cs b/Engine/Source/Programs/Shared/EpicGames.Horde/Storage/Bundles/V2/PacketHandle.cs
--- a/Engine/Source/Programs/Shared/EpicGames.Horde/Storage/Bundles/V2/PacketHandle.cs
+++ b/Engine/Source/Programs/Shared/EpicGames.Horde/Storage/Bundles/V2/PacketHandle.cs
@@ -4,6 +4,7 @@ using System;
 using System.Buffers;
 using System.Buffers.Text;
 using System.IO;
+using System.Linq;
 using System.Text;
 using System.Threading;
 using System.Threading.Tasks;
@@ -209,10 +210,10 @@ namespace EpicGames.Horde.Storage.Bundles.V2
 				=> $"packet-reader:{Bundle}@{Offset}";
 		}
 
-		record struct EncodedPacketCacheKey(BundleHandle Bundle, int Offset)
+		record struct BundlePageCacheKey(BundleHandle Bundle, int Index)
 		{
 			public override string ToString()
-				=> $"encoded-packet:{Bundle}@{Offset}";
+				=> $"bundle-page:{Bundle}:{Index}";
 		}
 
 		async ValueTask<IRefCountedHandle<PacketReader>> GetPacketReaderAsync(CancellationToken cancellationToken = default)
@@ -223,88 +224,62 @@ namespace EpicGames.Horde.Storage.Bundles.V2
 
 		async Task<PacketReader> CreatePacketReaderAsync(PacketReaderCacheKey cacheKey, CancellationToken cancellationToken)
 		{
-			using IRefCountedHandle<IReadOnlyMemoryOwner<byte>> encodedData = await ReadEncodedPacketAsync(cancellationToken);
-			IRefCountedHandle<Packet> packet = Packet.Decode(encodedData.Target.Memory, _cache.Allocator, cacheKey);
+			using IReadOnlyMemoryOwner<byte> encodedData = await ReadEncodedPacketAsync(cancellationToken);
+			IRefCountedHandle<Packet> packet = Packet.Decode(encodedData.Memory, _cache.Allocator, cacheKey);
 			return new PacketReader(_storageClient, _cache, _outer, this, packet.Target, packet);
 		}
 
-		async ValueTask<IRefCountedHandle<IReadOnlyMemoryOwner<byte>>> ReadEncodedPacketAsync(CancellationToken cancellationToken)
-		{
-			EncodedPacketCacheKey encodedPacketCacheKey = new EncodedPacketCacheKey(_outer, _packetOffset);
-			return await _cache.FindOrAddAsync(encodedPacketCacheKey, ReadEncodedPacketInternalAsync, cancellationToken);
-		}
+		const int BundlePageSize = 1024 * 1024;
 
-		async Task<IReadOnlyMemoryOwner<byte>> ReadEncodedPacketInternalAsync(EncodedPacketCacheKey key, CancellationToken cancellationToken)
+		async ValueTask<IReadOnlyMemoryOwner<byte>> ReadEncodedPacketAsync(CancellationToken cancellationToken)
 		{
-			IMemoryOwner<byte>? leadingPacket = null;
-			IMemoryOwner<byte>? trailingPacket = null;
+			int minPageIdx = _packetOffset / BundlePageSize;
+			int maxPageIdx = ((_packetOffset + _packetLength) + (BundlePageSize - 1)) / BundlePageSize;
+
+			// If the whole packet is contained within on page, just return that.
+			if (maxPageIdx == minPageIdx + 1)
+			{
+				IRefCountedHandle<IReadOnlyMemoryOwner<byte>> pageData = await ReadBundlePageAsync(minPageIdx, cancellationToken);
+				return ReadOnlyMemoryOwner.Create(pageData.Target.Memory.Slice(_packetOffset & (BundlePageSize - 1), _packetLength), pageData);
+			}
+
+			// Otherwise read sections from each page that contributes to the output and copy into a shared buffer
+			IMemoryOwner<byte> owner = PoolAllocator.Shared.Alloc(_packetLength, null);
 			try
 			{
-				// Read more data than was requested so we can add additional packets to the cache
-				int readLength = Math.Max(_packetLength, 512 * 1024);
-				using Stream stream = await _outer.OpenAsync(_packetOffset, readLength, cancellationToken);
-
-				// Read the first packet
-				leadingPacket = _cache.Allocator.Alloc(_packetLength, key);
-				Memory<byte> memory = leadingPacket.Memory.Slice(0, _packetLength);
-				await stream.ReadFixedLengthBytesAsync(memory, cancellationToken);
-
-				PacketReaderStats packetReaderStats = _storageClient.PacketReaderStats;
-				Interlocked.Increment(ref packetReaderStats._numReads);
-				Interlocked.Add(ref packetReaderStats._numBytesRead, _packetLength);
-				Interlocked.Increment(ref packetReaderStats._numPacketsRead);
-
-				// Read any other packets in the same stream
-				byte[] header = new byte[BundleSignature.NumBytes];
-				for (int readOffset = _packetLength; readOffset + BundleSignature.NumBytes < readLength;)
-				{
-					int readBytes = await stream.ReadGreedyAsync(header, cancellationToken);
-					if (readBytes < header.Length)
-					{
-						Interlocked.Add(ref packetReaderStats._numWastedBytesRead, readBytes);
-						break;
-					}
-
-					BundleSignature signature = BundleSignature.Read(header);
-					if (readOffset + signature.HeaderLength >= readLength)
-					{
-						Interlocked.Add(ref packetReaderStats._numWastedBytesRead, readLength - readOffset);
-						break;
-					}
-
-					trailingPacket = _cache.Allocator.Alloc(signature.HeaderLength, key);
-					memory = trailingPacket.Memory.Slice(0, signature.HeaderLength);
-					header.CopyTo(memory);
-					await stream.ReadFixedLengthBytesAsync(memory.Slice(BundleSignature.NumBytes), cancellationToken);
-
-					EncodedPacketCacheKey trailingKey = new EncodedPacketCacheKey(key.Bundle, _packetOffset + readOffset);
-
-#pragma warning disable CA2000
-					IReadOnlyMemoryOwner<byte> trailingValue = ReadOnlyMemoryOwner.Create<byte>(memory, trailingPacket);
-					if (!_cache.TryAdd(trailingKey, trailingValue))
-					{
-						trailingValue.Dispose();
-					}
-					trailingPacket = null;
-#pragma warning restore CA2000
-
-					readOffset += signature.HeaderLength;
-
-					Interlocked.Add(ref packetReaderStats._numBytesRead, signature.HeaderLength);
-					Interlocked.Increment(ref packetReaderStats._numPacketsRead);
-				}
+				await Parallel.ForEachAsync(Enumerable.Range(minPageIdx, maxPageIdx - minPageIdx), cancellationToken, async (pageIdx, ctx) => await ReadBundlePageAsync(pageIdx, owner.Memory, ctx));
+				return ReadOnlyMemoryOwner.Create<byte>(owner.Memory.Slice(0, _packetLength), owner);
 			}
 			catch
 			{
-#pragma warning disable CA1508 // trailingPacket is always null (...not if it's thrown an exception)
-				trailingPacket?.Dispose();
-#pragma warning restore CA1508
-				leadingPacket?.Dispose();
+				owner.Dispose();
 				throw;
 			}
-			return ReadOnlyMemoryOwner.Create<byte>(leadingPacket.Memory.Slice(0, _packetLength), leadingPacket);
 		}
 
-		#endregion
+		async ValueTask<IRefCountedHandle<IReadOnlyMemoryOwner<byte>>> ReadBundlePageAsync(int pageIdx, CancellationToken cancellationToken)
+		{
+			BundlePageCacheKey bundlePageCacheKey = new BundlePageCacheKey(_outer, pageIdx);
+			return await _cache.FindOrAddAsync(bundlePageCacheKey, ReadBundlePageInternalAsync, cancellationToken);
+		}
+
+		async Task ReadBundlePageAsync(int pageIdx, Memory<byte> targetMemory, CancellationToken cancellationToken)
+		{
+			using IRefCountedHandle<IReadOnlyMemoryOwner<byte>> pageData = await ReadBundlePageAsync(pageIdx, cancellationToken);
+
+			int pageBase = pageIdx * BundlePageSize;
+
+			int minOffset = Math.Max(_packetOffset, pageBase);
+			int maxOffset = Math.Min(_packetOffset + _packetLength, pageBase + BundlePageSize);
+
+			ReadOnlyMemory<byte> sourceMemory = pageData.Target.Memory.Slice(minOffset - pageBase, maxOffset - minOffset);
+			sourceMemory.CopyTo(targetMemory.Slice(minOffset - _packetOffset));
+		}
+
+		async Task<IReadOnlyMemoryOwner<byte>> ReadBundlePageInternalAsync(BundlePageCacheKey key, CancellationToken cancellationToken)
+		{
+			return await _outer.ReadAsync(key.Index * BundlePageSize, BundlePageSize, cancellationToken);
+		}
+#endregion
 	}
 }
diff --git a/Engine/Source/Programs/Shared/EpicGames.Horde/Storage/StorageBackendCache.cs b/Engine/Source/Programs/Shared/EpicGames.Horde/Storage/StorageBackendCache.cs
--- a/Engine/Source/Programs/Shared/EpicGames.Horde/Storage/StorageBackendCache.cs
+++ b/Engine/Source/Programs/Shared/EpicGames.Horde/Storage/StorageBackendCache.cs
@@ -107,7 +107,7 @@ namespace EpicGames.Horde.Storage
 
 			public BackendWrapper(string keyPrefix, StorageBackendCache cacheStorage, IStorageBackend inner)
 			{
-				_keyPrefix = keyPrefix;
+				_keyPrefix = keyPrefix.TrimEnd('/') + "/";
 				_cacheStorage = cacheStorage;
 				_inner = inner;
 			}
