Horde: Static analysis fixes.

diff --git a/Engine/Source/Programs/Shared/EpicGames.Core/IReadOnlyMemoryOwner.cs b/Engine/Source/Programs/Shared/EpicGames.Core/IReadOnlyMemoryOwner.cs
--- a/Engine/Source/Programs/Shared/EpicGames.Core/IReadOnlyMemoryOwner.cs
+++ b/Engine/Source/Programs/Shared/EpicGames.Core/IReadOnlyMemoryOwner.cs
@@ -39,7 +39,7 @@ namespace EpicGames.Core
 				_memory = null;
 				if (_owner != null)
 				{
-					_owner?.Dispose();
+					_owner.Dispose();
 					_owner = null;
 				}
 			}
diff --git a/Engine/Source/Programs/Shared/EpicGames.Horde/Storage/Bundles/V2/PacketHandle.cs b/Engine/Source/Programs/Shared/EpicGames.Horde/Storage/Bundles/V2/PacketHandle.cs
--- a/Engine/Source/Programs/Shared/EpicGames.Horde/Storage/Bundles/V2/PacketHandle.cs
+++ b/Engine/Source/Programs/Shared/EpicGames.Horde/Storage/Bundles/V2/PacketHandle.cs
@@ -3,7 +3,6 @@
 using System;
 using System.Buffers;
 using System.Buffers.Text;
-using System.IO;
 using System.Linq;
 using System.Text;
 using System.Threading;
@@ -280,6 +279,6 @@ namespace EpicGames.Horde.Storage.Bundles.V2
 		{
 			return await _outer.ReadAsync(key.Index * BundlePageSize, BundlePageSize, cancellationToken);
 		}
-#endregion
+		#endregion
 	}
 }
