Horde: Fix test failure.

diff --git a/Engine/Source/Programs/Shared/EpicGames.Horde/Storage/StorageBackendCache.cs b/Engine/Source/Programs/Shared/EpicGames.Horde/Storage/StorageBackendCache.cs
--- a/Engine/Source/Programs/Shared/EpicGames.Horde/Storage/StorageBackendCache.cs
+++ b/Engine/Source/Programs/Shared/EpicGames.Horde/Storage/StorageBackendCache.cs
@@ -107,7 +107,11 @@ namespace EpicGames.Horde.Storage
 
 			public BackendWrapper(string keyPrefix, StorageBackendCache cacheStorage, IStorageBackend inner)
 			{
-				_keyPrefix = keyPrefix.TrimEnd('/') + "/";
+				_keyPrefix = keyPrefix.TrimEnd('/');
+				if (_keyPrefix.Length > 0)
+				{
+					_keyPrefix += "/";
+				}
 				_cacheStorage = cacheStorage;
 				_inner = inner;
 			}
