- Moved all versions from *SDK.Versions.cs files to *_SDK.json files
- Some static variables made it tricky, so added some support to make it simpler on users of the static variables

diff --git a/Engine/Config/Android/Android_SDK.json b/Engine/Config/Android/Android_SDK.json
--- a/Engine/Config/Android/Android_SDK.json
+++ b/Engine/Config/Android/Android_SDK.json
@@ -1,7 +1,12 @@
 {
-"MainVersion" : "r25b",
-"AutoSDKDirectory" : "-26",
+	"MainVersion": "r25b",
+	"MinVersion": "r25b",
+	"MaxVersion": "r26c",
 
-"MinVersion" : "r25b",
-"MaxVersion" : "r26c"
+	"AutoSDKDirectory": "-26",
+
+	"platforms": "android-33",
+	"build-tools": "33.0.1",
+	"cmake": "3.22.1",
+	"ndk": "25.1.8937393",
 }
diff --git a/Engine/Config/Apple/Apple_SDK.json b/Engine/Config/Apple/Apple_SDK.json
--- a/Engine/Config/Apple/Apple_SDK.json
+++ b/Engine/Config/Apple/Apple_SDK.json
@@ -1,8 +1,10 @@
 {
-"MainVersion" : "14.1",
-"MinVersion" : "14.1.0",
-"MaxVersion" : "15.9.0",
-// The versions on Windows are iTunes versions
-"MinVersion_Win64" : "1100.0.0.0",
-"MaxVersion_Win64" : "8999.0",
+	// Xcode versions
+	"MainVersion": "14.1",
+	"MinVersion": "14.1.0",
+	"MaxVersion": "15.9.0",
+
+	// The versions on Windows are iTunes versions
+	"MinVersion_Win64": "1100.0.0.0",
+	"MaxVersion_Win64": "8999.0"
 }
diff --git a/Engine/Config/IOS/IOS_SDK.json b/Engine/Config/IOS/IOS_SDK.json
--- a/Engine/Config/IOS/IOS_SDK.json
+++ b/Engine/Config/IOS/IOS_SDK.json
@@ -1,3 +1,5 @@
 {
-"ParentSDKFile" : "../Apple/Apple_SDK.json"
+	"ParentSDKFile": "../Apple/Apple_SDK.json",
+
+	"MinSoftwareVersion":  "15.0",
 }
diff --git a/Engine/Config/Mac/Mac_SDK.json b/Engine/Config/Mac/Mac_SDK.json
--- a/Engine/Config/Mac/Mac_SDK.json
+++ b/Engine/Config/Mac/Mac_SDK.json
@@ -1,3 +1,6 @@
 {
-"ParentSDKFile" : "../Apple/Apple_SDK.json"
+	"ParentSDKFile": "../Apple/Apple_SDK.json",
+
+	// macOS Monterey 12.5 (Released 2022/07/22)
+	"MinSoftwareVersion": "12.5.0"
 }
diff --git a/Engine/Config/TVOS/TVOS_SDK.json b/Engine/Config/TVOS/TVOS_SDK.json
--- a/Engine/Config/TVOS/TVOS_SDK.json
+++ b/Engine/Config/TVOS/TVOS_SDK.json
@@ -1,3 +1,4 @@
 {
-"ParentSDKFile" : "../Apple/Apple_SDK.json"
+	// inherit the IOS MinSoftwareVersion
+	"ParentSDKFile": "../IOS/IOS_SDK.json"
 }
diff --git a/Engine/Config/Windows/Windows_SDK.json b/Engine/Config/Windows/Windows_SDK.json
--- a/Engine/Config/Windows/Windows_SDK.json
+++ b/Engine/Config/Windows/Windows_SDK.json
@@ -1,5 +1,39 @@
 {
-"MainVersion" : "10.0.19041.0",
-"MinVersion" : "10.0.00000.0",
-"MaxVersion" : "10.9.99999.0",
+	"MainVersion": "10.0.19041.0",
+	"MinVersion": "10.0.18362.0",
+	"MaxVersion": "10.9.99999.0",
+
+	"MinSoftwareVersion": "10.0.18362.0",
+
+	"PreferredClangVersions": [
+		// Clang 17: [REDACTED_PROFILE_URL]/llvm-project/issues/71976 should not be used as a preferred version until this issue is resolved and the max affected version is added to ClangWarnings.cs
+		"16.0.0-16.999",
+		"15.0.0-15.999",
+	],
+	"MinimumClangVersion": "15.0.0",
+	
+	// this is a strange format, but instead of a version range, this mapping min clang version for each MSVC version
+	"MinimumRequiredClangVersion": [
+		"14.37-16", // VS2022 17.7.x - 17.8.x
+		"14.35-15", // VS2022 17.5.x - 17.6.x
+		"14.34-14", // VS2022 17.4.x
+	],
+
+	"PreferredVisualCppVersions": [
+		"14.38.33130-14.38.99999", // VS2022 17.8.x
+		"14.37.32822-14.37.99999", // VS2022 17.7.x
+		"14.36.32532-14.36.99999", // VS2022 17.6.x
+		"14.35.32215-14.35.99999", // VS2022 17.5.x
+		"14.34.31933-14.34.99999", // VS2022 17.4.x
+	],
+	"BannedVisualCppVersions": [
+	],
+	"MinimumVisualCppVersion": "14.34.31933",
+
+	"PreferredIntelOneApiVersions": [
+		"2024.0.0-2024.9999"
+	],
+	"MinimumIntelOneApiVersion": "2024.0.0"
+
+
 }
diff --git a/Engine/Platforms/VisionOS/Config/VisionOS_SDK.json b/Engine/Platforms/VisionOS/Config/VisionOS_SDK.json
--- a/Engine/Platforms/VisionOS/Config/VisionOS_SDK.json
+++ b/Engine/Platforms/VisionOS/Config/VisionOS_SDK.json
@@ -1,3 +1,4 @@
 {
-"ParentSDKFile" : "../../../Config/Apple/Apple_SDK.json"
+	// inherit the IOS MinSoftwareVersion
+	"ParentSDKFile": "../../../Config/IOS/IOS_SDK.json"
 }
diff --git a/Engine/Platforms/VisionOS/Source/Programs/UnrealBuildTool/UEBuildVisionOS.cs b/Engine/Platforms/VisionOS/Source/Programs/UnrealBuildTool/UEBuildVisionOS.cs
--- a/Engine/Platforms/VisionOS/Source/Programs/UnrealBuildTool/UEBuildVisionOS.cs
+++ b/Engine/Platforms/VisionOS/Source/Programs/UnrealBuildTool/UEBuildVisionOS.cs
@@ -119,7 +119,7 @@ namespace UnrealBuildTool
 		/// </summary>
 		public override void RegisterBuildPlatforms(ILogger Logger)
 		{
-			ApplePlatformSDK SDK = new IOSPlatformSDK(Logger);
+			ApplePlatformSDK SDK = new ApplePlatformSDK(Logger);
 
 			// Register this build platform for IOS
 			UEBuildPlatform.RegisterBuildPlatform(new VisionOSPlatform(SDK, Logger), Logger);
diff --git a/Engine/Source/Programs/Shared/EpicGames.Build/UEBuildPlatformSDK.cs b/Engine/Source/Programs/Shared/EpicGames.Build/UEBuildPlatformSDK.cs
--- a/Engine/Source/Programs/Shared/EpicGames.Build/UEBuildPlatformSDK.cs
+++ b/Engine/Source/Programs/Shared/EpicGames.Build/UEBuildPlatformSDK.cs
@@ -377,6 +377,27 @@ namespace EpicGames.Core
 			return SDK;
 		}
 
+		public static T? GetSDKForPlatformOrMakeTemp<T>(string PlatformName) where T : UEBuildPlatformSDK
+		{
+			UEBuildPlatformSDK? SDK;
+			SDKRegistry.TryGetValue(PlatformName, out SDK);
+
+			// make a temp one if needed, this is not expected to happen often at all
+			if (SDK == null)
+			{
+				if (!TempSDKRegistry.TryGetValue(PlatformName, out SDK))
+				{
+					object[] parameter = new object[1];
+					parameter[0] = Log.Logger;
+					SDK = (UEBuildPlatformSDK)Activator.CreateInstance(typeof(T), parameter)!;
+					SDK.LoadJsonFile(PlatformName);
+					TempSDKRegistry.Add(PlatformName, SDK);
+				}
+			}
+			return (T?)SDK;
+		}
+
+
 		/// <summary>
 		/// Gets the set of all known SDKs
 		/// </summary>
@@ -575,7 +596,8 @@ namespace EpicGames.Core
 		/// <param name="MaxVersion">Largest version allowed (inclusive), or null if no maximum (in other words, MinVersion - infinity)y</param>
 		protected virtual void GetValidSoftwareVersionRange(out string? MinVersion, out string? MaxVersion)
 		{
-			throw new Exception($"This platform's Sdk class ({GetType().Name}) must implement either GetValidSoftwareVersionRange() or GetValidSoftwareVersionRanges(). If this triggers, and GetValidSoftwareVersionRanges() is implemented, that means GetValidSoftwareVersionRange was called directly");
+			MinVersion = GetVersionFromConfig("MinSoftwareVersion");
+			MaxVersion = GetVersionFromConfig("MaxSoftwareVersion");
 		}
 
 		protected virtual SoftwareCollection GetValidSoftwareVersions()
@@ -623,7 +645,7 @@ namespace EpicGames.Core
 		/// <returns>Version string, or empty string if not known/supported</returns>
 		public virtual string GetPlatformSpecificVersion(string VersionType)
 		{
-			return "";
+			return GetRequiredVersionFromConfig(VersionType);
 		}
 
 		/// <summary>
@@ -937,6 +959,10 @@ namespace EpicGames.Core
 		protected Dictionary<string, string> CachedManualSDKVersions = new Dictionary<string, string>();
 		private static Dictionary<string, UEBuildPlatformSDK> SDKRegistry = new Dictionary<string, UEBuildPlatformSDK>();
 
+		// this map holds on to some temporary SDK objects that are generally used once and don't want to stick around, but they could be used multiple times, 
+		// like in MicrosoftPlatofrmSDK.Version.cs, if one of the versions is needed, C# will go construct every version, needing the SDK multiple times
+		private static Dictionary<string, UEBuildPlatformSDK> TempSDKRegistry = new Dictionary<string, UEBuildPlatformSDK>();
+
 		private SDKDescriptor? GetSDKVersionForHint(SDKCollection Collection, string? Hint)
 		{
 			// if the hint is found, use it always
@@ -961,7 +987,7 @@ namespace EpicGames.Core
 
 		// cached SDK info from the SDK.json file
 		private Dictionary<string, string> ConfigSDKVersions = new(StringComparer.OrdinalIgnoreCase);
-		private Dictionary<string, string> DefaultConfigSDKVersions = new(StringComparer.OrdinalIgnoreCase);
+		private Dictionary<string, string[]> ConfigSDKVersionArrays = new(StringComparer.OrdinalIgnoreCase);
 
 		private void LoadJsonFile(string Platform)
 		{
@@ -993,9 +1019,10 @@ namespace EpicGames.Core
 			FileReference EngineSDKConfigFile = MakeConfigFilename(Unreal.EngineDirectory, true)!;
 
 			// load the file, along with any chained group file
-			ProcessJsonFile(EngineSDKConfigFile, ConfigSDKVersions);
+			ProcessJsonFile(EngineSDKConfigFile, ConfigSDKVersions, ConfigSDKVersionArrays);
 
-			// copy off the versions to the defaults, so we can later get the 
+			// copy off the versions to the defaults, so we can check if overridden by the project
+			Dictionary<string, string> DefaultConfigSDKVersions = new(ConfigSDKVersions);
 			foreach (string Key in ConfigSDKVersions.Keys)
 			{
 				DefaultConfigSDKVersions[Key] = ConfigSDKVersions[Key];
@@ -1009,7 +1036,12 @@ namespace EpicGames.Core
 				{
 					// load a project's SDK file if thre is one, into a temp dictionary
 					Dictionary<string, string> OverrideConfigSDKVersions = new(StringComparer.OrdinalIgnoreCase);
-					ProcessJsonFile(ProjectSDKConfigFile, OverrideConfigSDKVersions);
+					Dictionary<string, string[]> OverrideConfigSDKVersionArrays = new(StringComparer.OrdinalIgnoreCase);
+					ProcessJsonFile(ProjectSDKConfigFile, OverrideConfigSDKVersions, OverrideConfigSDKVersionArrays);
+					if (OverrideConfigSDKVersionArrays.Count > 0)
+					{
+						throw new Exception($"Overriding version arrays, in project '{ProjectFile.GetFileNameWithoutExtension()}', platform {Platform} is not currently supported");
+					}
 
 					// currently only care about MainVersion in the overrides
 					string? OverrideMainVersion;
@@ -1037,7 +1069,7 @@ namespace EpicGames.Core
 			}
 		}
 
-		private void ProcessJsonFile(FileReference SDKConfigFile, Dictionary<string, string> VersionMap)
+		private void ProcessJsonFile(FileReference SDKConfigFile, Dictionary<string, string> VersionMap, Dictionary<string, string[]> VersionArrayMap)
 		{
 			string Contents = FileReference.ReadAllText(SDKConfigFile);
 
@@ -1047,7 +1079,7 @@ namespace EpicGames.Core
 				AllowTrailingCommas = true,
 				ReadCommentHandling = JsonCommentHandling.Skip,
 			};
-			Dictionary<string, string>? LoadedDictionary = JsonSerializer.Deserialize<Dictionary<string, string>>(Contents, Options);
+			Dictionary<string, JsonElement>? LoadedDictionary = JsonSerializer.Deserialize<Dictionary<string, JsonElement>>(Contents, Options);
 			if (LoadedDictionary == null)
 			{
 				throw new Exception($"Failed to parse SDK version file '{SDKConfigFile}'");
@@ -1057,14 +1089,25 @@ namespace EpicGames.Core
 			List<string> Parents = new();
 			foreach (string Key in LoadedDictionary.Keys)
 			{
-				if (Key.Equals("ParentSDKFile", StringComparison.OrdinalIgnoreCase))
+				JsonElement Obj = LoadedDictionary[Key];
+				string? StringValue = Obj.ValueKind == JsonValueKind.String ? Obj.GetString() : null;
+				string[]? ArrayValue = Obj.ValueKind == JsonValueKind.Array ? Obj.EnumerateArray().Select(x => x.GetString()!).ToArray() : null;
+
+				if (StringValue != null)
 				{
-					Parents.Add(LoadedDictionary[Key]);
+					if (Key.Equals("ParentSDKFile", StringComparison.OrdinalIgnoreCase))
+					{
+						Parents.Add(StringValue);
+					}
+					// add this key if it's not already in the (case-insensitive) dictionary 
+					else if (!VersionMap.ContainsKey(Key))
+					{
+						VersionMap.Add(Key, StringValue);
+					}
 				}
-				// add this key if it's not already in the (case-insensitive) dictionary 
-				else if (!VersionMap.ContainsKey(Key))
+				else if (ArrayValue != null)
 				{
-					VersionMap.Add(Key, LoadedDictionary[Key]);
+					VersionArrayMap.Add(Key, ArrayValue);
 				}
 			}
 
@@ -1072,19 +1115,12 @@ namespace EpicGames.Core
 			foreach (string Parent in Parents)
 			{
 				FileReference ParentConfigFile = FileReference.Combine(SDKConfigFile.Directory, Parent);
-				ProcessJsonFile(ParentConfigFile, VersionMap);
+				ProcessJsonFile(ParentConfigFile, VersionMap, VersionArrayMap);
 			}
 		}
 
-		protected string GetRequiredVersionFromConfig(string VersionName)
+		private string GetHostSpecificVersionName(string VersionName)
 		{
-			// when bIsRequired is true, then we know it will return non-null
-			return GetVersionFromConfig(VersionName, bIsRequired:true)!;
-		}
-
-		protected string? GetVersionFromConfig(string VersionName, bool bIsRequired=false)
-		{
-			string? Version;
 			string HostPlatform = "Win64";
 			if (OperatingSystem.IsMacOS())
 			{
@@ -1094,9 +1130,19 @@ namespace EpicGames.Core
 			{
 				HostPlatform = "Linux";
 			}
+			return $"{VersionName}_{HostPlatform}";
+		}
+		public string GetRequiredVersionFromConfig(string VersionName)
+		{
+			// when bIsRequired is true, then we know it will return non-null
+			return GetVersionFromConfig(VersionName, bIsRequired:true)!;
+		}
+
+		public string? GetVersionFromConfig(string VersionName, bool bIsRequired=false)
+		{
+			string? Version;
 			// look up both Version_Host and Version (Host specific version wins)
-			string HostSpecificVersionName = $"{VersionName}_{HostPlatform}";
-			if (ConfigSDKVersions!.TryGetValue(HostSpecificVersionName, out Version) || ConfigSDKVersions.TryGetValue(VersionName, out Version))
+			if (ConfigSDKVersions!.TryGetValue(GetHostSpecificVersionName(VersionName), out Version) || ConfigSDKVersions.TryGetValue(VersionName, out Version))
 			{
 				return Version;
 			}
@@ -1108,6 +1154,81 @@ namespace EpicGames.Core
 			return null;
 		}
 
+		protected VersionNumber GetRequiredVersionNumberFromConfig(string VersionName)
+		{
+			// required won't ever return null
+			return GetVersionNumberFromConfig(VersionName, true)!;
+		}
+
+		public VersionNumber? GetVersionNumberFromConfig(string VersionName, bool bIsRequired=false)
+		{
+			string? VersionString = GetVersionFromConfig(VersionName, bIsRequired);
+
+			if (VersionString == null)
+			{
+				return null;
+			}
+
+			return VersionNumber.Parse(VersionString);
+		}
+
+
+		private VersionNumberRange? ParseVersionNumberRange(string Range)
+		{
+			string[] Versions = Range.Split("-");
+			if (Versions.Length != 2)
+			{
+				return null;
+			}
+
+			return VersionNumberRange.Parse(Versions[0], Versions[1]);
+		}
+
+		public VersionNumberRange? GetRequiredVersionNumberRangeFromConfig(string VersionName, bool bIsRequired = false)
+		{
+			// required won't ever return null
+			return GetVersionNumberRangeFromConfig(VersionName, true)!;
+		}
+
+		public VersionNumberRange? GetVersionNumberRangeFromConfig(string VersionName, bool bIsRequired = false)
+		{
+			string? VersionRange;
+			if (!ConfigSDKVersions!.TryGetValue(GetHostSpecificVersionName(VersionName), out VersionRange) && !ConfigSDKVersions.TryGetValue(VersionName, out VersionRange))
+			{
+				if (bIsRequired)
+				{
+					throw new Exception($"Unable to find required SDK version range '{VersionName}' for platform {PlatformName}. Check your SDK.json files");
+				}
+				return null;
+			}
+
+			VersionNumberRange? Range = ParseVersionNumberRange(VersionRange);
+			if (Range == null && bIsRequired)
+			{
+				throw new Exception($"Unable to parse the version number range for required version '{VersionName}' for platform {PlatformName}. Check your SDK.json files");
+			}
+			return Range;
+		}
+
+		public VersionNumberRange[] GetVersionNumberRangeArrayFromConfig(string VersionName)
+		{
+			string[]? VersionRanges;
+			List<VersionNumberRange> Ranges = new();
+			if (ConfigSDKVersionArrays.TryGetValue(GetHostSpecificVersionName(VersionName), out VersionRanges) || ConfigSDKVersionArrays.TryGetValue(VersionName, out VersionRanges))
+			{
+				foreach (string VersionRange in VersionRanges) 
+				{
+					VersionNumberRange? Range = ParseVersionNumberRange(VersionRange);
+					if (Range != null)
+					{
+						Ranges.Add(Range);
+					}
+				}
+			}
+
+			return Ranges.ToArray();
+		}
+
 		#endregion
 
 		// AutoSDKs handling portion
diff --git a/Engine/Source/Programs/Shared/EpicGames.Core/VersionNumber.cs b/Engine/Source/Programs/Shared/EpicGames.Core/VersionNumber.cs
new file mode 100644
--- /dev/null
+++ b/Engine/Source/Programs/Shared/EpicGames.Core/VersionNumber.cs
@@ -0,0 +1,250 @@
+﻿// Copyright Epic Games, Inc. All Rights Reserved.
+
+using System;
+using System.Collections.Generic;
+using System.Diagnostics.CodeAnalysis;
+using System.Text;
+
+namespace EpicGames.Core
+{
+	/// <summary>
+	/// Stores a numeric version consisting of any number of components.
+	/// </summary>
+	[Serializable]
+	public class VersionNumber : IComparable<VersionNumber>
+	{
+		/// <summary>
+		/// Set of delimiters for version numbers
+		/// </summary>
+		static readonly char[] Delimiters = { '.', ',' };
+
+		/// <summary>
+		/// The individual version components
+		/// </summary>
+		int[] Components;
+
+		/// <summary>
+		/// Constructor
+		/// </summary>
+		/// <param name="Components">The individual version components. At least one value must be given.</param>
+		public VersionNumber(params int[] Components)
+		{
+			if (Components.Length == 0)
+			{
+				throw new InvalidOperationException("Version number must have at least one component");
+			}
+			this.Components = Components;
+		}
+
+		/// <summary>
+		/// Returns the component at the given index
+		/// </summary>
+		/// <param name="Idx">The zero-based component index to return</param>
+		/// <returns>The component at the given index</returns>
+		public int GetComponent(int Idx)
+		{
+			return Components[Idx];
+		}
+
+		/// <summary>
+		/// Tests two objects for equality. VersionNumber behaves like a value type.
+		/// </summary>
+		/// <param name="Obj">Object to compare against</param>
+		/// <returns>True if the objects are equal, false otherwise.</returns>
+		public override bool Equals(object? Obj)
+		{
+			VersionNumber? Version = Obj as VersionNumber;
+			return !ReferenceEquals(Version, null) && this == Version;
+		}
+
+		/// <summary>
+		/// Returns a hash of the version number.
+		/// </summary>
+		/// <returns>A hash value for the version number.</returns>
+		public override int GetHashCode()
+		{
+			int Result = 5831;
+			for (int Idx = 0; Idx < Components.Length; Idx++)
+			{
+				Result = (Result * 33) + Components[Idx];
+			}
+			return Result;
+		}
+
+		/// <summary>
+		/// Compares whether two versions are equal.
+		/// </summary>
+		/// <param name="Lhs">The first version number</param>
+		/// <param name="Rhs">The second version number</param>
+		/// <returns>True if the versions are equal.</returns>
+		public static bool operator ==(VersionNumber? Lhs, VersionNumber? Rhs)
+		{
+			if (Object.ReferenceEquals(Lhs, null))
+			{
+				return Object.ReferenceEquals(Rhs, null);
+			}
+			else
+			{
+				return !Object.ReferenceEquals(Rhs, null) && Compare(Lhs, Rhs) == 0;
+			}
+		}
+
+		/// <summary>
+		/// Compares whether two versions are not equal.
+		/// </summary>
+		/// <param name="Lhs">The first version number</param>
+		/// <param name="Rhs">The second version number</param>
+		/// <returns>True if the versions are not equal.</returns>
+		public static bool operator !=(VersionNumber? Lhs, VersionNumber? Rhs)
+		{
+			return !(Lhs == Rhs);
+		}
+
+		/// <summary>
+		/// Compares whether one version is less than another.
+		/// </summary>
+		/// <param name="Lhs">The first version number</param>
+		/// <param name="Rhs">The second version number</param>
+		/// <returns>True if the first version is less than the second.</returns>
+		public static bool operator <(VersionNumber Lhs, VersionNumber Rhs)
+		{
+			return Compare(Lhs, Rhs) < 0;
+		}
+
+		/// <summary>
+		/// Compares whether one version is less or equal to another.
+		/// </summary>
+		/// <param name="Lhs">The first version number</param>
+		/// <param name="Rhs">The second version number</param>
+		/// <returns>True if the first version is less or equal to the second.</returns>
+		public static bool operator <=(VersionNumber Lhs, VersionNumber Rhs)
+		{
+			return Compare(Lhs, Rhs) <= 0;
+		}
+
+		/// <summary>
+		/// Compares whether one version is greater than another.
+		/// </summary>
+		/// <param name="Lhs">The first version number</param>
+		/// <param name="Rhs">The second version number</param>
+		/// <returns>True if the first version is greater than the second.</returns>
+		public static bool operator >(VersionNumber Lhs, VersionNumber Rhs)
+		{
+			return Compare(Lhs, Rhs) > 0;
+		}
+
+		/// <summary>
+		/// Compares whether one version is greater or equal to another.
+		/// </summary>
+		/// <param name="Lhs">The first version number</param>
+		/// <param name="Rhs">The second version number</param>
+		/// <returns>True if the first version is greater or equal to the second.</returns>
+		public static bool operator >=(VersionNumber Lhs, VersionNumber Rhs)
+		{
+			return Compare(Lhs, Rhs) >= 0;
+		}
+
+		/// <summary>
+		/// Comparison function for IComparable
+		/// </summary>
+		/// <param name="Other">Other version number to compare to</param>
+		/// <returns>A negative value if this version is before Other, a positive value if this version is after Other, and zero otherwise.</returns>
+		public int CompareTo(VersionNumber? Other)
+		{
+			return ReferenceEquals(Other, null) ? 1 : Compare(this, Other);
+		}
+
+		/// <summary>
+		/// Compares two version numbers and returns an integer indicating their order
+		/// </summary>
+		/// <param name="Lhs">The first version to check</param>
+		/// <param name="Rhs">The second version to check</param>
+		/// <returns>A negative value if Lhs is before Rhs, a positive value if Lhs is after Rhs, and zero otherwise.</returns>
+		public static int Compare(VersionNumber Lhs, VersionNumber Rhs)
+		{
+			for (int Idx = 0; ; Idx++)
+			{
+				if (Idx == Lhs.Components.Length)
+				{
+					if (Idx == Rhs.Components.Length)
+					{
+						return 0;
+					}
+					else
+					{
+						return -1;
+					}
+				}
+				else
+				{
+					if (Idx == Rhs.Components.Length)
+					{
+						return +1;
+					}
+					else if (Lhs.Components[Idx] != Rhs.Components[Idx])
+					{
+						return Lhs.Components[Idx] - Rhs.Components[Idx];
+					}
+				}
+			}
+		}
+
+		/// <summary>
+		/// Parses the version number from a string
+		/// </summary>
+		/// <param name="Text">The string to parse</param>
+		/// <returns>A version number object</returns>
+		public static VersionNumber Parse(string Text)
+		{
+			List<int> Components = new List<int>();
+			foreach (string TextElement in Text.Split(Delimiters))
+			{
+				Components.Add(Int32.Parse(TextElement));
+			}
+			return new VersionNumber(Components.ToArray());
+		}
+
+		/// <summary>
+		/// Parses the version number from a string
+		/// </summary>
+		/// <param name="Text">The string to parse</param>
+		/// <param name="OutNumber">Variable to receive the parsed version number</param>
+		/// <returns>A version number object</returns>
+		public static bool TryParse(string Text, [NotNullWhen(true)] out VersionNumber? OutNumber)
+		{
+			List<int> Components = new List<int>();
+			foreach (string TextElement in Text.Split(Delimiters))
+			{
+				int Component;
+				if (!Int32.TryParse(TextElement, out Component))
+				{
+					OutNumber = null;
+					return false;
+				}
+				Components.Add(Component);
+			}
+
+			OutNumber = new VersionNumber(Components.ToArray());
+			return true;
+		}
+
+		/// <summary>
+		/// Returns a string version number, eg. 1.4
+		/// </summary>
+		/// <returns>The stringized version number</returns>
+		public override string ToString()
+		{
+			StringBuilder Result = new StringBuilder();
+			if (Components.Length > 0)
+			{
+				Result.Append(Components[0]);
+				for (int Idx = 1; Idx < Components.Length; Idx++)
+				{
+					Result.Append('.');
+					Result.Append(Components[Idx]);
+				}
+			}
+			return Result.ToString();
+		}
+	}
+}
diff --git a/Engine/Source/Programs/Shared/EpicGames.Core/VersionNumberRange.cs b/Engine/Source/Programs/Shared/EpicGames.Core/VersionNumberRange.cs
new file mode 100644
--- /dev/null
+++ b/Engine/Source/Programs/Shared/EpicGames.Core/VersionNumberRange.cs
@@ -0,0 +1,60 @@
+// Copyright Epic Games, Inc. All Rights Reserved.
+
+using System;
+
+namespace EpicGames.Core
+{
+	/// <summary>
+	/// Range of version numbers
+	/// </summary>
+	public class VersionNumberRange
+	{
+		/// <summary>
+		/// Minimum version number
+		/// </summary>
+		public VersionNumber Min { get; }
+
+		/// <summary>
+		/// Maximum version number
+		/// </summary>
+		public VersionNumber Max { get; }
+
+		/// <summary>
+		/// Constructor
+		/// </summary>
+		/// <param name="Min"></param>
+		/// <param name="Max"></param>
+		public VersionNumberRange(VersionNumber Min, VersionNumber Max)
+		{
+			this.Min = Min;
+			this.Max = Max;
+		}
+
+		/// <summary>
+		/// Tests whether this range contains the given version
+		/// </summary>
+		/// <param name="Version"></param>
+		/// <returns></returns>
+		public bool Contains(VersionNumber Version)
+		{
+			return Version >= Min && Version <= Max;
+		}
+
+		/// <summary>
+		/// Parse a version range from two strings
+		/// </summary>
+		/// <param name="MinText"></param>
+		/// <param name="MaxText"></param>
+		/// <returns></returns>
+		public static VersionNumberRange Parse(string MinText, string MaxText)
+		{
+			return new VersionNumberRange(VersionNumber.Parse(MinText), VersionNumber.Parse(MaxText));
+		}
+
+		/// <inheritdoc/>
+		public override string ToString()
+		{
+			return String.Format("{0}-{1}", Min, Max);
+		}
+	}
+}
diff --git a/Engine/Source/Programs/UnrealBuildTool/Platform/Android/AndroidPlatformSDK.Versions.cs b/Engine/Source/Programs/UnrealBuildTool/Platform/Android/AndroidPlatformSDK.Versions.cs
--- a/Engine/Source/Programs/UnrealBuildTool/Platform/Android/AndroidPlatformSDK.Versions.cs
+++ b/Engine/Source/Programs/UnrealBuildTool/Platform/Android/AndroidPlatformSDK.Versions.cs
@@ -4,28 +4,7 @@ using EpicGames.Core;
 
 namespace UnrealBuildTool
 {
-	/////////////////////////////////////////////////////////////////////////////////////
-	// If you are looking for any version numbers not listed here, see GDK_SDK.json
-	/////////////////////////////////////////////////////////////////////////////////////
-
-	partial class AndroidPlatformSDK : UEBuildPlatformSDK
-	{
-		protected override void GetValidSoftwareVersionRange(out string? MinVersion, out string? MaxVersion)
-		{
-			MinVersion = MaxVersion = null;
-		}
-
-		public override string GetPlatformSpecificVersion(string VersionType)
-		{
-			switch (VersionType.ToLower())
-			{
-				case "platforms": return "android-33";
-				case "build-tools": return "33.0.1";
-				case "cmake": return "3.22.1";
-				case "ndk": return "25.1.8937393";
-			}
-
-			return "";
-		}
-	}
+	////////////////////////////////////////////////////////////////////////////////////////////
+	// If you are looking for any version numbers, see Engine/Config/Android/Android_SDK.json
+	////////////////////////////////////////////////////////////////////////////////////////////
 }
diff --git a/Engine/Source/Programs/UnrealBuildTool/Platform/IOS/IOSPlatformSDK.Versions.cs b/Engine/Source/Programs/UnrealBuildTool/Platform/IOS/IOSPlatformSDK.Versions.cs
--- a/Engine/Source/Programs/UnrealBuildTool/Platform/IOS/IOSPlatformSDK.Versions.cs
+++ b/Engine/Source/Programs/UnrealBuildTool/Platform/IOS/IOSPlatformSDK.Versions.cs
@@ -1,31 +1,8 @@
 // Copyright Epic Games, Inc. All Rights Reserved.
 
-using System;
-using Microsoft.Extensions.Logging;
-
 namespace UnrealBuildTool
 {
-	internal class IOSPlatformSDK : ApplePlatformSDK
-	{
-		public IOSPlatformSDK(ILogger Logger)
-			: base(Logger)
-		{
-		}
-
-		protected override void GetValidSoftwareVersionRange(out string MinVersion, out string? MaxVersion)
-		{
-			// what is our min IOS version?
-			MinVersion = "15.0";
-			MaxVersion = null;
-		}
-
-		/// <summary>
-		/// The minimum Clang version that a static library can be built with
-		/// </summary>
-		public Version MinimumStaticLibClangVersion =>
-				// return new Version("1200.00.32.29");	// UE 5.0 (Xcode 12.4)
-				// return new Version("1300.00.29.30"); // (old) UE 5.1 sets Xcode13.2.1 as minimum supported
-				// return new Version("1316.0.21.2"); // UE 5.1 sets Xcode13.4 as minimum supported
-				new Version("1400.0.29.202"); // UE 5.2 sets Xcode14.1 as minimum supported
-	}
+	///////////////////////////////////////////////////////////////////////////////////////
+	// If you are looking for any version numbers, see Engine/Config/IOS/IOS_SDK.json
+	///////////////////////////////////////////////////////////////////////////////////////
 }
diff --git a/Engine/Source/Programs/UnrealBuildTool/Platform/IOS/UEBuildIOS.cs b/Engine/Source/Programs/UnrealBuildTool/Platform/IOS/UEBuildIOS.cs
--- a/Engine/Source/Programs/UnrealBuildTool/Platform/IOS/UEBuildIOS.cs
+++ b/Engine/Source/Programs/UnrealBuildTool/Platform/IOS/UEBuildIOS.cs
@@ -771,7 +771,7 @@ namespace UnrealBuildTool
 #if false
 			if (BuildHostPlatform.Current.Platform == UnrealTargetPlatform.Mac && !Target.IOSPlatform.bSkipClangValidation)
 			{
-				IOSPlatformSDK SDK = (IOSPlatformSDK?)GetSDK() ?? new IOSPlatformSDK(Logger);
+				ApplePlatformSDK SDK = (ApplePlatformSDK?)GetSDK() ?? new ApplePlatformSDK(Logger);
 				foreach (FileReference LibLoc in Module.PublicLibraries)
 				{
 					switch (LibLoc.GetExtension())
@@ -1254,7 +1254,7 @@ namespace UnrealBuildTool
 		/// </summary>
 		public override void RegisterBuildPlatforms(ILogger Logger)
 		{
-			ApplePlatformSDK SDK = new IOSPlatformSDK(Logger);
+			ApplePlatformSDK SDK = new ApplePlatformSDK(Logger);
 
 			// Register this build platform for IOS
 			UEBuildPlatform.RegisterBuildPlatform(new IOSPlatform(SDK, Logger), Logger);
diff --git a/Engine/Source/Programs/UnrealBuildTool/Platform/Linux/LinuxPlatformSDK.Versions.cs b/Engine/Source/Programs/UnrealBuildTool/Platform/Linux/LinuxPlatformSDK.Versions.cs
--- a/Engine/Source/Programs/UnrealBuildTool/Platform/Linux/LinuxPlatformSDK.Versions.cs
+++ b/Engine/Source/Programs/UnrealBuildTool/Platform/Linux/LinuxPlatformSDK.Versions.cs
@@ -1,18 +1,8 @@
 // Copyright Epic Games, Inc. All Rights Reserved.
 
-using EpicGames.Core;
-
 namespace UnrealBuildTool
 {
-	/////////////////////////////////////////////////////////////////////////////////////
-	// If you are looking for any version numbers not listed here, see Linux_SDK.json
-	/////////////////////////////////////////////////////////////////////////////////////
-
-	partial class LinuxPlatformSDK : UEBuildPlatformSDK
-	{
-		protected override void GetValidSoftwareVersionRange(out string? MinVersion, out string? MaxVersion)
-		{
-			MinVersion = MaxVersion = null;
-		}
-	}
+	////////////////////////////////////////////////////////////////////////////////////////
+	// If you are looking for version numbers here, see Engine/Config/Linux/Linux_SDK.json
+	////////////////////////////////////////////////////////////////////////////////////////
 }
diff --git a/Engine/Source/Programs/UnrealBuildTool/Platform/Mac/ApplePlatformSDK.Versions.cs b/Engine/Source/Programs/UnrealBuildTool/Platform/Mac/ApplePlatformSDK.Versions.cs
--- a/Engine/Source/Programs/UnrealBuildTool/Platform/Mac/ApplePlatformSDK.Versions.cs
+++ b/Engine/Source/Programs/UnrealBuildTool/Platform/Mac/ApplePlatformSDK.Versions.cs
@@ -1,27 +1,8 @@
 // Copyright Epic Games, Inc. All Rights Reserved.
 
-using System;
-using EpicGames.Core;
-
-
-using UnrealBuildBase;
-
 namespace UnrealBuildTool
 {
 	/////////////////////////////////////////////////////////////////////////////////////
-	// If you are looking for any version numbers not listed here, see Apple_SDK.json
+	// If you are looking for version numbers, see Engine/Config/Apple/Apple_SDK.json
 	/////////////////////////////////////////////////////////////////////////////////////
-
-	internal partial class ApplePlatformSDK : UEBuildPlatformSDK
-	{
-		protected override void GetValidSoftwareVersionRange(out string? MinVersion, out string? MaxVersion)
-		{
-			MinVersion = MaxVersion = null;
-		}
-
-		/// <summary>
-		/// The minimum macOS SDK version that a dynamic library can be built with
-		/// </summary>
-		public virtual Version MinimumDynamicLibSDKVersion => new Version("12.1");      // SDK used in Xcode13.2.1
-	}
 }
diff --git a/Engine/Source/Programs/UnrealBuildTool/Platform/Mac/MacPlatformSDK.Versions.cs b/Engine/Source/Programs/UnrealBuildTool/Platform/Mac/MacPlatformSDK.Versions.cs
--- a/Engine/Source/Programs/UnrealBuildTool/Platform/Mac/MacPlatformSDK.Versions.cs
+++ b/Engine/Source/Programs/UnrealBuildTool/Platform/Mac/MacPlatformSDK.Versions.cs
@@ -1,20 +1,8 @@
 // Copyright Epic Games, Inc. All Rights Reserved.
 
-using Microsoft.Extensions.Logging;
-
 namespace UnrealBuildTool
 {
-	internal partial class MacPlatformSDK : ApplePlatformSDK
-	{
-		public MacPlatformSDK(ILogger Logger)
-			: base(Logger)
-		{
-		}
-
-		protected override void GetValidSoftwareVersionRange(out string MinVersion, out string? MaxVersion)
-		{
-			MinVersion = "12.5.0";      // macOS Monterey 12.5 (Released 2022/07/22)
-			MaxVersion = null;
-		}
-	}
+	//////////////////////////////////////////////////////////////////////////////////
+	// If you are looking for version numbers, see Engine/Config/Mac/Mac_SDK.json
+	//////////////////////////////////////////////////////////////////////////////////
 }
\ No newline at end of file
diff --git a/Engine/Source/Programs/UnrealBuildTool/Platform/Mac/UEBuildMac.cs b/Engine/Source/Programs/UnrealBuildTool/Platform/Mac/UEBuildMac.cs
--- a/Engine/Source/Programs/UnrealBuildTool/Platform/Mac/UEBuildMac.cs
+++ b/Engine/Source/Programs/UnrealBuildTool/Platform/Mac/UEBuildMac.cs
@@ -637,7 +637,7 @@ namespace UnrealBuildTool
 		/// </summary>
 		public override void RegisterBuildPlatforms(ILogger Logger)
 		{
-			MacPlatformSDK SDK = new MacPlatformSDK(Logger);
+			ApplePlatformSDK SDK = new ApplePlatformSDK(Logger);
 
 			// Register this build platform for Mac
 			UEBuildPlatform.RegisterBuildPlatform(new MacPlatform(SDK, Logger), Logger);
diff --git a/Engine/Source/Programs/UnrealBuildTool/Platform/TVOS/UEBuildTVOS.cs b/Engine/Source/Programs/UnrealBuildTool/Platform/TVOS/UEBuildTVOS.cs
--- a/Engine/Source/Programs/UnrealBuildTool/Platform/TVOS/UEBuildTVOS.cs
+++ b/Engine/Source/Programs/UnrealBuildTool/Platform/TVOS/UEBuildTVOS.cs
@@ -130,7 +130,7 @@ namespace UnrealBuildTool
 		/// </summary>
 		public override void RegisterBuildPlatforms(ILogger Logger)
 		{
-			ApplePlatformSDK SDK = new IOSPlatformSDK(Logger);
+			ApplePlatformSDK SDK = new ApplePlatformSDK(Logger);
 
 			// Register this build platform for IOS
 			UEBuildPlatform.RegisterBuildPlatform(new TVOSPlatform(SDK, Logger), Logger);
diff --git a/Engine/Source/Programs/UnrealBuildTool/Platform/Windows/MicrosoftPlatformSDK.Versions.cs b/Engine/Source/Programs/UnrealBuildTool/Platform/Windows/MicrosoftPlatformSDK.Versions.cs
--- a/Engine/Source/Programs/UnrealBuildTool/Platform/Windows/MicrosoftPlatformSDK.Versions.cs
+++ b/Engine/Source/Programs/UnrealBuildTool/Platform/Windows/MicrosoftPlatformSDK.Versions.cs
@@ -4,94 +4,61 @@ using System;
 using System.Diagnostics;
 using System.Linq;
 using EpicGames.Core;
+using Microsoft.Build.Framework;
 
 namespace UnrealBuildTool
 {
-	/////////////////////////////////////////////////////////////////////////////////////
-	// If you are looking for any version numbers not listed here, see Windows_SDK.json
-	/////////////////////////////////////////////////////////////////////////////////////
+	////////////////////////////////////////////////////////////////////////////////
+	// If you are looking for version numbers, see Engine/Config/Windows_SDK.json
+	////////////////////////////////////////////////////////////////////////////////
 
 	partial class MicrosoftPlatformSDK : UEBuildPlatformSDK
 	{
-		/// <summary>
-		/// The minimum Windows SDK version to be used. If this is null then it means there is no minimum version
-		/// </summary>
-		static readonly VersionNumber? MinimumWindowsSDKVersion = new VersionNumber(10, 0, 18362, 0);
-
-		/// <summary>
-		/// The maximum Windows SDK version to be used. If this is null then it means "Latest"
-		/// </summary>
-		static readonly VersionNumber? MaximumWindowsSDKVersion = null;
+		private static MicrosoftPlatformSDK SDK => GetSDKForPlatformOrMakeTemp<MicrosoftPlatformSDK>("Win64")!;
 
 		/// <summary>
 		/// The default compiler version to be used, if installed. 
 		/// </summary>
-		static readonly VersionNumberRange[] PreferredClangVersions =
-		{
-			// Clang 17: [REDACTED_PROFILE_URL]/llvm-project/issues/71976 should not be used as a preferred version until this issue is resolved and the max affected version is added to ClangWarnings.cs
-			VersionNumberRange.Parse("16.0.0", "16.999"), // VS2022 17.7.x runtime requires Clang 16
-			VersionNumberRange.Parse("15.0.0", "15.999"), // VS2022 17.5.x runtime requires Clang 15
-		};
+		static readonly VersionNumberRange[] PreferredClangVersions = SDK.GetVersionNumberRangeArrayFromConfig("PreferredClangVersions");
 
 		/// <summary>
 		/// The minimum supported Clang compiler
 		/// </summary>
-		static readonly VersionNumber MinimumClangVersion = new VersionNumber(15, 0, 0);
+		static readonly VersionNumber MinimumClangVersion = SDK.GetRequiredVersionNumberFromConfig("MinimumClangVersion");
 
 		/// <summary>
 		/// Ranges of tested compiler toolchains to be used, in order of preference. If multiple toolchains in a range are present, the latest version will be preferred.
 		/// Note that the numbers here correspond to the installation *folders* rather than precise executable versions.
 		/// </summary>
 		/// <seealso href="https://learn.microsoft.com/en-us/lifecycle/products/visual-studio-2022"/>
-		static readonly VersionNumberRange[] PreferredVisualCppVersions = new VersionNumberRange[]
-		{
-			VersionNumberRange.Parse("14.38.33130", "14.38.99999"), // VS2022 17.8.x
-			VersionNumberRange.Parse("14.37.32822", "14.37.99999"), // VS2022 17.7.x
-			VersionNumberRange.Parse("14.36.32532", "14.36.99999"), // VS2022 17.6.x
-			VersionNumberRange.Parse("14.35.32215", "14.35.99999"), // VS2022 17.5.x
-			VersionNumberRange.Parse("14.34.31933", "14.34.99999"), // VS2022 17.4.x
-		};
+		static readonly VersionNumberRange[] PreferredVisualCppVersions = SDK.GetVersionNumberRangeArrayFromConfig("PreferredVisualCppVersions");
 
 		/// <summary>
 		/// Minimum Clang version required for MSVC toolchain versions
 		/// </summary>
-		static readonly Tuple<VersionNumber, VersionNumber>[] MinimumRequiredClangVersion = new Tuple<VersionNumber, VersionNumber>[]
-		{
-			new(new VersionNumber(14, 37), new VersionNumber(16)), // VS2022 17.7.x - 17.8.x
-			new(new VersionNumber(14, 35), new VersionNumber(15)), // VS2022 17.5.x - 17.6.x
-			new(new VersionNumber(14, 34), new VersionNumber(14)), // VS2022 17.4.x
-		};
+		static readonly Tuple<VersionNumber, VersionNumber>[] MinimumRequiredClangVersion = SDK.GetVersionNumberRangeArrayFromConfig("MinimumRequiredClangVersion").
+			Select(x => new Tuple<VersionNumber, VersionNumber>(x.Min, x.Max)).ToArray();
 
 		/// <summary>
 		/// Tested compiler toolchains that should not be allowed.
 		/// </summary>
-		static readonly VersionNumberRange[] BannedVisualCppVersions = Array.Empty<VersionNumberRange>();
+		static readonly VersionNumberRange[] BannedVisualCppVersions = SDK.GetVersionNumberRangeArrayFromConfig("BannedVisualCppVersions");
 
 		/// <summary>
 		/// The minimum supported MSVC compiler
 		/// </summary>
-		static readonly VersionNumber MinimumVisualCppVersion = new VersionNumber(14, 34, 31933);
+		static readonly VersionNumber MinimumVisualCppVersion = SDK.GetRequiredVersionNumberFromConfig("MinimumVisualCppVersion");
 
 		/// <summary>
 		/// The default compiler version to be used, if installed. 
 		/// https://www.intel.com/content/www/us/en/developer/articles/tool/oneapi-standalone-components.html#dpcpp-cpp
 		/// </summary>
-		static readonly VersionNumberRange[] PreferredIntelOneApiVersions =
-		{
-			VersionNumberRange.Parse("2024.0.0", "2024.9999"),
-		};
+		static readonly VersionNumberRange[] PreferredIntelOneApiVersions = SDK.GetVersionNumberRangeArrayFromConfig("PreferredIntelOneApiVersions");
 
 		/// <summary>
 		/// The minimum supported Intel compiler
 		/// </summary>
-		static readonly VersionNumber MinimumIntelOneApiVersion = new VersionNumber(2024, 0, 0);
-
-		/// <inheritdoc/>
-		protected override void GetValidSoftwareVersionRange(out string? minVersion, out string? maxVersion)
-		{
-			minVersion = MinimumWindowsSDKVersion?.ToString();
-			maxVersion = MaximumWindowsSDKVersion?.ToString();
-		}
+		static readonly VersionNumber MinimumIntelOneApiVersion = SDK.GetRequiredVersionNumberFromConfig("MinimumIntelOneApiVersion");
 
 		/// <summary>
 		/// If a toolchain version is a preferred version
diff --git a/Engine/Source/Programs/UnrealBuildTool/System/VersionNumber.cs b/Engine/Source/Programs/UnrealBuildTool/System/VersionNumber.cs
deleted file mode 100644
diff --git a/Engine/Source/Programs/UnrealBuildTool/System/VersionNumberRange.cs b/Engine/Source/Programs/UnrealBuildTool/System/VersionNumberRange.cs
deleted file mode 100644
diff --git a/Engine/Source/Programs/UnrealBuildTool/ToolChain/ClangWarnings.cs b/Engine/Source/Programs/UnrealBuildTool/ToolChain/ClangWarnings.cs
--- a/Engine/Source/Programs/UnrealBuildTool/ToolChain/ClangWarnings.cs
+++ b/Engine/Source/Programs/UnrealBuildTool/ToolChain/ClangWarnings.cs
@@ -3,6 +3,7 @@
 using EpicGames.UHT.Utils;
 using Microsoft.Extensions.Logging;
 using System.Collections.Generic;
+using EpicGames.Core;
 
 namespace UnrealBuildTool
 {
