EpicGames.Core: Fix static analysis warnings.

diff --git a/Engine/Source/Programs/Shared/EpicGames.Core/VersionNumber.cs b/Engine/Source/Programs/Shared/EpicGames.Core/VersionNumber.cs
--- a/Engine/Source/Programs/Shared/EpicGames.Core/VersionNumber.cs
+++ b/Engine/Source/Programs/Shared/EpicGames.Core/VersionNumber.cs
@@ -16,45 +16,47 @@ namespace EpicGames.Core
 		/// <summary>
 		/// Set of delimiters for version numbers
 		/// </summary>
-		static readonly char[] Delimiters = { '.', ',' };
+		static readonly char[] s_delimiters = { '.', ',' };
 
 		/// <summary>
 		/// The individual version components
 		/// </summary>
-		int[] Components;
+		public IReadOnlyList<int> Components => _components;
+
+		readonly int[] _components;
 
 		/// <summary>
 		/// Constructor
 		/// </summary>
-		/// <param name="Components">The individual version components. At least one value must be given.</param>
-		public VersionNumber(params int[] Components)
+		/// <param name="components">The individual version components. At least one value must be given.</param>
+		public VersionNumber(params int[] components)
 		{
-			if (Components.Length == 0)
+			if (components.Length == 0)
 			{
 				throw new InvalidOperationException("Version number must have at least one component");
 			}
-			this.Components = Components;
+			_components = components;
 		}
 
 		/// <summary>
 		/// Returns the component at the given index
 		/// </summary>
-		/// <param name="Idx">The zero-based component index to return</param>
+		/// <param name="idx">The zero-based component index to return</param>
 		/// <returns>The component at the given index</returns>
-		public int GetComponent(int Idx)
+		public int GetComponent(int idx)
 		{
-			return Components[Idx];
+			return _components[idx];
 		}
 
 		/// <summary>
 		/// Tests two objects for equality. VersionNumber behaves like a value type.
 		/// </summary>
-		/// <param name="Obj">Object to compare against</param>
+		/// <param name="obj">Object to compare against</param>
 		/// <returns>True if the objects are equal, false otherwise.</returns>
-		public override bool Equals(object? Obj)
+		public override bool Equals(object? obj)
 		{
-			VersionNumber? Version = Obj as VersionNumber;
-			return !ReferenceEquals(Version, null) && this == Version;
+			VersionNumber? version = obj as VersionNumber;
+			return !ReferenceEquals(version, null) && this == version;
 		}
 
 		/// <summary>
@@ -63,110 +65,110 @@ namespace EpicGames.Core
 		/// <returns>A hash value for the version number.</returns>
 		public override int GetHashCode()
 		{
-			int Result = 5831;
-			for (int Idx = 0; Idx < Components.Length; Idx++)
+			int result = 5831;
+			for (int idx = 0; idx < _components.Length; idx++)
 			{
-				Result = (Result * 33) + Components[Idx];
+				result = (result * 33) + _components[idx];
 			}
-			return Result;
+			return result;
 		}
 
 		/// <summary>
 		/// Compares whether two versions are equal.
 		/// </summary>
-		/// <param name="Lhs">The first version number</param>
-		/// <param name="Rhs">The second version number</param>
+		/// <param name="lhs">The first version number</param>
+		/// <param name="rhs">The second version number</param>
 		/// <returns>True if the versions are equal.</returns>
-		public static bool operator ==(VersionNumber? Lhs, VersionNumber? Rhs)
+		public static bool operator ==(VersionNumber? lhs, VersionNumber? rhs)
 		{
-			if (Object.ReferenceEquals(Lhs, null))
+			if (Object.ReferenceEquals(lhs, null))
 			{
-				return Object.ReferenceEquals(Rhs, null);
+				return Object.ReferenceEquals(rhs, null);
 			}
 			else
 			{
-				return !Object.ReferenceEquals(Rhs, null) && Compare(Lhs, Rhs) == 0;
+				return !Object.ReferenceEquals(rhs, null) && Compare(lhs, rhs) == 0;
 			}
 		}
 
 		/// <summary>
 		/// Compares whether two versions are not equal.
 		/// </summary>
-		/// <param name="Lhs">The first version number</param>
-		/// <param name="Rhs">The second version number</param>
+		/// <param name="lhs">The first version number</param>
+		/// <param name="rhs">The second version number</param>
 		/// <returns>True if the versions are not equal.</returns>
-		public static bool operator !=(VersionNumber? Lhs, VersionNumber? Rhs)
+		public static bool operator !=(VersionNumber? lhs, VersionNumber? rhs)
 		{
-			return !(Lhs == Rhs);
+			return !(lhs == rhs);
 		}
 
 		/// <summary>
 		/// Compares whether one version is less than another.
 		/// </summary>
-		/// <param name="Lhs">The first version number</param>
-		/// <param name="Rhs">The second version number</param>
+		/// <param name="lhs">The first version number</param>
+		/// <param name="rhs">The second version number</param>
 		/// <returns>True if the first version is less than the second.</returns>
-		public static bool operator <(VersionNumber Lhs, VersionNumber Rhs)
+		public static bool operator <(VersionNumber lhs, VersionNumber rhs)
 		{
-			return Compare(Lhs, Rhs) < 0;
+			return Compare(lhs, rhs) < 0;
 		}
 
 		/// <summary>
 		/// Compares whether one version is less or equal to another.
 		/// </summary>
-		/// <param name="Lhs">The first version number</param>
-		/// <param name="Rhs">The second version number</param>
+		/// <param name="lhs">The first version number</param>
+		/// <param name="rhs">The second version number</param>
 		/// <returns>True if the first version is less or equal to the second.</returns>
-		public static bool operator <=(VersionNumber Lhs, VersionNumber Rhs)
+		public static bool operator <=(VersionNumber lhs, VersionNumber rhs)
 		{
-			return Compare(Lhs, Rhs) <= 0;
+			return Compare(lhs, rhs) <= 0;
 		}
 
 		/// <summary>
 		/// Compares whether one version is greater than another.
 		/// </summary>
-		/// <param name="Lhs">The first version number</param>
-		/// <param name="Rhs">The second version number</param>
+		/// <param name="lhs">The first version number</param>
+		/// <param name="rhs">The second version number</param>
 		/// <returns>True if the first version is greater than the second.</returns>
-		public static bool operator >(VersionNumber Lhs, VersionNumber Rhs)
+		public static bool operator >(VersionNumber lhs, VersionNumber rhs)
 		{
-			return Compare(Lhs, Rhs) > 0;
+			return Compare(lhs, rhs) > 0;
 		}
 
 		/// <summary>
 		/// Compares whether one version is greater or equal to another.
 		/// </summary>
-		/// <param name="Lhs">The first version number</param>
-		/// <param name="Rhs">The second version number</param>
+		/// <param name="lhs">The first version number</param>
+		/// <param name="rhs">The second version number</param>
 		/// <returns>True if the first version is greater or equal to the second.</returns>
-		public static bool operator >=(VersionNumber Lhs, VersionNumber Rhs)
+		public static bool operator >=(VersionNumber lhs, VersionNumber rhs)
 		{
-			return Compare(Lhs, Rhs) >= 0;
+			return Compare(lhs, rhs) >= 0;
 		}
 
 		/// <summary>
 		/// Comparison function for IComparable
 		/// </summary>
-		/// <param name="Other">Other version number to compare to</param>
+		/// <param name="other">Other version number to compare to</param>
 		/// <returns>A negative value if this version is before Other, a positive value if this version is after Other, and zero otherwise.</returns>
-		public int CompareTo(VersionNumber? Other)
+		public int CompareTo(VersionNumber? other)
 		{
-			return ReferenceEquals(Other, null) ? 1 : Compare(this, Other);
+			return ReferenceEquals(other, null) ? 1 : Compare(this, other);
 		}
 
 		/// <summary>
 		/// Compares two version numbers and returns an integer indicating their order
 		/// </summary>
-		/// <param name="Lhs">The first version to check</param>
-		/// <param name="Rhs">The second version to check</param>
+		/// <param name="lhs">The first version to check</param>
+		/// <param name="rhs">The second version to check</param>
 		/// <returns>A negative value if Lhs is before Rhs, a positive value if Lhs is after Rhs, and zero otherwise.</returns>
-		public static int Compare(VersionNumber Lhs, VersionNumber Rhs)
+		public static int Compare(VersionNumber lhs, VersionNumber rhs)
 		{
-			for (int Idx = 0; ; Idx++)
+			for (int idx = 0; ; idx++)
 			{
-				if (Idx == Lhs.Components.Length)
+				if (idx == lhs._components.Length)
 				{
-					if (Idx == Rhs.Components.Length)
+					if (idx == rhs._components.Length)
 					{
 						return 0;
 					}
@@ -177,13 +179,13 @@ namespace EpicGames.Core
 				}
 				else
 				{
-					if (Idx == Rhs.Components.Length)
+					if (idx == rhs._components.Length)
 					{
 						return +1;
 					}
-					else if (Lhs.Components[Idx] != Rhs.Components[Idx])
+					else if (lhs._components[idx] != rhs._components[idx])
 					{
-						return Lhs.Components[Idx] - Rhs.Components[Idx];
+						return lhs._components[idx] - rhs._components[idx];
 					}
 				}
 			}
@@ -192,39 +194,39 @@ namespace EpicGames.Core
 		/// <summary>
 		/// Parses the version number from a string
 		/// </summary>
-		/// <param name="Text">The string to parse</param>
+		/// <param name="text">The string to parse</param>
 		/// <returns>A version number object</returns>
-		public static VersionNumber Parse(string Text)
+		public static VersionNumber Parse(string text)
 		{
-			List<int> Components = new List<int>();
-			foreach (string TextElement in Text.Split(Delimiters))
+			List<int> components = new List<int>();
+			foreach (string textElement in text.Split(s_delimiters))
 			{
-				Components.Add(Int32.Parse(TextElement));
+				components.Add(Int32.Parse(textElement));
 			}
-			return new VersionNumber(Components.ToArray());
+			return new VersionNumber(components.ToArray());
 		}
 
 		/// <summary>
 		/// Parses the version number from a string
 		/// </summary>
-		/// <param name="Text">The string to parse</param>
-		/// <param name="OutNumber">Variable to receive the parsed version number</param>
+		/// <param name="text">The string to parse</param>
+		/// <param name="outNumber">Variable to receive the parsed version number</param>
 		/// <returns>A version number object</returns>
-		public static bool TryParse(string Text, [NotNullWhen(true)] out VersionNumber? OutNumber)
+		public static bool TryParse(string text, [NotNullWhen(true)] out VersionNumber? outNumber)
 		{
-			List<int> Components = new List<int>();
-			foreach (string TextElement in Text.Split(Delimiters))
+			List<int> components = new List<int>();
+			foreach (string textElement in text.Split(s_delimiters))
 			{
-				int Component;
-				if (!Int32.TryParse(TextElement, out Component))
+				int component;
+				if (!Int32.TryParse(textElement, out component))
 				{
-					OutNumber = null;
+					outNumber = null;
 					return false;
 				}
-				Components.Add(Component);
+				components.Add(component);
 			}
 
-			OutNumber = new VersionNumber(Components.ToArray());
+			outNumber = new VersionNumber(components.ToArray());
 			return true;
 		}
 
@@ -234,17 +236,17 @@ namespace EpicGames.Core
 		/// <returns>The stringized version number</returns>
 		public override string ToString()
 		{
-			StringBuilder Result = new StringBuilder();
-			if (Components.Length > 0)
+			StringBuilder result = new StringBuilder();
+			if (_components.Length > 0)
 			{
-				Result.Append(Components[0]);
-				for (int Idx = 1; Idx < Components.Length; Idx++)
+				result.Append(_components[0]);
+				for (int idx = 1; idx < _components.Length; idx++)
 				{
-					Result.Append('.');
-					Result.Append(Components[Idx]);
+					result.Append('.');
+					result.Append(_components[idx]);
 				}
 			}
-			return Result.ToString();
+			return result.ToString();
 		}
 	}
 }
diff --git a/Engine/Source/Programs/Shared/EpicGames.Core/VersionNumberRange.cs b/Engine/Source/Programs/Shared/EpicGames.Core/VersionNumberRange.cs
--- a/Engine/Source/Programs/Shared/EpicGames.Core/VersionNumberRange.cs
+++ b/Engine/Source/Programs/Shared/EpicGames.Core/VersionNumberRange.cs
@@ -22,33 +22,33 @@ namespace EpicGames.Core
 		/// <summary>
 		/// Constructor
 		/// </summary>
-		/// <param name="Min"></param>
-		/// <param name="Max"></param>
-		public VersionNumberRange(VersionNumber Min, VersionNumber Max)
+		/// <param name="min"></param>
+		/// <param name="max"></param>
+		public VersionNumberRange(VersionNumber min, VersionNumber max)
 		{
-			this.Min = Min;
-			this.Max = Max;
+			this.Min = min;
+			this.Max = max;
 		}
 
 		/// <summary>
 		/// Tests whether this range contains the given version
 		/// </summary>
-		/// <param name="Version"></param>
+		/// <param name="version"></param>
 		/// <returns></returns>
-		public bool Contains(VersionNumber Version)
+		public bool Contains(VersionNumber version)
 		{
-			return Version >= Min && Version <= Max;
+			return version >= Min && version <= Max;
 		}
 
 		/// <summary>
 		/// Parse a version range from two strings
 		/// </summary>
-		/// <param name="MinText"></param>
-		/// <param name="MaxText"></param>
+		/// <param name="minText"></param>
+		/// <param name="maxText"></param>
 		/// <returns></returns>
-		public static VersionNumberRange Parse(string MinText, string MaxText)
+		public static VersionNumberRange Parse(string minText, string maxText)
 		{
-			return new VersionNumberRange(VersionNumber.Parse(MinText), VersionNumber.Parse(MaxText));
+			return new VersionNumberRange(VersionNumber.Parse(minText), VersionNumber.Parse(maxText));
 		}
 
 		/// <inheritdoc/>
