- Fixed Mac and maybe Linux with new SDK.json changes

diff --git a/Engine/Source/Programs/Shared/EpicGames.Build/UEBuildPlatformSDK.cs b/Engine/Source/Programs/Shared/EpicGames.Build/UEBuildPlatformSDK.cs
--- a/Engine/Source/Programs/Shared/EpicGames.Build/UEBuildPlatformSDK.cs
+++ b/Engine/Source/Programs/Shared/EpicGames.Build/UEBuildPlatformSDK.cs
@@ -325,10 +325,7 @@ namespace EpicGames.Core
 			bIsSdkAllowedOnHost = bInIsSdkAllowedOnHost;
 
 			// load the SDK config file
-			if (bIsSdkAllowedOnHost)
-			{
-				LoadJsonFile(PlatformName);
-			}
+			LoadJsonFile(PlatformName);
 
 			// if the parent set up autosdk, the env vars will be wrong, but we can still get the manual SDK version from before it was setup
 			string? ParentManualSDKVersions = Environment.GetEnvironmentVariable(GetPlatformManualSDKSetupEnvVar());
