- Fixed a static analysis C# issue

diff --git a/Engine/Source/Programs/Shared/EpicGames.Build/UEBuildPlatformSDK.cs b/Engine/Source/Programs/Shared/EpicGames.Build/UEBuildPlatformSDK.cs
--- a/Engine/Source/Programs/Shared/EpicGames.Build/UEBuildPlatformSDK.cs
+++ b/Engine/Source/Programs/Shared/EpicGames.Build/UEBuildPlatformSDK.cs
@@ -825,7 +825,7 @@ namespace EpicGames.Core
 		/// <returns>Comparison integer</returns>
 		public int SdkVersionsCompare( string? StringValueA, string? StringValueB, string? Hint = null )
 		{
-			UInt64 ValueA = 0, ValueB = 0;
+			UInt64 ValueA, ValueB;
 			TryConvertVersionToInt(StringValueA, out ValueA, Hint);
 			TryConvertVersionToInt(StringValueB, out ValueB, Hint);
 			if (ValueA == ValueB)
