Horde: Add missing ECDSA case for server compute client

diff --git a/Engine/Source/Programs/Shared/EpicGames.Horde/Compute/Clients/ServerComputeClient.cs b/Engine/Source/Programs/Shared/EpicGames.Horde/Compute/Clients/ServerComputeClient.cs
--- a/Engine/Source/Programs/Shared/EpicGames.Horde/Compute/Clients/ServerComputeClient.cs
+++ b/Engine/Source/Programs/Shared/EpicGames.Horde/Compute/Clients/ServerComputeClient.cs
@@ -342,6 +342,7 @@ namespace EpicGames.Horde.Compute.Clients
 			switch (response.Encryption)
 			{
 				case Encryption.Ssl:
+				case Encryption.SslEcdsaP256:
 					TcpSslTransport sslTransport = new(socket, StringUtils.ParseHexString(response.Certificate), false);
 					await sslTransport.AuthenticateAsync(cancellationToken);
 					return sslTransport;
