Horde: Fix race condition in tracking memory allocator.

diff --git a/Engine/Source/Programs/Shared/EpicGames.Core/IMemoryAllocator.cs b/Engine/Source/Programs/Shared/EpicGames.Core/IMemoryAllocator.cs
--- a/Engine/Source/Programs/Shared/EpicGames.Core/IMemoryAllocator.cs
+++ b/Engine/Source/Programs/Shared/EpicGames.Core/IMemoryAllocator.cs
@@ -243,8 +243,11 @@ namespace EpicGames.Core
 			{
 				if (_allocation != null)
 				{
-					_allocator._allocatedSize -= _allocation.Memory.Length;
-					_allocator._allocations.Remove(_node);
+					lock (_allocator._lockObject)
+					{
+						_allocator._allocatedSize -= _allocation.Memory.Length;
+						_allocator._allocations.Remove(_node);
+					}
 
 					_allocation.Dispose();
 					_allocation = null!;
