- Fixed an issue with VisionOS missing SDK.json on WIndows. This is a stopgap method, but removes a crash

#changelist validated
#virtualized

diff --git a/Engine/Source/Programs/Shared/EpicGames.Build/UEBuildPlatformSDK.cs b/Engine/Source/Programs/Shared/EpicGames.Build/UEBuildPlatformSDK.cs
--- a/Engine/Source/Programs/Shared/EpicGames.Build/UEBuildPlatformSDK.cs
+++ b/Engine/Source/Programs/Shared/EpicGames.Build/UEBuildPlatformSDK.cs
@@ -387,6 +387,7 @@ namespace EpicGames.Core
 					object[] parameter = new object[1];
 					parameter[0] = Log.Logger;
 					SDK = (UEBuildPlatformSDK)Activator.CreateInstance(typeof(T), parameter)!;
+					SDK.bIsSdkAllowedOnHost = true;
 					SDK.LoadJsonFile(PlatformName);
 					TempSDKRegistry.Add(PlatformName, SDK);
 				}
@@ -1013,7 +1014,8 @@ namespace EpicGames.Core
 				return null;
 			};
 
-			FileReference EngineSDKConfigFile = MakeConfigFilename(Unreal.EngineDirectory, true)!;
+			// if the SDK isn't allowed on the host, then allow it to not exist
+			FileReference EngineSDKConfigFile = MakeConfigFilename(Unreal.EngineDirectory, bIsSdkAllowedOnHost)!;
 
 			// load the file, along with any chained group file
 			ProcessJsonFile(EngineSDKConfigFile, ConfigSDKVersions, ConfigSDKVersionArrays);
