Horde: Add cancellation tokens to all async methods.

diff --git a/Engine/Source/Programs/Horde/Horde.Server/Tools/ToolCollection.cs b/Engine/Source/Programs/Horde/Horde.Server/Tools/ToolCollection.cs
--- a/Engine/Source/Programs/Horde/Horde.Server/Tools/ToolCollection.cs
+++ b/Engine/Source/Programs/Horde/Horde.Server/Tools/ToolCollection.cs
@@ -183,21 +183,18 @@ namespace Horde.Server.Tools
 			_logger = logger;
 		}
 
-		/// <summary>
-		/// Gets a tool with the given identifier
-		/// </summary>
-		/// <param name="id">The tool identifier</param>
-		/// <param name="globalConfig">The current global configuration</param>
-		/// <returns></returns>
-		public async Task<ITool?> GetAsync(ToolId id, GlobalConfig globalConfig) => await GetInternalAsync(id, globalConfig);
+		/// <inheritdoc/>
+		public async Task<ITool?> GetAsync(ToolId id, GlobalConfig globalConfig, CancellationToken cancellationToken)
+			=> await GetInternalAsync(id, globalConfig, cancellationToken);
 
 		/// <summary>
 		/// Gets a tool with the given identifier
 		/// </summary>
 		/// <param name="toolId">The tool identifier</param>
 		/// <param name="globalConfig">The current global configuration</param>
+		/// <param name="cancellationToken">Cancellation token for the operation</param>
 		/// <returns></returns>
-		async Task<Tool?> GetInternalAsync(ToolId toolId, GlobalConfig globalConfig)
+		async Task<Tool?> GetInternalAsync(ToolId toolId, GlobalConfig globalConfig, CancellationToken cancellationToken)
 		{
 			ToolConfig? toolConfig;
 			if (globalConfig.TryGetTool(toolId, out toolConfig))
@@ -205,14 +202,14 @@ namespace Horde.Server.Tools
 				Tool? tool;
 				for (; ; )
 				{
-					tool = await _tools.Find(x => x.Id == toolId).FirstOrDefaultAsync();
+					tool = await _tools.Find(x => x.Id == toolId).FirstOrDefaultAsync(cancellationToken);
 					if (tool != null)
 					{
 						break;
 					}
 
 					tool = new Tool(toolId);
-					if (await _tools.InsertOneIgnoreDuplicatesAsync(tool))
+					if (await _tools.InsertOneIgnoreDuplicatesAsync(tool, cancellationToken))
 					{
 						break;
 					}
@@ -314,7 +311,7 @@ namespace Horde.Server.Tools
 					break;
 				}
 
-				newTool = await GetInternalAsync(tool.Id, globalConfig);
+				newTool = await GetInternalAsync(tool.Id, globalConfig, cancellationToken);
 				if (newTool == null)
 				{
 					return null;
@@ -334,7 +331,7 @@ namespace Horde.Server.Tools
 			const int MaxDeploymentCount = 5;
 			while (newTool.Deployments.Count >= MaxDeploymentCount)
 			{
-				newTool = await UpdateAsync(newTool, Builders<Tool>.Update.PopFirst(x => x.Deployments));
+				newTool = await UpdateAsync(newTool, Builders<Tool>.Update.PopFirst(x => x.Deployments), cancellationToken);
 				if (newTool == null)
 				{
 					return null;
@@ -346,27 +343,21 @@ namespace Horde.Server.Tools
 			}
 
 			// Add the new deployment
-			return await UpdateAsync(newTool, Builders<Tool>.Update.Push(x => x.Deployments, deployment));
+			return await UpdateAsync(newTool, Builders<Tool>.Update.Push(x => x.Deployments, deployment), cancellationToken);
 		}
 
-		/// <summary>
-		/// Updates the state of the current deployment
-		/// </summary>
-		/// <param name="tool">Tool to be updated</param>
-		/// <param name="deploymentId">Identifier for the deployment to modify</param>
-		/// <param name="action">New state of the deployment</param>
-		/// <returns></returns>
-		public async Task<ITool?> UpdateDeploymentAsync(ITool tool, ToolDeploymentId deploymentId, ToolDeploymentState action)
+		/// <inheritdoc/>
+		public async Task<ITool?> UpdateDeploymentAsync(ITool tool, ToolDeploymentId deploymentId, ToolDeploymentState action, CancellationToken cancellationToken)
 		{
 			if (tool.Config is BundledToolConfig)
 			{
 				throw new InvalidOperationException("Cannot update the state of bundled tools.");
 			}
 
-			return await UpdateDeploymentInternalAsync((Tool)tool, deploymentId, action);
+			return await UpdateDeploymentInternalAsync((Tool)tool, deploymentId, action, cancellationToken);
 		}
 
-		async Task<Tool?> UpdateDeploymentInternalAsync(Tool tool, ToolDeploymentId deploymentId, ToolDeploymentState action)
+		async Task<Tool?> UpdateDeploymentInternalAsync(Tool tool, ToolDeploymentId deploymentId, ToolDeploymentState action, CancellationToken cancellationToken)
 		{
 			int idx = tool.Deployments.FindIndex(x => x.Id == deploymentId);
 			if (idx == -1)
@@ -378,11 +369,11 @@ namespace Horde.Server.Tools
 			switch (action)
 			{
 				case ToolDeploymentState.Complete:
-					return await UpdateAsync(tool, Builders<Tool>.Update.Set(x => x.Deployments[idx].BaseProgress, 1.0).Unset(x => x.Deployments[idx].StartedAt));
+					return await UpdateAsync(tool, Builders<Tool>.Update.Set(x => x.Deployments[idx].BaseProgress, 1.0).Unset(x => x.Deployments[idx].StartedAt), cancellationToken);
 
 				case ToolDeploymentState.Cancelled:
 					List<ToolDeployment> newDeployments = tool.Deployments.Where(x => x != deployment).ToList();
-					return await UpdateAsync(tool, Builders<Tool>.Update.Set(x => x.Deployments, newDeployments));
+					return await UpdateAsync(tool, Builders<Tool>.Update.Set(x => x.Deployments, newDeployments), cancellationToken);
 
 				case ToolDeploymentState.Paused:
 					if (deployment.StartedAt == null)
@@ -391,7 +382,7 @@ namespace Horde.Server.Tools
 					}
 					else
 					{
-						return await UpdateAsync(tool, Builders<Tool>.Update.Set(x => x.Deployments[idx].BaseProgress, deployment.GetProgressValue(_clock.UtcNow)).Set(x => x.Deployments[idx].StartedAt, null));
+						return await UpdateAsync(tool, Builders<Tool>.Update.Set(x => x.Deployments[idx].BaseProgress, deployment.GetProgressValue(_clock.UtcNow)).Set(x => x.Deployments[idx].StartedAt, null), cancellationToken);
 					}
 
 				case ToolDeploymentState.Active:
@@ -401,7 +392,7 @@ namespace Horde.Server.Tools
 					}
 					else
 					{
-						return await UpdateAsync(tool, Builders<Tool>.Update.Set(x => x.Deployments[idx].StartedAt, _clock.UtcNow));
+						return await UpdateAsync(tool, Builders<Tool>.Update.Set(x => x.Deployments[idx].StartedAt, _clock.UtcNow), cancellationToken);
 					}
 
 				default:
@@ -467,12 +458,12 @@ namespace Horde.Server.Tools
 #pragma warning restore CA2000
 		}
 
-		async Task<Tool> UpdateAsync(Tool tool, UpdateDefinition<Tool> update)
+		async Task<Tool> UpdateAsync(Tool tool, UpdateDefinition<Tool> update, CancellationToken cancellationToken)
 		{
 			update = update.Set(x => x.LastUpdateTime, new DateTime(Math.Max(tool.LastUpdateTime.Ticks + 1, DateTime.UtcNow.Ticks)));
 
 			FilterDefinition<Tool> filter = Builders<Tool>.Filter.Eq(x => x.Id, tool.Id) & Builders<Tool>.Filter.Eq(x => x.LastUpdateTime, tool.LastUpdateTime);
-			return await _tools.FindOneAndUpdateAsync(filter, update, new FindOneAndUpdateOptions<Tool> { ReturnDocument = ReturnDocument.After });
+			return await _tools.FindOneAndUpdateAsync(filter, update, new FindOneAndUpdateOptions<Tool> { ReturnDocument = ReturnDocument.After }, cancellationToken);
 		}
 	}
 }
diff --git a/Engine/Source/Programs/Horde/Horde.Server/Tools/ToolsController.cs b/Engine/Source/Programs/Horde/Horde.Server/Tools/ToolsController.cs
--- a/Engine/Source/Programs/Horde/Horde.Server/Tools/ToolsController.cs
+++ b/Engine/Source/Programs/Horde/Horde.Server/Tools/ToolsController.cs
@@ -53,7 +53,7 @@ namespace Horde.Server.Tools
 		[Route("/api/v1/tools/{id}/blobs")]
 		public async Task<ActionResult<WriteBlobResponse>> WriteBlobAsync(ToolId id, IFormFile? file, CancellationToken cancellationToken = default)
 		{
-			ITool? tool = await _toolCollection.GetAsync(id, _globalConfig.Value);
+			ITool? tool = await _toolCollection.GetAsync(id, _globalConfig.Value, cancellationToken);
 			if (tool == null)
 			{
 				return NotFound(id);
@@ -75,7 +75,7 @@ namespace Horde.Server.Tools
 		[Route("/api/v1/tools/{id}/deployments")]
 		public async Task<ActionResult<CreateToolDeploymentResponse>> CreateDeploymentAsync(ToolId id, [FromForm] ToolDeploymentConfig options, [FromForm] IFormFile file, CancellationToken cancellationToken)
 		{
-			ITool? tool = await _toolCollection.GetAsync(id, _globalConfig.Value);
+			ITool? tool = await _toolCollection.GetAsync(id, _globalConfig.Value, cancellationToken);
 
 			if (tool == null)
 			{
@@ -105,7 +105,7 @@ namespace Horde.Server.Tools
 		[Route("/api/v2/tools/{id}/deployments")]
 		public async Task<ActionResult<CreateToolDeploymentResponse>> CreateDeploymentAsync(ToolId id, CreateToolDeploymentRequest request, CancellationToken cancellationToken)
 		{
-			ITool? tool = await _toolCollection.GetAsync(id, _globalConfig.Value);
+			ITool? tool = await _toolCollection.GetAsync(id, _globalConfig.Value, cancellationToken);
 
 			if (tool == null)
 			{
@@ -225,7 +225,7 @@ namespace Horde.Server.Tools
 		[Route("/api/v1/tools/{id}")]
 		public async Task<ActionResult> GetToolAsync(ToolId id, GetToolAction action = GetToolAction.Info, CancellationToken cancellationToken = default)
 		{
-			ITool? tool = await _toolCollection.GetAsync(id, _globalConfig.Value);
+			ITool? tool = await _toolCollection.GetAsync(id, _globalConfig.Value, cancellationToken);
 			if (tool == null)
 			{
 				return NotFound(id);
@@ -273,7 +273,7 @@ namespace Horde.Server.Tools
 		[Route("/api/v1/tools/{id}/deployments")]
 		public async Task<ActionResult> FindDeploymentAsync(ToolId id, [FromQuery] double phase = 0.0, [FromQuery] GetToolAction action = GetToolAction.Info, CancellationToken cancellationToken = default)
 		{
-			ITool? tool = await _toolCollection.GetAsync(id, _globalConfig.Value);
+			ITool? tool = await _toolCollection.GetAsync(id, _globalConfig.Value, cancellationToken);
 			if (tool == null)
 			{
 				return NotFound(id);
@@ -300,7 +300,7 @@ namespace Horde.Server.Tools
 		[Route("/api/v1/tools/{id}/deployments/{deploymentId}")]
 		public async Task<ActionResult> GetDeploymentAsync(ToolId id, ToolDeploymentId deploymentId, [FromQuery] GetToolAction action = GetToolAction.Info, CancellationToken cancellationToken = default)
 		{
-			ITool? tool = await _toolCollection.GetAsync(id, _globalConfig.Value);
+			ITool? tool = await _toolCollection.GetAsync(id, _globalConfig.Value, cancellationToken);
 			if (tool == null)
 			{
 				return NotFound(id);
@@ -382,7 +382,7 @@ namespace Horde.Server.Tools
 		[Route("/api/v1/tools/{id}/blobs/{*locator}")]
 		public async Task<ActionResult<object>> ReadToolBlobAsync(ToolId id, BlobLocator locator, CancellationToken cancellationToken = default)
 		{
-			ITool? tool = await _toolCollection.GetAsync(id, _globalConfig.Value);
+			ITool? tool = await _toolCollection.GetAsync(id, _globalConfig.Value, cancellationToken);
 			if (tool == null)
 			{
 				return NotFound(id);
diff --git a/Engine/Source/Programs/Horde/Horde.Server/Ugs/IUgsMetadataCollection.cs b/Engine/Source/Programs/Horde/Horde.Server/Ugs/IUgsMetadataCollection.cs
--- a/Engine/Source/Programs/Horde/Horde.Server/Ugs/IUgsMetadataCollection.cs
+++ b/Engine/Source/Programs/Horde/Horde.Server/Ugs/IUgsMetadataCollection.cs
@@ -2,6 +2,7 @@
 
 using System;
 using System.Collections.Generic;
+using System.Threading;
 using System.Threading.Tasks;
 
 namespace Horde.Server.Ugs
@@ -17,8 +18,9 @@ namespace Horde.Server.Ugs
 		/// <param name="stream">The stream containing the change</param>
 		/// <param name="change">The changelist number</param>
 		/// <param name="project">Project identifier</param>
+		/// <param name="cancellationToken">Cancellation token for the operation</param>
 		/// <returns>The metadata instance</returns>
-		Task<IUgsMetadata> FindOrAddAsync(string stream, int change, string? project);
+		Task<IUgsMetadata> FindOrAddAsync(string stream, int change, string? project, CancellationToken cancellationToken = default);
 
 		/// <summary>
 		/// Adds information to a change
@@ -30,8 +32,9 @@ namespace Horde.Server.Ugs
 		/// <param name="investigating">Whether the user is investigating the change</param>
 		/// <param name="starred">Whether the change is starred</param>
 		/// <param name="comment">New comment from this user</param>
+		/// <param name="cancellationToken">Cancellation token for the operation</param>
 		/// <returns>Async task task</returns>
-		Task<IUgsMetadata> UpdateUserAsync(IUgsMetadata metadata, string userName, bool? synced, UgsUserVote? vote, bool? investigating, bool? starred, string? comment);
+		Task<IUgsMetadata> UpdateUserAsync(IUgsMetadata metadata, string userName, bool? synced, UgsUserVote? vote, bool? investigating, bool? starred, string? comment, CancellationToken cancellationToken = default);
 
 		/// <summary>
 		/// Updates the state of a badge
@@ -40,8 +43,9 @@ namespace Horde.Server.Ugs
 		/// <param name="name">Name of the badge</param>
 		/// <param name="url">Url to link to for the badge</param>
 		/// <param name="state">State of the badge</param>
+		/// <param name="cancellationToken">Cancellation token for the operation</param>
 		/// <returns>Async task</returns>
-		Task<IUgsMetadata> UpdateBadgeAsync(IUgsMetadata metadata, string name, Uri? url, UgsBadgeState state);
+		Task<IUgsMetadata> UpdateBadgeAsync(IUgsMetadata metadata, string name, Uri? url, UgsBadgeState state, CancellationToken cancellationToken = default);
 
 		/// <summary>
 		/// Searches for metadata updates
@@ -50,7 +54,8 @@ namespace Horde.Server.Ugs
 		/// <param name="minChange">Minimum changelist number</param>
 		/// <param name="maxChange">Maximum changelist number</param>
 		/// <param name="afterTicks">Last query time</param>
+		/// <param name="cancellationToken">Cancellation token for the operation</param>
 		/// <returns>List of metadata updates</returns>
-		Task<List<IUgsMetadata>> FindAsync(string stream, int minChange, int? maxChange = null, long? afterTicks = null);
+		Task<List<IUgsMetadata>> FindAsync(string stream, int minChange, int? maxChange = null, long? afterTicks = null, CancellationToken cancellationToken = default);
 	}
 }
diff --git a/Engine/Source/Programs/Horde/Horde.Server/Ugs/UgsController.cs b/Engine/Source/Programs/Horde/Horde.Server/Ugs/UgsController.cs
--- a/Engine/Source/Programs/Horde/Horde.Server/Ugs/UgsController.cs
+++ b/Engine/Source/Programs/Horde/Horde.Server/Ugs/UgsController.cs
@@ -151,23 +151,24 @@ namespace Horde.Server.Ugs
 		/// <param name="user"></param>
 		/// <param name="includeResolved">Whether to include resolved issues</param>
 		/// <param name="maxResults">Maximum number of results to return</param>
+		/// <param name="cancellationToken">Cancellation token for the operation</param>
 		/// <returns>List of matching agents</returns>
 		[HttpGet]
 		[Route("/ugs/api/issues")]
 		[ProducesResponseType(typeof(GetUgsIssueResponse), 200)]
-		public async Task<ActionResult<List<GetUgsIssueResponse>>> GetIssuesAsync([FromQuery] string? user = null, [FromQuery] bool includeResolved = false, [FromQuery] int maxResults = 100)
+		public async Task<ActionResult<List<GetUgsIssueResponse>>> GetIssuesAsync([FromQuery] string? user = null, [FromQuery] bool includeResolved = false, [FromQuery] int maxResults = 100, CancellationToken cancellationToken = default)
 		{
-			IUser? userInfo = (user != null) ? await _userCollection.FindUserByLoginAsync(user) : null;
+			IUser? userInfo = (user != null) ? await _userCollection.FindUserByLoginAsync(user, cancellationToken) : null;
 
 			List<GetUgsIssueResponse> responses = new List<GetUgsIssueResponse>();
 			if (includeResolved)
 			{
-				List<IIssue> issues = await _issueService.Collection.FindIssuesAsync(null, resolved: null, count: maxResults);
+				IReadOnlyList<IIssue> issues = await _issueService.Collection.FindIssuesAsync(null, resolved: null, count: maxResults, cancellationToken: cancellationToken);
 				foreach(IIssue issue in issues)
 				{
-					IIssueDetails details = await _issueService.GetIssueDetailsAsync(issue);
+					IIssueDetails details = await _issueService.GetIssueDetailsAsync(issue, cancellationToken);
 					bool notify = userInfo != null && details.Suspects.Any(x => x.AuthorId == userInfo.Id);
-					responses.Add(await CreateIssueResponseAsync(details, notify));
+					responses.Add(await CreateIssueResponseAsync(details, notify, cancellationToken));
 				}
 			}
 			else
@@ -182,7 +183,7 @@ namespace Horde.Server.Ugs
 					if (cachedOpenIssue.ShowNotifications())
 					{
 						bool notify = userInfo != null && cachedOpenIssue.IncludeForUser(userInfo.Id);
-						responses.Add(await CreateIssueResponseAsync(cachedOpenIssue, notify));
+						responses.Add(await CreateIssueResponseAsync(cachedOpenIssue, notify, cancellationToken));
 					}
 				}
 			}
@@ -194,19 +195,20 @@ namespace Horde.Server.Ugs
 		/// </summary>
 		/// <param name="issueId">Id of the issue to get information about</param>
 		/// <param name="filter">Filter for the properties to return</param>
+		/// <param name="cancellationToken">Cancellation token for the operation</param>
 		/// <returns>List of matching agents</returns>
 		[HttpGet]
 		[Route("/ugs/api/issues/{issueId}")]
 		[ProducesResponseType(typeof(GetUgsIssueBuildResponse), 200)]
-		public async Task<ActionResult<object>> GetIssueAsync(int issueId, [FromQuery] PropertyFilter? filter = null)
+		public async Task<ActionResult<object>> GetIssueAsync(int issueId, [FromQuery] PropertyFilter? filter = null, CancellationToken cancellationToken = default)
 		{
-			IIssueDetails? issue = await _issueService.GetIssueDetailsAsync(issueId);
+			IIssueDetails? issue = await _issueService.GetIssueDetailsAsync(issueId, cancellationToken);
 			if (issue == null)
 			{
 				return NotFound();
 			}
 
-			return PropertyFilter.Apply(await CreateIssueResponseAsync(issue, false), filter);
+			return PropertyFilter.Apply(await CreateIssueResponseAsync(issue, false, cancellationToken), filter);
 		}
 
 		/// <summary>
@@ -214,13 +216,14 @@ namespace Horde.Server.Ugs
 		/// </summary>
 		/// <param name="issueId">Id of the issue to get information about</param>
 		/// <param name="filter">Filter for the properties to return</param>
+		/// <param name="cancellationToken">Cancellation token for the operation</param>
 		/// <returns>List of matching agents</returns>
 		[HttpGet]
 		[Route("/ugs/api/issues/{issueId}/builds")]
 		[ProducesResponseType(typeof(List<GetUgsIssueBuildResponse>), 200)]
-		public async Task<ActionResult<List<object>>> GetIssueBuildsAsync(int issueId, [FromQuery] PropertyFilter? filter = null)
+		public async Task<ActionResult<List<object>>> GetIssueBuildsAsync(int issueId, [FromQuery] PropertyFilter? filter = null, CancellationToken cancellationToken = default)
 		{
-			IIssueDetails? issue = await _issueService.GetCachedIssueDetailsAsync(issueId);
+			IIssueDetails? issue = await _issueService.GetCachedIssueDetailsAsync(issueId, cancellationToken);
 			if (issue == null)
 			{
 				return NotFound();
@@ -261,7 +264,7 @@ namespace Horde.Server.Ugs
 
 			Dictionary<LogId, ILogFile?> logFiles = new Dictionary<LogId, ILogFile?>();
 
-			List<IIssueSpan> spans = await _issueService.Collection.FindSpansAsync(issueId);
+			IReadOnlyList<IIssueSpan> spans = await _issueService.Collection.FindSpansAsync(issueId, cancellationToken);
 			List<ILogEvent> events = await _logFileService.FindEventsForSpansAsync(spans.Select(x => x.Id), null, 0, count: 10, cancellationToken);
 
 			foreach (ILogEvent logEvent in events)
@@ -291,13 +294,14 @@ namespace Horde.Server.Ugs
 		/// </summary>
 		/// <param name="details">The issue to get a URL for</param>
 		/// <param name="notify">Whether to show notifications for this issue</param>
+		/// <param name="cancellationToken">Cancellation token for the operation</param>
 		/// <returns>The issue response</returns>
-		async Task<GetUgsIssueResponse> CreateIssueResponseAsync(IIssueDetails details, bool notify)
+		async Task<GetUgsIssueResponse> CreateIssueResponseAsync(IIssueDetails details, bool notify, CancellationToken cancellationToken)
 		{
 			Uri? buildUrl = GetIssueBuildUrl(details);
 
-			IUser? owner = details.Issue.OwnerId.HasValue ? await _userCollection.GetCachedUserAsync(details.Issue.OwnerId.Value) : null;
-			IUser? nominatedBy = details.Issue.NominatedById.HasValue ? await _userCollection.GetCachedUserAsync(details.Issue.NominatedById.Value) : null;
+			IUser? owner = details.Issue.OwnerId.HasValue ? await _userCollection.GetCachedUserAsync(details.Issue.OwnerId.Value, cancellationToken) : null;
+			IUser? nominatedBy = details.Issue.NominatedById.HasValue ? await _userCollection.GetCachedUserAsync(details.Issue.NominatedById.Value, cancellationToken) : null;
 
 			return new GetUgsIssueResponse(details, owner, nominatedBy, notify, buildUrl);
 		}
diff --git a/Engine/Source/Programs/Horde/Horde.Server/Ugs/UgsMetadataCollection.cs b/Engine/Source/Programs/Horde/Horde.Server/Ugs/UgsMetadataCollection.cs
--- a/Engine/Source/Programs/Horde/Horde.Server/Ugs/UgsMetadataCollection.cs
+++ b/Engine/Source/Programs/Horde/Horde.Server/Ugs/UgsMetadataCollection.cs
@@ -4,6 +4,7 @@ using System;
 using System.Collections.Generic;
 using System.Diagnostics.CodeAnalysis;
 using System.Linq;
+using System.Threading;
 using System.Threading.Tasks;
 using Horde.Server.Server;
 using Horde.Server.Utilities;
@@ -108,15 +109,16 @@ namespace Horde.Server.Ugs
 		/// <param name="stream">Stream containing the change</param>
 		/// <param name="change">The changelist number to add a document for</param>
 		/// <param name="project">Arbitrary identifier for this project</param>
+		/// <param name="cancellationToken">Cancellation token for the operation</param>
 		/// <returns>The metadata document</returns>
-		public async Task<IUgsMetadata> FindOrAddAsync(string stream, int change, string? project)
+		public async Task<IUgsMetadata> FindOrAddAsync(string stream, int change, string? project, CancellationToken cancellationToken)
 		{
 			string normalizedStream = GetNormalizedStream(stream);
 			string normalizedProject = GetNormalizedProject(project);
 			for (; ; )
 			{
 				// Find an existing document
-				UgsMetadataDocument? existing = await _collection.Find(x => x.Stream == normalizedStream && x.Change == change && x.Project == normalizedProject).FirstOrDefaultAsync();
+				UgsMetadataDocument? existing = await _collection.Find(x => x.Stream == normalizedStream && x.Change == change && x.Project == normalizedProject).FirstOrDefaultAsync(cancellationToken);
 				if (existing != null)
 				{
 					return existing;
@@ -126,7 +128,7 @@ namespace Horde.Server.Ugs
 				try
 				{
 					UgsMetadataDocument newDocument = new UgsMetadataDocument(normalizedStream, change, normalizedProject);
-					await _collection.InsertOneAsync(newDocument);
+					await _collection.InsertOneAsync(newDocument, null, cancellationToken);
 					return newDocument;
 				}
 				catch (MongoWriteException ex)
@@ -140,7 +142,7 @@ namespace Horde.Server.Ugs
 		}
 
 		/// <inheritdoc/>
-		public async Task<IUgsMetadata> UpdateUserAsync(IUgsMetadata metadata, string userName, bool? synced, UgsUserVote? vote, bool? investigating, bool? starred, string? comment)
+		public async Task<IUgsMetadata> UpdateUserAsync(IUgsMetadata metadata, string userName, bool? synced, UgsUserVote? vote, bool? investigating, bool? starred, string? comment, CancellationToken cancellationToken)
 		{
 			UpdateDefinitionBuilder<UgsMetadataDocument> updateBuilder = Builders<UgsMetadataDocument>.Update;
 			for (; ; )
@@ -164,7 +166,7 @@ namespace Horde.Server.Ugs
 					userData.Starred = (starred == true) ? starred : null;
 					userData.Comment = comment;
 
-					if (await TryUpdateAsync(document, updateBuilder.Push(x => x.Users, userData)))
+					if (await TryUpdateAsync(document, updateBuilder.Push(x => x.Users, userData), cancellationToken))
 					{
 						document.Users.Add(userData);
 						return metadata;
@@ -202,7 +204,7 @@ namespace Horde.Server.Ugs
 						updates.Add(updateBuilder.Set(x => x.Users![userIdx].Comment, comment));
 					}
 
-					if (updates.Count == 0 || await TryUpdateAsync(document, updateBuilder.Combine(updates)))
+					if (updates.Count == 0 || await TryUpdateAsync(document, updateBuilder.Combine(updates), cancellationToken))
 					{
 						userData.SyncTime = newSyncTime;
 						userData.Vote = vote;
@@ -214,12 +216,12 @@ namespace Horde.Server.Ugs
 				}
 
 				// Update the document and try again
-				metadata = await FindOrAddAsync(metadata.Stream, metadata.Change, metadata.Project);
+				metadata = await FindOrAddAsync(metadata.Stream, metadata.Change, metadata.Project, cancellationToken);
 			}
 		}
 
 		/// <inheritdoc/>
-		public async Task<IUgsMetadata> UpdateBadgeAsync(IUgsMetadata metadata, string name, Uri? url, UgsBadgeState state)
+		public async Task<IUgsMetadata> UpdateBadgeAsync(IUgsMetadata metadata, string name, Uri? url, UgsBadgeState state, CancellationToken cancellationToken)
 		{
 			UpdateDefinitionBuilder<UgsMetadataDocument> updateBuilder = Builders<UgsMetadataDocument>.Update;
 			for (; ; )
@@ -236,7 +238,7 @@ namespace Horde.Server.Ugs
 					newBadge.Url = url;
 					newBadge.State = state;
 
-					if (await TryUpdateAsync(document, updateBuilder.Push(x => x.Badges, newBadge)))
+					if (await TryUpdateAsync(document, updateBuilder.Push(x => x.Badges, newBadge), cancellationToken))
 					{
 						document.Badges.Add(newBadge);
 						return metadata;
@@ -257,7 +259,7 @@ namespace Horde.Server.Ugs
 						updates.Add(updateBuilder.Set(x => x.Badges![badgeIdx].State, state));
 					}
 
-					if (updates.Count == 0 || await TryUpdateAsync(document, updateBuilder.Combine(updates)))
+					if (updates.Count == 0 || await TryUpdateAsync(document, updateBuilder.Combine(updates), cancellationToken))
 					{
 						badgeData.Url = url;
 						badgeData.State = state;
@@ -266,12 +268,12 @@ namespace Horde.Server.Ugs
 				}
 
 				// Update the document and try again
-				metadata = await FindOrAddAsync(metadata.Stream, metadata.Change, metadata.Project);
+				metadata = await FindOrAddAsync(metadata.Stream, metadata.Change, metadata.Project, cancellationToken);
 			}
 		}
 
 		/// <inheritdoc/>
-		public async Task<List<IUgsMetadata>> FindAsync(string stream, int minChange, int? maxChange = null, long? minTime = null)
+		public async Task<List<IUgsMetadata>> FindAsync(string stream, int minChange, int? maxChange = null, long? minTime = null, CancellationToken cancellationToken = default)
 		{
 			FilterDefinitionBuilder<UgsMetadataDocument> filterBuilder = Builders<UgsMetadataDocument>.Filter;
 
@@ -287,7 +289,7 @@ namespace Horde.Server.Ugs
 				filter &= filterBuilder.Gt(x => x.UpdateTicks, minTime.Value);
 			}
 
-			List<UgsMetadataDocument> documents = await _collection.Find(filter).ToListAsync();
+			List<UgsMetadataDocument> documents = await _collection.Find(filter).ToListAsync(cancellationToken);
 			return documents.ConvertAll<IUgsMetadata>(x =>
 			{
 				// Remove polluting null entries. Need to find the source of these.
@@ -331,8 +333,9 @@ namespace Horde.Server.Ugs
 		/// </summary>
 		/// <param name="document">The document to update</param>
 		/// <param name="update">Update definition</param>
+		/// <param name="cancellationToken">Cancellation token for the operation</param>
 		/// <returns>True if the update succeeded</returns>
-		private async Task<bool> TryUpdateAsync(UgsMetadataDocument document, UpdateDefinition<UgsMetadataDocument> update)
+		private async Task<bool> TryUpdateAsync(UgsMetadataDocument document, UpdateDefinition<UgsMetadataDocument> update, CancellationToken cancellationToken)
 		{
 			int newUpdateIndex = document.UpdateIndex + 1;
 			update = update.Set(x => x.UpdateIndex, newUpdateIndex);
@@ -343,7 +346,7 @@ namespace Horde.Server.Ugs
 			FilterDefinition<UgsMetadataDocument> filter = Builders<UgsMetadataDocument>.Filter.Expr(x => x.Id == document.Id && x.UpdateIndex == document.UpdateIndex);
 			try
 			{
-				UpdateResult result = await _collection.UpdateOneAsync(filter, update);
+				UpdateResult result = await _collection.UpdateOneAsync(filter, update, null, cancellationToken);
 				if (result.ModifiedCount > 0)
 				{
 					document.UpdateIndex = newUpdateIndex;
diff --git a/Engine/Source/Programs/Horde/Horde.Server/Users/IAvatarService.cs b/Engine/Source/Programs/Horde/Horde.Server/Users/IAvatarService.cs
--- a/Engine/Source/Programs/Horde/Horde.Server/Users/IAvatarService.cs
+++ b/Engine/Source/Programs/Horde/Horde.Server/Users/IAvatarService.cs
@@ -1,5 +1,6 @@
 // Copyright Epic Games, Inc. All Rights Reserved.
 
+using System.Threading;
 using System.Threading.Tasks;
 
 namespace Horde.Server.Users
@@ -13,8 +14,9 @@ namespace Horde.Server.Users
 		/// Gets a users's avatar
 		/// </summary>
 		/// <param name="user"></param>
+		/// <param name="cancellationToken">CAncellation token for the operation</param>
 		/// <returns></returns>
-		Task<IAvatar?> GetAvatarAsync(IUser user);
+		Task<IAvatar?> GetAvatarAsync(IUser user, CancellationToken cancellationToken);
 	}
 
 	/// <summary>
@@ -23,7 +25,7 @@ namespace Horde.Server.Users
 	public class NullAvatarService : IAvatarService
 	{
 		/// <inheritdoc/>
-		public Task<IAvatar?> GetAvatarAsync(IUser user) => Task.FromResult<IAvatar?>(null);
+		public Task<IAvatar?> GetAvatarAsync(IUser user, CancellationToken cancellationToken) => Task.FromResult<IAvatar?>(null);
 	}
 }
 
diff --git a/Engine/Source/Programs/Horde/Horde.Server/Users/IUserCollection.cs b/Engine/Source/Programs/Horde/Horde.Server/Users/IUserCollection.cs
--- a/Engine/Source/Programs/Horde/Horde.Server/Users/IUserCollection.cs
+++ b/Engine/Source/Programs/Horde/Horde.Server/Users/IUserCollection.cs
@@ -8,6 +8,7 @@ using MongoDB.Bson;
 using EpicGames.Horde.Users;
 using EpicGames.Horde.Jobs;
 using EpicGames.Horde.Jobs.Bisect;
+using System.Threading;
 
 namespace Horde.Server.Users
 {
@@ -20,15 +21,17 @@ namespace Horde.Server.Users
 		/// Gets a user by unique id
 		/// </summary>
 		/// <param name="id">Id of the user</param>
+		/// <param name="cancellationToken">Cancellation token for the operation</param>
 		/// <returns>The user information</returns>
-		Task<IUser?> GetUserAsync(UserId id);
+		Task<IUser?> GetUserAsync(UserId id, CancellationToken cancellationToken = default);
 
 		/// <summary>
 		/// Gets a cached user by unique id
 		/// </summary>
 		/// <param name="id">Id of the user</param>
+		/// <param name="cancellationToken">Cancellation token for the operation</param>
 		/// <returns>The user information</returns>
-		ValueTask<IUser?> GetCachedUserAsync(UserId? id);
+		ValueTask<IUser?> GetCachedUserAsync(UserId? id, CancellationToken cancellationToken = default);
 
 		/// <summary>
 		/// Gets a user by unique id
@@ -37,22 +40,25 @@ namespace Horde.Server.Users
 		/// <param name="nameRegex">Name regex to match for the users</param>
 		/// <param name="index">Maximum number of results</param>
 		/// <param name="count">Number of results to return</param>
+		/// <param name="cancellationToken">Cancellation token for the operation</param>
 		/// <returns>The user information</returns>
-		Task<List<IUser>> FindUsersAsync(IEnumerable<UserId>? ids = null, string? nameRegex = null, int? index = null, int? count = null);
+		Task<IReadOnlyList<IUser>> FindUsersAsync(IEnumerable<UserId>? ids = null, string? nameRegex = null, int? index = null, int? count = null, CancellationToken cancellationToken = default);
 
 		/// <summary>
 		/// Gets a user by login
 		/// </summary>
 		/// <param name="login">Login for the user</param>
+		/// <param name="cancellationToken">Cancellation token for the operation</param>
 		/// <returns>The user information</returns>
-		Task<IUser?> FindUserByLoginAsync(string login);
+		Task<IUser?> FindUserByLoginAsync(string login, CancellationToken cancellationToken = default);
 
 		/// <summary>
 		/// Gets a user by email address
 		/// </summary>
 		/// <param name="email">Email for the user</param>
+		/// <param name="cancellationToken">Cancellation token for the operation</param>
 		/// <returns>The user information</returns>
-		Task<IUser?> FindUserByEmailAsync(string email);
+		Task<IUser?> FindUserByEmailAsync(string email, CancellationToken cancellationToken = default);
 
 		/// <summary>
 		/// Find or add a user with the given claims. Claims will be updated if the user exists.
@@ -60,30 +66,34 @@ namespace Horde.Server.Users
 		/// <param name="login">Login id of the user</param>
 		/// <param name="name">Full name of the user</param>
 		/// <param name="email">Email address of the user</param>
+		/// <param name="cancellationToken">Cancellation token for the operation</param>
 		/// <returns>The user document</returns>
-		Task<IUser> FindOrAddUserByLoginAsync(string login, string? name = null, string? email = null);
+		Task<IUser> FindOrAddUserByLoginAsync(string login, string? name = null, string? email = null, CancellationToken cancellationToken = default);
 
 		/// <summary>
 		/// Gets the claims for a user
 		/// </summary>
 		/// <param name="userId"></param>
+		/// <param name="cancellationToken">Cancellation token for the operation</param>
 		/// <returns></returns>
-		Task<IUserClaims> GetClaimsAsync(UserId userId);
+		Task<IUserClaims> GetClaimsAsync(UserId userId, CancellationToken cancellationToken = default);
 
 		/// <summary>
 		/// Update the claims for a user
 		/// </summary>
 		/// <param name="userId"></param>
 		/// <param name="claims"></param>
+		/// <param name="cancellationToken">Cancellation token for the operation</param>
 		/// <returns></returns>
-		Task UpdateClaimsAsync(UserId userId, IEnumerable<IUserClaim> claims);
+		Task UpdateClaimsAsync(UserId userId, IEnumerable<IUserClaim> claims, CancellationToken cancellationToken = default);
 
 		/// <summary>
 		/// Get settings for a user
 		/// </summary>
 		/// <param name="userId"></param>
+		/// <param name="cancellationToken">Cancellation token for the operation</param>
 		/// <returns></returns>
-		Task<IUserSettings> GetSettingsAsync(UserId userId);
+		Task<IUserSettings> GetSettingsAsync(UserId userId, CancellationToken cancellationToken = default);
 
 		/// <summary>
 		/// Update a user
@@ -97,8 +107,9 @@ namespace Horde.Server.Users
 		/// <param name="templateOptions"></param>
 		/// <param name="addBisectTaskIds"></param>
 		/// <param name="removeBisectTaskIds"></param>
+		/// <param name="cancellationToken">Cancellation token for the operation</param>
 		/// <returns>Updated user object</returns>
-		Task UpdateSettingsAsync(UserId userId, bool? enableExperimentalFeatures = null, bool? alwaysTagPreflightCL = null, BsonValue? dashboardSettings = null, IEnumerable<JobId>? addPinnedJobIds = null, IEnumerable<JobId>? removePinnedJobIds = null, UpdateUserJobTemplateOptions? templateOptions = null, IEnumerable<BisectTaskId>? addBisectTaskIds = null, IEnumerable<BisectTaskId>? removeBisectTaskIds = null);
+		Task UpdateSettingsAsync(UserId userId, bool? enableExperimentalFeatures = null, bool? alwaysTagPreflightCL = null, BsonValue? dashboardSettings = null, IEnumerable<JobId>? addPinnedJobIds = null, IEnumerable<JobId>? removePinnedJobIds = null, UpdateUserJobTemplateOptions? templateOptions = null, IEnumerable<BisectTaskId>? addBisectTaskIds = null, IEnumerable<BisectTaskId>? removeBisectTaskIds = null, CancellationToken cancellationToken = default);
 	}
 
 	/// <summary>
@@ -111,13 +122,14 @@ namespace Horde.Server.Users
 		/// </summary>
 		/// <param name="userCollection"></param>
 		/// <param name="principal"></param>
+		/// <param name="cancellationToken">Cancellation token for the operation</param>
 		/// <returns></returns>
-		public static Task<IUser?> GetUserAsync(this IUserCollection userCollection, ClaimsPrincipal principal)
+		public static Task<IUser?> GetUserAsync(this IUserCollection userCollection, ClaimsPrincipal principal, CancellationToken cancellationToken = default)
 		{
 			UserId? userId = principal.GetUserId();
 			if (userId != null)
 			{
-				return userCollection.GetUserAsync(userId.Value);
+				return userCollection.GetUserAsync(userId.Value, cancellationToken);
 			}
 			else
 			{
diff --git a/Engine/Source/Programs/Horde/Horde.Server/Users/UserCollection.cs b/Engine/Source/Programs/Horde/Horde.Server/Users/UserCollection.cs
--- a/Engine/Source/Programs/Horde/Horde.Server/Users/UserCollection.cs
+++ b/Engine/Source/Programs/Horde/Horde.Server/Users/UserCollection.cs
@@ -14,6 +14,8 @@ using MongoDB.Driver;
 using EpicGames.Horde.Users;
 using EpicGames.Horde.Jobs;
 using EpicGames.Horde.Jobs.Bisect;
+using System.Threading;
+using System.Runtime.CompilerServices;
 
 namespace Horde.Server.Users
 {
@@ -98,13 +100,13 @@ namespace Horde.Server.Users
 		}
 
 		/// <inheritdoc/>
-		public async Task<IUser?> GetUserAsync(UserId id)
+		public async Task<IUser?> GetUserAsync(UserId id, CancellationToken cancellationToken)
 		{
-			return await _users.Find(x => x.Id == id).FirstOrDefaultAsync();
+			return await _users.Find(x => x.Id == id).FirstOrDefaultAsync(cancellationToken);
 		}
 
 		/// <inheritdoc/>
-		public async ValueTask<IUser?> GetCachedUserAsync(UserId? id)
+		public async ValueTask<IUser?> GetCachedUserAsync(UserId? id, CancellationToken cancellationToken)
 		{
 			if (id == null)
 			{
@@ -112,59 +114,59 @@ namespace Horde.Server.Users
 			}
 			else
 			{
-				return await GetUserAsync(id.Value);
+				return await GetUserAsync(id.Value, cancellationToken);
 			}
 		}
 
 		/// <inheritdoc/>
-		public async Task<List<IUser>> FindUsersAsync(IEnumerable<UserId>? ids, string? nameRegex, int? index, int? count)
+		public async Task<IReadOnlyList<IUser>> FindUsersAsync(IEnumerable<UserId>? ids, string? nameRegex, int? index, int? count, CancellationToken cancellationToken)
 		{
 			FilterDefinition<UserDocument> filter = Builders<UserDocument>.Filter.In(x => x.Id, ids);
-			return await _users.Find(filter).Range(index, count).ToListAsync<UserDocument, IUser>();
+			return await _users.Find(filter).Range(index, count).ToListAsync(cancellationToken);
 		}
 
 		/// <inheritdoc/>
-		public async Task<IUser?> FindUserByLoginAsync(string login)
+		public async Task<IUser?> FindUserByLoginAsync(string login, CancellationToken cancellationToken)
 		{
 			ClaimDocument primaryClaim = new ClaimDocument(ClaimTypes.Name, login);
-			return await _users.Find(x => x.PrimaryClaim == primaryClaim).FirstOrDefaultAsync();
+			return await _users.Find(x => x.PrimaryClaim == primaryClaim).FirstOrDefaultAsync(cancellationToken);
 		}
 
 		/// <inheritdoc/>
-		public Task<IUser?> FindUserByEmailAsync(string email)
+		public Task<IUser?> FindUserByEmailAsync(string email, CancellationToken cancellationToken)
 		{
 			return Task.FromResult<IUser?>(null);
 		}
 
 		/// <inheritdoc/>
-		public async Task<IUser> FindOrAddUserByLoginAsync(string login, string? name, string? email)
+		public async Task<IUser> FindOrAddUserByLoginAsync(string login, string? name, string? email, CancellationToken cancellationToken)
 		{
 			ClaimDocument newPrimaryClaim = new ClaimDocument(ClaimTypes.Name, login);
 			UpdateDefinition<UserDocument> update = Builders<UserDocument>.Update.SetOnInsert(x => x.Id, new UserId(BinaryIdUtils.CreateNew()));
-			return await _users.FindOneAndUpdateAsync<UserDocument>(x => x.PrimaryClaim == newPrimaryClaim, update, new FindOneAndUpdateOptions<UserDocument> { IsUpsert = true, ReturnDocument = ReturnDocument.After });
+			return await _users.FindOneAndUpdateAsync<UserDocument>(x => x.PrimaryClaim == newPrimaryClaim, update, new FindOneAndUpdateOptions<UserDocument> { IsUpsert = true, ReturnDocument = ReturnDocument.After }, cancellationToken);
 		}
 
 		/// <inheritdoc/>
-		public async Task<IUserClaims> GetClaimsAsync(UserId userId)
+		public async Task<IUserClaims> GetClaimsAsync(UserId userId, CancellationToken cancellationToken)
 		{
-			return await _users.Find(x => x.Id == userId).FirstOrDefaultAsync() ?? new UserDocument { Id = userId };
+			return await _users.Find(x => x.Id == userId).FirstOrDefaultAsync(cancellationToken) ?? new UserDocument { Id = userId };
 		}
 
 		/// <inheritdoc/>
-		public async Task UpdateClaimsAsync(UserId userId, IEnumerable<IUserClaim> claims)
+		public async Task UpdateClaimsAsync(UserId userId, IEnumerable<IUserClaim> claims, CancellationToken cancellationToken)
 		{
 			List<ClaimDocument> newClaims = claims.Select(x => new ClaimDocument(x)).ToList();
-			await _users.FindOneAndUpdateAsync(x => x.Id == userId, Builders<UserDocument>.Update.Set(x => x.Claims, newClaims));
+			await _users.FindOneAndUpdateAsync(x => x.Id == userId, Builders<UserDocument>.Update.Set(x => x.Claims, newClaims), cancellationToken: cancellationToken);
 		}
 
 		/// <inheritdoc/>
-		public async Task<IUserSettings> GetSettingsAsync(UserId userId)
+		public async Task<IUserSettings> GetSettingsAsync(UserId userId, CancellationToken cancellationToken)
 		{
-			return await _users.Find(x => x.Id == userId).FirstOrDefaultAsync() ?? new UserDocument { Id = userId };
+			return await _users.Find(x => x.Id == userId).FirstOrDefaultAsync(cancellationToken) ?? new UserDocument { Id = userId };
 		}
 
 		/// <inheritdoc/>
-		public async Task UpdateSettingsAsync(UserId userId, bool? enableExperimentalFeatures, bool? alwaysTagPreflightCL = null, BsonValue ? dashboardSettings = null, IEnumerable<JobId>? addPinnedJobIds = null, IEnumerable<JobId>? removePinnedJobIds = null, UpdateUserJobTemplateOptions? templateOptions = null, IEnumerable<BisectTaskId>? addBisectTaskIds = null, IEnumerable<BisectTaskId>? removeBisectTaskIds = null)
+		public async Task UpdateSettingsAsync(UserId userId, bool? enableExperimentalFeatures, bool? alwaysTagPreflightCL = null, BsonValue ? dashboardSettings = null, IEnumerable<JobId>? addPinnedJobIds = null, IEnumerable<JobId>? removePinnedJobIds = null, UpdateUserJobTemplateOptions? templateOptions = null, IEnumerable<BisectTaskId>? addBisectTaskIds = null, IEnumerable<BisectTaskId>? removeBisectTaskIds = null, CancellationToken cancellationToken = default)
 		{
 			if (addPinnedJobIds != null)
 			{
@@ -172,7 +174,7 @@ namespace Horde.Server.Users
 				{
 					FilterDefinition<UserDocument> filter = Builders<UserDocument>.Filter.Eq(x => x.Id, userId) & Builders<UserDocument>.Filter.AnyNin<JobId>(x => x.PinnedJobIds, new[] { pinnedJobId });
 					UpdateDefinition<UserDocument> update = Builders<UserDocument>.Update.PushEach(x => x.PinnedJobIds, new[] { pinnedJobId }, 50);
-					await _users.UpdateOneAsync(filter, update);
+					await _users.UpdateOneAsync(filter, update, null, cancellationToken);
 				}
 			}
 
@@ -195,7 +197,7 @@ namespace Horde.Server.Users
 			}
 			if (updates.Count > 0)
 			{
-				await _users.UpdateOneAsync<UserDocument>(x => x.Id == userId, Builders<UserDocument>.Update.Combine(updates));
+				await _users.UpdateOneAsync<UserDocument>(x => x.Id == userId, Builders<UserDocument>.Update.Combine(updates), (UpdateOptions?)null, cancellationToken);
 			}
 		}
 
@@ -203,11 +205,11 @@ namespace Horde.Server.Users
 		/// Enumerate all the documents in this collection
 		/// </summary>
 		/// <returns></returns>
-		public async IAsyncEnumerable<(IUser, IUserClaims, IUserSettings)> EnumerateDocumentsAsync()
+		public async IAsyncEnumerable<(IUser, IUserClaims, IUserSettings)> EnumerateDocumentsAsync([EnumeratorCancellation] CancellationToken cancellationToken)
 		{
-			using (IAsyncCursor<UserDocument> cursor = await _users.Find(FilterDefinition<UserDocument>.Empty).ToCursorAsync())
+			using (IAsyncCursor<UserDocument> cursor = await _users.Find(FilterDefinition<UserDocument>.Empty).ToCursorAsync(cancellationToken))
 			{
-				while (await cursor.MoveNextAsync())
+				while (await cursor.MoveNextAsync(cancellationToken))
 				{
 					foreach (UserDocument document in cursor.Current)
 					{
diff --git a/Engine/Source/Programs/Horde/Horde.Server/Users/UserCollectionV2.cs b/Engine/Source/Programs/Horde/Horde.Server/Users/UserCollectionV2.cs
--- a/Engine/Source/Programs/Horde/Horde.Server/Users/UserCollectionV2.cs
+++ b/Engine/Source/Programs/Horde/Horde.Server/Users/UserCollectionV2.cs
@@ -18,6 +18,7 @@ using EpicGames.Horde.Users;
 using MongoDB.Driver.Linq;
 using EpicGames.Horde.Jobs;
 using EpicGames.Horde.Jobs.Bisect;
+using System.Threading;
 
 namespace Horde.Server.Users
 {
@@ -235,7 +236,7 @@ namespace Horde.Server.Users
 		}
 
 		/// <inheritdoc/>
-		public async Task<IUser?> GetUserAsync(UserId id)
+		public async Task<IUser?> GetUserAsync(UserId id, CancellationToken cancellationToken)
 		{
 			if (id == UserId.Anonymous)
 			{
@@ -243,7 +244,7 @@ namespace Horde.Server.Users
 				return new UserDocument(id, "Anonymous", "anonymous", null);
 			}
 
-			IUser? user = await _users.Find(x => x.Id == id).FirstOrDefaultAsync();
+			IUser? user = await _users.Find(x => x.Id == id).FirstOrDefaultAsync(cancellationToken);
 			using (ICacheEntry entry = _userCache.CreateEntry(id))
 			{
 				entry.SetValue(user);
@@ -253,7 +254,7 @@ namespace Horde.Server.Users
 		}
 
 		/// <inheritdoc/>
-		public async ValueTask<IUser?> GetCachedUserAsync(UserId? id)
+		public async ValueTask<IUser?> GetCachedUserAsync(UserId? id, CancellationToken cancellationToken)
 		{
 			IUser? user;
 			if(id == null)
@@ -266,12 +267,12 @@ namespace Horde.Server.Users
 			}
 			else
 			{
-				return await GetUserAsync(id.Value);
+				return await GetUserAsync(id.Value, cancellationToken);
 			}
 		}
 
 		/// <inheritdoc/>
-		public async Task<List<IUser>> FindUsersAsync(IEnumerable<UserId>? ids, string? nameRegex = null, int? index = null, int? count = null)
+		public async Task<IReadOnlyList<IUser>> FindUsersAsync(IEnumerable<UserId>? ids, string? nameRegex = null, int? index = null, int? count = null, CancellationToken cancellationToken = default)
 		{
 			FilterDefinition<UserDocument> filter = FilterDefinition<UserDocument>.Empty;
 			if (ids != null)
@@ -287,27 +288,27 @@ namespace Horde.Server.Users
 
 			filter &= Builders<UserDocument>.Filter.Ne(x => x.Hidden, true);
 
-			return await _users.Find(filter).Range(index, count ?? 100).ToListAsync<UserDocument, IUser>();
+			return await _users.Find(filter).Range(index, count ?? 100).ToListAsync(cancellationToken);
 		}
 
 		/// <inheritdoc/>
-		public async Task<IUser?> FindUserByLoginAsync(string login)
+		public async Task<IUser?> FindUserByLoginAsync(string login, CancellationToken cancellationToken)
 		{
 			string loginUpper = login.ToUpperInvariant();
 			FilterDefinition<UserDocument> filter = Builders<UserDocument>.Filter.Eq(x => x.LoginUpper, loginUpper) & Builders<UserDocument>.Filter.Ne(x => x.Hidden, true);
-			return await _users.Find(filter).FirstOrDefaultAsync();
+			return await _users.Find(filter).FirstOrDefaultAsync(cancellationToken);
 		}
 
 		/// <inheritdoc/>
-		public async Task<IUser?> FindUserByEmailAsync(string email)
+		public async Task<IUser?> FindUserByEmailAsync(string email, CancellationToken cancellationToken)
 		{
 			string emailUpper = email.ToUpperInvariant();
 			FilterDefinition<UserDocument> filter = Builders<UserDocument>.Filter.Eq(x => x.EmailUpper, emailUpper) & Builders<UserDocument>.Filter.Ne(x => x.Hidden, true);
-			return await _users.Find(filter).FirstOrDefaultAsync();
+			return await _users.Find(filter).FirstOrDefaultAsync(cancellationToken);
 		}
 
 		/// <inheritdoc/>
-		public async Task<IUser> FindOrAddUserByLoginAsync(string login, string? name, string? email)
+		public async Task<IUser> FindOrAddUserByLoginAsync(string login, string? name, string? email, CancellationToken cancellationToken)
 		{
 			UserId newUserId = new UserId(BinaryIdUtils.CreateNew());
 			UpdateDefinition<UserDocument> update = Builders<UserDocument>.Update.SetOnInsert(x => x.Id, newUserId).SetOnInsert(x => x.Login, login).Unset(x => x.Hidden);
@@ -328,7 +329,7 @@ namespace Horde.Server.Users
 
 			string loginUpper = login.ToUpperInvariant();
 
-			IUser user = await _users.FindOneAndUpdateAsync<UserDocument>(x => x.LoginUpper == loginUpper, update, new FindOneAndUpdateOptions<UserDocument, UserDocument> { IsUpsert = true, ReturnDocument = ReturnDocument.After });
+			IUser user = await _users.FindOneAndUpdateAsync<UserDocument>(x => x.LoginUpper == loginUpper, update, new FindOneAndUpdateOptions<UserDocument, UserDocument> { IsUpsert = true, ReturnDocument = ReturnDocument.After }, cancellationToken);
 			if (user.Id == newUserId)
 			{
 				_logger.LogInformation("Added new user {Name} ({UserId}, {Login}, {Email})", user.Name, user.Id, user.Login, user.Email);
@@ -338,31 +339,31 @@ namespace Horde.Server.Users
 		}
 
 		/// <inheritdoc/>
-		public async Task<IUserClaims> GetClaimsAsync(UserId userId)
+		public async Task<IUserClaims> GetClaimsAsync(UserId userId, CancellationToken cancellationToken)
 		{
-			IUserClaims? claims = await _userClaims.Find(x => x.Id == userId).FirstOrDefaultAsync();
+			IUserClaims? claims = await _userClaims.Find(x => x.Id == userId).FirstOrDefaultAsync(cancellationToken);
 			claims ??= new UserClaimsDocument(userId);
 			return claims;
 		}
 
 		/// <inheritdoc/>
-		public async Task UpdateClaimsAsync(UserId userId, IEnumerable<IUserClaim> claims)
+		public async Task UpdateClaimsAsync(UserId userId, IEnumerable<IUserClaim> claims, CancellationToken cancellationToken)
 		{
 			UserClaimsDocument newDocument = new UserClaimsDocument(userId);
 			newDocument.Claims.AddRange(claims.Select(x => new UserClaim(x)));
-			await _userClaims.ReplaceOneAsync(x => x.Id == userId, newDocument, new ReplaceOptions { IsUpsert = true });
+			await _userClaims.ReplaceOneAsync(x => x.Id == userId, newDocument, new ReplaceOptions { IsUpsert = true }, cancellationToken);
 		}
 
 		/// <inheritdoc/>
-		public async Task<IUserSettings> GetSettingsAsync(UserId userId)
+		public async Task<IUserSettings> GetSettingsAsync(UserId userId, CancellationToken cancellationToken)
 		{
-			IUserSettings? settings = await _userSettings.Find(x => x.Id == userId).FirstOrDefaultAsync();
+			IUserSettings? settings = await _userSettings.Find(x => x.Id == userId).FirstOrDefaultAsync(cancellationToken);
 			settings ??= new UserSettingsDocument(userId);
 			return settings;
 		}
 
 		/// <inheritdoc/>
-		public async Task UpdateSettingsAsync(UserId userId, bool? enableExperimentalFeatures = null, bool? alwaysTagPreflightCL = null, BsonValue ? dashboardSettings = null, IEnumerable<JobId>? addPinnedJobIds = null, IEnumerable<JobId>? removePinnedJobIds = null, UpdateUserJobTemplateOptions? templateOptions = null, IEnumerable<BisectTaskId>? addBisectTaskIds = null, IEnumerable<BisectTaskId>? removeBisectTaskIds = null)
+		public async Task UpdateSettingsAsync(UserId userId, bool? enableExperimentalFeatures = null, bool? alwaysTagPreflightCL = null, BsonValue ? dashboardSettings = null, IEnumerable<JobId>? addPinnedJobIds = null, IEnumerable<JobId>? removePinnedJobIds = null, UpdateUserJobTemplateOptions? templateOptions = null, IEnumerable<BisectTaskId>? addBisectTaskIds = null, IEnumerable<BisectTaskId>? removeBisectTaskIds = null, CancellationToken cancellationToken = default)
 		{
 			List<UpdateDefinition<UserSettingsDocument>> updates = new List<UpdateDefinition<UserSettingsDocument>>();
 			if (enableExperimentalFeatures != null)
@@ -399,7 +400,7 @@ namespace Horde.Server.Users
 			{
 				JobTemplateSettingsDocument doc = new JobTemplateSettingsDocument(templateOptions.StreamId, templateOptions.TemplateId, templateOptions.TemplateHash, templateOptions.Arguments.ToList());
 				FilterDefinition<UserSettingsDocument> filter = Builders<UserSettingsDocument>.Filter.Eq(x => x.Id, userId) & Builders<UserSettingsDocument>.Filter.ElemMatch(x => x.JobTemplateSettings, t => t.StreamId == templateOptions.StreamId && t.TemplateId == templateOptions.TemplateId);
-				UpdateResult result = await _userSettings.UpdateOneAsync(filter, Builders<UserSettingsDocument>.Update.Set(x => x.JobTemplateSettings[-1], doc));
+				UpdateResult result = await _userSettings.UpdateOneAsync(filter, Builders<UserSettingsDocument>.Update.Set(x => x.JobTemplateSettings[-1], doc), null, cancellationToken);
 				if (result.ModifiedCount == 0)
 				{
 					updates.Add(Builders<UserSettingsDocument>.Update.PushEach(x => x.JobTemplateSettings, new[] { doc }, -100));
@@ -408,7 +409,7 @@ namespace Horde.Server.Users
 
 			if (updates.Count > 0)
 			{
-				await _userSettings.UpdateOneAsync<UserSettingsDocument>(x => x.Id == userId, Builders<UserSettingsDocument>.Update.Combine(updates), new UpdateOptions { IsUpsert = true });
+				await _userSettings.UpdateOneAsync<UserSettingsDocument>(x => x.Id == userId, Builders<UserSettingsDocument>.Update.Combine(updates), new UpdateOptions { IsUpsert = true }, cancellationToken);
 			}
 		}
 
@@ -416,16 +417,17 @@ namespace Horde.Server.Users
 		/// Upgrade from V1 collection
 		/// </summary>
 		/// <param name="userCollectionV1"></param>
+		/// <param name="cancellationToken">Cancellation token for the operation</param>
 		/// <returns></returns>
-		public async Task ResaveDocumentsAsync(UserCollectionV1 userCollectionV1)
+		public async Task ResaveDocumentsAsync(UserCollectionV1 userCollectionV1, CancellationToken cancellationToken)
 		{
-			await foreach ((IUser user, IUserClaims claims, IUserSettings settings) in userCollectionV1.EnumerateDocumentsAsync())
+			await foreach ((IUser user, IUserClaims claims, IUserSettings settings) in userCollectionV1.EnumerateDocumentsAsync(cancellationToken))
 			{
 				try
 				{
-					await _users.ReplaceOneAsync(x => x.Id == user.Id, new UserDocument(user), new ReplaceOptions { IsUpsert = true });
-					await _userClaims.ReplaceOneAsync(x => x.Id == user.Id, new UserClaimsDocument(claims), new ReplaceOptions { IsUpsert = true });
-					await _userSettings.ReplaceOneAsync(x => x.Id == user.Id, new UserSettingsDocument(settings), new ReplaceOptions { IsUpsert = true });
+					await _users.ReplaceOneAsync(x => x.Id == user.Id, new UserDocument(user), new ReplaceOptions { IsUpsert = true }, cancellationToken);
+					await _userClaims.ReplaceOneAsync(x => x.Id == user.Id, new UserClaimsDocument(claims), new ReplaceOptions { IsUpsert = true }, cancellationToken);
+					await _userSettings.ReplaceOneAsync(x => x.Id == user.Id, new UserSettingsDocument(settings), new ReplaceOptions { IsUpsert = true }, cancellationToken);
 					_logger.LogDebug("Updated user {UserId}", user.Id);
 				}
 				catch (MongoWriteException ex)
@@ -435,7 +437,7 @@ namespace Horde.Server.Users
 
 				if(settings.PinnedJobIds.Count > 0)
 				{
-					await UpdateSettingsAsync(user.Id, addPinnedJobIds: settings.PinnedJobIds);
+					await UpdateSettingsAsync(user.Id, addPinnedJobIds: settings.PinnedJobIds, cancellationToken: cancellationToken);
 				}
 			}
 		}
diff --git a/Engine/Source/Programs/Horde/Horde.Server/Users/UserController.cs b/Engine/Source/Programs/Horde/Horde.Server/Users/UserController.cs
--- a/Engine/Source/Programs/Horde/Horde.Server/Users/UserController.cs
+++ b/Engine/Source/Programs/Horde/Horde.Server/Users/UserController.cs
@@ -16,6 +16,7 @@ using Microsoft.Extensions.Options;
 using EpicGames.Horde.Users;
 using EpicGames.Horde.Jobs;
 using EpicGames.Horde.Jobs.Bisect;
+using System.Threading;
 
 namespace Horde.Server.Users
 {
@@ -48,17 +49,17 @@ namespace Horde.Server.Users
 		[HttpGet]
 		[Route("/api/v1/user")]
 		[ProducesResponseType(typeof(GetUserResponse), 200)]
-		public async Task<ActionResult<object>> GetUserAsync([FromQuery] PropertyFilter? filter = null)
+		public async Task<ActionResult<object>> GetUserAsync([FromQuery] PropertyFilter? filter = null, CancellationToken cancellationToken = default)
 		{
-			IUser? internalUser = await _userCollection.GetUserAsync(User);
+			IUser? internalUser = await _userCollection.GetUserAsync(User, cancellationToken);
 			if (internalUser == null)
 			{
 				return NotFound();
 			}
 
-			IAvatar? avatar = (_avatarService == null)? (IAvatar?)null : await _avatarService.GetAvatarAsync(internalUser);
-			IUserClaims claims = await _userCollection.GetClaimsAsync(internalUser.Id);
-			IUserSettings settings = await _userCollection.GetSettingsAsync(internalUser.Id);
+			IAvatar? avatar = (_avatarService == null)? (IAvatar?)null : await _avatarService.GetAvatarAsync(internalUser, cancellationToken);
+			IUserClaims claims = await _userCollection.GetClaimsAsync(internalUser.Id, cancellationToken);
+			IUserSettings settings = await _userCollection.GetSettingsAsync(internalUser.Id, cancellationToken);
 
 			GetUserResponse response = internalUser.ToApiResponse(avatar, claims, settings);
 			response.DashboardFeatures = GetDashboardFeatures(_globalConfig.Value, User);
diff --git a/Engine/Source/Programs/Horde/Horde.Server/Users/UsersController.cs b/Engine/Source/Programs/Horde/Horde.Server/Users/UsersController.cs
--- a/Engine/Source/Programs/Horde/Horde.Server/Users/UsersController.cs
+++ b/Engine/Source/Programs/Horde/Horde.Server/Users/UsersController.cs
@@ -8,6 +8,7 @@ using Horde.Server.Utilities;
 using Microsoft.AspNetCore.Authorization;
 using Microsoft.AspNetCore.Mvc;
 using EpicGames.Horde.Users;
+using System.Threading;
 
 namespace Horde.Server.Users
 {
@@ -47,7 +48,7 @@ namespace Horde.Server.Users
 		[HttpGet]
 		[Route("/api/v1/users/{id}")]
 		[ProducesResponseType(typeof(List<GetUserResponse>), 200)]
-		public async Task<ActionResult<object>> GetUserAsync(string id, [FromQuery] PropertyFilter? filter = null)
+		public async Task<ActionResult<object>> GetUserAsync(string id, [FromQuery] PropertyFilter? filter = null, CancellationToken cancellationToken = default)
 		{
 			UserId? userId = ParseUserId(id);
 			if(userId == null)
@@ -55,15 +56,15 @@ namespace Horde.Server.Users
 				return BadRequest("Invalid user id '{Id}'", id);
 			}
 
-			IUser? user = await UserCollection.GetUserAsync(userId.Value);
+			IUser? user = await UserCollection.GetUserAsync(userId.Value, cancellationToken);
 			if (user == null)
 			{
 				return NotFound(userId.Value);
 			}
 
-			IAvatar? avatar = (AvatarService == null) ? (IAvatar?)null : await AvatarService.GetAvatarAsync(user);
-			IUserClaims? claims = await UserCollection.GetClaimsAsync(user.Id);
-			IUserSettings? settings = await UserCollection.GetSettingsAsync(user.Id);
+			IAvatar? avatar = (AvatarService == null) ? (IAvatar?)null : await AvatarService.GetAvatarAsync(user, cancellationToken);
+			IUserClaims? claims = await UserCollection.GetClaimsAsync(user.Id, cancellationToken);
+			IUserSettings? settings = await UserCollection.GetSettingsAsync(user.Id, cancellationToken);
 			return PropertyFilter.Apply(user.ToApiResponse(avatar, claims, settings), filter);
 		}
 
@@ -80,7 +81,8 @@ namespace Horde.Server.Users
 			[FromQuery] int index = 0,
 			[FromQuery] int count = 100,
 			[FromQuery] bool includeClaims = false,
-			[FromQuery] bool includeAvatar = false)
+			[FromQuery] bool includeAvatar = false,
+			CancellationToken cancellationToken = default)
 		{
 
 			UserId[]? userIds = null;
@@ -89,13 +91,13 @@ namespace Horde.Server.Users
 				userIds = ids.Select(x => UserId.Parse(x)).ToArray();
 			}
 
-			List<IUser> users = await UserCollection.FindUsersAsync(userIds, nameRegex, index, count);
+			IReadOnlyList<IUser> users = await UserCollection.FindUsersAsync(userIds, nameRegex, index, count, cancellationToken);
 
 			List<GetUserResponse> response = new List<GetUserResponse>();
 			foreach (IUser user in users)
 			{
-				IAvatar? avatar = (AvatarService == null || !includeAvatar) ? (IAvatar?)null : await AvatarService.GetAvatarAsync(user);
-				IUserClaims? claims = (!includeClaims) ? null : await UserCollection.GetClaimsAsync(user.Id);
+				IAvatar? avatar = (AvatarService == null || !includeAvatar) ? (IAvatar?)null : await AvatarService.GetAvatarAsync(user, cancellationToken);
+				IUserClaims? claims = (!includeClaims) ? null : await UserCollection.GetClaimsAsync(user.Id, cancellationToken);
 				response.Add(user.ToApiResponse(avatar, claims, null));
 			}
 
diff --git a/Engine/Source/Programs/Horde/Horde.Server/Utilities/AsyncCachedValue.cs b/Engine/Source/Programs/Horde/Horde.Server/Utilities/AsyncCachedValue.cs
--- a/Engine/Source/Programs/Horde/Horde.Server/Utilities/AsyncCachedValue.cs
+++ b/Engine/Source/Programs/Horde/Horde.Server/Utilities/AsyncCachedValue.cs
@@ -4,6 +4,7 @@ using System;
 using System.Diagnostics;
 using System.Threading;
 using System.Threading.Tasks;
+using EpicGames.Core;
 
 namespace Horde.Server.Utilities
 {
@@ -11,7 +12,7 @@ namespace Horde.Server.Utilities
 	/// Caches a value and asynchronously updates it after a period of time
 	/// </summary>
 	/// <typeparam name="T"></typeparam>
-	public class AsyncCachedValue<T>
+	public sealed class AsyncCachedValue<T> : IAsyncDisposable
 	{
 		class State
 		{
@@ -28,6 +29,8 @@ namespace Horde.Server.Utilities
 			}
 		}
 
+		CancellationTokenSource _cancellationTokenSource = new CancellationTokenSource();
+
 		/// <summary>
 		/// The current state
 		/// </summary>
@@ -36,7 +39,7 @@ namespace Horde.Server.Utilities
 		/// <summary>
 		/// Generator for the new value
 		/// </summary>
-		readonly Func<Task<T>> _generator;
+		readonly Func<CancellationToken, Task<T>> _generator;
 
 		/// <summary>
 		/// Time at which to start to refresh the value
@@ -51,7 +54,7 @@ namespace Horde.Server.Utilities
 		/// <summary>
 		/// Default constructor
 		/// </summary>
-		public AsyncCachedValue(Func<Task<T>> generator, TimeSpan refreshTime)
+		public AsyncCachedValue(Func<CancellationToken, Task<T>> generator, TimeSpan refreshTime)
 			: this(generator, refreshTime * 0.75, refreshTime)
 		{
 		}
@@ -59,13 +62,29 @@ namespace Horde.Server.Utilities
 		/// <summary>
 		/// Default constructor
 		/// </summary>
-		public AsyncCachedValue(Func<Task<T>> generator, TimeSpan minRefreshTime, TimeSpan maxRefreshTime)
+		public AsyncCachedValue(Func<CancellationToken, Task<T>> generator, TimeSpan minRefreshTime, TimeSpan maxRefreshTime)
 		{
 			_generator = generator;
 			_minRefreshTime = minRefreshTime;
 			_maxRefreshTime = maxRefreshTime;
 		}
 
+		/// <inheritdoc/>
+		public async ValueTask DisposeAsync()
+		{
+			if (_current != null)
+			{
+				await _cancellationTokenSource.CancelAsync();
+				await _current.IgnoreCanceledExceptionsAsync();
+			}
+
+			if (_cancellationTokenSource != null)
+			{
+				_cancellationTokenSource.Dispose();
+				_cancellationTokenSource = null!;
+			}
+		}
+
 		/// <summary>
 		/// Invalidates the current value
 		/// </summary>
@@ -148,7 +167,7 @@ namespace Horde.Server.Utilities
 
 		async Task<State> CreateStateAsync()
 		{
-			return new State(await _generator());
+			return new State(await _generator(_cancellationTokenSource.Token));
 		}
 	}
 }
diff --git a/Engine/Source/Programs/Horde/Horde.Server/Utilities/MongoExtensions.cs b/Engine/Source/Programs/Horde/Horde.Server/Utilities/MongoExtensions.cs
--- a/Engine/Source/Programs/Horde/Horde.Server/Utilities/MongoExtensions.cs
+++ b/Engine/Source/Programs/Horde/Horde.Server/Utilities/MongoExtensions.cs
@@ -125,27 +125,6 @@ namespace Horde.Server.Utilities
 			return query;
 		}
 
-		/// <summary>
-		/// Filters the documents returned from a search
-		/// </summary>
-		/// <param name="query">The query to filter</param>
-		/// <returns>New query</returns>
-		public static async Task<List<TResult>> ToListAsync<TDocument, TResult>(this IAsyncCursorSource<TDocument> query) where TDocument : TResult
-		{
-			List<TResult> results = new List<TResult>();
-			using (IAsyncCursor<TDocument> cursor = await query.ToCursorAsync())
-			{
-				while (await cursor.MoveNextAsync())
-				{
-					foreach(TDocument document in cursor.Current)
-					{
-						results.Add(document);
-					}
-				}
-			}
-			return results;
-		}
-
 		/// <summary>
 		/// Attempts to insert a document into a collection, handling the error case that a document with the given key already exists
 		/// </summary>
diff --git a/Engine/Source/Programs/Horde/Horde.Server/Utilities/MongoQueryCache.cs b/Engine/Source/Programs/Horde/Horde.Server/Utilities/MongoQueryCache.cs
--- a/Engine/Source/Programs/Horde/Horde.Server/Utilities/MongoQueryCache.cs
+++ b/Engine/Source/Programs/Horde/Horde.Server/Utilities/MongoQueryCache.cs
@@ -4,6 +4,7 @@ using System;
 using System.Collections.Generic;
 using System.Diagnostics;
 using System.Linq;
+using System.Threading;
 using System.Threading.Tasks;
 using Microsoft.Extensions.Caching.Memory;
 using MongoDB.Bson;
@@ -44,9 +45,9 @@ namespace Horde.Server.Utilities
 			_cache.Dispose();
 		}
 
-		async Task RefreshAsync(QueryState state, FilterDefinition<TDocument> filter)
+		async Task RefreshAsync(QueryState state, FilterDefinition<TDocument> filter, CancellationToken cancellationToken = default)
 		{
-			state._results = await _collection.Find(filter).ToListAsync();
+			state._results = await _collection.Find(filter).ToListAsync(cancellationToken);
 			state._timer.Restart();
 		}
 
@@ -76,7 +77,7 @@ namespace Horde.Server.Utilities
 
 			if (state._queryTask == null && (state._results == null || state._timer.Elapsed > _maxLatency))
 			{
-				state._queryTask = Task.Run(() => RefreshAsync(state, filter));
+				state._queryTask = Task.Run(() => RefreshAsync(state, filter, CancellationToken.None), CancellationToken.None);
 			}
 			if (state._queryTask != null && (state._results == null || state._timer.Elapsed > _maxLatency))
 			{
diff --git a/Engine/Source/Programs/Horde/Horde.Server/Utilities/SingletonDocument.cs b/Engine/Source/Programs/Horde/Horde.Server/Utilities/SingletonDocument.cs
--- a/Engine/Source/Programs/Horde/Horde.Server/Utilities/SingletonDocument.cs
+++ b/Engine/Source/Programs/Horde/Horde.Server/Utilities/SingletonDocument.cs
@@ -5,6 +5,7 @@ using System.Threading.Tasks;
 using Horde.Server.Server;
 using MongoDB.Bson.Serialization.Attributes;
 using EpicGames.Horde;
+using System.Threading;
 
 namespace Horde.Server.Utilities
 {
@@ -97,15 +98,17 @@ namespace Horde.Server.Utilities
 		/// <summary>
 		/// Gets the current document
 		/// </summary>
+		/// <param name="cancellationToken">Cancellation token for the operation</param>
 		/// <returns>The current document</returns>
-		Task<T> GetAsync();
+		Task<T> GetAsync(CancellationToken cancellationToken);
 
 		/// <summary>
 		/// Attempts to update the document
 		/// </summary>
 		/// <param name="value">New state of the document</param>
+		/// <param name="cancellationToken">Cancellation token for the operation</param>
 		/// <returns>True if the document was updated, false otherwise</returns>
-		Task<bool> TryUpdateAsync(T value);
+		Task<bool> TryUpdateAsync(T value, CancellationToken cancellationToken);
 	}
 
 	/// <summary>
@@ -129,15 +132,15 @@ namespace Horde.Server.Utilities
 		}
 
 		/// <inheritdoc/>
-		public async Task<T> GetAsync()
+		public async Task<T> GetAsync(CancellationToken cancellationToken)
 		{
-			return await _mongoService.GetSingletonAsync<T>();
+			return await _mongoService.GetSingletonAsync<T>(cancellationToken);
 		}
 
 		/// <inheritdoc/>
-		public Task<bool> TryUpdateAsync(T value)
+		public Task<bool> TryUpdateAsync(T value, CancellationToken cancellationToken)
 		{
-			return _mongoService.TryUpdateSingletonAsync<T>(value);
+			return _mongoService.TryUpdateSingletonAsync<T>(value, cancellationToken);
 		}
 	}
 
@@ -152,15 +155,16 @@ namespace Horde.Server.Utilities
 		/// <typeparam name="T"></typeparam>
 		/// <param name="singleton"></param>
 		/// <param name="updateAction"></param>
+		/// <param name="cancellationToken">Cancellation token for the operation</param>
 		/// <returns></returns>
-		public static async Task<T> UpdateAsync<T>(this ISingletonDocument<T> singleton, Action<T> updateAction) where T : SingletonBase, new()
+		public static async Task<T> UpdateAsync<T>(this ISingletonDocument<T> singleton, Action<T> updateAction, CancellationToken cancellationToken) where T : SingletonBase, new()
 		{
 			for (; ; )
 			{
-				T value = await singleton.GetAsync();
+				T value = await singleton.GetAsync(cancellationToken);
 				updateAction(value);
 
-				if (await singleton.TryUpdateAsync(value))
+				if (await singleton.TryUpdateAsync(value, cancellationToken))
 				{
 					return value;
 				}
diff --git a/Engine/Source/Programs/Horde/Horde.Server/Utilities/StreamExtensions.cs b/Engine/Source/Programs/Horde/Horde.Server/Utilities/StreamExtensions.cs
--- a/Engine/Source/Programs/Horde/Horde.Server/Utilities/StreamExtensions.cs
+++ b/Engine/Source/Programs/Horde/Horde.Server/Utilities/StreamExtensions.cs
@@ -2,6 +2,7 @@
 
 using System;
 using System.IO;
+using System.Threading;
 using System.Threading.Tasks;
 
 namespace Horde.Server.Utilities
@@ -18,10 +19,11 @@ namespace Horde.Server.Utilities
 		/// <param name="data">Buffer to receive the read data</param>
 		/// <param name="offset">Offset within the buffer to read the new data</param>
 		/// <param name="length">Length of the data to read</param>
+		/// <param name="cancellationToken">Cancellation token for the operation</param>
 		/// <returns>Async task</returns>
-		public static Task ReadFixedSizeDataAsync(this Stream stream, byte[] data, int offset, int length)
+		public static Task ReadFixedSizeDataAsync(this Stream stream, byte[] data, int offset, int length, CancellationToken cancellationToken = default)
 		{
-			return ReadFixedSizeDataAsync(stream, data.AsMemory(offset, length));
+			return ReadFixedSizeDataAsync(stream, data.AsMemory(offset, length), cancellationToken);
 		}
 
 		/// <summary>
@@ -29,12 +31,13 @@ namespace Horde.Server.Utilities
 		/// </summary>
 		/// <param name="stream">The stream to read from</param>
 		/// <param name="data">Buffer to receive the read data</param>
+		/// <param name="cancellationToken">Cancellation token for the operation</param>
 		/// <returns>Async task</returns>
-		public static async Task ReadFixedSizeDataAsync(this Stream stream, Memory<byte> data)
+		public static async Task ReadFixedSizeDataAsync(this Stream stream, Memory<byte> data, CancellationToken cancellationToken = default)
 		{
 			while (data.Length > 0)
 			{
-				int count = await stream.ReadAsync(data);
+				int count = await stream.ReadAsync(data, cancellationToken);
 				if (count == 0)
 				{
 					throw new EndOfStreamException();
diff --git a/Engine/Source/Programs/Shared/EpicGames.Slack/SlackClient.cs b/Engine/Source/Programs/Shared/EpicGames.Slack/SlackClient.cs
--- a/Engine/Source/Programs/Shared/EpicGames.Slack/SlackClient.cs
+++ b/Engine/Source/Programs/Shared/EpicGames.Slack/SlackClient.cs
@@ -6,6 +6,7 @@ using System.Net.Http;
 using System.Text;
 using System.Text.Json;
 using System.Text.Json.Serialization;
+using System.Threading;
 using System.Threading.Tasks;
 using Microsoft.Extensions.Logging;
 
@@ -99,12 +100,12 @@ namespace EpicGames.Slack
 
 		static bool ShouldLogError<TResponse>(TResponse response) where TResponse : SlackResponse => !response.Ok;
 
-		private Task<TResponse> SendRequestAsync<TResponse>(string requestUrl, object request) where TResponse : SlackResponse
+		private Task<TResponse> SendRequestAsync<TResponse>(string requestUrl, object request, CancellationToken cancellationToken) where TResponse : SlackResponse
 		{
-			return SendRequestAsync<TResponse>(requestUrl, request, ShouldLogError);
+			return SendRequestAsync<TResponse>(requestUrl, request, ShouldLogError, cancellationToken);
 		}
 
-		private async Task<TResponse> SendRequestAsync<TResponse>(string requestUrl, object request, Func<TResponse, bool> shouldLogError) where TResponse : SlackResponse
+		private async Task<TResponse> SendRequestAsync<TResponse>(string requestUrl, object request, Func<TResponse, bool> shouldLogError, CancellationToken cancellationToken) where TResponse : SlackResponse
 		{
 			using (HttpRequestMessage sendMessageRequest = new HttpRequestMessage(HttpMethod.Post, requestUrl))
 			{
@@ -112,20 +113,20 @@ namespace EpicGames.Slack
 				using (StringContent messageContent = new StringContent(requestJson, Encoding.UTF8, "application/json"))
 				{
 					sendMessageRequest.Content = messageContent;
-					return await SendRequestAsync<TResponse>(sendMessageRequest, requestJson, shouldLogError);
+					return await SendRequestAsync<TResponse>(sendMessageRequest, requestJson, shouldLogError, cancellationToken);
 				}
 			}
 		}
 
-		private Task<TResponse> SendRequestAsync<TResponse>(HttpRequestMessage request, string requestJson) where TResponse : SlackResponse
+		private Task<TResponse> SendRequestAsync<TResponse>(HttpRequestMessage request, string requestJson, CancellationToken cancellationToken) where TResponse : SlackResponse
 		{
-			return SendRequestAsync<TResponse>(request, requestJson, ShouldLogError);
+			return SendRequestAsync<TResponse>(request, requestJson, ShouldLogError, cancellationToken);
 		}
 
-		private async Task<TResponse> SendRequestAsync<TResponse>(HttpRequestMessage request, string requestJson, Func<TResponse, bool> shouldLogError) where TResponse : SlackResponse
+		private async Task<TResponse> SendRequestAsync<TResponse>(HttpRequestMessage request, string requestJson, Func<TResponse, bool> shouldLogError, CancellationToken cancellationToken) where TResponse : SlackResponse
 		{
-			HttpResponseMessage response = await _httpClient.SendAsync(request);
-			byte[] responseBytes = await response.Content.ReadAsByteArrayAsync();
+			HttpResponseMessage response = await _httpClient.SendAsync(request, cancellationToken);
+			byte[] responseBytes = await response.Content.ReadAsByteArrayAsync(cancellationToken);
 
 			TResponse responseObject = JsonSerializer.Deserialize<TResponse>(responseBytes)!;
 			if (shouldLogError(responseObject))
@@ -196,10 +197,23 @@ namespace EpicGames.Slack
 		/// </summary>
 		/// <param name="recipient">Recipient of the message. May be a channel or Slack user id.</param>
 		/// <param name="message">New message to post</param>
+ 		/// <param name="cancellationToken">Cancellation token for the operation</param>
 		/// <returns>Identifier for the message. Note that the returned channel value will respond to a concrete channel identifier if a user is specified as recipient, and must be used to update the message.</returns>
-		public async Task<SlackMessageId> PostMessageAsync(string recipient, SlackMessage message)
+		public async Task<SlackMessageId> PostMessageAsync(string recipient, SlackMessage message, CancellationToken cancellationToken = default)
 		{
-			return await PostOrUpdateMessageAsync(recipient, null, null, message, false);
+			return await PostOrUpdateMessageAsync(recipient, null, null, message, false, cancellationToken);
+		}
+
+		/// <summary>
+		/// Posts a message to a recipient
+		/// </summary>
+		/// <param name="threadId">Parent message to post to</param>
+		/// <param name="message">New message to post</param>
+		/// <param name="cancellationToken">Cancellation token for the operation</param>
+		/// <returns>Identifier for the message. Note that the returned channel value will respond to a concrete channel identifier if a user is specified as recipient, and must be used to update the message.</returns>
+		public Task<SlackMessageId> PostMessageToThreadAsync(SlackMessageId threadId, SlackMessage message, CancellationToken cancellationToken = default)
+		{
+			return PostMessageToThreadAsync(threadId, message, false, cancellationToken);
 		}
 
 		/// <summary>
@@ -208,10 +222,22 @@ namespace EpicGames.Slack
 		/// <param name="threadId">Parent message to post to</param>
 		/// <param name="message">New message to post</param>
 		/// <param name="replyBroadcast">Whether to broadcast the message to the channel</param>
+		/// <param name="cancellationToken">Cancellation token for the operation</param>
 		/// <returns>Identifier for the message. Note that the returned channel value will respond to a concrete channel identifier if a user is specified as recipient, and must be used to update the message.</returns>
-		public async Task<SlackMessageId> PostMessageToThreadAsync(SlackMessageId threadId, SlackMessage message, bool replyBroadcast = false)
+		public async Task<SlackMessageId> PostMessageToThreadAsync(SlackMessageId threadId, SlackMessage message, bool replyBroadcast, CancellationToken cancellationToken = default)
+		{
+			return await PostOrUpdateMessageAsync(threadId.Channel, null, threadId.Ts, message, replyBroadcast, cancellationToken);
+		}
+
+		/// <summary>
+		/// Updates an existing message
+		/// </summary>
+		/// <param name="id">The message id</param>
+		/// <param name="message">New message to post</param>
+		/// <param name="cancellationToken">Cancellation token for the operation</param>
+		public Task UpdateMessageAsync(SlackMessageId id, SlackMessage message, CancellationToken cancellationToken = default)
 		{
-			return await PostOrUpdateMessageAsync(threadId.Channel, null, threadId.Ts, message, replyBroadcast);
+			return UpdateMessageAsync(id, message, false, cancellationToken);
 		}
 
 		/// <summary>
@@ -220,12 +246,13 @@ namespace EpicGames.Slack
 		/// <param name="id">The message id</param>
 		/// <param name="message">New message to post</param>
 		/// <param name="replyBroadcast">For threaded messages, whether to broadcast the message to the channel</param>
-		public async Task UpdateMessageAsync(SlackMessageId id, SlackMessage message, bool replyBroadcast = false)
+		/// <param name="cancellationToken">Cancellation token for the operation</param>
+		public async Task UpdateMessageAsync(SlackMessageId id, SlackMessage message, bool replyBroadcast = false, CancellationToken cancellationToken = default)
 		{
-			await PostOrUpdateMessageAsync(id.Channel, id.Ts, id.ThreadTs, message, replyBroadcast);
+			await PostOrUpdateMessageAsync(id.Channel, id.Ts, id.ThreadTs, message, replyBroadcast, cancellationToken);
 		}
 
-		async Task<SlackMessageId> PostOrUpdateMessageAsync(string recipient, string? ts, string? threadTs, SlackMessage message, bool replyBroadcast)
+		async Task<SlackMessageId> PostOrUpdateMessageAsync(string recipient, string? ts, string? threadTs, SlackMessage message, bool replyBroadcast, CancellationToken cancellationToken)
 		{
 			PostMessageRequest request = new PostMessageRequest();
 			request.Channel = recipient;
@@ -252,7 +279,7 @@ namespace EpicGames.Slack
 				request.UnfurlMedia = false;
 			}
 
-			PostMessageResponse response = await SendRequestAsync<PostMessageResponse>(ts == null? PostMessageUrl : UpdateMessageUrl, request);
+			PostMessageResponse response = await SendRequestAsync<PostMessageResponse>(ts == null? PostMessageUrl : UpdateMessageUrl, request, cancellationToken);
 			if (!response.Ok || response.Ts == null)
 			{
 				throw new SlackException(response.Error ?? "unknown");
@@ -283,13 +310,14 @@ namespace EpicGames.Slack
 		/// Gets a permalink for a message
 		/// </summary>
 		/// <param name="id">Message identifier</param>
+		/// <param name="cancellationToken">Cancellation token for the operation</param>
 		/// <returns>Link to the message</returns>
-		public async Task<string> GetPermalinkAsync(SlackMessageId id)
+		public async Task<string> GetPermalinkAsync(SlackMessageId id, CancellationToken cancellationToken = default)
 		{
 			string requestUrl = $"{GetPermalinkUrl}?channel={id.Channel}&message_ts={id.Ts}";
 			using (HttpRequestMessage sendMessageRequest = new HttpRequestMessage(HttpMethod.Get, requestUrl))
 			{
-				GetPermalinkResponse response = await SendRequestAsync<GetPermalinkResponse>(sendMessageRequest, "");
+				GetPermalinkResponse response = await SendRequestAsync<GetPermalinkResponse>(sendMessageRequest, "", cancellationToken);
 				if (!response.Ok || response.Permalink == null)
 				{
 					throw new SlackException(response.Error ?? "unknown");
@@ -322,7 +350,8 @@ namespace EpicGames.Slack
 		/// </summary>
 		/// <param name="messageId">Message to react to</param>
 		/// <param name="name">Name of the reaction to post</param>
-		public async Task AddReactionAsync(SlackMessageId messageId, string name)
+		/// <param name="cancellationToken">Cancellation token for the operation</param>
+		public async Task AddReactionAsync(SlackMessageId messageId, string name, CancellationToken cancellationToken = default)
 		{
 			ReactionMessage message = new ReactionMessage();
 			message.Channel = messageId.Channel;
@@ -331,7 +360,7 @@ namespace EpicGames.Slack
 
 			static bool ShouldLogError(SlackResponse response) => !response.Ok && !String.Equals(response.Error, "already_reacted", StringComparison.Ordinal);
 
-			await SendRequestAsync<SlackResponse>(AddReactionUrl, message, ShouldLogError);
+			await SendRequestAsync<SlackResponse>(AddReactionUrl, message, ShouldLogError, cancellationToken);
 		}
 
 		/// <summary>
@@ -339,7 +368,8 @@ namespace EpicGames.Slack
 		/// </summary>
 		/// <param name="messageId">Message to react to</param>
 		/// <param name="name">Name of the reaction to post</param>
-		public async Task RemoveReactionAsync(SlackMessageId messageId, string name)
+		/// <param name="cancellationToken">Cancellation token for the operation</param>
+		public async Task RemoveReactionAsync(SlackMessageId messageId, string name, CancellationToken cancellationToken = default)
 		{
 			ReactionMessage message = new ReactionMessage();
 			message.Channel = messageId.Channel;
@@ -348,7 +378,7 @@ namespace EpicGames.Slack
 
 			static bool ShouldLogError(SlackResponse response) => !response.Ok && !String.Equals(response.Error, "no_reaction", StringComparison.Ordinal);
 
-			await SendRequestAsync<SlackResponse>(RemoveReactionUrl, message, ShouldLogError);
+			await SendRequestAsync<SlackResponse>(RemoveReactionUrl, message, ShouldLogError, cancellationToken);
 		}
 
 		#endregion
@@ -371,14 +401,16 @@ namespace EpicGames.Slack
 		/// </summary>
 		/// <param name="channel">Channel identifier to invite the user to</param>
 		/// <param name="userId">The user id</param>
-		public Task InviteUserAsync(string channel, string userId) => InviteUsersAsync(channel, new[] { userId });
+		/// <param name="cancellationToken">Cancellation token for the operation</param>
+		public Task InviteUserAsync(string channel, string userId, CancellationToken cancellationToken = default) => InviteUsersAsync(channel, new[] { userId }, cancellationToken);
 
 		/// <summary>
 		/// Invite a set of users to a channel
 		/// </summary>
 		/// <param name="channel">Channel identifier to invite the user to</param>
 		/// <param name="userIds">The user id</param>
-		public async Task InviteUsersAsync(string channel, IEnumerable<string> userIds)
+		/// <param name="cancellationToken">Cancellation token for the operation</param>
+		public async Task InviteUsersAsync(string channel, IEnumerable<string> userIds, CancellationToken cancellationToken = default)
 		{
 			InviteMessage message = new InviteMessage();
 			message.Channel = channel;
@@ -386,7 +418,7 @@ namespace EpicGames.Slack
 
 			static bool ShouldLogError(SlackResponse response) => !response.Ok && !String.Equals(response.Error, "already_in_channel", StringComparison.Ordinal);
 
-			await SendRequestAsync<SlackResponse>(ConversationsInviteUrl, message, ShouldLogError);
+			await SendRequestAsync<SlackResponse>(ConversationsInviteUrl, message, ShouldLogError, cancellationToken);
 		}
 
 		/// <summary>
@@ -394,7 +426,8 @@ namespace EpicGames.Slack
 		/// </summary>
 		/// <param name="channel">Channel identifier to invite the user to</param>
 		/// <param name="userIds">The user id</param>
-		public async Task<string?> TryInviteUsersAsync(string channel, IEnumerable<string> userIds)
+		/// <param name="cancellationToken">Cancellation token for the operation</param>
+		public async Task<string?> TryInviteUsersAsync(string channel, IEnumerable<string> userIds, CancellationToken cancellationToken = default)
 		{
 			InviteMessage message = new InviteMessage();
 			message.Channel = channel;
@@ -402,7 +435,7 @@ namespace EpicGames.Slack
 
 			static bool ShouldLogError(SlackResponse response) => false;
 
-			SlackResponse response = await SendRequestAsync<SlackResponse>(ConversationsInviteUrl, message, ShouldLogError);
+			SlackResponse response = await SendRequestAsync<SlackResponse>(ConversationsInviteUrl, message, ShouldLogError, cancellationToken);
 			return response.Ok ? null : response.Error;
 		}
 
@@ -413,13 +446,14 @@ namespace EpicGames.Slack
 		/// </summary>
 		/// <param name="channel">Channel identifier to invite the user to</param>
 		/// <param name="userIds">The user id</param>
-		public async Task AdminInviteUsersAsync(string channel, IEnumerable<string> userIds)
+		/// <param name="cancellationToken">Cancellation token for the operation</param>
+		public async Task AdminInviteUsersAsync(string channel, IEnumerable<string> userIds, CancellationToken cancellationToken = default)
 		{
 			object message = new { channel_id = channel, user_ids = String.Join(",", userIds) };
 
 			static bool ShouldLogError(SlackResponse response) => !response.Ok && !String.Equals(response.Error, "already_in_channel", StringComparison.Ordinal);
 
-			await SendRequestAsync<SlackResponse>(AdminConversationsInviteUrl, message, ShouldLogError);
+			await SendRequestAsync<SlackResponse>(AdminConversationsInviteUrl, message, ShouldLogError, cancellationToken);
 		}
 
 		#endregion
@@ -439,12 +473,13 @@ namespace EpicGames.Slack
 		/// Gets a user's profile
 		/// </summary>
 		/// <param name="userId">The user id</param>
+		/// <param name="cancellationToken">Cancellation token for the operation</param>
 		/// <returns>User profile</returns>
-		public async Task<SlackUser> GetUserAsync(string userId)
+		public async Task<SlackUser> GetUserAsync(string userId, CancellationToken cancellationToken = default)
 		{
 			using (HttpRequestMessage request = new HttpRequestMessage(HttpMethod.Post, $"{UsersInfoUrl}?user={userId}"))
 			{
-				UserResponse response = await SendRequestAsync<UserResponse>(request, "", ShouldLogError);
+				UserResponse response = await SendRequestAsync<UserResponse>(request, "", ShouldLogError, cancellationToken);
 				return response.User!;
 			}
 		}
@@ -453,14 +488,15 @@ namespace EpicGames.Slack
 		/// Finds a user by email address
 		/// </summary>
 		/// <param name="email">The user's email address</param>
+		/// <param name="cancellationToken">Cancellation token for the operation</param>
 		/// <returns>User profile</returns>
-		public async Task<SlackUser?> FindUserByEmailAsync(string email)
+		public async Task<SlackUser?> FindUserByEmailAsync(string email, CancellationToken cancellationToken = default)
 		{
 			using (HttpRequestMessage request = new HttpRequestMessage(HttpMethod.Post, $"{UsersLookupByEmailUrl}?email={email}"))
 			{
 				static bool ShouldLogError(UserResponse response) => !response.Ok && !String.Equals(response.Error, "users_not_found", StringComparison.Ordinal);
 
-				UserResponse response = await SendRequestAsync<UserResponse>(request, "", ShouldLogError);
+				UserResponse response = await SendRequestAsync<UserResponse>(request, "", ShouldLogError, cancellationToken);
 				return response.User;
 			}
 		}
@@ -485,12 +521,13 @@ namespace EpicGames.Slack
 		/// </summary>
 		/// <param name="triggerId">The trigger id, as returned as part of an interaction payload</param>
 		/// <param name="view">Definition for the view</param>
-		public async Task OpenViewAsync(string triggerId, SlackView view)
+		/// <param name="cancellationToken">Cancellation token for the operation</param>
+		public async Task OpenViewAsync(string triggerId, SlackView view, CancellationToken cancellationToken = default)
 		{
 			ViewsOpenRequest request = new ViewsOpenRequest();
 			request.TriggerId = triggerId;
 			request.View = view;
-			await SendRequestAsync<PostMessageResponse>(ViewsOpenUrl, request);
+			await SendRequestAsync<PostMessageResponse>(ViewsOpenUrl, request, cancellationToken);
 		}
 
 		#endregion
