UAT: Fix static analysis warnings (and error building from VS?) due to style issues. Created separate exception types to make the new code more idiomatic.

diff --git a/Engine/Source/Programs/AutomationTool/BuildGraph/Tasks/MoveTask.cs b/Engine/Source/Programs/AutomationTool/BuildGraph/Tasks/MoveTask.cs
--- a/Engine/Source/Programs/AutomationTool/BuildGraph/Tasks/MoveTask.cs
+++ b/Engine/Source/Programs/AutomationTool/BuildGraph/Tasks/MoveTask.cs
@@ -135,23 +135,15 @@ namespace AutomationTool.Tasks
 				}
 
 			}
-			catch (FilePatternException Ex)
+			catch (FilePatternSourceFileMissingException Ex)
 			{
-				if (Ex.GetFilePatternExceptionType() == FilePatternException.FilePatternExceptionType.SourceFileDoesNotExist)
+				if (Parameters.ErrorIfNotFound)
 				{
-					if (Parameters.ErrorIfNotFound)
-					{
-						Logger.LogError("Error while trying to create file pattern match for '{SourcePattern}', error {ExceptionString}", SourcePattern, Ex.ToString());
-					}
-					else
-					{
-						Logger.LogInformation("Error while trying to create file pattern match for '{SourcePattern}', error {ExceptionString}", SourcePattern, Ex.ToString());
-					}
-					return Task.CompletedTask;
+					Logger.LogError("Error while trying to create file pattern match for '{SourcePattern}', error {ExceptionString}", SourcePattern, Ex.ToString());
 				}
 				else
 				{
-					throw Ex;
+					Logger.LogInformation("Error while trying to create file pattern match for '{SourcePattern}', error {ExceptionString}", SourcePattern, Ex.ToString());
 				}
 			}
 
diff --git a/Engine/Source/Programs/Shared/EpicGames.Core/FilePattern.cs b/Engine/Source/Programs/Shared/EpicGames.Core/FilePattern.cs
--- a/Engine/Source/Programs/Shared/EpicGames.Core/FilePattern.cs
+++ b/Engine/Source/Programs/Shared/EpicGames.Core/FilePattern.cs
@@ -21,40 +21,41 @@ namespace EpicGames.Core
 		public FilePatternException(string message)
 			: base(message)
 		{
-			Type = FilePatternExceptionType.Unknown;
 		}
 
-		/// <summary>
-		/// Constructor
-		/// </summary>
-		public FilePatternException(string format, params object[] args)
-			: base(String.Format(format, args))
+		/// <inheritdoc/>
+		public override string ToString()
 		{
-			Type = FilePatternExceptionType.Unknown;
+			return Message;
 		}
+	}
 
-		public FilePatternException(FilePatternExceptionType InType, string format, params object[] args)
-			: base(String.Format(format, args))
+	/// <summary>
+	/// Exception thrown to indicate that a source file is not under the given base directory
+	/// </summary>
+	public class FilePatternSourceFileNotUnderBaseDirException : FilePatternException
+	{
+		/// <summary>
+		/// Constructor
+		/// </summary>
+		public FilePatternSourceFileNotUnderBaseDirException(string message)
+			: base(message)
 		{
-			Type = InType;
 		}
+	}
 
-		public FilePatternExceptionType GetFilePatternExceptionType() 
-		{ 
-			return Type; 
-		}
-		/// <inheritdoc/>
-		public override string ToString()
+	/// <summary>
+	/// Exception thrown to indicate that a source file is missing
+	/// </summary>
+	public class FilePatternSourceFileMissingException : FilePatternException
+	{
+		/// <summary>
+		/// Constructor
+		/// </summary>
+		public FilePatternSourceFileMissingException(string message)
+			: base(message)
 		{
-			return Message;
 		}
-		public enum FilePatternExceptionType
-		{
-			Unknown,
-			SourceFileDoesNotExist,
-			SourceFileNotUnderBaseDir
-		};
-		FilePatternExceptionType Type;
 	}
 
 	/// <summary>
@@ -356,7 +357,7 @@ namespace EpicGames.Core
 				// Check the two patterns are compatible
 				if(!sourcePattern.IsCompatibleWith(targetPattern))
 				{
-					throw new FilePatternException("File patterns '{0}' and '{1}' do not have matching wildcards", sourcePattern, targetPattern);
+					throw new FilePatternException($"File patterns '{sourcePattern}' and '{targetPattern}' do not have matching wildcards");
 				}
 
 				// Create a filter to match the source files
@@ -393,7 +394,7 @@ namespace EpicGames.Core
 					FileReference? existingSourceFile;
 					if(targetFileToSourceFile.TryGetValue(targetFiles[idx], out existingSourceFile) && existingSourceFile != sourceFiles[idx])
 					{
-						throw new FilePatternException("Output file '{0}' is mapped from '{1}' and '{2}'", targetFiles[idx], existingSourceFile, sourceFiles[idx]);
+						throw new FilePatternException($"Output file '{targetFiles[idx]}' is mapped from '{existingSourceFile}' and '{sourceFiles[idx]}'");
 					}
 					targetFileToSourceFile[targetFiles[idx]] = sourceFiles[idx];
 				}
@@ -409,7 +410,7 @@ namespace EpicGames.Core
 				}
 				else
 				{
-					throw new FilePatternException(FilePatternException.FilePatternExceptionType.SourceFileDoesNotExist, "Source file '{0}' does not exist", sourceFile);
+					throw new FilePatternSourceFileMissingException($"Source file '{sourceFile}' does not exist");
 				}
 			}
 
@@ -418,7 +419,7 @@ namespace EpicGames.Core
 			{
 				if(targetFileToSourceFile.ContainsKey(sourceFile))
 				{
-					throw new FilePatternException("'{0}' is listed as a source and target file", sourceFile);
+					throw new FilePatternException($"'{sourceFile}' is listed as a source and target file");
 				}
 			}
 
@@ -439,11 +440,11 @@ namespace EpicGames.Core
 			{
 				if(!inputFile.IsUnderDirectory(baseDirectory))
 				{
-					throw new FilePatternException(FilePatternException.FilePatternExceptionType.SourceFileNotUnderBaseDir, "Source file '{0}' is not under '{1}'", inputFile, baseDirectory);
+					throw new FilePatternSourceFileNotUnderBaseDirException($"Source file '{inputFile}' is not under '{baseDirectory}'");
 				}
 				else if(!FileReference.Exists(inputFile))
 				{
-					throw new FilePatternException(FilePatternException.FilePatternExceptionType.SourceFileDoesNotExist, "Source file '{0}' does not exist", inputFile);
+					throw new FilePatternSourceFileMissingException($"Source file '{inputFile}' does not exist");
 				}
 				else
 				{
