Horde: Ensure file read access is requested when verifying sandbox files for compute tasks

Without this access flag, reads could fail with "access denied" (depending on OS? Wine vs Windows?)

diff --git a/Engine/Source/Programs/Shared/EpicGames.Horde/Compute/AgentMessageHandler.cs b/Engine/Source/Programs/Shared/EpicGames.Horde/Compute/AgentMessageHandler.cs
--- a/Engine/Source/Programs/Shared/EpicGames.Horde/Compute/AgentMessageHandler.cs
+++ b/Engine/Source/Programs/Shared/EpicGames.Horde/Compute/AgentMessageHandler.cs
@@ -217,11 +217,9 @@ namespace EpicGames.Horde.Compute
 				}
 				else
 				{
-					IoHash hash;
-					using (FileStream stream = FileReference.Open(file, FileMode.Open))
-					{
-						hash = await IoHash.ComputeAsync(stream, cancellationToken);
-					}
+					await using FileStream stream = FileReference.Open(file, FileMode.Open, FileAccess.Read);
+					IoHash hash = await IoHash.ComputeAsync(stream, cancellationToken);
+					
 					if (hash == fileEntry.StreamHash)
 					{
 						_logger.LogInformation("Hash of {File} is correct ({Hash})", file, hash);
