Horde: Throw any received remote compute exceptions

- Try surface remote errors better on the client's side
- Wrap remote compute exceptions as an inner exception to capture both local and remote callstacks

diff --git a/Engine/Source/Programs/Shared/EpicGames.Horde/Compute/AgentManagedProcess.cs b/Engine/Source/Programs/Shared/EpicGames.Horde/Compute/AgentManagedProcess.cs
--- a/Engine/Source/Programs/Shared/EpicGames.Horde/Compute/AgentManagedProcess.cs
+++ b/Engine/Source/Programs/Shared/EpicGames.Horde/Compute/AgentManagedProcess.cs
@@ -72,7 +72,7 @@ namespace EpicGames.Horde.Compute
 				{
 					case AgentMessageType.Exception:
 						ExceptionMessage exception = message.ParseExceptionMessage();
-						throw new ComputeRemoteException(exception);
+						throw new ComputeException("Error while executing remote process", new ComputeRemoteException(exception));
 					case AgentMessageType.ExecuteOutput:
 						AppendData(message.Data.Span);
 						break;
@@ -81,7 +81,8 @@ namespace EpicGames.Horde.Compute
 						_result.TrySetResult(executeProcessResponse.ExitCode);
 						return;
 					default:
-						throw new InvalidAgentMessageException(message);
+						message.ThrowIfUnexpectedType();
+						return;
 				}
 			}
 		}
diff --git a/Engine/Source/Programs/Shared/EpicGames.Horde/Compute/AgentMessage.cs b/Engine/Source/Programs/Shared/EpicGames.Horde/Compute/AgentMessage.cs
--- a/Engine/Source/Programs/Shared/EpicGames.Horde/Compute/AgentMessage.cs
+++ b/Engine/Source/Programs/Shared/EpicGames.Horde/Compute/AgentMessage.cs
@@ -192,8 +192,8 @@ namespace EpicGames.Horde.Compute
 		/// <summary>
 		/// Constructor
 		/// </summary>
-		public InvalidAgentMessageException(AgentMessage message)
-			: base($"Unexpected message {message.Type}")
+		public InvalidAgentMessageException(AgentMessage actualMessage, AgentMessageType? expectedType, ComputeRemoteException? remoteException)
+			: base($"Unexpected message {actualMessage.Type}" + (expectedType != null ? $". Wanted {expectedType}" : ""), remoteException)
 		{
 		}
 	}
@@ -351,10 +351,26 @@ namespace EpicGames.Horde.Compute
 		public static async ValueTask WaitForAttachAsync(this AgentMessageChannel channel, CancellationToken cancellationToken = default)
 		{
 			using AgentMessage message = await channel.ReceiveAsync(cancellationToken);
-			if (message.Type != AgentMessageType.Attach)
+			message.ThrowIfUnexpectedType(AgentMessageType.Attach);
+		}
+
+		/// <summary>
+		/// Throw an exception if message is not of expected type
+		/// </summary>
+		/// <param name="message">Agent message to extend</param>
+		/// <param name="expectedType">Optional type to expect. If not specified, assume type was unwanted no matter what</param>
+		public static void ThrowIfUnexpectedType(this AgentMessage message, AgentMessageType? expectedType = null)
+		{
+			if (message.Type == expectedType)
 			{
-				throw new InvalidAgentMessageException(message);
+				return;
 			}
+
+			ComputeRemoteException? cre = message.Type == AgentMessageType.Exception
+				? new ComputeRemoteException(message.ParseExceptionMessage())
+				: null;
+
+			throw new InvalidAgentMessageException(message, expectedType, cre);
 		}
 
 		#region Process
@@ -394,10 +410,7 @@ namespace EpicGames.Horde.Compute
 			}
 
 			using AgentMessage response = await RunStorageServerAsync(channel, storage, cancellationToken);
-			if (response.Type != AgentMessageType.WriteFilesResponse)
-			{
-				throw new InvalidAgentMessageException(response);
-			}
+			response.ThrowIfUnexpectedType(AgentMessageType.WriteFilesResponse);
 		}
 
 		/// <summary>
@@ -650,10 +663,7 @@ namespace EpicGames.Horde.Compute
 				try
 				{
 					response = await channel.ReceiveAsync(cancellationToken);
-					if (response.Type != AgentMessageType.ReadBlobResponse)
-					{
-						throw new InvalidAgentMessageException(response);
-					}
+					response.ThrowIfUnexpectedType(AgentMessageType.ReadBlobResponse);
 
 					int chunkOffset = BinaryPrimitives.ReadInt32LittleEndian(response.Data.Span.Slice(0, 4));
 					int chunkLength = response.Data.Length - 8;
diff --git a/Engine/Source/Programs/Shared/EpicGames.Horde/Compute/AgentMessageChannel.cs b/Engine/Source/Programs/Shared/EpicGames.Horde/Compute/AgentMessageChannel.cs
--- a/Engine/Source/Programs/Shared/EpicGames.Horde/Compute/AgentMessageChannel.cs
+++ b/Engine/Source/Programs/Shared/EpicGames.Horde/Compute/AgentMessageChannel.cs
@@ -250,10 +250,7 @@ namespace EpicGames.Horde.Compute
 		public static async ValueTask<AgentMessage> ReceiveAsync(this AgentMessageChannel channel, AgentMessageType type, CancellationToken cancellationToken = default)
 		{
 			AgentMessage message = await channel.ReceiveAsync(cancellationToken);
-			if (message.Type != type)
-			{
-				throw new InvalidAgentMessageException(message);
-			}
+			message.ThrowIfUnexpectedType(type);
 			return message;
 		}
 
diff --git a/Engine/Source/Programs/Shared/EpicGames.Horde/Compute/AgentMessageHandler.cs b/Engine/Source/Programs/Shared/EpicGames.Horde/Compute/AgentMessageHandler.cs
--- a/Engine/Source/Programs/Shared/EpicGames.Horde/Compute/AgentMessageHandler.cs
+++ b/Engine/Source/Programs/Shared/EpicGames.Horde/Compute/AgentMessageHandler.cs
@@ -84,9 +84,9 @@ namespace EpicGames.Horde.Compute
 		async Task RunAsync(ComputeSocket socket, int channelId, int bufferSize, Action<Exception> postException, CancellationToken cancellationToken)
 		{
 			List<Task> childTasks = new List<Task>();
+			using AgentMessageChannel channel = socket.CreateAgentMessageChannel(channelId, bufferSize);
 			try
 			{
-				using AgentMessageChannel channel = socket.CreateAgentMessageChannel(channelId, bufferSize);
 				await channel.AttachAsync(cancellationToken);
 
 				for (; ; )
@@ -144,7 +144,8 @@ namespace EpicGames.Horde.Compute
 							}
 							break;
 						default:
-							throw new InvalidAgentMessageException(message);
+							message.ThrowIfUnexpectedType();
+							return;
 					}
 				}
 			}
@@ -156,6 +157,7 @@ namespace EpicGames.Horde.Compute
 			catch (Exception ex)
 			{
 				_logger.LogInformation(ex, "Compute Channel {ChannelId}: Exception: {Message}", channelId, ex.Message);
+				await channel.SendExceptionAsync(ex, cancellationToken);
 				postException(ex);
 			}
 			finally
diff --git a/Engine/Source/Programs/Shared/EpicGames.Horde/Compute/ComputeExceptions.cs b/Engine/Source/Programs/Shared/EpicGames.Horde/Compute/ComputeExceptions.cs
--- a/Engine/Source/Programs/Shared/EpicGames.Horde/Compute/ComputeExceptions.cs
+++ b/Engine/Source/Programs/Shared/EpicGames.Horde/Compute/ComputeExceptions.cs
@@ -5,17 +5,17 @@ using System;
 namespace EpicGames.Horde.Compute
 {
 	/// <summary>
-	/// Base class for compute errors
+	/// Generic class for compute errors
 	/// </summary>
-	public abstract class ComputeException : Exception
+	public class ComputeException : Exception
 	{
 		/// <inheritdoc/>
-		protected ComputeException(string message) : base(message)
+		public ComputeException(string message) : base(message)
 		{
 		}
 
 		/// <inheritdoc/>
-		protected ComputeException(string? message, Exception? innerException) : base(message, innerException)
+		public ComputeException(string? message, Exception? innerException) : base(message, innerException)
 		{
 		}
 	}
@@ -46,7 +46,7 @@ namespace EpicGames.Horde.Compute
 		/// Constructor
 		/// </summary>
 		public ComputeRemoteException(ExceptionMessage message)
-			: this(message.Message, message.Description)
+			: this("Remote exception: " + message.Message, message.Description)
 		{
 		}
 
@@ -56,7 +56,7 @@ namespace EpicGames.Horde.Compute
 		public ComputeRemoteException(string message, string description)
 			: base(message)
 		{
-			_description = description;
+			_description = "From remote machine: " + description;
 		}
 
 		/// <inheritdoc/>
