Horde: Add a managed workspace option to use the native P4 client rather than P4.EXE. Currently disabled by default, but intend to enable for testing in dev environment.

Also change ManagedWorkspaceOptions to be a fancy new record type, which implements all the comparison/hashing operators automagically for us.

diff --git a/Engine/Source/Programs/Horde/Horde.Agent/Utility/WorkspaceInfo.cs b/Engine/Source/Programs/Horde/Horde.Agent/Utility/WorkspaceInfo.cs
--- a/Engine/Source/Programs/Horde/Horde.Agent/Utility/WorkspaceInfo.cs
+++ b/Engine/Source/Programs/Horde/Horde.Agent/Utility/WorkspaceInfo.cs
@@ -141,7 +141,7 @@ namespace Horde.Agent.Utility
 			}
 
 			// Create the connection
-			IPerforceConnection perforce = await PerforceConnection.CreateAsync(new PerforceSettings(serverAndPort, userName) { PreferNativeClient = false, Password = ticket}, logger);
+			IPerforceConnection perforce = await PerforceConnection.CreateAsync(new PerforceSettings(serverAndPort, userName) { PreferNativeClient = options.PreferNativeClient, Password = ticket}, logger);
 			if (userName != null)
 			{
 				if (ticket != null)
@@ -280,9 +280,10 @@ namespace Horde.Agent.Utility
 			const string MaxFileConcurrencyKey = "maxFileConcurrency";
 			const string MinScratchSpaceKey = "minScratchSpace";
 			const string UseHaveTableKey = "useHaveTable";
+			const string PreferNativeClientKey = "preferNativeClient";
 
 			ManagedWorkspaceOptions options = new ManagedWorkspaceOptions();
-			options.Partitioned = workspace.Partitioned;
+			options = options with { Partitioned = workspace.Partitioned };
 
 			string? method = workspace.Method;
 			if (!String.IsNullOrEmpty(method))
@@ -292,19 +293,23 @@ namespace Horde.Agent.Utility
 				{
 					if (Int32.TryParse(nameValues[NumParallelSyncThreadsKey], out int v))
 					{
-						options.NumParallelSyncThreads = v;
+						options = options with { NumParallelSyncThreads = v };
 					}
 					if (Int32.TryParse(nameValues[MaxFileConcurrencyKey], out v))
 					{
-						options.MaxFileConcurrency = v;
+						options = options with { MaxFileConcurrency = v };
 					}
 					if (Int32.TryParse(nameValues[MinScratchSpaceKey], out v))
 					{
-						options.MinScratchSpace = v;
+						options = options with { MinScratchSpace = v };
 					}
 					if (String.Equals(nameValues[UseHaveTableKey], "false", StringComparison.OrdinalIgnoreCase))
 					{
-						options.UseHaveTable = false;
+						options = options with { UseHaveTable = false };
+					}
+					if (String.Equals(nameValues[PreferNativeClientKey], "true", StringComparison.OrdinalIgnoreCase))
+					{
+						options = options with { PreferNativeClient = true };
 					}
 				}
 			}
diff --git a/Engine/Source/Programs/Shared/EpicGames.Perforce.Managed/ManagedWorkspace.cs b/Engine/Source/Programs/Shared/EpicGames.Perforce.Managed/ManagedWorkspace.cs
--- a/Engine/Source/Programs/Shared/EpicGames.Perforce.Managed/ManagedWorkspace.cs
+++ b/Engine/Source/Programs/Shared/EpicGames.Perforce.Managed/ManagedWorkspace.cs
@@ -64,77 +64,30 @@ namespace EpicGames.Perforce.Managed
 			View = view;
 		}
 	}
-	
+
 	/// <summary>
 	/// Extra options for configuring ManagedWorkspace
 	/// </summary>
-	public class ManagedWorkspaceOptions
-	{
-		/// <summary>
-		/// Maximum number of threads to sync in parallel
-		/// </summary>
-		public int NumParallelSyncThreads { get; set; } = 4;
-
-		/// <summary>
-		/// Maximum number of concurrent file system operations (copying, moving, deleting etc)
-		/// </summary>
-		public int MaxFileConcurrency { get; set; } = 4;
-		
-		/// <summary>
-		/// Minimum amount of space that must be on a drive after a branch is synced
-		/// </summary>
-		public long MinScratchSpace { get; set; } = 50L * 1024 * 1024 * 1024;
-		
-		/// <summary>
-		/// Use the client's have table when syncing.
-		/// 
-		/// When set to false, updates to the have table will be prevented through use of "sync -p".
-		/// Actual files to sync will be gathered through "fstat".
-		/// This puts less strain on the Perforce server and can improve sync performance.
-		/// </summary>
-		public bool UseHaveTable { get; set; } = true;
-
-		/// <summary>
-		/// Whether to allow using partitioned workspaces
-		/// </summary>
-		public bool Partitioned { get; set; }
-
-		/// <inheritdoc/>
-		protected bool Equals(ManagedWorkspaceOptions other)
-		{
-			return NumParallelSyncThreads == other.NumParallelSyncThreads && MaxFileConcurrency == other.MaxFileConcurrency && MinScratchSpace == other.MinScratchSpace && UseHaveTable == other.UseHaveTable;
-		}
-
-		/// <inheritdoc/>
-		public override bool Equals(object? obj)
-		{
-			if (ReferenceEquals(null, obj))
-			{ 
-				return false; 
-			}
-			if (ReferenceEquals(this, obj))
-			{ 
-				return true;
-			}
-			if (obj.GetType() != GetType())
-			{
-				return false;
-			}
-			return Equals((ManagedWorkspaceOptions)obj);
-		}
-
-		/// <inheritdoc/>
-		public override int GetHashCode()
-		{
-			return HashCode.Combine(NumParallelSyncThreads, MaxFileConcurrency, MinScratchSpace, UseHaveTable);
-		}
-
-		/// <inheritdoc/>
-		public override string ToString()
-		{
-			return $"{nameof(NumParallelSyncThreads)}={NumParallelSyncThreads} {nameof(MaxFileConcurrency)}={MaxFileConcurrency} {nameof(MinScratchSpace)}={MinScratchSpace} {nameof(UseHaveTable)}={UseHaveTable}";
-		}
-	}
+	/// <param name="NumParallelSyncThreads">Maximum number of threads to sync in parallel</param>
+	/// <param name="MaxFileConcurrency">Maximum number of concurrent file system operations (copying, moving, deleting etc)</param>
+	/// <param name="MinScratchSpace">Minimum amount of space that must be on a drive after a branch is synced</param>
+	/// <param name="UseHaveTable">
+	///		Use the client's have table when syncing.
+	///		
+	///		When set to false, updates to the have table will be prevented through use of "sync -p".
+	///		Actual files to sync will be gathered through "fstat". This puts less strain on the Perforce server and can improve sync performance.
+	///	</param>
+	/// <param name="Partitioned">Whether to allow using partitioned workspaces</param>
+	/// <param name="PreferNativeClient">Whether to prefer the native p4 client</param>
+	public record class ManagedWorkspaceOptions
+	(
+		int NumParallelSyncThreads = 4,
+		int MaxFileConcurrency = 4,
+		long MinScratchSpace = 50L * 1024 * 1024 * 1024,
+		bool UseHaveTable = true,
+		bool Partitioned = false,
+		bool PreferNativeClient = false
+	);
 
 	/// <summary>
 	/// Version number for managed workspace cache files
