Horde: Reduce log spam when a large number of files are missing from the cache.

diff --git a/Engine/Source/Programs/Shared/EpicGames.Perforce.Managed/ManagedWorkspace.cs b/Engine/Source/Programs/Shared/EpicGames.Perforce.Managed/ManagedWorkspace.cs
--- a/Engine/Source/Programs/Shared/EpicGames.Perforce.Managed/ManagedWorkspace.cs
+++ b/Engine/Source/Programs/Shared/EpicGames.Perforce.Managed/ManagedWorkspace.cs
@@ -13,6 +13,7 @@ using System.Threading;
 using System.Threading.Tasks;
 using EpicGames.Core;
 using Microsoft.Extensions.Logging;
+using Microsoft.Extensions.Logging.Abstractions;
 
 namespace EpicGames.Perforce.Managed
 {
@@ -569,15 +570,23 @@ namespace EpicGames.Perforce.Managed
 				// Make sure all the folders exist in the cache
 				CreateCacheHierarchy();
 
-				// Check that all the files in the cache appear as we expect them to
 				List<CachedFileInfo> trackedFiles = _contentIdToTrackedFile.Values.ToList();
+
+				// Check that all the files in the cache appear as we expect them to
+				const int MaxLoggedMissingFiles = 250;
+				int numMissingFiles = 0;
 				foreach (CachedFileInfo trackedFile in trackedFiles)
 				{
-					if (!trackedFile.CheckIntegrity(_logger))
+					if (!trackedFile.CheckIntegrity((numMissingFiles < MaxLoggedMissingFiles)? _logger : NullLogger.Instance))
 					{
 						RemoveTrackedFile(trackedFile);
+						numMissingFiles++;
 					}
 				}
+				if (numMissingFiles > MaxLoggedMissingFiles)
+				{
+					_logger.LogWarning("+ {Count} more", numMissingFiles - MaxLoggedMissingFiles);
+				}
 
 				// Clear the repair flag
 				_requiresRepair = false;
