Horde: Log exceptions when moving files to the cache at the end of a job, and remove invalid entries that were added to the cache on failure.

diff --git a/Engine/Source/Programs/Shared/EpicGames.Core/KnownLogEvents.cs b/Engine/Source/Programs/Shared/EpicGames.Core/KnownLogEvents.cs
--- a/Engine/Source/Programs/Shared/EpicGames.Core/KnownLogEvents.cs
+++ b/Engine/Source/Programs/Shared/EpicGames.Core/KnownLogEvents.cs
@@ -321,6 +321,11 @@ namespace EpicGames.Core
 		/// </summary>
 		public static EventId Systemic_OutOfDiskSpace { get; } = new EventId(768);
 
+		/// <summary>
+		/// Error moving files to cache
+		/// </summary>
+		public static EventId Systemic_ManagedWorkspace { get; } = new EventId(769);
+
 		/// <summary>
 		/// Maximum code for systemic events. Add new events in the 700-799 range.
 		/// </summary>
diff --git a/Engine/Source/Programs/Shared/EpicGames.Perforce.Managed/ManagedWorkspace.cs b/Engine/Source/Programs/Shared/EpicGames.Perforce.Managed/ManagedWorkspace.cs
--- a/Engine/Source/Programs/Shared/EpicGames.Perforce.Managed/ManagedWorkspace.cs
+++ b/Engine/Source/Programs/Shared/EpicGames.Perforce.Managed/ManagedWorkspace.cs
@@ -1758,13 +1758,13 @@ namespace EpicGames.Perforce.Managed
 					Stopwatch timer = Stopwatch.StartNew();
 
 					// Add any new files to the cache
-					List<KeyValuePair<FileReference, FileReference>> sourceAndTargetFiles = new List<KeyValuePair<FileReference, FileReference>>();
+					List<(FileContentId ContentId, FileReference Source, FileReference Target)> files = new List<(FileContentId, FileReference, FileReference)>();
 					foreach (KeyValuePair<FileContentId, WorkspaceFileInfo> fileToMove in filesToMove)
 					{
 						ulong cacheId = GetUniqueCacheId(fileToMove.Key);
 						CachedFileInfo newTrackingInfo = new CachedFileInfo(_cacheDir, fileToMove.Key, cacheId, fileToMove.Value._length, fileToMove.Value._lastModifiedTicks, fileToMove.Value._readOnly, _nextSequenceNumber);
 						_contentIdToTrackedFile.Add(fileToMove.Key, newTrackingInfo);
-						sourceAndTargetFiles.Add(new KeyValuePair<FileReference, FileReference>(fileToMove.Value.GetLocation(), newTrackingInfo.GetLocation()));
+						files.Add((fileToMove.Key, fileToMove.Value.GetLocation(), newTrackingInfo.GetLocation()));
 					}
 					_nextSequenceNumber++;
 
@@ -1773,11 +1773,7 @@ namespace EpicGames.Perforce.Managed
 
 					// Execute all the moves and deletes
 					ParallelOptions options = new() { MaxDegreeOfParallelism = _options.MaxFileConcurrency, CancellationToken = cancellationToken };
-					await Parallel.ForEachAsync(sourceAndTargetFiles, options, (sourceAndTargetFile, ct) =>
-					{
-						FileUtils.ForceMoveFile(sourceAndTargetFile.Key, sourceAndTargetFile.Value);
-						return ValueTask.CompletedTask;
-					});
+					await Parallel.ForEachAsync(files, options, (file, ctx) => MoveFileToCache(file.ContentId, file.Source, file.Target, _contentIdToTrackedFile, ctx));
 
 					scope.Progress = $"({timer.Elapsed.TotalSeconds:0.0}s)";
 				}
@@ -1826,6 +1822,24 @@ namespace EpicGames.Perforce.Managed
 			await SaveAsync(TransactionState.Clean, cancellationToken);
 		}
 
+		ValueTask MoveFileToCache(FileContentId contentId, FileReference sourceFile, FileReference targetFile, Dictionary<FileContentId, CachedFileInfo> contentIdToTrackedFile, CancellationToken cancellationToken)
+		{
+			cancellationToken.ThrowIfCancellationRequested();
+			try
+			{
+				FileUtils.ForceMoveFile(sourceFile, targetFile);
+			}
+			catch (Exception ex)
+			{
+				_logger.LogWarning(KnownLogEvents.Systemic_ManagedWorkspace, ex, "Unable to move {SourceFile} to {TargetFile}: {Error}", sourceFile, targetFile, ex.Message);
+				lock (contentIdToTrackedFile)
+				{
+					contentIdToTrackedFile.Remove(contentId);
+				}
+			}
+			return default;
+		}
+
 		/// <summary>
 		/// Helper function to delete a file from the workspace, and output any failure as a warning.
 		/// </summary>
