Horde: Improve errors from file move exceptions.

diff --git a/Engine/Source/Programs/Shared/EpicGames.Core/FileUtils.cs b/Engine/Source/Programs/Shared/EpicGames.Core/FileUtils.cs
--- a/Engine/Source/Programs/Shared/EpicGames.Core/FileUtils.cs
+++ b/Engine/Source/Programs/Shared/EpicGames.Core/FileUtils.cs
@@ -396,11 +396,11 @@ namespace EpicGames.Core
 				}
 				catch (Exception deleteEx)
 				{
-					throw new WrappedFileOrDirectoryException(new AggregateException(ex, deleteEx), String.Format("Unable to move {0} to {1} (also tried delete/move)", sourceLocation, targetLocation));
+					throw new WrappedFileOrDirectoryException(new AggregateException(ex, deleteEx), $"Unable to move {sourceLocation} to {targetLocation} (also tried delete/move): {ex.Message}");
 				}
 
 				// Throw the original exception
-				throw new WrappedFileOrDirectoryException(ex, String.Format("Unable to move {0} to {1}", sourceLocation, targetLocation));
+				throw new WrappedFileOrDirectoryException(ex, $"Unable to move {sourceLocation} to {targetLocation}: {ex.Message}");
 			}
 		}
 
