Horde: Prevent switching back to using p4.exe for syncing managed workspaces when native p4 client is in use.

diff --git a/Engine/Source/Programs/Shared/EpicGames.Perforce.Managed/ManagedWorkspace.cs b/Engine/Source/Programs/Shared/EpicGames.Perforce.Managed/ManagedWorkspace.cs
--- a/Engine/Source/Programs/Shared/EpicGames.Perforce.Managed/ManagedWorkspace.cs
+++ b/Engine/Source/Programs/Shared/EpicGames.Perforce.Managed/ManagedWorkspace.cs
@@ -1831,7 +1831,9 @@ namespace EpicGames.Perforce.Managed
 			}
 			catch (Exception ex)
 			{
-				_logger.LogWarning(KnownLogEvents.Systemic_ManagedWorkspace, ex, "Unable to move {SourceFile} to {TargetFile}: {Error}", sourceFile, targetFile, ex.Message);
+				Exception innerException = (ex as WrappedFileOrDirectoryException)?.InnerException ?? ex;
+				_logger.LogWarning(KnownLogEvents.Systemic_ManagedWorkspace, innerException, "Unable to move {SourceFile} to {TargetFile}: {Error}", sourceFile, targetFile, innerException.Message);
+
 				lock (contentIdToTrackedFile)
 				{
 					contentIdToTrackedFile.Remove(contentId);
@@ -2072,6 +2074,24 @@ namespace EpicGames.Perforce.Managed
 					await FileReference.WriteAllBytesAsync(localFile, Array.Empty<byte>(), cancellationToken);
 				}
 			}
+			else if (client is NativePerforceConnection)
+			{
+				List<string> files = new List<string>();
+				for (int idx = beginIdx; idx < endIdx; idx++)
+				{
+					files.Add($"{filesToSync[idx]._streamFile.Path}#{filesToSync[idx]._streamFile.Revision}");
+				}
+
+				using IPerforceConnection threadedClient = await PerforceConnection.CreateAsync(client.Settings, client.Logger);
+				if (_options.UseHaveTable)
+				{
+					await threadedClient.SyncAsync(SyncOptions.Force | SyncOptions.FullDepotSyntax, -1, files, cancellationToken).ToListAsync(cancellationToken);
+				}
+				else
+				{
+					await threadedClient.SyncAsync(SyncOptions.DoNotUpdateHaveList | SyncOptions.FullDepotSyntax, -1, files, cancellationToken).ToListAsync(cancellationToken);
+				}
+			}
 			else
 			{
 				FileReference syncFileName = FileReference.Combine(_baseDir, $"SyncList-{beginIdx}.txt");
@@ -2085,14 +2105,14 @@ namespace EpicGames.Perforce.Managed
 
 				if (_options.UseHaveTable)
 				{
-					using PerforceConnection clientWithFileList = new (client.Settings, client.Logger);
+					using PerforceConnection clientWithFileList = new(client.Settings, client.Logger);
 					clientWithFileList.GlobalOptions.Add($"-x\"{syncFileName}\"");
 					await clientWithFileList.SyncAsync(SyncOptions.Force | SyncOptions.FullDepotSyntax, -1, Array.Empty<string>(), cancellationToken).ToListAsync(cancellationToken);
 				}
 				else
 				{
 					// Ensure a client with an empty have table is used to not interfere with the DoNotUpdateHaveList option.
-					using PerforceConnection clientWithFileList = new (client.Settings, client.Logger);
+					using PerforceConnection clientWithFileList = new(client.Settings, client.Logger);
 					clientWithFileList.ClientName = client.Settings.ClientName!;
 					clientWithFileList.GlobalOptions.Add($"-x\"{syncFileName}\"");
 					await clientWithFileList.SyncAsync(SyncOptions.DoNotUpdateHaveList | SyncOptions.FullDepotSyntax, -1, Array.Empty<string>(), cancellationToken).ToListAsync(cancellationToken);
