Horde: Add pool categories to dashboard config response
#rnx

diff --git a/Engine/Source/Programs/Horde/Horde.Server/Dashboard/DashboardController.cs b/Engine/Source/Programs/Horde/Horde.Server/Dashboard/DashboardController.cs
--- a/Engine/Source/Programs/Horde/Horde.Server/Dashboard/DashboardController.cs
+++ b/Engine/Source/Programs/Horde/Horde.Server/Dashboard/DashboardController.cs
@@ -1,10 +1,8 @@
 // Copyright Epic Games, Inc. All Rights Reserved.
 
 using System;
-using System.Collections;
 using System.Collections.Generic;
 using System.Linq;
-using System.Reflection;
 using System.Text;
 using System.Threading;
 using System.Threading.Tasks;
@@ -152,6 +150,11 @@ namespace Horde.Server.Dashboard
 				dashboardConfigResponse.AgentCategories.Add(new GetDashboardAgentCategoryResponse { Name = category.Name, Condition = category.Condition });
 			}
 
+			foreach (DashboardPoolCategoryConfig category in _globalConfig.Value.Dashboard.PoolCategories)
+			{
+				dashboardConfigResponse.PoolCategories.Add(new GetDashboardPoolCategoryResponse { Name = category.Name, Condition = category.Condition });
+			}
+
 			foreach (TelemetryViewConfig telemetry in _globalConfig.Value.Dashboard.Analytics)
 			{
 				GetTelemetryViewResponse rview = new GetTelemetryViewResponse();
diff --git a/Engine/Source/Programs/Shared/EpicGames.Horde/Dashboard/DashboardMessages.cs b/Engine/Source/Programs/Shared/EpicGames.Horde/Dashboard/DashboardMessages.cs
--- a/Engine/Source/Programs/Shared/EpicGames.Horde/Dashboard/DashboardMessages.cs
+++ b/Engine/Source/Programs/Shared/EpicGames.Horde/Dashboard/DashboardMessages.cs
@@ -48,6 +48,11 @@ namespace EpicGames.Horde.Dashboard
 		/// </summary>
 		public List<GetDashboardAgentCategoryResponse> AgentCategories { get; set; } = new List<GetDashboardAgentCategoryResponse>();
 
+		/// <summary>
+		/// Categories to display on the pools page
+		/// </summary>
+		public List<GetDashboardPoolCategoryResponse> PoolCategories { get; set; } = new List<GetDashboardPoolCategoryResponse>();
+
 		/// <summary>
 		/// Telemetry to display on the telemetry page
 		/// </summary>
@@ -55,6 +60,22 @@ namespace EpicGames.Horde.Dashboard
 
 	}
 
+	/// <summary>
+	/// Describes a category for the pools page
+	/// </summary>
+	public class GetDashboardPoolCategoryResponse
+	{
+		/// <summary>
+		/// Title for the tab
+		/// </summary>
+		public string Name { get; set; } = "Unnamed";
+
+		/// <summary>
+		/// Condition for pools to be included in this category
+		/// </summary>
+		public Condition? Condition { get; set; }
+	}
+
 	/// <summary>
 	/// Describes a category for the agents page
 	/// </summary>
