Added UHT support for ``UE_DEPRECATED(all, "...")`` macro, which is already used in the engine and will serve as a way to leave the UE_DEPRECATED mark on code that has been deprecated for more than 2 engine versions but has to stay in the codebase "forever" because it's reflected (UPROPERTY, UFUNCTION, UCLASS, USTRUCT).
Use case :
```
// Let's say we've deprecated the following code in UE5.1
UE_DEPRECATED(5.1, "Use the other thing instead...")
UPROPERTY(meta = (DeprecatedProperty, DeprecationMessage = "Use the other thing instead"))
float SomeUProperty_DEPRECATED = 1.0f;

UE_DEPRECATED(5.1, "Use the other thing instead...")
UFUNCTION(meta = (DeprecatedFunction, DeprecationMessage = "Use the other thing instead..."))
void FooBar(int32 X);

// Then when shipping UE5.4, we would turn that into:
UE_DEPRECATED(all, "Use the other thing instead...")
UPROPERTY(meta = (DeprecatedProperty, DeprecationMessage = "Use the other thing instead"))
float SomeUProperty_DEPRECATED = 1.0f;

UE_DEPRECATED(all, "Use the other thing instead...")
UFUNCTION(meta = (DeprecatedFunction, DeprecationMessage = "Use the other thing instead..."))
void FooBar(int32 X);
```

Therefore, the act of deprecating code that is more than 2 engine versions old consists in a simple text-search for "UE_DEPRECATED" in your code and for the occurrences that are not tagged as "all" and more than 2 versions old. Then:
- If the code is reflected, turn them into UE_DEPRECATED(all, ...)
- Otherwise, remove the code altogether

diff --git a/Engine/Source/Programs/Shared/EpicGames.UHT/Tokenizer/UhtTokenReaderGeneralExtensions.cs b/Engine/Source/Programs/Shared/EpicGames.UHT/Tokenizer/UhtTokenReaderGeneralExtensions.cs
--- a/Engine/Source/Programs/Shared/EpicGames.UHT/Tokenizer/UhtTokenReaderGeneralExtensions.cs
+++ b/Engine/Source/Programs/Shared/EpicGames.UHT/Tokenizer/UhtTokenReaderGeneralExtensions.cs
@@ -227,10 +227,13 @@ namespace EpicGames.UHT.Tokenizer
 				if (token.IsIdentifier("DEPRECATED") || token.IsIdentifier("UE_DEPRECATED"))
 				{
 					tokenReader.ConsumeToken();
-					tokenReader
-						.Require('(', "deprecation macro")
-						.RequireConstFloat("version in deprecation macro")
-						.Require(',', "deprecation macro")
+					tokenReader.Require('(', "deprecation macro");
+					token = tokenReader.GetToken();
+					if (!token.IsConstFloat() && !token.IsIdentifier("all", ignoreCase:true))
+					{
+						throw new UhtTokenException(tokenReader, token, "'all' or constant float version in deprecation macro");
+					}
+					tokenReader.Require(',', "deprecation macro")
 						.RequireConstString("message in deprecation macro")
 						.Require(')', "deprecation macro");
 					attributeAction?.Invoke("deprecated");
