Adding List<string>.AddDefinition extension.

diff --git a/Engine/Source/Programs/Shared/EpicGames.Core/ListExtensions.cs b/Engine/Source/Programs/Shared/EpicGames.Core/ListExtensions.cs
--- a/Engine/Source/Programs/Shared/EpicGames.Core/ListExtensions.cs
+++ b/Engine/Source/Programs/Shared/EpicGames.Core/ListExtensions.cs
@@ -101,6 +101,39 @@ namespace EpicGames.Core
 			list.AddRange(arguments.AsEnumerable());
 		}
 
+		/// <summary>
+		/// Add to a list of definitions with the form of "Name=Value".
+		/// </summary>
+		/// <param name="list">List to add to</param>
+		/// <param name="definitionName">Definition name</param>
+		/// <param name="Value">Value to assign to Definition</param>
+		public static void AddDefinition(this List<string> list, string definitionName, string value)
+		{
+			list.Add($"{definitionName}={value}");
+		}
+
+		/// <summary>
+		/// Add to a list of definitions with the form of "Name=Value" where Value is 1 or 0 based on bValue.
+		/// </summary>
+		/// <param name="list">List to add to</param>
+		/// <param name="definitionName">Definition name</param>
+		/// <param name="bValue">Value to assign to Definition</param>
+		public static void AddDefinition(this List<string> list, string definitionName, bool bValue)
+		{
+			list.Add($"{definitionName}={(bValue ? "1" : "0")}");
+		}
+
+		/// <summary>
+		/// Add to a list of definitions with the form of "Name=Value".
+		/// </summary>
+		/// <param name="list">List to add to</param>
+		/// <param name="definitionName">Definition name</param>
+		/// <param name="bValue">Value to assign to Definition</param>
+		public static void AddDefinition(this List<string> list, string definitionName, int value)
+		{
+			list.Add($"{definitionName}={value}");
+		}
+
 		/// <summary>
 		/// Convert all elements of a list to a base type
 		/// </summary>
