Fixing documentation strings.

diff --git a/Engine/Source/Programs/Shared/EpicGames.Core/ListExtensions.cs b/Engine/Source/Programs/Shared/EpicGames.Core/ListExtensions.cs
--- a/Engine/Source/Programs/Shared/EpicGames.Core/ListExtensions.cs
+++ b/Engine/Source/Programs/Shared/EpicGames.Core/ListExtensions.cs
@@ -106,7 +106,7 @@ namespace EpicGames.Core
 		/// </summary>
 		/// <param name="list">List to add to</param>
 		/// <param name="definitionName">Definition name</param>
-		/// <param name="Value">Value to assign to Definition</param>
+		/// <param name="value">Value to assign to Definition</param>
 		public static void AddDefinition(this List<string> list, string definitionName, string value)
 		{
 			list.Add($"{definitionName}={value}");
@@ -128,7 +128,7 @@ namespace EpicGames.Core
 		/// </summary>
 		/// <param name="list">List to add to</param>
 		/// <param name="definitionName">Definition name</param>
-		/// <param name="bValue">Value to assign to Definition</param>
+		/// <param name="value">Value to assign to Definition</param>
 		public static void AddDefinition(this List<string> list, string definitionName, int value)
 		{
 			list.Add($"{definitionName}={value}");
