Horde: Fix file attributes getting screwed up on Windows when extracting a read-only file that didn't previously exist.

FileInfo.Attributes would be cached as being set to -1, so OR'ing in the FileAttribute.ReadOnly flag would set all attribute flags (including system, hidden, etc...).

diff --git a/Engine/Source/Programs/Shared/EpicGames.Horde/Storage/Nodes/DirectoryNode.Extract.cs b/Engine/Source/Programs/Shared/EpicGames.Horde/Storage/Nodes/DirectoryNode.Extract.cs
--- a/Engine/Source/Programs/Shared/EpicGames.Horde/Storage/Nodes/DirectoryNode.Extract.cs
+++ b/Engine/Source/Programs/Shared/EpicGames.Horde/Storage/Nodes/DirectoryNode.Extract.cs
@@ -172,14 +172,14 @@ namespace EpicGames.Horde.Storage.Nodes
 					{
 						if (FileInfo.Exists)
 						{
-							if ((FileInfo.Attributes & FileAttributes.ReadOnly) != 0)
-							{
-								FileInfo.Attributes &= ~FileAttributes.ReadOnly;
-							}
 							if (FileInfo.LinkTarget != null)
 							{
 								FileInfo.Delete();
 							}
+							else if (FileInfo.IsReadOnly)
+							{
+								FileInfo.IsReadOnly = false;
+							}
 						}
 						else
 						{
@@ -638,6 +638,7 @@ namespace EpicGames.Horde.Storage.Nodes
 			// Set correct permissions on the output file
 			if (remainingChunks == 0)
 			{
+				file.FileInfo.Refresh();
 				FileEntry.SetPermissions(file.FileInfo!, file.FileEntry.Flags);
 				copyStats?.Update(1, 0);
 			}
diff --git a/Engine/Source/Programs/Shared/EpicGames.Horde/Storage/Nodes/FileEntry.cs b/Engine/Source/Programs/Shared/EpicGames.Horde/Storage/Nodes/FileEntry.cs
--- a/Engine/Source/Programs/Shared/EpicGames.Horde/Storage/Nodes/FileEntry.cs
+++ b/Engine/Source/Programs/Shared/EpicGames.Horde/Storage/Nodes/FileEntry.cs
@@ -203,7 +203,7 @@ namespace EpicGames.Horde.Storage.Nodes
 			{
 				if ((flags & FileEntryFlags.ReadOnly) != 0)
 				{
-					fileInfo.Attributes |= FileAttributes.ReadOnly;
+					fileInfo.IsReadOnly = true;
 				}
 			}
 		}
