EpicGames.Perforce: Include the Status field when serializing change records. Updating submitted changelists fails without this.

diff --git a/Engine/Source/Programs/Shared/EpicGames.Perforce/PerforceConnection.cs b/Engine/Source/Programs/Shared/EpicGames.Perforce/PerforceConnection.cs
--- a/Engine/Source/Programs/Shared/EpicGames.Perforce/PerforceConnection.cs
+++ b/Engine/Source/Programs/Shared/EpicGames.Perforce/PerforceConnection.cs
@@ -897,6 +897,10 @@ namespace EpicGames.Perforce
 			{
 				nameToValue.Add(new KeyValuePair<string, object>("Description", input.Description));
 			}
+			if (input.Status != ChangeStatus.All)
+			{
+				nameToValue.Add(new KeyValuePair<string, object>("Status", PerforceReflection.GetEnumText(typeof(ChangeStatus), input.Status)));
+			}
 			if (input.Files.Count > 0)
 			{
 				nameToValue.Add(new KeyValuePair<string, object>("Files", input.Files));
