Horde: Fix null pointer dereference due to compute buffer being disposed before send task is terminated.

diff --git a/Engine/Source/Programs/Shared/EpicGames.Horde/Compute/ComputeSocket.cs b/Engine/Source/Programs/Shared/EpicGames.Horde/Compute/ComputeSocket.cs
--- a/Engine/Source/Programs/Shared/EpicGames.Horde/Compute/ComputeSocket.cs
+++ b/Engine/Source/Programs/Shared/EpicGames.Horde/Compute/ComputeSocket.cs
@@ -769,6 +769,9 @@ namespace EpicGames.Horde.Compute
 				sendBuffer.AddRef();
 			}
 
+			// Wait for the send task to complete
+			await sendBuffer.Task.DisposeAsync();
+
 			// Release the reader
 			await sendBuffer.Semaphore.WaitAsync(cancellationToken);
 			try
@@ -782,9 +785,6 @@ namespace EpicGames.Horde.Compute
 				sendBuffer.Release(); // Added above
 			}
 
-			// Wait for the send task to complete
-			await sendBuffer.Task.DisposeAsync();
-
 			// Remove the buffer from the dictionary
 			lock (_lockObject)
 			{
@@ -816,6 +816,10 @@ namespace EpicGames.Horde.Compute
 					await reader.WaitToReadAsync(1, cancellationToken);
 				}
 			}
+			catch (OperationCanceledException e)
+			{
+				_logger.LogDebug(e, "Background send task cancelled for channel {ChannelId}", channelId);
+			}
 			catch (Exception e)
 			{
 				_logger.LogInformation(e, "Error in background send");
