Horde: Remove category field from status response messages.

#rnx

diff --git a/Engine/Source/Programs/Shared/EpicGames.Horde/Server/ServerMessages.cs b/Engine/Source/Programs/Shared/EpicGames.Horde/Server/ServerMessages.cs
--- a/Engine/Source/Programs/Shared/EpicGames.Horde/Server/ServerMessages.cs
+++ b/Engine/Source/Programs/Shared/EpicGames.Horde/Server/ServerMessages.cs
@@ -165,11 +165,7 @@ namespace EpicGames.Horde.Server
 	/// </summary>
 	public class ServerStatusSubsystem
 	{
-		/// <summary>
-		/// Category of this subsystem
-		/// </summary>
-		public string Category { get; set; } = "";
-		
+	
 		/// <summary>
 		/// Name of the subsystem
 		/// </summary>
