Add telemetry service to EpicGames.Core

#rnx

diff --git a/Engine/Source/Programs/Shared/EpicGames.Core/EpicGames.Core.csproj b/Engine/Source/Programs/Shared/EpicGames.Core/EpicGames.Core.csproj
--- a/Engine/Source/Programs/Shared/EpicGames.Core/EpicGames.Core.csproj
+++ b/Engine/Source/Programs/Shared/EpicGames.Core/EpicGames.Core.csproj
@@ -52,8 +52,10 @@
   <ItemGroup>
     <PackageReference Include="JetBrains.Annotations" Version="2022.1.0" />
     <PackageReference Include="Microsoft.Extensions.Configuration" Version="8.0.0" />
+    <PackageReference Include="Microsoft.Extensions.Http.Polly" Version="6.0.26" />
     <PackageReference Include="Microsoft.Extensions.Logging" Version="6.0.0" />
     <PackageReference Include="Microsoft.Extensions.Logging.Abstractions" Version="6.0.1" />
+    <PackageReference Include="Microsoft.Extensions.Logging.Console" Version="6.0.0" />
     <PackageReference Include="Microsoft.Extensions.ObjectPool" Version="6.0.8" />
     <PackageReference Include="OpenTracing" Version="0.12.1" />
     <PackageReference Include="System.Memory" Version="4.5.5" />
diff --git a/Engine/Source/Programs/Shared/EpicGames.Core/Telemetry/DataRouterTelemetryService.cs b/Engine/Source/Programs/Shared/EpicGames.Core/Telemetry/DataRouterTelemetryService.cs
new file mode 100644
--- /dev/null
+++ b/Engine/Source/Programs/Shared/EpicGames.Core/Telemetry/DataRouterTelemetryService.cs
@@ -0,0 +1,405 @@
+// Copyright Epic Games, Inc. All Rights Reserved.
+
+using System;
+using System.Collections.Concurrent;
+using System.Collections.Generic;
+using System.Net.Http;
+using System.Net.Http.Headers;
+using System.Reflection;
+using System.Text.Json;
+using System.Text.Json.Serialization;
+using System.Threading.Tasks;
+using System.Timers;
+using Microsoft.Extensions.Logging;
+
+namespace EpicGames.Core.Telemetry
+{
+	public abstract class DataRouterEvent
+	{
+		protected string AppId { get; } = String.Empty;
+		protected string AppVersion { get; } = String.Empty;
+		protected string AppEnvironment { get; } = String.Empty;
+		protected string UploadType { get; } = String.Empty;
+		protected string UserId { get; } = String.Empty;
+		protected string SessionId { get; } = String.Empty;
+
+		// This is the TimeStamp of when the event was added to queue UNTIL we serialize
+		// so we use the name TimeStamp in code, and serialize as DateOffset
+		[JsonPropertyName("DateOffset")]
+		[JsonConverter(typeof(DateOffsetDataRouterEventConverter))]
+		public DateTime TimeStamp
+		{
+			get;
+			protected set;
+		}
+
+		public abstract string EventName
+		{
+			get;
+		}
+
+		protected DataRouterEvent(string inAppId, string inAppVersion, string inAppEnvironment, string inUploadType, string inUserId, string inSessionId, DateTime? inTimeStamp = null)
+		{
+			AppId = inAppId;
+			AppVersion = inAppVersion;
+			AppEnvironment = inAppEnvironment;
+			UploadType = inUploadType;
+			UserId = inUserId;
+			SessionId = inSessionId;
+			TimeStamp = inTimeStamp ?? DateTime.UtcNow;
+		}
+
+		[System.Diagnostics.CodeAnalysis.SuppressMessage("Design", "CA1055:URI-like return values should not be strings", Justification = "Legacy interoperability")]
+		public virtual string GetUrlParameters()
+		{
+			string urlParameters = $"AppID={AppId}&AppVersion={AppVersion}&AppEnvironment={AppEnvironment}&UploadType={UploadType}";
+			urlParameters += !String.IsNullOrEmpty(UserId) ? $"&UserID={UserId}" : String.Empty;
+			urlParameters += !String.IsNullOrEmpty(SessionId) ? $"&SessionID={SessionId}" : String.Empty;
+
+			return urlParameters;
+		}
+	}
+
+	public class DataRouterPost
+	{
+		// List of objects to force polymorphic JSON serialization
+		public List<object> Events
+		{
+			get;
+		} = new List<object>();
+	}
+
+	// Used in conjunction with a JsonSerializer class to mimic serialization of DataRouterPost
+	public class SerializedDataRouterPost
+	{
+		public List<string> EventsSerialized
+		{
+			get;
+		} = new List<string>();
+
+	}
+
+	/// <summary>
+	/// Implementation of JsonConverter for SerializedDataRouterPost. Matches output of serializing a DataRouterPost.
+	/// <see>https://confluence.it.epicgames.com/display/DPE/Data+Router</see>
+	/// </summary>
+	public class SerializedDataRouterPostConverter : JsonConverter<SerializedDataRouterPost>
+	{
+		public override SerializedDataRouterPost? Read(ref Utf8JsonReader reader, Type typeToConvert, JsonSerializerOptions options)
+		{
+			throw new NotImplementedException();
+		}
+
+		public override void Write(Utf8JsonWriter writer, SerializedDataRouterPost value, JsonSerializerOptions options)
+		{
+			writer.WriteStartObject();
+			writer.WriteStartArray("Events");
+			foreach (string serializedEvent in value.EventsSerialized)
+			{
+				writer.WriteRawValue(serializedEvent);
+			}
+			writer.WriteEndArray();
+			writer.WriteEndObject();
+		}
+	}
+
+	public class DateOffsetDataRouterEventConverter : JsonConverter<DateTime>
+	{
+		public override DateTime Read(ref Utf8JsonReader reader, Type typeToConvert, JsonSerializerOptions options)
+		{
+			throw new NotImplementedException();
+		}
+
+		public override void Write(Utf8JsonWriter writer, DateTime value, JsonSerializerOptions options)
+		{
+			TimeSpan offset = DateTime.UtcNow - value;
+
+			if (offset.Ticks <= 0)
+			{
+				offset = TimeSpan.Zero;
+			}
+			else if (offset.TotalDays > 1.0)
+			{
+				offset = new TimeSpan(23, 59, 59);
+			}
+
+			string offsetString = offset.ToString(@"hh\:mm\:ss\.fff");
+
+			writer.WriteStringValue($"+{offsetString}");
+		}
+	}
+
+	/// <summary>
+	/// Implementation for TelemetryService using DataRouter.
+	/// FlushEvents() should be called before application exits otherwise some events may remain unsent
+	/// <see>https://confluence.it.epicgames.com/display/DPE/Data+Router</see>
+	/// </summary>
+	public class DataRouterTelemetryService : ITelemetryService<DataRouterEvent>
+	{
+#if DEBUG
+		public bool IsDryRun
+		{
+			get;
+			set;
+		} = false;
+#endif
+
+		protected ILogger<DataRouterTelemetryService> Logger { get; }
+		protected HttpClient HttpClient { get; }
+		protected ConcurrentQueue<DataRouterEvent> EventsQueue { get; } = new ConcurrentQueue<DataRouterEvent>();
+
+		private int _autoFlushIntervalMilliseconds = 60 * 1000;
+		public int AutoFlushIntervalMilliseconds
+		{
+			get => _autoFlushIntervalMilliseconds;
+			set
+			{
+				_autoFlushIntervalMilliseconds = value;
+				if (_autoFlushTimer != null)
+				{
+					_autoFlushTimer.Stop();
+
+					if (_autoFlushIntervalMilliseconds <= 0)
+					{
+						Logger.LogWarning("Invalid value for AutoFlushIntervalMilliseconds {AutoFlushIntervalMs} <= 0 Ms", _autoFlushIntervalMilliseconds);
+						return;
+					}
+
+					_autoFlushTimer.Interval = _autoFlushIntervalMilliseconds;
+					_autoFlushTimer.Start();
+				}
+			}
+		}
+
+		// Max size of json body for sending telemetry events
+		private const int MaxEventSize = 2 * 1000 * 1000;
+
+		public string BaseAddress { get; set; }
+
+		private readonly Timer _autoFlushTimer = new Timer();
+
+		private static int s_dataRouterPostEmptyJsonLength = -1;
+		public static int DataRouterPostEmptyJsonLength
+		{
+			get
+			{
+				if (s_dataRouterPostEmptyJsonLength < 0)
+				{
+					DataRouterPost emptyPost = new DataRouterPost();
+					s_dataRouterPostEmptyJsonLength = JsonSerializer.Serialize(emptyPost).Length;
+				}
+				return s_dataRouterPostEmptyJsonLength;
+			}
+		}
+
+		private void SetHttpClientHeaders()
+		{
+			Assembly? entryAssembly = Assembly.GetEntryAssembly();
+			if (entryAssembly != null)
+			{
+				string productName = String.IsNullOrEmpty(entryAssembly.GetName().Name) ? GetType().Name : entryAssembly.GetName().Name!;
+				string productVersion = entryAssembly.GetName().Version == null ? "1.0.0.0" : entryAssembly.GetName().Version!.ToString();
+				ProductInfoHeaderValue userAgent = new ProductInfoHeaderValue(productName, productVersion);
+				HttpClient.DefaultRequestHeaders.UserAgent.Add(userAgent);
+			}
+			else
+			{
+				HttpClient.DefaultRequestHeaders.UserAgent.Add(new ProductInfoHeaderValue(GetType().Name, "1.0.0.0"));
+			}
+		}
+
+		public DataRouterTelemetryService(HttpClient inHttpClient, string inBaseAddress)
+		{
+			using (ILoggerFactory loggerFactory = LoggerFactory.Create(builder =>
+			{
+				builder
+				.ClearProviders()
+				.AddConsole();
+			}))
+			{
+				Logger = loggerFactory.CreateLogger<DataRouterTelemetryService>();
+			}
+
+			HttpClient = inHttpClient;
+			BaseAddress = inBaseAddress;
+			SetHttpClientHeaders();
+
+			_autoFlushTimer.Interval = AutoFlushIntervalMilliseconds;
+			_autoFlushTimer.Elapsed += AutoFlushEventsAsync;
+		}
+
+		private Task? _prevFlushTask = null;
+		private async void AutoFlushEventsAsync(object? sender, ElapsedEventArgs e)
+		{
+			if (_prevFlushTask != null && !_prevFlushTask.IsCompleted)
+			{
+				await _prevFlushTask;
+			}
+
+			_prevFlushTask = FlushEventsAsync();
+		}
+
+		public virtual void RecordEvent(DataRouterEvent eventData)
+		{
+			EventsQueue.Enqueue(eventData);
+		}
+
+		public void FlushEvents()
+		{
+			bool isAutoFlushEnabled = _autoFlushTimer.Enabled;
+			_autoFlushTimer.Enabled = false;
+
+			Task flushTask = Task.Run(async () => await FlushEventsAsync());
+			flushTask.Wait(millisecondsTimeout: _autoFlushIntervalMilliseconds);
+
+			_autoFlushTimer.Enabled = isAutoFlushEnabled;
+		}
+
+		// pass the URI in so we can be more informative with our error message
+		private void CreateBatchedSerializedDataRouterPosts(DataRouterPost inPost, string postURI, out List<SerializedDataRouterPost> outPostsList)
+		{
+			outPostsList = new List<SerializedDataRouterPost>();
+
+			SerializedDataRouterPost currentBatch = new SerializedDataRouterPost();
+
+			int currentBatchSize = DataRouterPostEmptyJsonLength;
+
+			// remember that we will have an extra comma for all but the last post in a batch
+			foreach (object @event in inPost.Events)
+			{
+				string serializedEvent = JsonSerializer.Serialize(@event);
+
+				int sizeOfEvent = serializedEvent.Length;
+
+				// skip events that are too big on their own
+				if (sizeOfEvent > MaxEventSize)
+				{
+					Logger.LogError("Unable to send DataRouter event with uri '{PostURI}' because event size was over the max of {MaxEventSize} bytes.", postURI, MaxEventSize);
+					continue;
+				}
+				// if we will go over, add current batch to list and start a new batch
+				else if (currentBatchSize + sizeOfEvent > MaxEventSize)
+				{
+					outPostsList.Add(currentBatch);
+
+					currentBatchSize = DataRouterPostEmptyJsonLength;
+
+					currentBatch = new SerializedDataRouterPost();
+				}
+				else
+				{
+					currentBatch.EventsSerialized.Add(serializedEvent);
+					// account for comma
+					currentBatchSize += sizeOfEvent + 1;
+				}
+			}
+			if (currentBatch.EventsSerialized.Count > 0)
+			{
+				outPostsList.Add(currentBatch);
+			}
+		}
+
+		private async Task<HttpResponseMessage> PostRequestAsync(string postDataKey, string jsonBody)
+		{
+			using HttpRequestMessage postRequest = new HttpRequestMessage(HttpMethod.Post, $"{BaseAddress}?{postDataKey}");
+			postRequest.Content = new StringContent(jsonBody, System.Text.Encoding.UTF8, "application/json");
+			return await HttpClient.SendAsync(postRequest);
+		}
+
+		public virtual async Task FlushEventsAsync()
+		{
+			ConcurrentDictionary<string, DataRouterPost> posts = new ConcurrentDictionary<string, DataRouterPost>();
+			foreach (DataRouterEvent eventData in EventsQueue)
+			{
+				string requestUri = eventData.GetUrlParameters();
+				if (!posts.ContainsKey(requestUri))
+				{
+					if (!posts.TryAdd(requestUri, new DataRouterPost()))
+					{
+						Logger.LogWarning("Failed to Flush RequestUri: {RequestUri}", requestUri);
+					}
+				}
+
+				posts[requestUri].Events.Add(eventData);
+			}
+
+			EventsQueue.Clear();
+
+#if DEBUG
+			if (IsDryRun)
+			{
+				foreach (KeyValuePair<string, DataRouterPost> postData in posts)
+				{
+					Logger.LogInformation("curl -H 'Content-Type: application/json' \\\n-H 'User-Agent:{UserAgent}' \\\n-X POST \\\n -d '{Data}' \\\n{Key}", HttpClient.DefaultRequestHeaders.UserAgent, JsonSerializer.Serialize(postData.Value, new JsonSerializerOptions { WriteIndented = true }), postData.Key);
+				}
+
+				return;
+			}
+#endif
+			JsonSerializerOptions serializeOptions = new JsonSerializerOptions
+			{
+				WriteIndented = false,
+				Converters =
+				{
+					new SerializedDataRouterPostConverter()
+				}
+			};
+
+			List<Task<HttpResponseMessage>> postRequests = new List<Task<HttpResponseMessage>>();
+			foreach (KeyValuePair<string, DataRouterPost> postData in posts)
+			{
+				// Break our post into multiple if the DataRouterPost serialized would go over our max event size
+				CreateBatchedSerializedDataRouterPosts(postData.Value, postData.Key, out List<SerializedDataRouterPost> serializedPosts);
+
+				foreach (SerializedDataRouterPost serializedPost in serializedPosts)
+				{
+					string jsonBody = JsonSerializer.Serialize(serializedPost, serializeOptions);
+
+					postRequests.Add(PostRequestAsync(postData.Key, jsonBody));
+				}
+			}
+
+			try
+			{
+				HttpResponseMessage[] responses = await Task.WhenAll(postRequests);
+				foreach (HttpResponseMessage? postResponse in responses)
+				{
+					Logger.LogDebug("{RequestUri} finished with status {StatusCode}", postResponse.RequestMessage?.RequestUri, postResponse.StatusCode);
+				}
+			}
+			catch (HttpRequestException ex)
+			{
+				Logger.LogError("{Exception}", ex.Message);
+			}
+			catch (Exception ex)
+			{
+				Logger.LogError("{Exception}", ex.ToString());
+			}
+		}
+
+		/// <inheritdoc/>
+		public void Dispose()
+		{
+			Dispose(true);
+			GC.SuppressFinalize(this);
+		}
+
+		/// <summary>
+		/// Standard Dispose pattern method
+		/// </summary>
+		/// <param name="disposing"></param>
+		protected virtual void Dispose(bool disposing)
+		{
+			_autoFlushTimer.Enabled = false;
+
+			// Make sure to flush any remaining events if the service is disposed of
+			FlushEvents();
+
+			_autoFlushTimer.Dispose();
+
+			if (disposing)
+			{
+			}
+		}
+	}
+}
diff --git a/Engine/Source/Programs/Shared/EpicGames.Core/Telemetry/DefaultTelemetryHttpClient.cs b/Engine/Source/Programs/Shared/EpicGames.Core/Telemetry/DefaultTelemetryHttpClient.cs
new file mode 100644
--- /dev/null
+++ b/Engine/Source/Programs/Shared/EpicGames.Core/Telemetry/DefaultTelemetryHttpClient.cs
@@ -0,0 +1,41 @@
+// Copyright Epic Games, Inc. All Rights Reserved.
+
+using Microsoft.Extensions.Http;
+using Polly;
+using Polly.Extensions.Http;
+using System;
+using System.Net.Http;
+
+namespace EpicGames.Core.Telemetry
+{
+	public static class HttpClientDefaults
+	{
+		// When services can be used, utilize this helper to get this policy that handles retries
+		//ex. ServicesCollection.AddHttpClient<TelemetryService>().AddPolicyHandler(HttpClientDefaults.GetPolicy());
+		public static Polly.Retry.AsyncRetryPolicy<HttpResponseMessage> GetRetryPolicy()
+		{
+			return HttpPolicyExtensions
+				// 408, 5XX responses 
+				.HandleTransientHttpError()
+				// 504 should have been covered by HandleTransientHttpError but adding here to be safe
+				.OrResult(msg => msg.StatusCode == System.Net.HttpStatusCode.NotFound ||
+								 msg.StatusCode == System.Net.HttpStatusCode.GatewayTimeout)
+				.WaitAndRetryAsync(3, retryAttempt => TimeSpan.FromSeconds(Math.Pow(2, retryAttempt)));
+		}
+
+		// Returns new HttpClient configured with our default retry policy
+		public static HttpClient GetClient()
+		{
+			// per MS documentation, with a lifetime specified we shouldn't have socket exhaustion issues
+			SocketsHttpHandler socketHandler = new SocketsHttpHandler { PooledConnectionLifetime = TimeSpan.FromMinutes(15) };
+#pragma warning disable CA2000 // Dispose objects before losing scope
+			PolicyHttpMessageHandler pollyHandler = new PolicyHttpMessageHandler(GetRetryPolicy())
+			{
+				InnerHandler = socketHandler,
+			};
+#pragma warning restore CA2000 // Dispose objects before losing scope
+
+			return new HttpClient(pollyHandler);
+		}
+	}
+}
\ No newline at end of file
diff --git a/Engine/Source/Programs/Shared/EpicGames.Core/Telemetry/TelemetryService.cs b/Engine/Source/Programs/Shared/EpicGames.Core/Telemetry/TelemetryService.cs
new file mode 100644
--- /dev/null
+++ b/Engine/Source/Programs/Shared/EpicGames.Core/Telemetry/TelemetryService.cs
@@ -0,0 +1,21 @@
+// Copyright Epic Games, Inc. All Rights Reserved.
+
+using System;
+
+namespace EpicGames.Core.Telemetry
+{
+	/// <summary>
+	/// Base interface for a TelemetryService.
+	/// Derived from IDisposable as we want to make sure implementation handles flushing
+	/// remaining events in the case that the service is disposed of. This usually would
+	/// happen in a using block with an instance of the ITelemetryService derived class.
+	/// If registering as part of a servicescollection, Flush should still be called in shutdown code. 
+	/// </summary>
+	/// <typeparam name="T">Type of event this service should record</typeparam>
+	public interface ITelemetryService<T> : IDisposable where T : class
+	{
+		public void RecordEvent(T eventData);
+
+		public void FlushEvents();
+	}
+}
