EpicGames.Core: Add missing xml comments

#rnx

diff --git a/Engine/Source/Programs/Shared/EpicGames.Core/Telemetry/DataRouterTelemetryService.cs b/Engine/Source/Programs/Shared/EpicGames.Core/Telemetry/DataRouterTelemetryService.cs
--- a/Engine/Source/Programs/Shared/EpicGames.Core/Telemetry/DataRouterTelemetryService.cs
+++ b/Engine/Source/Programs/Shared/EpicGames.Core/Telemetry/DataRouterTelemetryService.cs
@@ -14,17 +14,44 @@ using Microsoft.Extensions.Logging;
 
 namespace EpicGames.Core.Telemetry
 {
+	/// <summary>
+	/// Abstract base class for a data router event
+	/// </summary>
 	public abstract class DataRouterEvent
 	{
+		/// <summary>
+		/// The app's id
+		/// </summary>
 		protected string AppId { get; } = String.Empty;
+
+		/// <summary>
+		/// The app's version
+		/// </summary>
 		protected string AppVersion { get; } = String.Empty;
+
+		/// <summary>
+		/// The app's environment
+		/// </summary>
 		protected string AppEnvironment { get; } = String.Empty;
+
+		/// <summary>
+		/// The upload type
+		/// </summary>
 		protected string UploadType { get; } = String.Empty;
+
+		/// <summary>
+		///  The user's id
+		/// </summary>
 		protected string UserId { get; } = String.Empty;
+
+		/// <summary>
+		/// The session's id
+		/// </summary>
 		protected string SessionId { get; } = String.Empty;
 
-		// This is the TimeStamp of when the event was added to queue UNTIL we serialize
-		// so we use the name TimeStamp in code, and serialize as DateOffset
+		/// <summary>
+		/// The TimeStamp of when the event was added to queue UNTIL we serialize so we use the name TimeStamp in code, and serialize as DateOffset
+		/// </summary>
 		[JsonPropertyName("DateOffset")]
 		[JsonConverter(typeof(DateOffsetDataRouterEventConverter))]
 		public DateTime TimeStamp
@@ -33,11 +60,24 @@ namespace EpicGames.Core.Telemetry
 			protected set;
 		}
 
+		/// <summary>
+		/// The name of the event
+		/// </summary>
 		public abstract string EventName
 		{
 			get;
 		}
 
+		/// <summary>
+		/// Constructor
+		/// </summary>
+		/// <param name="inAppId"></param>
+		/// <param name="inAppVersion"></param>
+		/// <param name="inAppEnvironment"></param>
+		/// <param name="inUploadType"></param>
+		/// <param name="inUserId"></param>
+		/// <param name="inSessionId"></param>
+		/// <param name="inTimeStamp"></param>
 		protected DataRouterEvent(string inAppId, string inAppVersion, string inAppEnvironment, string inUploadType, string inUserId, string inSessionId, DateTime? inTimeStamp = null)
 		{
 			AppId = inAppId;
@@ -49,6 +89,10 @@ namespace EpicGames.Core.Telemetry
 			TimeStamp = inTimeStamp ?? DateTime.UtcNow;
 		}
 
+		/// <summary>
+		/// Generates url parameter strings for this event
+		/// </summary>
+		/// <returns>The generated url parameters string</returns>
 		[System.Diagnostics.CodeAnalysis.SuppressMessage("Design", "CA1055:URI-like return values should not be strings", Justification = "Legacy interoperability")]
 		public virtual string GetUrlParameters()
 		{
@@ -60,6 +104,9 @@ namespace EpicGames.Core.Telemetry
 		}
 	}
 
+	/// <summary>
+	/// Object to post data router events
+	/// </summary>
 	public class DataRouterPost
 	{
 		// List of objects to force polymorphic JSON serialization
@@ -69,9 +116,12 @@ namespace EpicGames.Core.Telemetry
 		} = new List<object>();
 	}
 
-	// Used in conjunction with a JsonSerializer class to mimic serialization of DataRouterPost
+	/// <summary>
+	/// Object to serialize data router post events
+	/// </summary>
 	public class SerializedDataRouterPost
 	{
+		// Used in conjunction with a JsonSerializer class to mimic serialization of DataRouterPost
 		public List<string> EventsSerialized
 		{
 			get;
@@ -85,11 +135,25 @@ namespace EpicGames.Core.Telemetry
 	/// </summary>
 	public class SerializedDataRouterPostConverter : JsonConverter<SerializedDataRouterPost>
 	{
+		/// <summary>
+		/// Reader
+		/// </summary>
+		/// <param name="reader"></param>
+		/// <param name="typeToConvert"></param>
+		/// <param name="options"></param>
+		/// <returns></returns>
+		/// <exception cref="NotImplementedException"></exception>
 		public override SerializedDataRouterPost? Read(ref Utf8JsonReader reader, Type typeToConvert, JsonSerializerOptions options)
 		{
 			throw new NotImplementedException();
 		}
 
+		/// <summary>
+		/// Writer
+		/// </summary>
+		/// <param name="writer"></param>
+		/// <param name="value"></param>
+		/// <param name="options"></param>
 		public override void Write(Utf8JsonWriter writer, SerializedDataRouterPost value, JsonSerializerOptions options)
 		{
 			writer.WriteStartObject();
@@ -103,13 +167,30 @@ namespace EpicGames.Core.Telemetry
 		}
 	}
 
+	/// <summary>
+	/// DateTime converter for data router events
+	/// </summary>
 	public class DateOffsetDataRouterEventConverter : JsonConverter<DateTime>
 	{
+		/// <summary>
+		/// Reader
+		/// </summary>
+		/// <param name="reader"></param>
+		/// <param name="typeToConvert"></param>
+		/// <param name="options"></param>
+		/// <returns></returns>
+		/// <exception cref="NotImplementedException"></exception>
 		public override DateTime Read(ref Utf8JsonReader reader, Type typeToConvert, JsonSerializerOptions options)
 		{
 			throw new NotImplementedException();
 		}
 
+		/// <summary>
+		/// Writer
+		/// </summary>
+		/// <param name="writer"></param>
+		/// <param name="value"></param>
+		/// <param name="options"></param>
 		public override void Write(Utf8JsonWriter writer, DateTime value, JsonSerializerOptions options)
 		{
 			TimeSpan offset = DateTime.UtcNow - value;
@@ -137,6 +218,9 @@ namespace EpicGames.Core.Telemetry
 	public class DataRouterTelemetryService : ITelemetryService<DataRouterEvent>
 	{
 #if DEBUG
+		/// <summary>
+		/// Denotes if this is a dry run
+		/// </summary>
 		public bool IsDryRun
 		{
 			get;
@@ -144,11 +228,24 @@ namespace EpicGames.Core.Telemetry
 		} = false;
 #endif
 
+		/// <summary>
+		/// Logging interface
+		/// </summary>
 		protected ILogger<DataRouterTelemetryService> Logger { get; }
+
+		/// <summary>
+		/// The http client to use for connections
+		/// </summary>
 		protected HttpClient HttpClient { get; }
+
+		/// <summary>
+		/// The queue of pending data router events
+		/// </summary>
 		protected ConcurrentQueue<DataRouterEvent> EventsQueue { get; } = new ConcurrentQueue<DataRouterEvent>();
 
-		private int _autoFlushIntervalMilliseconds = 60 * 1000;
+		/// <summary>
+		/// Interval to flush events, in milliseconds
+		/// </summary>
 		public int AutoFlushIntervalMilliseconds
 		{
 			get => _autoFlushIntervalMilliseconds;
@@ -170,15 +267,23 @@ namespace EpicGames.Core.Telemetry
 				}
 			}
 		}
+		private int _autoFlushIntervalMilliseconds = 60 * 1000;
 
-		// Max size of json body for sending telemetry events
+		/// <summary>
+		/// Max size of json body for sending telemetry events
+		/// </summary>
 		private const int MaxEventSize = 2 * 1000 * 1000;
 
+		/// <summary>
+		/// Base url address to post events to
+		/// </summary>
 		public string BaseAddress { get; set; }
 
 		private readonly Timer _autoFlushTimer = new Timer();
 
-		private static int s_dataRouterPostEmptyJsonLength = -1;
+		/// <summary>
+		/// Length of an empty data router post when serialized to json
+		/// </summary>
 		public static int DataRouterPostEmptyJsonLength
 		{
 			get
@@ -191,7 +296,11 @@ namespace EpicGames.Core.Telemetry
 				return s_dataRouterPostEmptyJsonLength;
 			}
 		}
+		private static int s_dataRouterPostEmptyJsonLength = -1;
 
+		/// <summary>
+		/// Sets up the default request headers for the http client
+		/// </summary>
 		private void SetHttpClientHeaders()
 		{
 			Assembly? entryAssembly = Assembly.GetEntryAssembly();
@@ -208,6 +317,11 @@ namespace EpicGames.Core.Telemetry
 			}
 		}
 
+		/// <summary>
+		/// Constructor
+		/// </summary>
+		/// <param name="inHttpClient"></param>
+		/// <param name="inBaseAddress"></param>
 		public DataRouterTelemetryService(HttpClient inHttpClient, string inBaseAddress)
 		{
 			using (ILoggerFactory loggerFactory = LoggerFactory.Create(builder =>
@@ -228,7 +342,11 @@ namespace EpicGames.Core.Telemetry
 			_autoFlushTimer.Elapsed += AutoFlushEventsAsync;
 		}
 
-		private Task? _prevFlushTask = null;
+		/// <summary>
+		/// Async function to automatically flush events
+		/// </summary>
+		/// <param name="sender"></param>
+		/// <param name="e"></param>
 		private async void AutoFlushEventsAsync(object? sender, ElapsedEventArgs e)
 		{
 			if (_prevFlushTask != null && !_prevFlushTask.IsCompleted)
@@ -238,12 +356,15 @@ namespace EpicGames.Core.Telemetry
 
 			_prevFlushTask = FlushEventsAsync();
 		}
+		private Task? _prevFlushTask = null;
 
+		/// <inheritdoc/>
 		public virtual void RecordEvent(DataRouterEvent eventData)
 		{
 			EventsQueue.Enqueue(eventData);
 		}
 
+		/// <inheritdoc/>
 		public void FlushEvents()
 		{
 			bool isAutoFlushEnabled = _autoFlushTimer.Enabled;
@@ -255,7 +376,12 @@ namespace EpicGames.Core.Telemetry
 			_autoFlushTimer.Enabled = isAutoFlushEnabled;
 		}
 
-		// pass the URI in so we can be more informative with our error message
+		/// <summary>
+		/// Creates a batch of post events
+		/// </summary>
+		/// <param name="inPost"></param>
+		/// <param name="postURI"></param>
+		/// <param name="outPostsList"></param>
 		private void CreateBatchedSerializedDataRouterPosts(DataRouterPost inPost, string postURI, out List<SerializedDataRouterPost> outPostsList)
 		{
 			outPostsList = new List<SerializedDataRouterPost>();
@@ -299,6 +425,12 @@ namespace EpicGames.Core.Telemetry
 			}
 		}
 
+		/// <summary>
+		/// Async function to post a single request
+		/// </summary>
+		/// <param name="postDataKey"></param>
+		/// <param name="jsonBody"></param>
+		/// <returns>Http responce</returns>
 		private async Task<HttpResponseMessage> PostRequestAsync(string postDataKey, string jsonBody)
 		{
 			using HttpRequestMessage postRequest = new HttpRequestMessage(HttpMethod.Post, $"{BaseAddress}?{postDataKey}");
@@ -306,6 +438,10 @@ namespace EpicGames.Core.Telemetry
 			return await HttpClient.SendAsync(postRequest);
 		}
 
+		/// <summary>
+		/// Async function to flush events
+		/// </summary>
+		/// <returns></returns>
 		public virtual async Task FlushEventsAsync()
 		{
 			ConcurrentDictionary<string, DataRouterPost> posts = new ConcurrentDictionary<string, DataRouterPost>();
diff --git a/Engine/Source/Programs/Shared/EpicGames.Core/Telemetry/DefaultTelemetryHttpClient.cs b/Engine/Source/Programs/Shared/EpicGames.Core/Telemetry/DefaultTelemetryHttpClient.cs
--- a/Engine/Source/Programs/Shared/EpicGames.Core/Telemetry/DefaultTelemetryHttpClient.cs
+++ b/Engine/Source/Programs/Shared/EpicGames.Core/Telemetry/DefaultTelemetryHttpClient.cs
@@ -10,8 +10,11 @@ namespace EpicGames.Core.Telemetry
 {
 	public static class HttpClientDefaults
 	{
-		// When services can be used, utilize this helper to get this policy that handles retries
-		//ex. ServicesCollection.AddHttpClient<TelemetryService>().AddPolicyHandler(HttpClientDefaults.GetPolicy());
+		/// <summary>
+		/// When services can be used, utilize this helper to get this policy that handles retries
+		/// ex. ServicesCollection.AddHttpClient<TelemetryService>().AddPolicyHandler(HttpClientDefaults.GetPolicy());
+		/// </summary>
+		/// <returns>The retry policy</returns>
 		public static Polly.Retry.AsyncRetryPolicy<HttpResponseMessage> GetRetryPolicy()
 		{
 			return HttpPolicyExtensions
@@ -23,7 +26,10 @@ namespace EpicGames.Core.Telemetry
 				.WaitAndRetryAsync(3, retryAttempt => TimeSpan.FromSeconds(Math.Pow(2, retryAttempt)));
 		}
 
-		// Returns new HttpClient configured with our default retry policy
+		/// <summary>
+		/// Creates a HttpClient with the default retry policy
+		/// </summary>
+		/// <returns>HttpClient</returns>
 		public static HttpClient GetClient()
 		{
 			// per MS documentation, with a lifetime specified we shouldn't have socket exhaustion issues
diff --git a/Engine/Source/Programs/Shared/EpicGames.Core/Telemetry/TelemetryService.cs b/Engine/Source/Programs/Shared/EpicGames.Core/Telemetry/TelemetryService.cs
--- a/Engine/Source/Programs/Shared/EpicGames.Core/Telemetry/TelemetryService.cs
+++ b/Engine/Source/Programs/Shared/EpicGames.Core/Telemetry/TelemetryService.cs
@@ -14,8 +14,15 @@ namespace EpicGames.Core.Telemetry
 	/// <typeparam name="T">Type of event this service should record</typeparam>
 	public interface ITelemetryService<T> : IDisposable where T : class
 	{
+		/// <summary>
+		/// Record an event to be processed
+		/// </summary>
+		/// <param name="eventData">The event to record</param>
 		public void RecordEvent(T eventData);
 
+		/// <summary>
+		/// Function to flush pending events
+		/// </summary>
 		public void FlushEvents();
 	}
 }
