EpicGames.Core: Fix static analysis warnings.

diff --git a/Engine/Source/Programs/Shared/EpicGames.Core/Telemetry/DataRouterTelemetryService.cs b/Engine/Source/Programs/Shared/EpicGames.Core/Telemetry/DataRouterTelemetryService.cs
--- a/Engine/Source/Programs/Shared/EpicGames.Core/Telemetry/DataRouterTelemetryService.cs
+++ b/Engine/Source/Programs/Shared/EpicGames.Core/Telemetry/DataRouterTelemetryService.cs
@@ -109,7 +109,9 @@ namespace EpicGames.Core.Telemetry
 	/// </summary>
 	public class DataRouterPost
 	{
-		// List of objects to force polymorphic JSON serialization
+		/// <summary>
+		/// List of objects to force polymorphic JSON serialization
+		/// </summary>
 		public List<object> Events
 		{
 			get;
@@ -121,7 +123,9 @@ namespace EpicGames.Core.Telemetry
 	/// </summary>
 	public class SerializedDataRouterPost
 	{
-		// Used in conjunction with a JsonSerializer class to mimic serialization of DataRouterPost
+		/// <summary>
+		/// Used in conjunction with a JsonSerializer class to mimic serialization of DataRouterPost
+		/// </summary>
 		public List<string> EventsSerialized
 		{
 			get;
diff --git a/Engine/Source/Programs/Shared/EpicGames.Core/Telemetry/DefaultTelemetryHttpClient.cs b/Engine/Source/Programs/Shared/EpicGames.Core/Telemetry/DefaultTelemetryHttpClient.cs
--- a/Engine/Source/Programs/Shared/EpicGames.Core/Telemetry/DefaultTelemetryHttpClient.cs
+++ b/Engine/Source/Programs/Shared/EpicGames.Core/Telemetry/DefaultTelemetryHttpClient.cs
@@ -8,11 +8,14 @@ using System.Net.Http;
 
 namespace EpicGames.Core.Telemetry
 {
+	/// <summary>
+	/// 
+	/// </summary>
 	public static class HttpClientDefaults
 	{
 		/// <summary>
 		/// When services can be used, utilize this helper to get this policy that handles retries
-		/// ex. ServicesCollection.AddHttpClient<TelemetryService>().AddPolicyHandler(HttpClientDefaults.GetPolicy());
+		/// ex. ServicesCollection.AddHttpClient&lt;TelemetryService&gt;().AddPolicyHandler(HttpClientDefaults.GetPolicy());
 		/// </summary>
 		/// <returns>The retry policy</returns>
 		public static Polly.Retry.AsyncRetryPolicy<HttpResponseMessage> GetRetryPolicy()
