Fix some malformed xml comments

#rnx

diff --git a/Engine/Source/Programs/Shared/EpicGames.Core/Telemetry/DefaultTelemetryHttpClient.cs b/Engine/Source/Programs/Shared/EpicGames.Core/Telemetry/DefaultTelemetryHttpClient.cs
--- a/Engine/Source/Programs/Shared/EpicGames.Core/Telemetry/DefaultTelemetryHttpClient.cs
+++ b/Engine/Source/Programs/Shared/EpicGames.Core/Telemetry/DefaultTelemetryHttpClient.cs
@@ -9,7 +9,7 @@ using System.Net.Http;
 namespace EpicGames.Core.Telemetry
 {
 	/// <summary>
-	/// 
+	/// Static class to set http client defaults
 	/// </summary>
 	public static class HttpClientDefaults
 	{
