Telemetry: Update DataRouterTelemetryService to take a Uri instead of string

diff --git a/Engine/Source/Programs/Shared/EpicGames.Core/Telemetry/DataRouterTelemetryService.cs b/Engine/Source/Programs/Shared/EpicGames.Core/Telemetry/DataRouterTelemetryService.cs
--- a/Engine/Source/Programs/Shared/EpicGames.Core/Telemetry/DataRouterTelemetryService.cs
+++ b/Engine/Source/Programs/Shared/EpicGames.Core/Telemetry/DataRouterTelemetryService.cs
@@ -281,7 +281,7 @@ namespace EpicGames.Core.Telemetry
 		/// <summary>
 		/// Base url address to post events to
 		/// </summary>
-		public string BaseAddress { get; set; }
+		public Uri BaseAddress { get; set; }
 
 		private readonly Timer _autoFlushTimer = new Timer();
 
@@ -326,7 +326,7 @@ namespace EpicGames.Core.Telemetry
 		/// </summary>
 		/// <param name="inHttpClient"></param>
 		/// <param name="inBaseAddress"></param>
-		public DataRouterTelemetryService(HttpClient inHttpClient, string inBaseAddress)
+		public DataRouterTelemetryService(HttpClient inHttpClient, Uri inBaseAddress)
 		{
 			using (ILoggerFactory loggerFactory = LoggerFactory.Create(builder =>
 			{
