Fixed UHT issue where code generation would fail to generate valid code if the header file had multiple '.'s in the file name.

diff --git a/Engine/Source/Programs/Shared/EpicGames.UHT/Exporters/CodeGen/UhtHeaderCodeGeneratorHFile.cs b/Engine/Source/Programs/Shared/EpicGames.UHT/Exporters/CodeGen/UhtHeaderCodeGeneratorHFile.cs
--- a/Engine/Source/Programs/Shared/EpicGames.UHT/Exporters/CodeGen/UhtHeaderCodeGeneratorHFile.cs
+++ b/Engine/Source/Programs/Shared/EpicGames.UHT/Exporters/CodeGen/UhtHeaderCodeGeneratorHFile.cs
@@ -70,7 +70,7 @@ namespace EpicGames.UHT.Exporters.CodeGen
 				builder.Append(DisableDeprecationWarnings).Append("\r\n");
 
 				string strippedName = Path.GetFileNameWithoutExtension(HeaderFile.FilePath);
-				string defineName = $"{Package.ShortName.ToString().ToUpper()}_{strippedName}_generated_h";
+				string defineName = $"{Package.ShortName.ToString().ToUpper()}_{strippedName}_generated_h".Replace('.', '_');
 
 				if (HeaderFile.References.ForwardDeclarations.Count > 0)
 				{
