Horde: Cap maximum length of logged Perforce command line to 512 characters.

#rnx

diff --git a/Engine/Source/Programs/Shared/EpicGames.Perforce/NativePerforceConnection.cs b/Engine/Source/Programs/Shared/EpicGames.Perforce/NativePerforceConnection.cs
--- a/Engine/Source/Programs/Shared/EpicGames.Perforce/NativePerforceConnection.cs
+++ b/Engine/Source/Programs/Shared/EpicGames.Perforce/NativePerforceConnection.cs
@@ -563,8 +563,19 @@ namespace EpicGames.Perforce
 
 		private void ExecCommand(string command, List<string> args, byte[]? inputData, string? promptResponse, bool interceptIo)
 		{
+			StringBuilder argList = new StringBuilder();
+			for (int idx = 0; idx < args.Count; idx++)
+			{
+				CommandLineArguments.Append(argList, args[idx]);
+				if (argList.Length > 512)
+				{
+					argList.Append($" {{+{args.Count - idx - 1} more}}");
+					break;
+				}
+			}
+
 			Stopwatch timer = Stopwatch.StartNew();
-			Logger.LogTrace("Conn {ConnectionId}: {Command} {Args}", _uniqueId, command, String.Join(" ", args));
+			Logger.LogTrace("Conn {ConnectionId}: {Command} {Args}", _uniqueId, command, argList.ToString());
 
 			List<IntPtr> nativeArgs = new List<IntPtr>();
 			try
