Horde: Explicitly disable parallel syncing for managed workspaces. This introduces a dependency on p4.exe being installed, which we don't want.

#rnx

diff --git a/Engine/Source/Programs/Shared/EpicGames.Perforce.Managed/ManagedWorkspace.cs b/Engine/Source/Programs/Shared/EpicGames.Perforce.Managed/ManagedWorkspace.cs
--- a/Engine/Source/Programs/Shared/EpicGames.Perforce.Managed/ManagedWorkspace.cs
+++ b/Engine/Source/Programs/Shared/EpicGames.Perforce.Managed/ManagedWorkspace.cs
@@ -2082,15 +2082,19 @@ namespace EpicGames.Perforce.Managed
 					files.Add($"{filesToSync[idx]._streamFile.Path}#{filesToSync[idx]._streamFile.Revision}");
 				}
 
-				using IPerforceConnection threadedClient = await PerforceConnection.CreateAsync(client.Settings, client.Logger);
+				SyncOptions options = SyncOptions.FullDepotSyntax;
 				if (_options.UseHaveTable)
 				{
-					await threadedClient.SyncAsync(SyncOptions.Force | SyncOptions.FullDepotSyntax, -1, files, cancellationToken).ToListAsync(cancellationToken);
+					options |= SyncOptions.Force;
 				}
 				else
 				{
-					await threadedClient.SyncAsync(SyncOptions.DoNotUpdateHaveList | SyncOptions.FullDepotSyntax, -1, files, cancellationToken).ToListAsync(cancellationToken);
+					options |= SyncOptions.DoNotUpdateHaveList;
 				}
+
+				// Note: Explicitly disable parallel syncing here; the P4 API attempts to shell out to p4.exe, which may not be installed.
+				using IPerforceConnection threadedClient = await PerforceConnection.CreateAsync(client.Settings, client.Logger);
+				await threadedClient.SyncAsync(options, -1, 0, -1, -1, -1, -1, files, cancellationToken).ToListAsync(cancellationToken);
 			}
 			else
 			{
