Horde: Improve error message when a parsed log event is malformed.

#rnx

diff --git a/Engine/Source/Programs/Shared/EpicGames.Core/LogEventParser.cs b/Engine/Source/Programs/Shared/EpicGames.Core/LogEventParser.cs
--- a/Engine/Source/Programs/Shared/EpicGames.Core/LogEventParser.cs
+++ b/Engine/Source/Programs/Shared/EpicGames.Core/LogEventParser.cs
@@ -333,7 +333,7 @@ namespace EpicGames.Core
 				}
 				catch (Exception ex)
 				{
-					_logger.LogError(ex, "Exception while parsing log event");
+					_logger.LogError(ex, "Exception while parsing log event: {Message}", ex.Message);
 				}
 			}
 
