Horde: Move IssueMessages into EpicGames.Horde.

#rnx

diff --git a/Engine/Source/Programs/Horde/Horde.Server/Issues/IIssue.cs b/Engine/Source/Programs/Horde/Horde.Server/Issues/IIssue.cs
--- a/Engine/Source/Programs/Horde/Horde.Server/Issues/IIssue.cs
+++ b/Engine/Source/Programs/Horde/Horde.Server/Issues/IIssue.cs
@@ -6,30 +6,10 @@ using System.Linq;
 using MongoDB.Bson;
 using EpicGames.Horde.Streams;
 using EpicGames.Horde.Users;
+using EpicGames.Horde.Issues;
 
 namespace Horde.Server.Issues
 {
-	/// <summary>
-	/// The severity of an issue
-	/// </summary>
-	public enum IssueSeverity
-	{
-		/// <summary>
-		/// Unspecified severity
-		/// </summary>
-		Unspecified,
-
-		/// <summary>
-		/// This error represents a warning
-		/// </summary>
-		Warning,
-
-		/// <summary>
-		/// This issue represents an error
-		/// </summary>
-		Error,
-	}
-
 	/// <summary>
 	/// Stores information about a build health issue
 	/// </summary>
diff --git a/Engine/Source/Programs/Horde/Horde.Server/Issues/IIssueStep.cs b/Engine/Source/Programs/Horde/Horde.Server/Issues/IIssueStep.cs
--- a/Engine/Source/Programs/Horde/Horde.Server/Issues/IIssueStep.cs
+++ b/Engine/Source/Programs/Horde/Horde.Server/Issues/IIssueStep.cs
@@ -1,6 +1,7 @@
 // Copyright Epic Games, Inc. All Rights Reserved.
 
 using System;
+using EpicGames.Horde.Issues;
 using EpicGames.Horde.Jobs;
 using EpicGames.Horde.Logs;
 using Horde.Server.Jobs.Graphs;
diff --git a/Engine/Source/Programs/Horde/Horde.Server/Issues/IssueMessages.cs b/Engine/Source/Programs/Horde/Horde.Server/Issues/IssueMessages.cs
deleted file mode 100644
diff --git a/Engine/Source/Programs/Horde/Horde.Server/Issues/IssueReportService.cs b/Engine/Source/Programs/Horde/Horde.Server/Issues/IssueReportService.cs
--- a/Engine/Source/Programs/Horde/Horde.Server/Issues/IssueReportService.cs
+++ b/Engine/Source/Programs/Horde/Horde.Server/Issues/IssueReportService.cs
@@ -6,6 +6,7 @@ using System.Linq;
 using System.Threading;
 using System.Threading.Tasks;
 using EpicGames.Core;
+using EpicGames.Horde.Issues;
 using EpicGames.Horde.Jobs.Templates;
 using EpicGames.Horde.Streams;
 using Horde.Server.Jobs;
diff --git a/Engine/Source/Programs/Horde/Horde.Server/Issues/IssuesController.cs b/Engine/Source/Programs/Horde/Horde.Server/Issues/IssuesController.cs
--- a/Engine/Source/Programs/Horde/Horde.Server/Issues/IssuesController.cs
+++ b/Engine/Source/Programs/Horde/Horde.Server/Issues/IssuesController.cs
@@ -5,7 +5,10 @@ using System.Collections.Generic;
 using System.Linq;
 using System.Threading;
 using System.Threading.Tasks;
+using EpicGames.Core;
+using EpicGames.Horde.Issues;
 using EpicGames.Horde.Jobs;
+using EpicGames.Horde.Jobs.Templates;
 using EpicGames.Horde.Logs;
 using EpicGames.Horde.Streams;
 using EpicGames.Horde.Users;
@@ -140,7 +143,7 @@ namespace Horde.Server.Issues
 							// Convert each issue to a response
 							foreach (IIssueSpan span in spansForIssue)
 							{
-								spanResponses.Add(new FindIssueSpanResponse(span, span.LastFailure.Annotations.WorkflowId));
+								spanResponses.Add(NewFindIssueSpanResponse(span, span.LastFailure.Annotations.WorkflowId));
 							}
 
 							// Find which workflows are affected
@@ -179,7 +182,7 @@ namespace Horde.Server.Issues
 							quarantinedBy = await _userCollection.GetCachedUserAsync(issue.QuarantinedByUserId.Value, cancellationToken);
 						}
 
-						FindIssueResponse response = new FindIssueResponse(issue, owner, nominatedBy, resolvedBy, quarantinedBy, streamSeverity, spanResponses, openWorkflowIds.ToList());
+						FindIssueResponse response = NewFindIssueResponse(issue, owner, nominatedBy, resolvedBy, quarantinedBy, streamSeverity, spanResponses, openWorkflowIds.ToList());
 						responses.Add(PropertyFilter.Apply(response, filter));
 					}
 				}
@@ -192,6 +195,89 @@ namespace Horde.Server.Issues
 			return responses;
 		}
 
+		static FindIssueSpanResponse NewFindIssueSpanResponse(IIssueSpan span, WorkflowId? workflowId)
+		{
+			FindIssueSpanResponse response = new FindIssueSpanResponse();
+			response.Id = span.Id.ToString();
+			response.TemplateId = span.TemplateRefId.ToString();
+			response.Name = span.NodeName;
+			response.WorkflowId = workflowId;
+			if (span.LastSuccess != null)
+			{
+				response.LastSuccess = NewGetIssueStepResponse(span.LastSuccess);
+			}
+			if (span.NextSuccess != null)
+			{
+				response.NextSuccess = NewGetIssueStepResponse(span.NextSuccess);
+			}
+			return response;
+		}
+
+		static GetIssueStepResponse NewGetIssueStepResponse(IIssueStep issueStep)
+		{
+			GetIssueStepResponse response = new GetIssueStepResponse();
+			response.Change = issueStep.Change;
+			response.Severity = issueStep.Severity;
+			response.JobName = issueStep.JobName;
+			response.JobId = issueStep.JobId;
+			response.BatchId = issueStep.BatchId;
+			response.StepId = issueStep.StepId;
+			response.StepTime = issueStep.StepTime;
+			response.LogId = issueStep.LogId;
+			return response;
+		}
+
+		/// <summary>
+		/// Constructs a new issue
+		/// </summary>
+		/// <param name="issue">The isseu information</param>
+		/// <param name="owner">Owner of the issue</param>
+		/// <param name="nominatedBy">User that nominated the current fixer</param>
+		/// <param name="resolvedBy">User that resolved the issue</param>
+		/// <param name="quarantinedBy">User that quarantined the issue</param>
+		/// <param name="streamSeverity">The current severity in the stream</param>
+		/// <param name="spans">Spans for this issue</param>
+		/// <param name="openWorkflows">List of workflows for which this issue is open</param>
+		static FindIssueResponse NewFindIssueResponse(IIssue issue, IUser? owner, IUser? nominatedBy, IUser? resolvedBy, IUser? quarantinedBy, IssueSeverity? streamSeverity, List<FindIssueSpanResponse> spans, List<WorkflowId> openWorkflows)
+		{
+			FindIssueResponse response = new FindIssueResponse();
+			response.Id = issue.Id;
+			response.CreatedAt = issue.CreatedAt;
+			response.RetrievedAt = DateTime.UtcNow;
+			response.Summary = String.IsNullOrEmpty(issue.UserSummary) ? issue.Summary : issue.UserSummary;
+			response.Description = issue.Description;
+			response.Severity = issue.Severity;
+			response.StreamSeverity = streamSeverity;
+			response.Promoted = issue.Promoted;
+			if (owner != null)
+			{
+				response.Owner = owner.ToThinApiResponse();
+			}
+			if (nominatedBy != null)
+			{
+				response.NominatedBy = nominatedBy.ToThinApiResponse();
+			}
+			response.AcknowledgedAt = issue.AcknowledgedAt;
+			response.FixChange = issue.FixChange;
+			response.ResolvedAt = issue.ResolvedAt;
+			if (resolvedBy != null)
+			{
+				response.ResolvedBy = resolvedBy.ToThinApiResponse();
+			}
+			response.VerifiedAt = issue.VerifiedAt;
+			response.LastSeenAt = issue.LastSeenAt;
+			response.Spans = spans;
+			response.ExternalIssueKey = issue.ExternalIssueKey;
+			if (quarantinedBy != null)
+			{
+				response.QuarantinedBy = quarantinedBy.ToThinApiResponse();
+				response.QuarantineTimeUtc = issue.QuarantineTimeUtc;
+			}
+			response.WorkflowThreadUrl = issue.WorkflowThreadUrl;
+			response.OpenWorkflows = openWorkflows;
+			return response;
+		}
+
 		/// <summary>
 		/// Retrieve information about a specific issue
 		/// </summary>
@@ -328,14 +414,129 @@ namespace Horde.Server.Issues
 				try
 				{
 					_globalConfig.Value.TryGetStream(streamSpans.Key, out StreamConfig? streamConfig);
-					affectedStreams.Add(new GetIssueAffectedStreamResponse(details, streamConfig, streamSpans));
+					affectedStreams.Add(NewGetIssueAffectedStreamResponse(details, streamConfig, streamSpans));
 				}
 				catch (Exception ex)
 				{
 					_logger.LogError(ex, "Unable to get {StreamId} for span key on issue {IssueId}", streamSpans.Key, details.Issue.Id);
 				}
 			}
-			return new GetIssueResponse(details, affectedStreams, showDesktopAlerts);
+			return NewGetIssueResponse(details, affectedStreams, showDesktopAlerts);
+		}
+
+		/// <summary>
+		/// Constructor
+		/// </summary>
+		/// <param name="details">Issue to construct from</param>
+		/// <param name="streamConfig"></param>
+		/// <param name="spans">The spans to construct from</param>
+		static GetIssueAffectedStreamResponse NewGetIssueAffectedStreamResponse(IIssueDetails details, StreamConfig? streamConfig, IEnumerable<IIssueSpan> spans)
+		{
+			GetIssueAffectedStreamResponse response = new GetIssueAffectedStreamResponse();
+
+			IIssueSpan firstSpan = spans.First();
+			response.StreamId = firstSpan.StreamId;
+			response.StreamName = firstSpan.StreamName;
+			response.Resolved = spans.All(x => x.NextSuccess != null);
+
+			response.AffectedTemplates = new List<GetIssueAffectedTemplateResponse>();
+			foreach (IGrouping<TemplateId, IIssueSpan> template in spans.GroupBy(x => x.TemplateRefId))
+			{
+				string templateName = template.Key.ToString();
+				if (streamConfig != null && streamConfig.TryGetTemplate(template.Key, out TemplateRefConfig? templateRefConfig))
+				{
+					templateName = templateRefConfig.Name;
+				}
+
+				HashSet<ObjectId> unresolvedTemplateSpans = new HashSet<ObjectId>(template.Where(x => x.NextSuccess == null).Select(x => x.Id));
+
+				IIssueStep? templateStep = details.Steps.Where(x => unresolvedTemplateSpans.Contains(x.SpanId)).OrderByDescending(x => x.StepTime).FirstOrDefault();
+
+				response.AffectedTemplates.Add(NewGetIssueAffectedTemplateResponse(template.Key, templateName, template.All(x => x.NextSuccess != null), templateStep?.Severity ?? IssueSeverity.Unspecified));
+			}
+
+			HashSet<TemplateId> templateIdsSet = new HashSet<TemplateId>(spans.Select(x => x.TemplateRefId));
+			response.TemplateIds = templateIdsSet.Select(x => x.ToString()).ToList();
+
+			HashSet<TemplateId> unresolvedTemplateIdsSet = new HashSet<TemplateId>(spans.Where(x => x.NextSuccess == null).Select(x => x.TemplateRefId));
+			response.UnresolvedTemplateIds = unresolvedTemplateIdsSet.Select(x => x.ToString()).ToList();
+			response.ResolvedTemplateIds = templateIdsSet.Except(unresolvedTemplateIdsSet).Select(x => x.ToString()).ToList();
+
+			return response;
+		}
+
+		/// <summary>
+		/// Constructor
+		/// </summary>
+		/// <param name="templateId"></param>
+		/// <param name="templateName"></param>
+		/// <param name="resolved"></param>
+		/// <param name="severity"></param>
+		static GetIssueAffectedTemplateResponse NewGetIssueAffectedTemplateResponse(TemplateId templateId, string templateName, bool resolved, IssueSeverity severity = IssueSeverity.Unspecified)
+		{
+			GetIssueAffectedTemplateResponse response = new GetIssueAffectedTemplateResponse();
+			response.TemplateId = templateId;
+			response.TemplateName = templateName;
+			response.Resolved = resolved;
+			response.Severity = severity;
+			return response;
+		}
+
+		/// <summary>
+		/// Constructs a new issue
+		/// </summary>
+		/// <param name="details">Issue to construct from</param>
+		/// <param name="affectedStreams">The affected streams</param>
+		/// <param name="showDesktopAlerts">Whether to show alerts for this issue</param>
+		static GetIssueResponse NewGetIssueResponse(IIssueDetails details, List<GetIssueAffectedStreamResponse> affectedStreams, bool showDesktopAlerts)
+		{
+			GetIssueResponse response = new GetIssueResponse();
+			IIssue issue = details.Issue;
+			response.Id = issue.Id;
+			response.CreatedAt = issue.CreatedAt;
+			response.RetrievedAt = DateTime.UtcNow;
+			response.Summary = String.IsNullOrEmpty(issue.UserSummary) ? issue.Summary : issue.UserSummary;
+			response.Description = issue.Description;
+			response.Severity = issue.Severity;
+			response.Promoted = issue.Promoted;
+			response.Owner = details.Owner?.Login;
+			response.OwnerId = details.Owner?.Id.ToString();
+			response.OwnerInfo = details.Owner?.ToThinApiResponse();
+			response.NominatedBy = details.NominatedBy?.Login;
+			response.NominatedByInfo = details.NominatedBy?.ToThinApiResponse();
+			response.AcknowledgedAt = issue.AcknowledgedAt;
+			response.FixChange = issue.FixChange;
+			response.ResolvedAt = issue.ResolvedAt;
+			response.ResolvedBy = details.ResolvedBy?.Login;
+			response.ResolvedById = details.ResolvedBy?.Id.ToString();
+			response.ResolvedByInfo = details.ResolvedBy?.ToThinApiResponse();
+			response.VerifiedAt = issue.VerifiedAt;
+			response.LastSeenAt = issue.LastSeenAt;
+			response.Streams = details.Spans.Select(x => x.StreamName).Distinct().ToList()!;
+			response.ResolvedStreams = new List<string>();
+			response.UnresolvedStreams = new List<string>();
+			response.AffectedStreams = affectedStreams;
+			foreach (IGrouping<StreamId, IIssueSpan> stream in details.Spans.GroupBy(x => x.StreamId))
+			{
+				if (stream.All(x => x.NextSuccess != null))
+				{
+					response.ResolvedStreams.Add(stream.Key.ToString());
+				}
+				else
+				{
+					response.UnresolvedStreams.Add(stream.Key.ToString());
+				}
+			}
+			response.PrimarySuspects = details.SuspectUsers.Where(x => x.Login != null).Select(x => x.Login).ToList();
+			response.PrimarySuspectIds = details.SuspectUsers.Select(x => x.Id.ToString()).ToList();
+			response.PrimarySuspectsInfo = details.SuspectUsers.ConvertAll(x => x.ToThinApiResponse());
+			response.ShowDesktopAlerts = showDesktopAlerts;
+			response.ExternalIssueKey = details.ExternalIssueKey;
+			response.QuarantinedByUserInfo = details.QuarantinedBy?.ToThinApiResponse();
+			response.QuarantineTimeUtc = details.QuarantineTimeUtc;
+			response.ForceClosedByUserInfo = details.ForceClosedBy?.ToThinApiResponse();
+			response.WorkflowThreadUrl = issue.WorkflowThreadUrl;
+			return response;
 		}
 
 		/// <summary>
@@ -367,12 +568,45 @@ namespace Horde.Server.Issues
 				{
 					HashSet<ObjectId> spanIds = new HashSet<ObjectId>(spanGroup.Select(x => x.Id));
 					List<IIssueStep> steps = issue.Steps.Where(x => spanIds.Contains(x.SpanId)).ToList();
-					responses.Add(PropertyFilter.Apply(new GetIssueStreamResponse(spanGroup.Key, spanGroup.ToList(), steps), filter));
+					responses.Add(PropertyFilter.Apply(NewGetIssueStreamResponse(spanGroup.Key, spanGroup.ToList(), steps), filter));
 				}
 			}
 			return responses;
 		}
 
+		static GetIssueStreamResponse NewGetIssueStreamResponse(StreamId streamId, List<IIssueSpan> spans, List<IIssueStep> steps)
+		{
+			GetIssueStreamResponse response = new GetIssueStreamResponse();
+			response.StreamId = streamId;
+
+			foreach (IIssueSpan span in spans)
+			{
+				if (span.LastSuccess != null && (response.MinChange == null || span.LastSuccess.Change < response.MinChange.Value))
+				{
+					response.MinChange = span.LastSuccess.Change;
+				}
+				if (span.NextSuccess != null && (response.MaxChange == null || span.NextSuccess.Change > response.MaxChange.Value))
+				{
+					response.MaxChange = span.NextSuccess.Change;
+				}
+				response.Nodes.Add(NewGetIssueSpanResponse(span, steps.Where(y => y.SpanId == span.Id).ToList()));
+			}
+			return response;
+		}
+
+		static GetIssueSpanResponse NewGetIssueSpanResponse(IIssueSpan span, List<IIssueStep> steps)
+		{
+			GetIssueSpanResponse response = new GetIssueSpanResponse();
+			response.Id = span.Id.ToString();
+			response.Name = span.NodeName;
+			response.TemplateId = span.TemplateRefId;
+			response.WorkflowId = span.LastFailure.Annotations.WorkflowId;
+			response.LastSuccess = (span.LastSuccess != null) ? NewGetIssueStepResponse(span.LastSuccess) : null;
+			response.Steps = steps.ConvertAll(x => NewGetIssueStepResponse(x));
+			response.NextSuccess = (span.NextSuccess != null) ? NewGetIssueStepResponse(span.NextSuccess) : null;
+			return response;
+		}
+
 		/// <summary>
 		/// Retrieve events for a specific issue
 		/// </summary>
@@ -409,7 +643,7 @@ namespace Horde.Server.Issues
 			HashSet<ObjectId> spanIds = new HashSet<ObjectId>(spans.Select(x => x.Id));
 			List<IIssueStep> steps = details.Steps.Where(x => spanIds.Contains(x.SpanId)).ToList();
 
-			return PropertyFilter.Apply(new GetIssueStreamResponse(streamId, spans, steps), filter);
+			return PropertyFilter.Apply(NewGetIssueStreamResponse(streamId, spans, steps), filter);
 		}
 
 		/// <summary>
@@ -636,13 +870,27 @@ namespace Horde.Server.Issues
 
 				for (int i = 0; i < issues.Count; i++)
 				{
-					response.Add(new GetExternalIssueResponse(issues[i]));
+					response.Add(NewGetExternalIssueResponse(issues[i]));
 				}
 			}
 
 			return response;
 		}
 
+		static GetExternalIssueResponse NewGetExternalIssueResponse(IExternalIssue issue)
+		{
+			GetExternalIssueResponse response = new GetExternalIssueResponse();
+			response.Key = issue.Key;
+			response.Link = issue.Link;
+			response.StatusName = issue.StatusName;
+			response.ResolutionName = issue.ResolutionName;
+			response.PriorityName = issue.PriorityName;
+			response.AssigneeName = issue.AssigneeName;
+			response.AssigneeDisplayName = issue.AssigneeDisplayName;
+			response.AssigneeEmailAddress = issue.AssigneeEmailAddress;
+			return response;
+		}
+
 		/// <summary>
 		/// Create a new external issue
 		/// </summary>
@@ -715,10 +963,29 @@ namespace Horde.Server.Issues
 
 			projects.ForEach(project =>
 			{
-				response.Add(new GetExternalIssueProjectResponse(project.Key, project.Name, project.Id, project.Components, project.IssueTypes));
+				response.Add(NewGetExternalIssueProjectResponse(project.Key, project.Name, project.Id, project.Components, project.IssueTypes));
 			});
 
 			return response;
 		}
+
+		/// <summary>
+		/// Constructor
+		/// </summary>
+		/// <param name="key"></param>
+		/// <param name="name"></param>
+		/// <param name="id"></param>
+		/// <param name="components"></param>
+		/// <param name="issueTypes"></param>
+		static GetExternalIssueProjectResponse NewGetExternalIssueProjectResponse(string key, string name, string id, Dictionary<string, string> components, Dictionary<string, string> issueTypes)
+		{
+			GetExternalIssueProjectResponse response = new GetExternalIssueProjectResponse();
+			response.ProjectKey = key;
+			response.Name = name;
+			response.Id = id;
+			response.Components = new Dictionary<string, string>(components);
+			response.IssueTypes = new Dictionary<string, string>(issueTypes);
+			return response;
+		}
 	}
 }
diff --git a/Engine/Source/Programs/Horde/Horde.Server/Issues/WorkflowConfig.cs b/Engine/Source/Programs/Horde/Horde.Server/Issues/WorkflowConfig.cs
--- a/Engine/Source/Programs/Horde/Horde.Server/Issues/WorkflowConfig.cs
+++ b/Engine/Source/Programs/Horde/Horde.Server/Issues/WorkflowConfig.cs
@@ -3,6 +3,7 @@
 using System;
 using System.Collections.Generic;
 using System.ComponentModel.DataAnnotations;
+using EpicGames.Horde.Issues;
 using Horde.Server.Jobs.Graphs;
 
 namespace Horde.Server.Issues
diff --git a/Engine/Source/Programs/Horde/Horde.Server/Issues/WorkflowId.cs b/Engine/Source/Programs/Horde/Horde.Server/Issues/WorkflowId.cs
deleted file mode 100644
diff --git a/Engine/Source/Programs/Horde/Horde.Server/Jobs/Graphs/NodeAnnotations.cs b/Engine/Source/Programs/Horde/Horde.Server/Jobs/Graphs/NodeAnnotations.cs
--- a/Engine/Source/Programs/Horde/Horde.Server/Jobs/Graphs/NodeAnnotations.cs
+++ b/Engine/Source/Programs/Horde/Horde.Server/Jobs/Graphs/NodeAnnotations.cs
@@ -2,7 +2,7 @@
 
 using System;
 using System.Collections.Generic;
-using Horde.Server.Issues;
+using EpicGames.Horde.Issues;
 using Horde.Server.Utilities;
 
 namespace Horde.Server.Jobs.Graphs
diff --git a/Engine/Source/Programs/Horde/Horde.Server/Notifications/Sinks/SlackNotificationSink.cs b/Engine/Source/Programs/Horde/Horde.Server/Notifications/Sinks/SlackNotificationSink.cs
--- a/Engine/Source/Programs/Horde/Horde.Server/Notifications/Sinks/SlackNotificationSink.cs
+++ b/Engine/Source/Programs/Horde/Horde.Server/Notifications/Sinks/SlackNotificationSink.cs
@@ -18,6 +18,7 @@ using System.Threading.Tasks;
 using EpicGames.Core;
 using EpicGames.Horde;
 using EpicGames.Horde.Agents;
+using EpicGames.Horde.Issues;
 using EpicGames.Horde.Jobs;
 using EpicGames.Horde.Jobs.Templates;
 using EpicGames.Horde.Logs;
diff --git a/Engine/Source/Programs/Horde/Horde.Server/Streams/StreamConfig.cs b/Engine/Source/Programs/Horde/Horde.Server/Streams/StreamConfig.cs
--- a/Engine/Source/Programs/Horde/Horde.Server/Streams/StreamConfig.cs
+++ b/Engine/Source/Programs/Horde/Horde.Server/Streams/StreamConfig.cs
@@ -15,6 +15,7 @@ using EpicGames.Horde;
 using EpicGames.Horde.Acls;
 using EpicGames.Horde.Agents.Pools;
 using EpicGames.Horde.Common;
+using EpicGames.Horde.Issues;
 using EpicGames.Horde.Jobs.Templates;
 using EpicGames.Horde.Replicators;
 using EpicGames.Horde.Streams;
diff --git a/Engine/Source/Programs/Horde/Horde.Server/Ugs/UgsController.cs b/Engine/Source/Programs/Horde/Horde.Server/Ugs/UgsController.cs
--- a/Engine/Source/Programs/Horde/Horde.Server/Ugs/UgsController.cs
+++ b/Engine/Source/Programs/Horde/Horde.Server/Ugs/UgsController.cs
@@ -15,6 +15,7 @@ using Microsoft.Extensions.Options;
 using MongoDB.Bson;
 using EpicGames.Horde.Users;
 using EpicGames.Horde.Logs;
+using EpicGames.Horde.Issues;
 
 namespace Horde.Server.Ugs
 {
diff --git a/Engine/Source/Programs/Horde/Horde.Server/Ugs/UgsMessages.cs b/Engine/Source/Programs/Horde/Horde.Server/Ugs/UgsMessages.cs
--- a/Engine/Source/Programs/Horde/Horde.Server/Ugs/UgsMessages.cs
+++ b/Engine/Source/Programs/Horde/Horde.Server/Ugs/UgsMessages.cs
@@ -5,6 +5,7 @@ using System.Collections.Generic;
 using System.ComponentModel.DataAnnotations;
 using System.Linq;
 using EpicGames.Core;
+using EpicGames.Horde.Issues;
 using Horde.Server.Issues;
 using Horde.Server.Users;
 
diff --git a/Engine/Source/Programs/Shared/EpicGames.Horde/Issues/IssueMessages.cs b/Engine/Source/Programs/Shared/EpicGames.Horde/Issues/IssueMessages.cs
new file mode 100644
--- /dev/null
+++ b/Engine/Source/Programs/Shared/EpicGames.Horde/Issues/IssueMessages.cs
@@ -0,0 +1,816 @@
+// Copyright Epic Games, Inc. All Rights Reserved.
+
+using System;
+using System.Collections.Generic;
+using EpicGames.Horde.Jobs;
+using EpicGames.Horde.Jobs.Templates;
+using EpicGames.Horde.Logs;
+using EpicGames.Horde.Streams;
+using EpicGames.Horde.Users;
+
+#pragma warning disable CA2227 // Change collections to be read-only
+
+namespace EpicGames.Horde.Issues
+{
+	/// <summary>
+	/// The severity of an issue
+	/// </summary>
+	public enum IssueSeverity
+	{
+		/// <summary>
+		/// Unspecified severity
+		/// </summary>
+		Unspecified,
+
+		/// <summary>
+		/// This error represents a warning
+		/// </summary>
+		Warning,
+
+		/// <summary>
+		/// This issue represents an error
+		/// </summary>
+		Error,
+	}
+
+	/// <summary>
+	/// Identifies a particular changelist and job
+	/// </summary>
+	public class GetIssueStepResponse
+	{
+		/// <summary>
+		/// The changelist number
+		/// </summary>
+		public int Change { get; set; }
+
+		/// <summary>
+		/// Severity of the issue in this step
+		/// </summary>
+		public IssueSeverity Severity { get; set; }
+
+		/// <summary>
+		/// Name of the job containing this step
+		/// </summary>
+		public string JobName { get; set; } = String.Empty;
+
+		/// <summary>
+		/// The unique job id
+		/// </summary>
+		public JobId JobId { get; set; }
+
+		/// <summary>
+		/// The unique batch id
+		/// </summary>
+		public JobStepBatchId BatchId { get; set; }
+
+		/// <summary>
+		/// The unique step id
+		/// </summary>
+		public JobStepId StepId { get; set; }
+
+		/// <summary>
+		/// Time at which the step ran
+		/// </summary>
+		public DateTime StepTime { get; set; }
+
+		/// <summary>
+		/// The unique log id
+		/// </summary>
+		public LogId? LogId { get; set; }
+	}
+
+	/// <summary>
+	/// Trace of a set of node failures across multiple steps
+	/// </summary>
+	public class GetIssueSpanResponse
+	{
+		/// <summary>
+		/// Unique id of this span
+		/// </summary>
+		public string Id { get; set; } = String.Empty;
+
+		/// <summary>
+		/// The template containing this step
+		/// </summary>
+		public TemplateId TemplateId { get; set; }
+
+		/// <summary>
+		/// Name of the step
+		/// </summary>
+		public string Name { get; set; } = String.Empty;
+
+		/// <summary>
+		/// Workflow that this span belongs to
+		/// </summary>
+		public WorkflowId? WorkflowId { get; set; }
+
+		/// <summary>
+		/// The previous build 
+		/// </summary>
+		public GetIssueStepResponse? LastSuccess { get; set; }
+
+		/// <summary>
+		/// The failing builds for a particular event
+		/// </summary>
+		public List<GetIssueStepResponse> Steps { get; set; } = new List<GetIssueStepResponse>();
+
+		/// <summary>
+		/// The following successful build
+		/// </summary>
+		public GetIssueStepResponse? NextSuccess { get; set; }
+	}
+
+	/// <summary>
+	/// Information about a particular step
+	/// </summary>
+	public class GetIssueStreamResponse
+	{
+		/// <summary>
+		/// Unique id of the stream
+		/// </summary>
+		public StreamId StreamId { get; set; }
+
+		/// <summary>
+		/// Minimum changelist affected by this issue (ie. last successful build)
+		/// </summary>
+		public int? MinChange { get; set; }
+
+		/// <summary>
+		/// Maximum changelist affected by this issue (ie. next successful build)
+		/// </summary>
+		public int? MaxChange { get; set; }
+
+		/// <summary>
+		/// Map of steps to (event signature id -> trace id)
+		/// </summary>
+		public List<GetIssueSpanResponse> Nodes { get; set; } = new List<GetIssueSpanResponse>();
+	}
+
+	/// <summary>
+	/// Outcome of a particular build
+	/// </summary>
+	public enum IssueBuildOutcome
+	{
+		/// <summary>
+		/// Unknown outcome
+		/// </summary>
+		Unknown,
+
+		/// <summary>
+		/// Build succeeded
+		/// </summary>
+		Success,
+
+		/// <summary>
+		/// Build failed
+		/// </summary>
+		Error,
+
+		/// <summary>
+		/// Build finished with warnings
+		/// </summary>
+		Warning,
+	}
+
+	/// <summary>
+	/// Information about a template affected by an issue
+	/// </summary>
+	public class GetIssueAffectedTemplateResponse
+	{
+		/// <summary>
+		/// The template id
+		/// </summary>
+		public TemplateId TemplateId { get; set; }
+
+		/// <summary>
+		/// The template name
+		/// </summary>
+		public string TemplateName { get; set; } = String.Empty;
+
+		/// <summary>
+		/// Whether it has been resolved or not
+		/// </summary>
+		public bool Resolved { get; set; }
+
+		/// <summary>
+		/// The issue severity of the affected template
+		/// </summary>
+		public IssueSeverity Severity { get; set; }
+	}
+
+	/// <summary>
+	/// Summary for the state of a stream in an issue
+	/// </summary>
+	public class GetIssueAffectedStreamResponse
+	{
+		/// <summary>
+		/// Id of the stream
+		/// </summary>
+		public StreamId StreamId { get; set; }
+
+		/// <summary>
+		/// Name of the stream
+		/// </summary>
+		public string StreamName { get; set; } = String.Empty;
+
+		/// <summary>
+		/// Whether the issue has been resolved in this stream
+		/// </summary>
+		public bool Resolved { get; set; }
+
+		/// <summary>
+		/// The affected templates
+		/// </summary>
+		public List<GetIssueAffectedTemplateResponse> AffectedTemplates { get; set; } = new List<GetIssueAffectedTemplateResponse>();
+
+		/// <summary>
+		/// List of affected template ids
+		/// </summary>
+		public List<string> TemplateIds { get; set; } = new List<string>();
+
+		/// <summary>
+		/// List of resolved template ids
+		/// </summary>
+		public List<string> ResolvedTemplateIds { get; set; } = new List<string>();
+
+		/// <summary>
+		/// List of unresolved template ids
+		/// </summary>
+		public List<string> UnresolvedTemplateIds { get; set; } = new List<string>();
+	}
+
+	/// <summary>
+	/// Stores information about a build health issue
+	/// </summary>
+	public class GetIssueResponse
+	{
+		/// <summary>
+		/// The unique object id
+		/// </summary>
+		public int Id { get; set; }
+
+		/// <summary>
+		/// Time at which the issue was created
+		/// </summary>
+		public DateTime CreatedAt { get; set; }
+
+		/// <summary>
+		/// Time at which the issue was retrieved
+		/// </summary>
+		public DateTime RetrievedAt { get; set; }
+
+		/// <summary>
+		/// The associated project for the issue
+		/// </summary>
+		public string? Project { get; set; }
+
+		/// <summary>
+		/// The summary text for this issue
+		/// </summary>
+		public string Summary { get; set; } = String.Empty;
+
+		/// <summary>
+		/// Detailed description text
+		/// </summary>
+		public string? Description { get; set; }
+
+		/// <summary>
+		/// Severity of this issue
+		/// </summary>
+		public IssueSeverity Severity { get; set; }
+
+		/// <summary>
+		/// Whether the issue is promoted
+		/// </summary>
+		public bool Promoted { get; set; }
+
+		/// <summary>
+		/// Owner of the issue [DEPRECATED]
+		/// </summary>
+		public string? Owner { get; set; }
+
+		/// <summary>
+		/// User id of the owner [DEPRECATED]
+		/// </summary>
+		public string? OwnerId { get; set; }
+
+		/// <summary>
+		/// Owner of the issue
+		/// </summary>
+		public GetThinUserInfoResponse? OwnerInfo { get; set; }
+
+		/// <summary>
+		/// User that nominated the current owner [DEPRECATED]
+		/// </summary>
+		public string? NominatedBy { get; set; }
+
+		/// <summary>
+		/// Owner of the issue
+		/// </summary>
+		public GetThinUserInfoResponse? NominatedByInfo { get; set; }
+
+		/// <summary>
+		/// Time that the issue was acknowledged
+		/// </summary>
+		public DateTime? AcknowledgedAt { get; set; }
+
+		/// <summary>
+		/// Changelist that fixed this issue
+		/// </summary>
+		public int? FixChange { get; set; }
+
+		/// <summary>
+		/// Time at which the issue was resolved
+		/// </summary>
+		public DateTime? ResolvedAt { get; set; }
+
+		/// <summary>
+		/// Name of the user that resolved the issue [DEPRECATED]
+		/// </summary>
+		public string? ResolvedBy { get; set; }
+
+		/// <summary>
+		/// User id of the person that resolved the issue [DEPRECATED]
+		/// </summary>
+		public string? ResolvedById { get; set; }
+
+		/// <summary>
+		/// User that resolved the issue
+		/// </summary>
+		public GetThinUserInfoResponse? ResolvedByInfo { get; set; }
+
+		/// <summary>
+		/// Time at which the issue was verified
+		/// </summary>
+		public DateTime? VerifiedAt { get; set; }
+
+		/// <summary>
+		/// Time that the issue was last seen
+		/// </summary>
+		public DateTime LastSeenAt { get; set; }
+
+		/// <summary>
+		/// List of stream paths affected by this issue
+		/// </summary>
+		public List<string> Streams { get; set; } = new List<string>();
+
+		/// <summary>
+		/// List of affected stream ids
+		/// </summary>
+		public List<string> ResolvedStreams { get; set; } = new List<string>();
+
+		/// <summary>
+		/// List of unresolved streams
+		/// </summary>
+		public List<string> UnresolvedStreams { get; set; } = new List<string>();
+
+		/// <summary>
+		/// List of affected streams
+		/// </summary>
+		public List<GetIssueAffectedStreamResponse> AffectedStreams { get; set; } = new List<GetIssueAffectedStreamResponse>();
+
+		/// <summary>
+		/// Most likely suspects for causing this issue [DEPRECATED]
+		/// </summary>
+		public List<string>? PrimarySuspects { get; set; }
+
+		/// <summary>
+		/// User ids of the most likely suspects [DEPRECATED]
+		/// </summary>
+		public List<string>? PrimarySuspectIds { get; set; }
+
+		/// <summary>
+		/// Most likely suspects for causing this issue
+		/// </summary>
+		public List<GetThinUserInfoResponse> PrimarySuspectsInfo { get; set; } = new List<GetThinUserInfoResponse>();
+
+		/// <summary>
+		/// Whether to show alerts for this issue
+		/// </summary>
+		public bool ShowDesktopAlerts { get; set; }
+
+		/// <summary>
+		/// Key for this issue in external issue tracker
+		/// </summary>
+		public string? ExternalIssueKey { get; set; }
+
+		/// <summary>
+		/// User who quarantined the issue
+		/// </summary>
+		public GetThinUserInfoResponse? QuarantinedByUserInfo { get; set; }
+
+		/// <summary>
+		/// The UTC time when the issue was quarantined
+		/// </summary>
+		public DateTime? QuarantineTimeUtc { get; set; }
+
+		/// <summary>
+		/// User who force closed the issue
+		/// </summary>
+		public GetThinUserInfoResponse? ForceClosedByUserInfo { get; set; }
+
+		/// <summary>
+		/// The workflow thread url for this issue
+		/// </summary>
+		public Uri? WorkflowThreadUrl { get; set; }
+	}
+
+	/// <summary>
+	/// Information about a span within an issue
+	/// </summary>
+	public class FindIssueSpanResponse
+	{
+		/// <summary>
+		/// Unique id of this span
+		/// </summary>
+		public string Id { get; set; } = String.Empty;
+
+		/// <summary>
+		/// The template containing this step
+		/// </summary>
+		public string TemplateId { get; set; } = String.Empty;
+
+		/// <summary>
+		/// Name of the step
+		/// </summary>
+		public string Name { get; set; } = String.Empty;
+
+		/// <summary>
+		/// Workflow for this span
+		/// </summary>
+		public WorkflowId? WorkflowId { get; set; }
+
+		/// <summary>
+		/// The previous build 
+		/// </summary>
+		public GetIssueStepResponse? LastSuccess { get; set; }
+
+		/// <summary>
+		/// The following successful build
+		/// </summary>
+		public GetIssueStepResponse? NextSuccess { get; set; }
+	}
+
+	/// <summary>
+	/// Stores information about a build health issue
+	/// </summary>
+	public class FindIssueResponse
+	{
+		/// <summary>
+		/// The unique object id
+		/// </summary>
+		public int Id { get; set; }
+
+		/// <summary>
+		/// Time at which the issue was created
+		/// </summary>
+		public DateTime CreatedAt { get; set; }
+
+		/// <summary>
+		/// Time at which the issue was retrieved
+		/// </summary>
+		public DateTime RetrievedAt { get; set; }
+
+		/// <summary>
+		/// The associated project for the issue
+		/// </summary>
+		public string? Project { get; set; }
+
+		/// <summary>
+		/// The summary text for this issue
+		/// </summary>
+		public string Summary { get; set; } = String.Empty;
+
+		/// <summary>
+		/// Detailed description text
+		/// </summary>
+		public string? Description { get; set; }
+
+		/// <summary>
+		/// Severity of this issue
+		/// </summary>
+		public IssueSeverity Severity { get; set; }
+
+		/// <summary>
+		/// Severity of this issue in the stream
+		/// </summary>
+		public IssueSeverity? StreamSeverity { get; set; }
+
+		/// <summary>
+		/// Whether the issue is promoted
+		/// </summary>
+		public bool Promoted { get; set; }
+
+		/// <summary>
+		/// Owner of the issue
+		/// </summary>
+		public GetThinUserInfoResponse? Owner { get; set; }
+
+		/// <summary>
+		/// Owner of the issue
+		/// </summary>
+		public GetThinUserInfoResponse? NominatedBy { get; set; }
+
+		/// <summary>
+		/// Time that the issue was acknowledged
+		/// </summary>
+		public DateTime? AcknowledgedAt { get; set; }
+
+		/// <summary>
+		/// Changelist that fixed this issue
+		/// </summary>
+		public int? FixChange { get; set; }
+
+		/// <summary>
+		/// Time at which the issue was resolved
+		/// </summary>
+		public DateTime? ResolvedAt { get; set; }
+
+		/// <summary>
+		/// User that resolved the issue
+		/// </summary>
+		public GetThinUserInfoResponse? ResolvedBy { get; set; }
+
+		/// <summary>
+		/// Time at which the issue was verified
+		/// </summary>
+		public DateTime? VerifiedAt { get; set; }
+
+		/// <summary>
+		/// Time that the issue was last seen
+		/// </summary>
+		public DateTime LastSeenAt { get; set; }
+
+		/// <summary>
+		/// Spans for this issue
+		/// </summary>
+		public List<FindIssueSpanResponse>? Spans { get; set; }
+
+		/// <summary>
+		/// Key for this issue in external issue tracker
+		/// </summary>
+		public string? ExternalIssueKey { get; set; }
+
+		/// <summary>
+		/// User who quarantined the issue
+		/// </summary>
+		public GetThinUserInfoResponse? QuarantinedBy { get; set; }
+
+		/// <summary>
+		/// The UTC time when the issue was quarantined
+		/// </summary>
+		public DateTime? QuarantineTimeUtc { get; set; }
+
+		/// <summary>
+		/// The workflow thread url for this issue
+		/// </summary>
+		public Uri? WorkflowThreadUrl { get; set; }
+
+		/// <summary>
+		/// Workflows for which this issue is open
+		/// </summary>
+		public List<WorkflowId>? OpenWorkflows { get; set; }
+	}
+
+	/// <summary>
+	/// Request an issue to be updated
+	/// </summary>
+	public class UpdateIssueRequest
+	{
+		/// <summary>
+		/// Summary of the issue
+		/// </summary>
+		public string? Summary { get; set; }
+
+		/// <summary>
+		/// Description of the issue
+		/// </summary>
+		public string? Description { get; set; }
+
+		/// <summary>
+		/// Whether the issue is promoted or not
+		/// </summary>
+		public bool? Promoted { get; set; }
+
+		/// <summary>
+		/// New user id for owner of the issue, can be cleared by passing empty string
+		/// </summary>
+		public string? OwnerId { get; set; }
+
+		/// <summary>
+		/// User id that nominated the new owner
+		/// </summary>
+		public string? NominatedById { get; set; }
+
+		/// <summary>
+		/// Whether the issue has been acknowledged
+		/// </summary>
+		public bool? Acknowledged { get; set; }
+
+		/// <summary>
+		/// Whether the user has declined this issue
+		/// </summary>
+		public bool? Declined { get; set; }
+
+		/// <summary>
+		/// The change at which the issue is claimed fixed. 0 = not fixed, -1 = systemic issue.
+		/// </summary>
+		public int? FixChange { get; set; }
+
+		/// <summary>
+		/// Whether the issue should be marked as resolved
+		/// </summary>
+		public bool? Resolved { get; set; }
+
+		/// <summary>
+		/// List of spans to add to this issue
+		/// </summary>
+		public List<string>? AddSpans { get; set; }
+
+		/// <summary>
+		/// List of spans to remove from this issue
+		/// </summary>
+		public List<string>? RemoveSpans { get; set; }
+
+		/// <summary>
+		/// A key to issue in external tracker
+		/// </summary>
+		public string? ExternalIssueKey { get; set; }
+
+		/// <summary>
+		/// Id of user quarantining issue
+		/// </summary>
+		public string? QuarantinedById { get; set; }
+
+		/// <summary>
+		/// Id of user who is forcibly closing this issue, skipping verification checks.  This is useful for when a failing step has been removed for example
+		/// </summary>
+		public string? ForceClosedById { get; set; }
+	}
+
+	/// <summary>
+	/// External issue project information
+	/// </summary>
+	public class GetExternalIssueProjectResponse
+	{
+		/// <summary>
+		/// The project key
+		/// </summary>
+		public string ProjectKey { get; set; } = String.Empty;
+
+		/// <summary>
+		/// The name of the project
+		/// </summary>
+		public string Name { get; set; } = String.Empty;
+
+		/// <summary>
+		/// The id of the project
+		/// </summary>
+		public string Id { get; set; } = String.Empty;
+
+		/// <summary>
+		/// component id => name
+		/// </summary>
+		public Dictionary<string, string> Components { get; set; } = new Dictionary<string, string>();
+
+		/// <summary>
+		/// IssueType id => name
+		/// </summary>
+		public Dictionary<string, string> IssueTypes { get; set; } = new Dictionary<string, string>();
+	}
+
+	/// <summary>
+	/// Marks an issue as fixed by another user. Designed for use from a Perforce trigger.
+	/// </summary>
+	public class MarkFixedViaPerforceRequest
+	{
+		/// <summary>
+		/// Name of the user that fixed the issue
+		/// </summary>
+		public string UserName { get; set; } = String.Empty;
+
+		/// <summary>
+		/// Change that fixed the issue
+		/// </summary>
+		public int FixChange { get; set; }
+	}
+
+	/// <summary>
+	/// Request an issue to be created on external issue tracking system
+	/// </summary>
+	public class CreateExternalIssueRequest
+	{
+		/// <summary>
+		/// Horde issue which is linked to external issue
+		/// </summary>
+		public int IssueId { get; set; } = 0;
+
+		/// <summary>
+		/// StreamId of a stream with this issue
+		/// </summary>
+		public string StreamId { get; set; } = String.Empty;
+
+		/// <summary>
+		/// Summary text for external issue
+		/// </summary>
+		public string Summary { get; set; } = String.Empty;
+
+		/// <summary>
+		/// External issue project id
+		/// </summary>
+		public string ProjectId { get; set; } = String.Empty;
+
+		/// <summary>
+		/// External issue component id
+		/// </summary>
+		public string ComponentId { get; set; } = String.Empty;
+
+		/// <summary>
+		/// External issue type id
+		/// </summary>
+		public string IssueTypeId { get; set; } = String.Empty;
+
+		/// <summary>
+		/// Optional description text for external issue
+		/// </summary>
+		public string? Description { get; set; }
+
+		/// <summary>
+		/// Optional link to issue on Horde
+		/// </summary>
+		public string? HordeIssueLink { get; set; }
+	}
+
+	/// <summary>
+	/// Response for externally created issue
+	/// </summary>
+	public class CreateExternalIssueResponse
+	{
+		/// <summary>
+		/// External issue key 
+		/// </summary>
+		public string Key { get; set; }
+
+		/// <summary>
+		/// Link to issue on external tracking site
+		/// </summary>
+		public string? Link { get; set; }
+
+		/// <summary>
+		/// Constructor
+		/// </summary>
+		/// <param name="key"></param>
+		/// <param name="link"></param>
+		public CreateExternalIssueResponse(string key, string? link)
+		{
+			Key = key;
+			Link = link;
+		}
+	}
+
+	/// <summary>
+	/// External issue response object
+	/// </summary>
+	public class GetExternalIssueResponse
+	{
+		/// <summary>
+		/// The external issue key
+		/// </summary>
+		public string Key { get; set; } = String.Empty;
+
+		/// <summary>
+		/// The issue link on external tracking site
+		/// </summary>
+		public string? Link { get; set; }
+
+		/// <summary>
+		/// The issue status name, "To Do", "In Progress", etc
+		/// </summary>
+		public string? StatusName { get; set; }
+
+		/// <summary>
+		/// The issue resolution name, "Fixed", "Closed", etc
+		/// </summary>
+		public string? ResolutionName { get; set; }
+
+		/// <summary>
+		/// The issue priority name, "1 - Critical", "2 - Major", etc
+		/// </summary>
+		public string? PriorityName { get; set; }
+
+		/// <summary>
+		/// The current assignee's user name
+		/// </summary>
+		public string? AssigneeName { get; set; }
+
+		/// <summary>
+		/// The current assignee's display name
+		/// </summary>
+		public string? AssigneeDisplayName { get; set; }
+
+		/// <summary>
+		/// The current assignee's email address
+		/// </summary>
+		public string? AssigneeEmailAddress { get; set; }
+	}
+}
diff --git a/Engine/Source/Programs/Shared/EpicGames.Horde/Issues/WorkflowId.cs b/Engine/Source/Programs/Shared/EpicGames.Horde/Issues/WorkflowId.cs
new file mode 100644
--- /dev/null
+++ b/Engine/Source/Programs/Shared/EpicGames.Horde/Issues/WorkflowId.cs
@@ -0,0 +1,49 @@
+// Copyright Epic Games, Inc. All Rights Reserved.
+
+using System.ComponentModel;
+using EpicGames.Core;
+using EpicGames.Serialization;
+
+namespace EpicGames.Horde.Issues
+{
+	/// <summary>
+	/// Identifier for a workflow
+	/// </summary>
+	/// <param name="Id">Id to construct from</param>
+	[JsonSchemaString]
+	[TypeConverter(typeof(StringIdTypeConverter<WorkflowId, WorkflowIdConverter>))]
+	[StringIdConverter(typeof(WorkflowIdConverter))]
+	[CbConverter(typeof(StringIdCbConverter<WorkflowId, WorkflowIdConverter>))]
+	public record struct WorkflowId(StringId Id)
+	{
+		/// <summary>
+		/// Empty workflow id constant
+		/// </summary>
+		public static WorkflowId Empty { get; } = default;
+
+		/// <summary>
+		/// Constructor
+		/// </summary>
+		public WorkflowId(string id) : this(new StringId(id))
+		{
+		}
+
+		/// <inheritdoc cref="StringId.IsEmpty"/>
+		public bool IsEmpty => Id.IsEmpty;
+
+		/// <inheritdoc/>
+		public override string ToString() => Id.ToString();
+	}
+
+	/// <summary>
+	/// Converter to and from <see cref="StringId"/> instances.
+	/// </summary>
+	class WorkflowIdConverter : StringIdConverter<WorkflowId>
+	{
+		/// <inheritdoc/>
+		public override WorkflowId FromStringId(StringId id) => new WorkflowId(id);
+
+		/// <inheritdoc/>
+		public override StringId ToStringId(WorkflowId value) => value.Id;
+	}
+}
