Horde: Move agent messages into EpicGames.Horde.

#rnx

diff --git a/Engine/Source/Programs/Horde/Horde.Agent.Tests/WorkerServiceTest.cs b/Engine/Source/Programs/Horde/Horde.Agent.Tests/WorkerServiceTest.cs
--- a/Engine/Source/Programs/Horde/Horde.Agent.Tests/WorkerServiceTest.cs
+++ b/Engine/Source/Programs/Horde/Horde.Agent.Tests/WorkerServiceTest.cs
@@ -388,7 +388,7 @@ namespace Horde.Agent.Tests
 			_leases[leaseId] = new Lease
 			{
 				Id = leaseId,
-				State = LeaseState.Pending,
+				State = RpcLeaseState.Pending,
 				Payload = Any.Pack(testTask)
 			};
 		}
@@ -499,7 +499,7 @@ namespace Horde.Agent.Tests
 				_logger.LogInformation("OnUpdateSessionRequest: {AgentId} {SessionId} {Status}", request.AgentId, request.SessionId, request.Status);
 				await Task.Delay(100, cancellationToken);
 				UpdateSessionResponse response = new () { ExpiryTime = Timestamp.FromDateTime(DateTime.UtcNow + TimeSpan.FromMinutes(120)) };
-				response.Leases.AddRange(_leases.Values.Where(x => x.State != LeaseState.Completed));
+				response.Leases.AddRange(_leases.Values.Where(x => x.State != RpcLeaseState.Completed));
 				await responseStream.Write(response);
 			}
 			
diff --git a/Engine/Source/Programs/Horde/Horde.Agent/Leases/LeaseManager.cs b/Engine/Source/Programs/Horde/Horde.Agent/Leases/LeaseManager.cs
--- a/Engine/Source/Programs/Horde/Horde.Agent/Leases/LeaseManager.cs
+++ b/Engine/Source/Programs/Horde/Horde.Agent/Leases/LeaseManager.cs
@@ -322,7 +322,7 @@ namespace Horde.Agent.Leases
 						// Now reconcile the local state to match what the server reports
 						if (updateSessionResponse != null)
 						{
-							bool atLeastOneLeaseFinished = _activeLeases.Any(x => x.Lease.State is LeaseState.Completed or LeaseState.Cancelled);
+							bool atLeastOneLeaseFinished = _activeLeases.Any(x => x.Lease.State is RpcLeaseState.Completed or RpcLeaseState.Cancelled);
 							if (atLeastOneLeaseFinished && shutdownAfterFinishedLease)
 							{
 								_logger.LogInformation("At least one lease executed. Requesting shutdown.");
@@ -332,13 +332,13 @@ namespace Horde.Agent.Leases
 							PoolIds = updateSessionResponse.PoolIds;
 
 							// Remove any leases which have completed
-							int numRemoved = _activeLeases.RemoveAll(x => (x.Lease.State == LeaseState.Completed || x.Lease.State == LeaseState.Cancelled) && !updateSessionResponse.Leases.Any(y => y.Id == x.Lease.Id && y.State != LeaseState.Cancelled));
+							int numRemoved = _activeLeases.RemoveAll(x => (x.Lease.State == RpcLeaseState.Completed || x.Lease.State == RpcLeaseState.Cancelled) && !updateSessionResponse.Leases.Any(y => y.Id == x.Lease.Id && y.State != RpcLeaseState.Cancelled));
 							NumLeasesCompleted += numRemoved;
 
 							// Create any new leases and cancel any running leases
 							foreach (Lease serverLease in updateSessionResponse.Leases)
 							{
-								if (serverLease.State == LeaseState.Cancelled)
+								if (serverLease.State == RpcLeaseState.Cancelled)
 								{
 									LeaseInfo? info = _activeLeases.FirstOrDefault(x => x.Lease.Id == serverLease.Id);
 									if (info != null)
@@ -347,9 +347,9 @@ namespace Horde.Agent.Leases
 										info.CancellationTokenSource.Cancel();
 									}
 								}
-								if (serverLease.State == LeaseState.Pending && !_activeLeases.Any(x => x.Lease.Id == serverLease.Id))
+								if (serverLease.State == RpcLeaseState.Pending && !_activeLeases.Any(x => x.Lease.Id == serverLease.Id))
 								{
-									serverLease.State = LeaseState.Active;
+									serverLease.State = RpcLeaseState.Active;
 
 									_logger.LogInformation("Adding lease {LeaseId}", serverLease.Id);
 									LeaseInfo info = new LeaseInfo(serverLease);
@@ -525,14 +525,14 @@ namespace Horde.Agent.Leases
 			{
 				if (leaseInfo.CancellationTokenSource.IsCancellationRequested)
 				{
-					leaseInfo.Lease.State = LeaseState.Cancelled;
-					leaseInfo.Lease.Outcome = LeaseOutcome.Failed;
+					leaseInfo.Lease.State = RpcLeaseState.Cancelled;
+					leaseInfo.Lease.Outcome = RpcLeaseOutcome.Failed;
 					leaseInfo.Lease.Output = ByteString.Empty;
 				}
 				else
 				{
-					leaseInfo.Lease.State = (result.Outcome == LeaseOutcome.Cancelled) ? LeaseState.Cancelled : LeaseState.Completed;
-					leaseInfo.Lease.Outcome = result.Outcome;
+					leaseInfo.Lease.State = (result.Outcome == LeaseOutcome.Cancelled) ? RpcLeaseState.Cancelled : RpcLeaseState.Completed;
+					leaseInfo.Lease.Outcome = (RpcLeaseOutcome)result.Outcome;
 					leaseInfo.Lease.Output = (result.Output != null) ? ByteString.CopyFrom(result.Output) : ByteString.Empty;
 				}
 				_logger.LogInformation("Transitioning lease {LeaseId} to {State}, outcome={Outcome}", leaseInfo.Lease.Id, leaseInfo.Lease.State, leaseInfo.Lease.Outcome);
diff --git a/Engine/Source/Programs/Horde/Horde.Agent/Leases/LeaseResult.cs b/Engine/Source/Programs/Horde/Horde.Agent/Leases/LeaseResult.cs
--- a/Engine/Source/Programs/Horde/Horde.Agent/Leases/LeaseResult.cs
+++ b/Engine/Source/Programs/Horde/Horde.Agent/Leases/LeaseResult.cs
@@ -1,7 +1,7 @@
 // Copyright Epic Games, Inc. All Rights Reserved.
 
 using Horde.Agent.Services;
-using HordeCommon;
+using EpicGames.Horde.Agents.Leases;
 
 namespace Horde.Agent.Leases
 {
diff --git a/Engine/Source/Programs/Horde/Horde.Server.Tests/Agents/AgentControllerTest.cs b/Engine/Source/Programs/Horde/Horde.Server.Tests/Agents/AgentControllerTest.cs
--- a/Engine/Source/Programs/Horde/Horde.Server.Tests/Agents/AgentControllerTest.cs
+++ b/Engine/Source/Programs/Horde/Horde.Server.Tests/Agents/AgentControllerTest.cs
@@ -1,6 +1,7 @@
 // Copyright Epic Games, Inc. All Rights Reserved.
 
 using System.Threading.Tasks;
+using EpicGames.Horde.Agents;
 using Horde.Server.Agents;
 using Microsoft.AspNetCore.Mvc;
 using Microsoft.VisualStudio.TestTools.UnitTesting;
diff --git a/Engine/Source/Programs/Horde/Horde.Server.Tests/Agents/Leases/LeasesControllerDbTest.cs b/Engine/Source/Programs/Horde/Horde.Server.Tests/Agents/Leases/LeasesControllerDbTest.cs
--- a/Engine/Source/Programs/Horde/Horde.Server.Tests/Agents/Leases/LeasesControllerDbTest.cs
+++ b/Engine/Source/Programs/Horde/Horde.Server.Tests/Agents/Leases/LeasesControllerDbTest.cs
@@ -3,16 +3,14 @@
 using System;
 using System.Collections.Generic;
 using System.Threading.Tasks;
-using Horde.Server.Agents;
+using EpicGames.Horde.Agents;
+using EpicGames.Horde.Agents.Leases;
+using EpicGames.Horde.Agents.Sessions;
 using Horde.Server.Agents.Leases;
 using Horde.Server.Agents.Sessions;
 using Horde.Server.Utilities;
-using HordeCommon;
 using Microsoft.AspNetCore.Mvc;
 using Microsoft.VisualStudio.TestTools.UnitTesting;
-using EpicGames.Horde.Agents;
-using EpicGames.Horde.Agents.Leases;
-using EpicGames.Horde.Agents.Sessions;
 
 namespace Horde.Server.Tests.Agents.Leases
 {
@@ -30,8 +28,9 @@ namespace Horde.Server.Tests.Agents.Leases
 
 			ILease lease1 = await CreateLeaseAsync(Clock.UtcNow - TimeSpan.FromMinutes(10), TimeSpan.FromMinutes(6));
 			ILease lease2 = await CreateLeaseAsync(Clock.UtcNow - TimeSpan.FromMinutes(7), TimeSpan.FromMinutes(3.1));
-			/*ILease outOfTimeWindow = */await CreateLeaseAsync(Clock.UtcNow - TimeSpan.FromMinutes(7), TimeSpan.FromMinutes(25));
-			
+			/*ILease outOfTimeWindow = */
+			await CreateLeaseAsync(Clock.UtcNow - TimeSpan.FromMinutes(7), TimeSpan.FromMinutes(25));
+
 			ActionResult<List<object>> res = await LeasesController.FindLeasesAsync(null, null, null, null, minTime, maxTime);
 			Assert.AreEqual(2, res.Value!.Count);
 			Assert.AreEqual(lease2.Id, (res.Value[0] as GetAgentLeaseResponse)!.Id);
diff --git a/Engine/Source/Programs/Horde/Horde.Server.Tests/Fleet/LeaseUtilizationAwsMetricStrategyTests.cs b/Engine/Source/Programs/Horde/Horde.Server.Tests/Fleet/LeaseUtilizationAwsMetricStrategyTests.cs
--- a/Engine/Source/Programs/Horde/Horde.Server.Tests/Fleet/LeaseUtilizationAwsMetricStrategyTests.cs
+++ b/Engine/Source/Programs/Horde/Horde.Server.Tests/Fleet/LeaseUtilizationAwsMetricStrategyTests.cs
@@ -5,6 +5,9 @@ using System.Collections.Generic;
 using System.Linq;
 using System.Threading.Tasks;
 using Amazon.CloudWatch.Model;
+using EpicGames.Horde.Agents.Leases;
+using EpicGames.Horde.Agents.Pools;
+using EpicGames.Horde.Streams;
 using Google.Protobuf;
 using Google.Protobuf.WellKnownTypes;
 using Horde.Server.Agents;
@@ -12,13 +15,9 @@ using Horde.Server.Agents.Fleet;
 using Horde.Server.Agents.Leases;
 using Horde.Server.Agents.Pools;
 using Horde.Server.Utilities;
-using HordeCommon;
 using HordeCommon.Rpc.Tasks;
-using Microsoft.VisualStudio.TestTools.UnitTesting;
-using EpicGames.Horde.Agents.Leases;
-using EpicGames.Horde.Streams;
 using Microsoft.Extensions.Logging.Abstractions;
-using EpicGames.Horde.Agents.Pools;
+using Microsoft.VisualStudio.TestTools.UnitTesting;
 
 namespace Horde.Server.Tests.Fleet
 {
@@ -26,7 +25,7 @@ namespace Horde.Server.Tests.Fleet
 	public class LeaseUtilizationAwsMetricStrategyTest : TestSetup
 	{
 		private readonly IPool _pool;
-		private readonly LeaseUtilizationAwsMetricSettings _settings = new (60, "myNamespace");
+		private readonly LeaseUtilizationAwsMetricSettings _settings = new(60, "myNamespace");
 		private readonly FakeAmazonCloudWatch _cloudWatch = new();
 		private readonly LeaseUtilizationAwsMetricStrategy _strategy;
 		private readonly NullLogger<LeaseUtilizationAwsMetricStrategy> _logger = NullLogger<LeaseUtilizationAwsMetricStrategy>.Instance;
@@ -34,12 +33,12 @@ namespace Horde.Server.Tests.Fleet
 		public LeaseUtilizationAwsMetricStrategyTest()
 		{
 			_pool = CreatePoolAsync(new() { Name = "my-pool", EnableAutoscaling = true, MinAgents = 0, NumReserveAgents = 0 }).Result;
-			_strategy = new (LeaseCollection, _cloudWatch.Get(), _settings, Clock, _logger);
+			_strategy = new(LeaseCollection, _cloudWatch.Get(), _settings, Clock, _logger);
 		}
 
 		private async Task<List<IAgent>> CreateAgentsAsync(IPool pool, int numAgents)
 		{
-			List<IAgent> agents = new (numAgents);
+			List<IAgent> agents = new(numAgents);
 			for (int i = 0; i < numAgents; i++)
 			{
 				agents.Add(await CreateAgentAsync(pool));
@@ -47,7 +46,7 @@ namespace Horde.Server.Tests.Fleet
 
 			return agents;
 		}
-		
+
 		[TestMethod]
 		public async Task MetadataAsync()
 		{
@@ -65,7 +64,7 @@ namespace Horde.Server.Tests.Fleet
 			Assert.AreEqual("Pool", datums[0].Dimensions[0].Name);
 			Assert.AreEqual("my-pool", datums[0].Dimensions[0].Value);
 		}
-		
+
 		[TestMethod]
 		public async Task UtilizationZeroAsync()
 		{
@@ -78,7 +77,7 @@ namespace Horde.Server.Tests.Fleet
 			// Assert
 			Assert.AreEqual(0.0, _cloudWatch.GetMetricData("myNamespace")[0].Value, 0.0001);
 		}
-		
+
 		[TestMethod]
 		public async Task NoAgentsAsync()
 		{
@@ -91,7 +90,7 @@ namespace Horde.Server.Tests.Fleet
 			// Assert
 			Assert.AreEqual(0.0, _cloudWatch.GetMetricData("myNamespace")[0].Value, 0.0001);
 		}
-		
+
 		[TestMethod]
 		public async Task UtilizationHalfAsync()
 		{
@@ -106,7 +105,7 @@ namespace Horde.Server.Tests.Fleet
 			// Assert
 			Assert.AreEqual(50.0, _cloudWatch.GetMetricData("myNamespace")[0].Value, 0.0001);
 		}
-		
+
 		[TestMethod]
 		public async Task UtilizationFullAsync()
 		{
@@ -123,7 +122,7 @@ namespace Horde.Server.Tests.Fleet
 			// Assert
 			Assert.AreEqual(100.0, _cloudWatch.GetMetricData("myNamespace")[0].Value, 0.0001);
 		}
-		
+
 		[TestMethod]
 		public async Task AgentMissingInParametersAsync()
 		{
@@ -141,7 +140,7 @@ namespace Horde.Server.Tests.Fleet
 			// Assert
 			Assert.AreEqual(100.0, _cloudWatch.GetMetricData("myNamespace")[0].Value, 0.0001);
 		}
-		
+
 		[TestMethod]
 		public async Task IgnoreOtherPoolsAsync()
 		{
@@ -160,7 +159,7 @@ namespace Horde.Server.Tests.Fleet
 			// Assert
 			Assert.AreEqual(50.0, _cloudWatch.GetMetricData("myNamespace")[0].Value, 0.0001);
 		}
-		
+
 		[TestMethod]
 		public async Task MoreLeasesThanAgentsAsync()
 		{
@@ -180,7 +179,7 @@ namespace Horde.Server.Tests.Fleet
 			// Assert
 			Assert.AreEqual(100.0, _cloudWatch.GetMetricData("myNamespace")[0].Value, 0.0001);
 		}
-		
+
 		[TestMethod]
 		public async Task LeasesInProgressAsync()
 		{
@@ -207,7 +206,7 @@ namespace Horde.Server.Tests.Fleet
 			ExecuteJob,
 			Compute
 		}
-		
+
 		private async Task<ILease> AddPlaceholderLeaseAsync(IAgent agent, IPool pool, DateTime startTime, TimeSpan? duration, LeaseType leaseType = LeaseType.ExecuteJob)
 		{
 			Assert.IsNotNull(agent.SessionId);
@@ -217,7 +216,7 @@ namespace Horde.Server.Tests.Fleet
 			executeJobTask.JobName = "placeholderJobName";
 			payload = Any.Pack(executeJobTask).ToByteArray();
 			PoolId? poolId = pool.Id;
-			
+
 			if (leaseType == LeaseType.Compute)
 			{
 				ComputeTask computeTask = new();
@@ -226,14 +225,14 @@ namespace Horde.Server.Tests.Fleet
 				payload = Any.Pack(computeTask).ToByteArray();
 				poolId = null;
 			}
-			
+
 			ILease lease = await LeaseCollection.AddAsync(new LeaseId(BinaryIdUtils.CreateNew()), null, "placeholderLease", agent.Id, agent.SessionId!.Value, new StreamId("placeholderStream"), poolId, null, startTime, payload);
 			if (duration != null)
 			{
 				bool wasModified = await LeaseCollection.TrySetOutcomeAsync(lease.Id, startTime + duration.Value, LeaseOutcome.Success, null);
-				Assert.IsTrue(wasModified);	
+				Assert.IsTrue(wasModified);
 			}
-			
+
 			return lease;
 		}
 	}
diff --git a/Engine/Source/Programs/Horde/Horde.Server.Tests/Fleet/LeaseUtilizationStrategyTest.cs b/Engine/Source/Programs/Horde/Horde.Server.Tests/Fleet/LeaseUtilizationStrategyTest.cs
--- a/Engine/Source/Programs/Horde/Horde.Server.Tests/Fleet/LeaseUtilizationStrategyTest.cs
+++ b/Engine/Source/Programs/Horde/Horde.Server.Tests/Fleet/LeaseUtilizationStrategyTest.cs
@@ -3,6 +3,8 @@
 using System;
 using System.Collections.Generic;
 using System.Threading.Tasks;
+using EpicGames.Horde.Agents.Leases;
+using EpicGames.Horde.Streams;
 using Google.Protobuf;
 using Google.Protobuf.WellKnownTypes;
 using Horde.Server.Agents;
@@ -10,11 +12,8 @@ using Horde.Server.Agents.Fleet;
 using Horde.Server.Agents.Leases;
 using Horde.Server.Agents.Pools;
 using Horde.Server.Utilities;
-using HordeCommon;
 using HordeCommon.Rpc.Tasks;
 using Microsoft.VisualStudio.TestTools.UnitTesting;
-using EpicGames.Horde.Streams;
-using EpicGames.Horde.Agents.Leases;
 
 namespace Horde.Server.Tests.Fleet
 {
@@ -41,25 +40,25 @@ namespace Horde.Server.Tests.Fleet
 		{
 			IPool pool = await CreatePoolAsync(new PoolConfig { Name = "AutoscalePool1", EnableAutoscaling = true, MinAgents = 0, NumReserveAgents = 0 });
 			CreateAgents(pool);
-			
+
 			await AddPlaceholderLeaseAsync(_agent1, pool, Clock.UtcNow - TimeSpan.FromMinutes(120), TimeSpan.FromMinutes(120));
 			await AddPlaceholderLeaseAsync(_agent2, pool, Clock.UtcNow - TimeSpan.FromMinutes(120), TimeSpan.FromMinutes(120));
 			await AddPlaceholderLeaseAsync(_agent3, pool, Clock.UtcNow - TimeSpan.FromMinutes(120), TimeSpan.FromMinutes(120));
 			await AddPlaceholderLeaseAsync(_agent4, pool, Clock.UtcNow - TimeSpan.FromMinutes(120), TimeSpan.FromMinutes(120), LeaseType.Compute);
 			await AssertPoolSizeAsync(pool, 4);
 		}
-		
+
 		[TestMethod]
 		public async Task UtilizationHalfAsync()
 		{
 			IPool pool = await CreatePoolAsync(new PoolConfig { Name = "AutoscalePool1", EnableAutoscaling = true, MinAgents = 0, NumReserveAgents = 0 });
 			CreateAgents(pool);
-			
+
 			await AddPlaceholderLeaseAsync(_agent1, pool, Clock.UtcNow - TimeSpan.FromMinutes(120), TimeSpan.FromMinutes(120));
 			await AddPlaceholderLeaseAsync(_agent2, pool, Clock.UtcNow - TimeSpan.FromMinutes(120), TimeSpan.FromMinutes(120));
 			await AssertPoolSizeAsync(pool, 2);
 		}
-		
+
 		[TestMethod]
 		public async Task UtilizationZeroAsync()
 		{
@@ -68,28 +67,28 @@ namespace Horde.Server.Tests.Fleet
 
 			await AssertPoolSizeAsync(pool, 0);
 		}
-		
+
 		[TestMethod]
 		public async Task ReserveAgentsAsync()
 		{
 			IPool pool = await CreatePoolAsync(new PoolConfig { Name = "AutoscalePool1", EnableAutoscaling = true, MinAgents = 0, NumReserveAgents = 5 });
 			CreateAgents(pool);
-			
+
 			await AddPlaceholderLeaseAsync(_agent1, pool, Clock.UtcNow - TimeSpan.FromMinutes(120), TimeSpan.FromMinutes(120));
 			await AddPlaceholderLeaseAsync(_agent2, pool, Clock.UtcNow - TimeSpan.FromMinutes(120), TimeSpan.FromMinutes(120));
 			await AddPlaceholderLeaseAsync(_agent3, pool, Clock.UtcNow - TimeSpan.FromMinutes(120), TimeSpan.FromMinutes(120));
 			await AddPlaceholderLeaseAsync(_agent4, pool, Clock.UtcNow - TimeSpan.FromMinutes(120), TimeSpan.FromMinutes(120));
-			
+
 			// Full utilization should mean all agents plus the reserve agents
 			await AssertPoolSizeAsync(pool, 9);
 		}
-		
+
 		[TestMethod]
 		public async Task MinAgentsAsync()
 		{
 			IPool pool = await CreatePoolAsync(new PoolConfig { Name = "AutoscalePool1", EnableAutoscaling = true, MinAgents = 2, NumReserveAgents = 0 });
 			CreateAgents(pool);
-			
+
 			// Even with no utilization, expect at least the min number of agents
 			await AssertPoolSizeAsync(pool, 2);
 		}
@@ -106,7 +105,7 @@ namespace Horde.Server.Tests.Fleet
 			ExecuteJob,
 			Compute
 		}
-		
+
 		private async Task<ILease> AddPlaceholderLeaseAsync(IAgent agent, IPool pool, DateTime startTime, TimeSpan duration, LeaseType leaseType = LeaseType.ExecuteJob)
 		{
 			Assert.IsNotNull(agent.SessionId);
@@ -115,7 +114,7 @@ namespace Horde.Server.Tests.Fleet
 			ExecuteJobTask executeJobTask = new();
 			executeJobTask.JobName = "placeholderJobName";
 			payload = Any.Pack(executeJobTask).ToByteArray();
-			
+
 			if (leaseType == LeaseType.Compute)
 			{
 				ComputeTask computeTask = new();
diff --git a/Engine/Source/Programs/Horde/Horde.Server/Agents/AgentMessages.cs b/Engine/Source/Programs/Horde/Horde.Server/Agents/AgentMessages.cs
deleted file mode 100644
diff --git a/Engine/Source/Programs/Horde/Horde.Server/Agents/AgentService.cs b/Engine/Source/Programs/Horde/Horde.Server/Agents/AgentService.cs
--- a/Engine/Source/Programs/Horde/Horde.Server/Agents/AgentService.cs
+++ b/Engine/Source/Programs/Horde/Horde.Server/Agents/AgentService.cs
@@ -25,6 +25,7 @@ using Horde.Server.Server;
 using Horde.Server.Tasks;
 using Horde.Server.Utilities;
 using HordeCommon;
+using HordeCommon.Rpc;
 using Microsoft.Extensions.Hosting;
 using Microsoft.Extensions.Logging;
 using StackExchange.Redis;
@@ -726,7 +727,7 @@ namespace Horde.Server.Agents
 					if (lease.State == LeaseState.Cancelled)
 					{
 						HordeCommon.Rpc.Messages.Lease? newLease;
-						if (!leaseIdToNewState.TryGetValue(lease.Id, out newLease) || newLease.State == LeaseState.Cancelled || newLease.State == LeaseState.Completed)
+						if (!leaseIdToNewState.TryGetValue(lease.Id, out newLease) || newLease.State == RpcLeaseState.Cancelled || newLease.State == RpcLeaseState.Completed)
 						{
 							await RemoveLeaseAsync(agent, lease, utcNow, LeaseOutcome.Cancelled, null, cancellationToken);
 							leases.RemoveAt(idx--);
@@ -736,14 +737,14 @@ namespace Horde.Server.Agents
 					else
 					{
 						HordeCommon.Rpc.Messages.Lease? newLease;
-						if (leaseIdToNewState.TryGetValue(lease.Id, out newLease) && newLease.State != lease.State)
+						if (leaseIdToNewState.TryGetValue(lease.Id, out newLease) && (LeaseState)newLease.State != lease.State)
 						{
-							if (newLease.State == LeaseState.Cancelled || newLease.State == LeaseState.Completed)
+							if (newLease.State == RpcLeaseState.Cancelled || newLease.State == RpcLeaseState.Completed)
 							{
-								await RemoveLeaseAsync(agent, lease, utcNow, newLease.Outcome, newLease.Output.ToByteArray(), cancellationToken);
+								await RemoveLeaseAsync(agent, lease, utcNow, (LeaseOutcome)newLease.Outcome, newLease.Output.ToByteArray(), cancellationToken);
 								leases.RemoveAt(idx--);
 							}
-							else if (newLease.State == LeaseState.Active && lease.State == LeaseState.Pending)
+							else if (newLease.State == RpcLeaseState.Active && lease.State == LeaseState.Pending)
 							{
 								lease.State = LeaseState.Active;
 							}
diff --git a/Engine/Source/Programs/Horde/Horde.Server/Agents/IAgent.cs b/Engine/Source/Programs/Horde/Horde.Server/Agents/IAgent.cs
--- a/Engine/Source/Programs/Horde/Horde.Server/Agents/IAgent.cs
+++ b/Engine/Source/Programs/Horde/Horde.Server/Agents/IAgent.cs
@@ -23,6 +23,7 @@ using Horde.Server.Perforce;
 using Horde.Server.Server;
 using Horde.Server.Streams;
 using HordeCommon;
+using HordeCommon.Rpc;
 using HordeCommon.Rpc.Messages;
 using HordeCommon.Rpc.Tasks;
 using MongoDB.Bson;
@@ -407,7 +408,7 @@ namespace Horde.Server.Agents
 			HordeCommon.Rpc.Messages.Lease lease = new HordeCommon.Rpc.Messages.Lease();
 			lease.Id = Id.ToString();
 			lease.Payload = Google.Protobuf.WellKnownTypes.Any.Parser.ParseFrom(Payload); 
-			lease.State = State;
+			lease.State = (RpcLeaseState)State;
 			return lease;
 		}
 	}
diff --git a/Engine/Source/Programs/Horde/Horde.Server/Agents/Leases/ILease.cs b/Engine/Source/Programs/Horde/Horde.Server/Agents/Leases/ILease.cs
--- a/Engine/Source/Programs/Horde/Horde.Server/Agents/Leases/ILease.cs
+++ b/Engine/Source/Programs/Horde/Horde.Server/Agents/Leases/ILease.cs
@@ -1,15 +1,14 @@
 // Copyright Epic Games, Inc. All Rights Reserved.
 
 using System;
-using Google.Protobuf;
-using Google.Protobuf.WellKnownTypes;
-using HordeCommon;
-using EpicGames.Horde.Streams;
 using EpicGames.Horde.Agents;
 using EpicGames.Horde.Agents.Leases;
-using EpicGames.Horde.Logs;
 using EpicGames.Horde.Agents.Pools;
 using EpicGames.Horde.Agents.Sessions;
+using EpicGames.Horde.Logs;
+using EpicGames.Horde.Streams;
+using Google.Protobuf;
+using Google.Protobuf.WellKnownTypes;
 
 namespace Horde.Server.Agents.Leases
 {
diff --git a/Engine/Source/Programs/Horde/Horde.Server/Agents/Leases/ILeaseCollection.cs b/Engine/Source/Programs/Horde/Horde.Server/Agents/Leases/ILeaseCollection.cs
--- a/Engine/Source/Programs/Horde/Horde.Server/Agents/Leases/ILeaseCollection.cs
+++ b/Engine/Source/Programs/Horde/Horde.Server/Agents/Leases/ILeaseCollection.cs
@@ -2,15 +2,14 @@
 
 using System;
 using System.Collections.Generic;
+using System.Threading;
 using System.Threading.Tasks;
-using HordeCommon;
 using EpicGames.Horde.Agents;
 using EpicGames.Horde.Agents.Leases;
-using EpicGames.Horde.Streams;
-using EpicGames.Horde.Logs;
 using EpicGames.Horde.Agents.Pools;
 using EpicGames.Horde.Agents.Sessions;
-using System.Threading;
+using EpicGames.Horde.Logs;
+using EpicGames.Horde.Streams;
 
 namespace Horde.Server.Agents.Leases
 {
diff --git a/Engine/Source/Programs/Horde/Horde.Server/Agents/Leases/LeaseCollection.cs b/Engine/Source/Programs/Horde/Horde.Server/Agents/Leases/LeaseCollection.cs
--- a/Engine/Source/Programs/Horde/Horde.Server/Agents/Leases/LeaseCollection.cs
+++ b/Engine/Source/Programs/Horde/Horde.Server/Agents/Leases/LeaseCollection.cs
@@ -13,7 +13,6 @@ using EpicGames.Horde.Logs;
 using EpicGames.Horde.Streams;
 using Horde.Server.Server;
 using Horde.Server.Utilities;
-using HordeCommon;
 using MongoDB.Bson.Serialization.Attributes;
 using MongoDB.Driver;
 using OpenTelemetry.Trace;
diff --git a/Engine/Source/Programs/Horde/Horde.Server/Compute/ComputeTaskSource.cs b/Engine/Source/Programs/Horde/Horde.Server/Compute/ComputeTaskSource.cs
--- a/Engine/Source/Programs/Horde/Horde.Server/Compute/ComputeTaskSource.cs
+++ b/Engine/Source/Programs/Horde/Horde.Server/Compute/ComputeTaskSource.cs
@@ -15,7 +15,6 @@ using Horde.Server.Agents.Leases;
 using Horde.Server.Agents.Relay;
 using Horde.Server.Server;
 using Horde.Server.Tasks;
-using HordeCommon;
 using HordeCommon.Rpc.Tasks;
 using Microsoft.Extensions.Logging;
 using Microsoft.Extensions.Options;
@@ -32,7 +31,7 @@ namespace Horde.Server.Compute
 		/// In direct connection mode, these two are identical.
 		/// </summary>
 		public int Port { get; }
-		
+
 		/// <summary>
 		/// Port the local process on the agent is listening on
 		/// </summary>
@@ -83,7 +82,7 @@ namespace Horde.Server.Compute
 			return $"Port={Port}, AgentPort={AgentPort}";
 		}
 	}
-	
+
 	/// <summary>
 	/// Information about a compute 
 	/// </summary>
@@ -91,15 +90,15 @@ namespace Horde.Server.Compute
 	{
 		/// <inheritdoc cref="AssignComputeResponse.ConnectionMode" />
 		public ConnectionMode ConnectionMode { get; }
-		
+
 		/// <summary>
 		/// IP address of the agent
 		/// </summary>
 		public IPAddress Ip { get; }
-		
+
 		/// <inheritdoc cref="AssignComputeResponse.ConnectionAddress" />
 		public string? ConnectionAddress { get; }
-		
+
 		/// <inheritdoc cref="AssignComputeResponse.Ports" />
 		public IReadOnlyDictionary<string, ComputeResourcePort> Ports { get; }
 
@@ -189,7 +188,7 @@ namespace Horde.Server.Compute
 		{
 			return Task.FromResult(WaitInternalAsync(agent, cancellationToken));
 		}
-		
+
 		/// <inheritdoc/>
 		public override async Task OnLeaseFinishedAsync(IAgent agent, LeaseId leaseId, ComputeTask payload, LeaseOutcome outcome, ReadOnlyMemory<byte> output, ILogger logger, CancellationToken cancellationToken)
 		{
diff --git a/Engine/Source/Programs/Horde/Horde.Server/Server/ConsistencyService.cs b/Engine/Source/Programs/Horde/Horde.Server/Server/ConsistencyService.cs
--- a/Engine/Source/Programs/Horde/Horde.Server/Server/ConsistencyService.cs
+++ b/Engine/Source/Programs/Horde/Horde.Server/Server/ConsistencyService.cs
@@ -13,6 +13,7 @@ using Microsoft.Extensions.Logging;
 using Horde.Server.Agents;
 using EpicGames.Horde.Agents;
 using EpicGames.Horde.Agents.Sessions;
+using EpicGames.Horde.Agents.Leases;
 
 namespace Horde.Server.Server
 {
diff --git a/Engine/Source/Programs/Horde/Horde.Server/Tasks/ITaskSource.cs b/Engine/Source/Programs/Horde/Horde.Server/Tasks/ITaskSource.cs
--- a/Engine/Source/Programs/Horde/Horde.Server/Tasks/ITaskSource.cs
+++ b/Engine/Source/Programs/Horde/Horde.Server/Tasks/ITaskSource.cs
@@ -9,13 +9,12 @@ using System.Text;
 using System.Threading;
 using System.Threading.Tasks;
 using EpicGames.Core;
+using EpicGames.Horde.Agents.Leases;
 using Google.Protobuf;
 using Google.Protobuf.Reflection;
 using Google.Protobuf.WellKnownTypes;
 using Horde.Server.Agents;
-using HordeCommon;
 using Microsoft.Extensions.Logging;
-using EpicGames.Horde.Agents.Leases;
 
 namespace Horde.Server.Tasks
 {
@@ -123,7 +122,7 @@ namespace Horde.Server.Tasks
 		/// <param name="cancellationToken">Cancellation token for the operation</param>
 		ValueTask GetLeaseDetailsAsync(Any payload, Dictionary<string, string> details, CancellationToken cancellationToken);
 	}
-	
+
 	/// <summary>
 	/// Base implementation of <see cref="ITaskSource"/>
 	/// </summary>
@@ -233,36 +232,36 @@ namespace Horde.Server.Tasks
 			return default;
 		}
 
-        /// <summary>
-        /// Creates a lease task which will wait until the given cancellation token is signalled.
-        /// </summary>
-        /// <param name="token">The cancellation token</param>
-        /// <returns>Lease task</returns>
-        protected static Task<AgentLease?> SkipAsync(CancellationToken token)
-        {
+		/// <summary>
+		/// Creates a lease task which will wait until the given cancellation token is signalled.
+		/// </summary>
+		/// <param name="token">The cancellation token</param>
+		/// <returns>Lease task</returns>
+		protected static Task<AgentLease?> SkipAsync(CancellationToken token)
+		{
 			_ = token;
-            return Task.FromResult<AgentLease?>(null);
-        }
-
-        /// <summary>
-        /// Waits until the cancellation token is signalled, then return an cancelled lease task.
-        /// </summary>
-        /// <param name="token">The cancellation token</param>
-        /// <returns>Lease task</returns>
-        protected static async Task<Task<AgentLease?>> DrainAsync(CancellationToken token)
-        {
-            await token.AsTask();
-            return Task.FromResult<AgentLease?>(null);
-        }
-
-        /// <summary>
-        /// Creates a lease task from a given lease
-        /// </summary>
-        /// <param name="lease">Lease to create the task from</param>
-        /// <returns></returns>
-        protected static Task<AgentLease?> LeaseAsync(AgentLease lease)
-        {
-            return Task.FromResult<AgentLease?>(lease);
-        }
-    }
+			return Task.FromResult<AgentLease?>(null);
+		}
+
+		/// <summary>
+		/// Waits until the cancellation token is signalled, then return an cancelled lease task.
+		/// </summary>
+		/// <param name="token">The cancellation token</param>
+		/// <returns>Lease task</returns>
+		protected static async Task<Task<AgentLease?>> DrainAsync(CancellationToken token)
+		{
+			await token.AsTask();
+			return Task.FromResult<AgentLease?>(null);
+		}
+
+		/// <summary>
+		/// Creates a lease task from a given lease
+		/// </summary>
+		/// <param name="lease">Lease to create the task from</param>
+		/// <returns></returns>
+		protected static Task<AgentLease?> LeaseAsync(AgentLease lease)
+		{
+			return Task.FromResult<AgentLease?>(lease);
+		}
+	}
 }
diff --git a/Engine/Source/Programs/Horde/Horde.Server/Tasks/RestartTaskSource.cs b/Engine/Source/Programs/Horde/Horde.Server/Tasks/RestartTaskSource.cs
--- a/Engine/Source/Programs/Horde/Horde.Server/Tasks/RestartTaskSource.cs
+++ b/Engine/Source/Programs/Horde/Horde.Server/Tasks/RestartTaskSource.cs
@@ -2,16 +2,15 @@
 
 using System.Threading;
 using System.Threading.Tasks;
+using EpicGames.Horde.Agents.Leases;
+using EpicGames.Horde.Jobs;
+using EpicGames.Horde.Logs;
 using Google.Protobuf;
 using Google.Protobuf.WellKnownTypes;
 using Horde.Server.Agents;
 using Horde.Server.Logs;
 using Horde.Server.Utilities;
-using HordeCommon;
 using HordeCommon.Rpc.Tasks;
-using EpicGames.Horde.Agents.Leases;
-using EpicGames.Horde.Jobs;
-using EpicGames.Horde.Logs;
 
 namespace Horde.Server.Tasks
 {
diff --git a/Engine/Source/Programs/Horde/Horde.Server/Tasks/ShutdownTaskSource.cs b/Engine/Source/Programs/Horde/Horde.Server/Tasks/ShutdownTaskSource.cs
--- a/Engine/Source/Programs/Horde/Horde.Server/Tasks/ShutdownTaskSource.cs
+++ b/Engine/Source/Programs/Horde/Horde.Server/Tasks/ShutdownTaskSource.cs
@@ -2,16 +2,15 @@
 
 using System.Threading;
 using System.Threading.Tasks;
+using EpicGames.Horde.Agents.Leases;
+using EpicGames.Horde.Jobs;
+using EpicGames.Horde.Logs;
 using Google.Protobuf;
 using Google.Protobuf.WellKnownTypes;
 using Horde.Server.Agents;
 using Horde.Server.Logs;
 using Horde.Server.Utilities;
-using HordeCommon;
 using HordeCommon.Rpc.Tasks;
-using EpicGames.Horde.Agents.Leases;
-using EpicGames.Horde.Jobs;
-using EpicGames.Horde.Logs;
 
 namespace Horde.Server.Tasks
 {
diff --git a/Engine/Source/Programs/Shared/EpicGames.Horde/Agents/AgentMessages.cs b/Engine/Source/Programs/Shared/EpicGames.Horde/Agents/AgentMessages.cs
new file mode 100644
--- /dev/null
+++ b/Engine/Source/Programs/Shared/EpicGames.Horde/Agents/AgentMessages.cs
@@ -0,0 +1,94 @@
+// Copyright Epic Games, Inc. All Rights Reserved.
+
+using System;
+using System.Collections.Generic;
+using EpicGames.Horde.Agents.Leases;
+using EpicGames.Horde.Agents.Sessions;
+using EpicGames.Horde.Logs;
+
+namespace EpicGames.Horde.Agents
+{
+	/// <summary>
+	/// Parameters to update an agent
+	/// </summary>
+	/// <param name="Enabled"> Whether the agent is currently enabled </param>
+	/// <param name="Ephemeral"> Whether the agent is ephemeral </param>
+	/// <param name="RequestConform"> Request a conform be performed using the current agent set </param>
+	/// <param name="RequestFullConform"> Request that a full conform be performed, removing all intermediate files </param>
+	/// <param name="RequestRestart"> Request the machine be restarted </param>
+	/// <param name="RequestShutdown"> Request the machine be shut down </param>
+	/// <param name="RequestForceRestart"> Request the machine be restarted without waiting for leases to complete </param>
+	/// <param name="Pools"> Pools for this agent </param>
+	/// <param name="Comment"> New comment </param>
+	public record UpdateAgentRequest(bool? Enabled = null, bool? Ephemeral = null, bool? RequestConform = null, bool? RequestFullConform = null, bool? RequestRestart = null, bool? RequestShutdown = null, bool? RequestForceRestart = null, List<string>? Pools = null, string? Comment = null);
+
+	/// <summary>
+	/// Response for queries to find a particular lease within an agent
+	/// </summary>
+	/// <param name="Id"> Identifier for the lease </param>
+	/// <param name="ParentId"> Identifier for the parent lease. Used to terminate hierarchies of leases. </param>
+	/// <param name="AgentId"> The agent id </param>
+	/// <param name="AgentRate"> Cost of this agent, per hour </param>
+	/// <param name="Name"> Name of the lease </param>
+	/// <param name="LogId"> Log id for this lease </param>
+	/// <param name="StartTime"> Time at which the lease started (UTC) </param>
+	/// <param name="FinishTime"> Time at which the lease started (UTC) </param>
+	/// <param name="Executing"> Whether this lease has started executing on the agent yet </param>
+	/// <param name="Details"> Details of the payload being executed </param>
+	/// <param name="Outcome"> Outcome of the lease </param>
+	/// <param name="State"> State of the lease (for AgentLeases) </param>
+	public record GetAgentLeaseResponse(LeaseId Id, LeaseId? ParentId, AgentId? AgentId, double? AgentRate, string Name, LogId? LogId, DateTime StartTime, DateTime? FinishTime, bool Executing, Dictionary<string, string>? Details, LeaseOutcome? Outcome, LeaseState? State);
+
+	/// <summary>
+	/// Information about an agent session
+	/// </summary>
+	/// <param name="Id"> Unique id for this session </param>
+	/// <param name="StartTime"> Start time for this session </param>
+	/// <param name="FinishTime"> Finishing time for this session </param>
+	/// <param name="Properties"> Properties of this agent </param>
+	/// <param name="Version"> Version of the software running during this session </param>
+	public record GetAgentSessionResponse(SessionId Id, DateTime StartTime, DateTime? FinishTime, List<string>? Properties, string? Version);
+
+	/// <summary>
+	/// Information about a workspace synced on an agent
+	/// </summary>
+	/// <param name="Cluster"> The Perforce server and port to connect to </param>
+	/// <param name="UserName"> User to log into Perforce with (eg. buildmachine) </param>
+	/// <param name="Identifier"> Identifier to distinguish this workspace from other workspaces </param>
+	/// <param name="Stream"> The stream to sync </param>
+	/// <param name="View"> Custom view for the workspace </param>
+	/// <param name="BIncremental"> Whether to use an incremental workspace </param>
+	/// <param name="Method"> Method to use when syncing/materializing data from Perforce </param>
+	public record GetAgentWorkspaceResponse(string? Cluster, string? UserName, string Identifier, string Stream, List<string>? View, bool BIncremental, string? Method);
+
+	/// <summary>
+	/// Information about an agent
+	/// </summary>
+	/// <param name="Id"> The agent's unique ID </param>
+	/// <param name="Name"> Friendly name of the agent </param>
+	/// <param name="Enabled"> Whether the agent is currently enabled </param>
+	/// <param name="Rate"> Cost estimate per-hour for this agent </param>
+	/// <param name="SessionId"> The current session id </param>
+	/// <param name="Ephemeral"> Whether the agent is ephemeral </param>
+	/// <param name="Online"> Whether the agent is currently online </param>
+	/// <param name="Deleted"> Whether this agent has expired </param>
+	/// <param name="PendingConform"> Whether a conform job is pending </param>
+	/// <param name="PendingFullConform"> Whether a full conform job is pending </param>
+	/// <param name="PendingRestart"> Whether a restart is pending </param>
+	/// <param name="PendingShutdown"> Whether a shutdown is pending </param>
+	/// <param name="LastShutdownReason"> The reason for the last shutdown </param>
+	/// <param name="LastConformTime"> Last time a conform was attempted </param>
+	/// <param name="ConformAttemptCount"> Number of times a conform has been attempted </param>
+	/// <param name="NextConformTime"> Last time a conform was attempted </param>
+	/// <param name="Version"> The current client version </param>
+	/// <param name="Properties"> Properties for the agent </param>
+	/// <param name="Resources"> Resources for the agent </param>
+	/// <param name="UpdateTime"> Last update time of this agent </param>
+	/// <param name="LastStatusChange"> Last time agent's status was changed </param>
+	/// <param name="Pools"> Pools for this agent </param>
+	/// <param name="Capabilities"> Capabilities of this agent </param>
+	/// <param name="Leases"> Array of active leases. </param>
+	/// <param name="Workspaces">Current workspaces synced on the agent</param>
+	/// <param name="Comment"> Comment for this agent </param>
+	public record GetAgentResponse(AgentId Id, string Name, bool Enabled, double? Rate, SessionId? SessionId, bool Ephemeral, bool Online, bool Deleted, bool PendingConform, bool PendingFullConform, bool PendingRestart, bool PendingShutdown, string LastShutdownReason, DateTime LastConformTime, int? ConformAttemptCount, DateTime? NextConformTime, string? Version, List<string> Properties, Dictionary<string, int> Resources, DateTime? UpdateTime, DateTime? LastStatusChange, List<string>? Pools, object? Capabilities, List<GetAgentLeaseResponse> Leases, List<GetAgentWorkspaceResponse> Workspaces, string? Comment);
+}
diff --git a/Engine/Source/Programs/Shared/EpicGames.Horde/Agents/Leases/LeaseMessages.cs b/Engine/Source/Programs/Shared/EpicGames.Horde/Agents/Leases/LeaseMessages.cs
--- a/Engine/Source/Programs/Shared/EpicGames.Horde/Agents/Leases/LeaseMessages.cs
+++ b/Engine/Source/Programs/Shared/EpicGames.Horde/Agents/Leases/LeaseMessages.cs
@@ -1,7 +1,66 @@
 // Copyright Epic Games, Inc. All Rights Reserved.
 
+#pragma warning disable CA1027 // Mark enums with FlagsAttribute
+
 namespace EpicGames.Horde.Agents.Leases
 {
+	/// <summary>
+	/// Outcome from a lease. Values must match lease_outcome.proto.
+	/// </summary>
+	public enum LeaseOutcome
+	{
+		/// <summary>
+		/// Default value.
+		/// </summary>
+		Unspecified = 0,
+
+		/// <summary>
+		/// The lease was executed successfully
+		/// </summary>
+		Success = 1,
+
+		/// <summary>
+		/// The lease was not executed succesfully, but cannot be run again.
+		/// </summary>
+		Failed = 2,
+
+		/// <summary>
+		/// The lease was cancelled by request
+		/// </summary>
+		Cancelled = 4
+	}
+
+	/// <summary>
+	/// State of a lease. Values must match lease_state.proto.
+	/// </summary>
+	public enum LeaseState
+	{
+		/// <summary>
+		/// Default value.
+		/// </summary>
+		Unspecified = 0,
+
+		/// <summary>
+		/// Set by the server when waiting for an agent to accept the lease. Once processed, the agent should transition the lease state to active.
+		/// </summary>
+		Pending = 1,
+
+		/// <summary>
+		/// The agent is actively working on this lease.
+		/// </summary>
+		Active = 2,
+
+		/// <summary>
+		/// The agent has finished working on this lease.
+		/// </summary>
+		Completed = 3,
+
+		/// <summary>
+		/// Set by the server to indicate that the lease should be cancelled.
+		/// </summary>
+		Cancelled = 4
+	}
+
 	/// <summary>
 	/// Updates an existing lease
 	/// </summary>
