Add new remove function to remove a key

diff --git a/Engine/Source/Programs/Shared/EpicGames.Core/JsonObject.cs b/Engine/Source/Programs/Shared/EpicGames.Core/JsonObject.cs
--- a/Engine/Source/Programs/Shared/EpicGames.Core/JsonObject.cs
+++ b/Engine/Source/Programs/Shared/EpicGames.Core/JsonObject.cs
@@ -656,6 +656,22 @@ namespace EpicGames.Core
 			return _rawOrderedObject.Contains(fieldName);
 		}
 
+		/// <summary>
+		/// Removes a field from the given json object.
+		/// </summary>
+		/// <param name="fieldName">Name of the field to remove </param>
+		public void RemoveField(string fieldName)
+		{
+			if (String.IsNullOrEmpty(fieldName))
+			{
+				return;
+			}
+			if (ContainsField(fieldName))
+			{
+				_rawOrderedObject.Remove(fieldName);
+			}
+		}
+
 		/// <summary>
 		/// Sets the integer value of a field if it exists. Otherwise, adds the field and value.
 		/// </summary>
@@ -824,6 +840,7 @@ namespace EpicGames.Core
 			string[] stringArray = value.Select(x => x.ToString()!).ToArray();
 			_rawOrderedObject[fieldName] = stringArray;
 		}
+
 		/// <summary>
 		/// Converts this Json Object to a string representation. This follows formatting of UE .uplugin files with 4 spaces for tabs and indentation enabled.
 		/// IMPORTANT: If this JsonObject contains HTML, the returned string should NOT be used directly for HTML or a script. Read the note below.
