Horde: Attempt graceful socket shutdown for 30s before falling back to destroying everything. Also add more diagnostics around errors shutting down when we do so.

diff --git a/Engine/Source/Programs/Shared/EpicGames.Horde/Compute/ComputeSocket.cs b/Engine/Source/Programs/Shared/EpicGames.Horde/Compute/ComputeSocket.cs
--- a/Engine/Source/Programs/Shared/EpicGames.Horde/Compute/ComputeSocket.cs
+++ b/Engine/Source/Programs/Shared/EpicGames.Horde/Compute/ComputeSocket.cs
@@ -757,32 +757,59 @@ namespace EpicGames.Horde.Compute
 		[SuppressMessage("Reliability", "CA2000:Dispose objects before losing scope")]
 		async Task DetachSendBufferAsync(int channelId, CancellationToken cancellationToken)
 		{
-			// Get the current send buffer state
-			SendBuffer? sendBuffer;
-			lock (_lockObject)
+			SendBuffer? sendBuffer = null;
+			try
 			{
-				if (!_sendBuffers.TryGetValue(channelId, out sendBuffer))
+				// Get the current send buffer state
+				lock (_lockObject)
 				{
-					_logger.LogWarning("No buffer is attached to channel {ChannelId}", channelId);
-					return;
+					if (!_sendBuffers.TryGetValue(channelId, out sendBuffer))
+					{
+						_logger.LogWarning("No buffer is attached to channel {ChannelId}", channelId);
+						return;
+					}
+					sendBuffer.AddRef();
 				}
-				sendBuffer.AddRef();
-			}
 
-			// Wait for the send task to complete
-			await sendBuffer.Task.DisposeAsync();
+				// Try to stop the send task
+				try
+				{
+					await sendBuffer.Task.StopAsync(cancellationToken).WaitAsync(TimeSpan.FromSeconds(30.0));
+				}
+				catch (TimeoutException ex)
+				{
+					_logger.LogWarning(ex, "Send task did not terminate gracefully in 30s. Tearing down socket buffers.");
+				}
+				catch (Exception ex)
+				{
+					_logger.LogError(ex, "Exception while stopping send buffer task: {Message}", ex.Message);
+				}
 
-			// Release the reader
-			await sendBuffer.Semaphore.WaitAsync(cancellationToken);
+				// Release the reader
+				await sendBuffer.Semaphore.WaitAsync(cancellationToken);
+				try
+				{
+					sendBuffer._reader?.Dispose();
+					sendBuffer._reader = null;
+				}
+				finally
+				{
+					sendBuffer.Semaphore.Release();
+				}
+			}
+			finally
+			{
+				sendBuffer?.Release(); // Added above
+			}
+
+			// Force the send task to complete
 			try
 			{
-				sendBuffer._reader?.Dispose();
-				sendBuffer._reader = null;
+				await sendBuffer.Task.DisposeAsync();
 			}
-			finally
+			catch (Exception ex)
 			{
-				sendBuffer.Semaphore.Release();
-				sendBuffer.Release(); // Added above
+				_logger.LogError(ex, "Exception disposing send buffer task: {Message}", ex.Message);
 			}
 
 			// Remove the buffer from the dictionary
@@ -816,13 +843,13 @@ namespace EpicGames.Horde.Compute
 					await reader.WaitToReadAsync(1, cancellationToken);
 				}
 			}
-			catch (OperationCanceledException e)
+			catch (OperationCanceledException ex)
 			{
-				_logger.LogDebug(e, "Background send task cancelled for channel {ChannelId}", channelId);
+				_logger.LogDebug(ex, "Background send task cancelled for channel {ChannelId}", channelId);
 			}
-			catch (Exception e)
+			catch (Exception ex)
 			{
-				_logger.LogInformation(e, "Error in background send");
+				_logger.LogInformation(ex, "Error in background send: {Message}", ex.Message);
 				throw;
 			}
 		}
