Horde: Add missing cancellation token to method call.

diff --git a/Engine/Source/Programs/Shared/EpicGames.Horde/Compute/ComputeSocket.cs b/Engine/Source/Programs/Shared/EpicGames.Horde/Compute/ComputeSocket.cs
--- a/Engine/Source/Programs/Shared/EpicGames.Horde/Compute/ComputeSocket.cs
+++ b/Engine/Source/Programs/Shared/EpicGames.Horde/Compute/ComputeSocket.cs
@@ -774,7 +774,11 @@ namespace EpicGames.Horde.Compute
 				// Try to stop the send task
 				try
 				{
-					await sendBuffer.Task.StopAsync(cancellationToken).WaitAsync(TimeSpan.FromSeconds(30.0));
+					await sendBuffer.Task.StopAsync(cancellationToken).WaitAsync(TimeSpan.FromSeconds(30.0), cancellationToken);
+				}
+				catch (OperationCanceledException)
+				{
+					throw;
 				}
 				catch (TimeoutException ex)
 				{
