Horde: Do not re-throw exceptions caught on the compute send task after they've been logged.

diff --git a/Engine/Source/Programs/Shared/EpicGames.Horde/Compute/ComputeSocket.cs b/Engine/Source/Programs/Shared/EpicGames.Horde/Compute/ComputeSocket.cs
--- a/Engine/Source/Programs/Shared/EpicGames.Horde/Compute/ComputeSocket.cs
+++ b/Engine/Source/Programs/Shared/EpicGames.Horde/Compute/ComputeSocket.cs
@@ -854,7 +854,6 @@ namespace EpicGames.Horde.Compute
 			catch (Exception ex)
 			{
 				_logger.LogInformation(ex, "Error in background send: {Message}", ex.Message);
-				throw;
 			}
 		}
 	}
