Horde: Add a better error message when there are no changes in view for a client.

diff --git a/Engine/Source/Programs/Shared/EpicGames.Perforce.Managed/ManagedWorkspace.cs b/Engine/Source/Programs/Shared/EpicGames.Perforce.Managed/ManagedWorkspace.cs
--- a/Engine/Source/Programs/Shared/EpicGames.Perforce.Managed/ManagedWorkspace.cs
+++ b/Engine/Source/Programs/Shared/EpicGames.Perforce.Managed/ManagedWorkspace.cs
@@ -17,10 +17,23 @@ using Microsoft.Extensions.Logging.Abstractions;
 
 namespace EpicGames.Perforce.Managed
 {
+	/// <summary>
+	/// Exception relating to managed workspace
+	/// </summary>
+	public class ManagedWorkspaceException : Exception
+	{
+		/// <summary>
+		/// Constructor
+		/// </summary>
+		public ManagedWorkspaceException(string message) : base(message)
+		{
+		}
+	}
+
 	/// <summary>
 	/// Exception thrown when there is not enough free space on the drive
 	/// </summary>
-	public class InsufficientSpaceException : Exception
+	public class InsufficientSpaceException : ManagedWorkspaceException
 	{
 		/// <summary>
 		/// Constructor
@@ -1261,6 +1274,10 @@ namespace EpicGames.Perforce.Managed
 				Stopwatch timer = Stopwatch.StartNew();
 
 				List<ChangesRecord> changes = await perforceClient.GetChangesAsync(ChangesOptions.None, 1, ChangeStatus.Submitted, new[] { String.Format("//{0}/...", perforceClient.Settings.ClientName) }, cancellationToken);
+				if (changes.Count == 0)
+				{
+					throw new ManagedWorkspaceException($"Unable to find latest change; no changes in view for {perforceClient.Settings.ClientName}.");
+				}
 				changeNumber = changes[0].Number;
 
 				status.Progress = String.Format("CL {0} ({1:0.0}s)", changeNumber, timer.Elapsed.TotalSeconds);
