Horde: Fix schema generation error.

#rnx

diff --git a/Engine/Source/Programs/Shared/EpicGames.Core/JsonSchema.cs b/Engine/Source/Programs/Shared/EpicGames.Core/JsonSchema.cs
--- a/Engine/Source/Programs/Shared/EpicGames.Core/JsonSchema.cs
+++ b/Engine/Source/Programs/Shared/EpicGames.Core/JsonSchema.cs
@@ -364,6 +364,10 @@ namespace EpicGames.Core
 			{
 				return new JsonSchemaString(JsonSchemaStringFormat.Uri);
 			}
+			if (type == typeof(System.Text.Json.Nodes.JsonObject))
+			{
+				return new JsonSchemaObject();
+			}
 
 			Type[] interfaceTypes = type.GetInterfaces();
 			foreach (Type interfaceType in interfaceTypes)
