Horde: Static analysis fixes.

#rnx

diff --git a/Engine/Source/Programs/Horde/Horde.Agent/AgentSettings.cs b/Engine/Source/Programs/Horde/Horde.Agent/AgentSettings.cs
--- a/Engine/Source/Programs/Horde/Horde.Agent/AgentSettings.cs
+++ b/Engine/Source/Programs/Horde/Horde.Agent/AgentSettings.cs
@@ -1,7 +1,6 @@
 // Copyright Epic Games, Inc. All Rights Reserved.
 
 using System.ComponentModel.DataAnnotations;
-using System.Text.Json.Serialization;
 using EpicGames.Core;
 using EpicGames.Horde;
 using EpicGames.Horde.Tools;
diff --git a/Engine/Source/Programs/Shared/EpicGames.Horde/HordeOptions.cs b/Engine/Source/Programs/Shared/EpicGames.Horde/HordeOptions.cs
--- a/Engine/Source/Programs/Shared/EpicGames.Horde/HordeOptions.cs
+++ b/Engine/Source/Programs/Shared/EpicGames.Horde/HordeOptions.cs
@@ -54,9 +54,9 @@ namespace EpicGames.Horde
 		{
 			if (RuntimeInformation.IsOSPlatform(OSPlatform.Windows))
 			{
-				string? url = null;
-				url ??= Registry.GetValue(@"HKEY_CURRENT_USER\SOFTWARE\Epic Games\Horde", "Url", null) as string;
-				url ??= Registry.GetValue(@"HKEY_LOCAL_MACHINE\SOFTWARE\Epic Games\Horde", "Url", null) as string;
+				string? url = 
+					(Registry.GetValue(@"HKEY_CURRENT_USER\SOFTWARE\Epic Games\Horde", "Url", null) as string) ??
+					(Registry.GetValue(@"HKEY_LOCAL_MACHINE\SOFTWARE\Epic Games\Horde", "Url", null) as string);
 
 				if (!String.IsNullOrEmpty(url))
 				{
