Horde: Download UGS deployment settings from Horde parameters endpoint.

#rnx

diff --git a/Engine/Source/Programs/Shared/EpicGames.Horde/HordeHttpClient.cs b/Engine/Source/Programs/Shared/EpicGames.Horde/HordeHttpClient.cs
--- a/Engine/Source/Programs/Shared/EpicGames.Horde/HordeHttpClient.cs
+++ b/Engine/Source/Programs/Shared/EpicGames.Horde/HordeHttpClient.cs
@@ -3,6 +3,7 @@
 using System;
 using System.Collections.Generic;
 using System.IO;
+using System.Linq;
 using System.Net.Http;
 using System.Net.Http.Json;
 using System.Text.Json;
@@ -32,6 +33,8 @@ using Polly.Timeout;
 
 namespace EpicGames.Horde
 {
+	using JsonObject = System.Text.Json.Nodes.JsonObject;
+
 	/// <summary>
 	/// Wraps an Http client which communicates with the Horde server
 	/// </summary>
@@ -199,6 +202,36 @@ namespace EpicGames.Horde
 
 		#endregion
 
+		#region Parameters
+
+		/// <summary>
+		/// Query parameters for other tools
+		/// </summary>
+		/// <param name="cancellationToken">Cancellation token for the operation</param>
+		/// <returns>Parameters for other tools</returns>
+		public Task<JsonObject> GetParametersAsync(CancellationToken cancellationToken = default)
+		{
+			return GetParametersAsync(null, cancellationToken);
+		}
+
+		/// <summary>
+		/// Query parameters for other tools
+		/// </summary>
+		/// <param name="filter">Filters for properties to return (eg. "ugs.server", "ugs.*")</param>
+		/// <param name="cancellationToken">Cancellation token for the operation</param>
+		/// <returns>Information about all the projects</returns>
+		public Task<JsonObject> GetParametersAsync(IEnumerable<string>? filter, CancellationToken cancellationToken = default)
+		{
+			string url = "api/v1/parameters";
+			if (filter != null && filter.Any())
+			{
+				url = $"{url}?filter={Uri.EscapeDataString(String.Join(",", filter))}";
+			}
+			return GetAsync<JsonObject>(_httpClient, url, cancellationToken);
+		}
+
+		#endregion
+
 		#region Projects
 
 		/// <summary>
diff --git a/Engine/Source/Programs/UnrealGameSync/UnrealGameSync/Launcher.cs b/Engine/Source/Programs/UnrealGameSync/UnrealGameSync/Launcher.cs
--- a/Engine/Source/Programs/UnrealGameSync/UnrealGameSync/Launcher.cs
+++ b/Engine/Source/Programs/UnrealGameSync/UnrealGameSync/Launcher.cs
@@ -14,12 +14,16 @@ using System.IO;
 using System.IO.Compression;
 using System.Linq;
 using System.Reflection;
+using System.Text.Json;
+using System.Text.Json.Nodes;
 using System.Threading;
 using System.Threading.Tasks;
 using System.Windows.Forms;
 
 namespace UnrealGameSync
 {
+	using JsonObject = System.Text.Json.Nodes.JsonObject;
+
 	static class Launcher
 	{
 		public static bool SyncAndRunLatest(Mutex instanceMutex, string[] args)
@@ -164,6 +168,17 @@ namespace UnrealGameSync
 							throw new UserErrorException("Couldn't write sync text to {SyncVersionFile}");
 						}
 					}
+
+					// Query the deployment settings and write them to the output directory if not already set
+					JsonObject parameters = await httpClient.GetParametersAsync(new[] { "ugs.*" }, cancellationToken);
+					JsonObject? ugsParameters = parameters["ugs"] as JsonObject;
+					if (ugsParameters != null)
+					{
+						byte[] deploymentData = JsonSerializer.SerializeToUtf8Bytes(ugsParameters, new JsonSerializerOptions { WriteIndented = true });
+
+						FileReference deploymentJson = new FileReference(Path.Combine(applicationFolder, "Deployment.json"));
+						await FileReference.WriteAllBytesAsync(deploymentJson, deploymentData, cancellationToken);
+					}
 				}
 				else if (launcherSettings.UpdateSource == LauncherUpdateSource.Perforce)
 				{
