Horde: Fix static analysis warning.

#rnx

diff --git a/Engine/Source/Programs/Shared/EpicGames.Horde/HordeHttpClient.cs b/Engine/Source/Programs/Shared/EpicGames.Horde/HordeHttpClient.cs
--- a/Engine/Source/Programs/Shared/EpicGames.Horde/HordeHttpClient.cs
+++ b/Engine/Source/Programs/Shared/EpicGames.Horde/HordeHttpClient.cs
@@ -3,7 +3,6 @@
 using System;
 using System.Collections.Generic;
 using System.IO;
-using System.Linq;
 using System.Net.Http;
 using System.Net.Http.Json;
 using System.Text.Json;
@@ -179,7 +178,7 @@ namespace EpicGames.Horde
 		/// <returns>Config information needed by the dashboard</returns>
 		public Task<GetDashboardPreviewResponse> CreateDashbordPreviewAsync(CreateDashboardPreviewRequest request, CancellationToken cancellationToken = default)
 		{
-			return PostAsync<GetDashboardPreviewResponse, CreateDashboardPreviewRequest>(_httpClient, $"api/v1/dashboard/preview", request, cancellationToken); 
+			return PostAsync<GetDashboardPreviewResponse, CreateDashboardPreviewRequest>(_httpClient, $"api/v1/dashboard/preview", request, cancellationToken);
 		}
 
 		/// <summary>
