Horde: Add a flag for whether to show tools in the dashboard. Defaults to true, but may be disabled for things like the UGS update package.

#rnx

diff --git a/Engine/Source/Programs/Horde/Horde.Server/Tools/ToolConfig.cs b/Engine/Source/Programs/Horde/Horde.Server/Tools/ToolConfig.cs
--- a/Engine/Source/Programs/Horde/Horde.Server/Tools/ToolConfig.cs
+++ b/Engine/Source/Programs/Horde/Horde.Server/Tools/ToolConfig.cs
@@ -58,6 +58,11 @@ namespace Horde.Server.Tools
 		/// </summary>
 		public bool ShowInUgs { get; set; }
 
+		/// <summary>
+		/// Whether to show this tool for download in the dashboard
+		/// </summary>
+		public bool ShowInDashboard { get; set; } = true;
+
 		/// <summary>
 		/// Default namespace for new deployments of this tool
 		/// </summary>
diff --git a/Engine/Source/Programs/Horde/Horde.Server/Tools/ToolsController.cs b/Engine/Source/Programs/Horde/Horde.Server/Tools/ToolsController.cs
--- a/Engine/Source/Programs/Horde/Horde.Server/Tools/ToolsController.cs
+++ b/Engine/Source/Programs/Horde/Horde.Server/Tools/ToolsController.cs
@@ -214,7 +214,7 @@ namespace Horde.Server.Tools
 		static GetToolSummaryResponse CreateGetToolSummaryResponse(ITool tool)
 		{
 			IToolDeployment? deployment = (tool.Deployments.Count == 0) ? null : tool.Deployments[^1];
-			return new GetToolSummaryResponse(tool.Id, tool.Config.Name, tool.Config.Description, tool.Config.ParentId, deployment?.Version, deployment?.Id, tool.Config.ShowInUgs);
+			return new GetToolSummaryResponse(tool.Id, tool.Config.Name, tool.Config.Description, tool.Config.ParentId, deployment?.Version, deployment?.Id, tool.Config.ShowInUgs, tool.Config.ShowInDashboard);
 		}
 
 		/// <summary>
@@ -258,7 +258,7 @@ namespace Horde.Server.Tools
 
 		static GetToolResponse CreateGetToolResponse(ITool tool, List<GetToolDeploymentResponse> deployments)
 		{
-			return new GetToolResponse(tool.Id, tool.Config.Name, tool.Config.Description, tool.Config.ParentId, deployments, tool.Config.Public, tool.Config.ShowInUgs);
+			return new GetToolResponse(tool.Id, tool.Config.Name, tool.Config.Description, tool.Config.ParentId, deployments, tool.Config.Public, tool.Config.ShowInUgs, tool.Config.ShowInDashboard);
 		}
 
 		/// <summary>
diff --git a/Engine/Source/Programs/Shared/EpicGames.Horde/Tools/ToolMessages.cs b/Engine/Source/Programs/Shared/EpicGames.Horde/Tools/ToolMessages.cs
--- a/Engine/Source/Programs/Shared/EpicGames.Horde/Tools/ToolMessages.cs
+++ b/Engine/Source/Programs/Shared/EpicGames.Horde/Tools/ToolMessages.cs
@@ -17,7 +17,8 @@ namespace EpicGames.Horde.Tools
 	/// <param name="Deployments">Current deployments of this tool, sorted by time.</param>
 	/// <param name="Public">Whether this tool should be exposed for download on a public endpoint without authentication</param>
 	/// <param name="ShowInUgs">Whether to show this tool for download inside UGS</param>
-	public record class GetToolResponse(ToolId Id, string Name, string Description, ToolId? ParentId, List<GetToolDeploymentResponse> Deployments, bool Public, bool ShowInUgs);
+	/// <param name="ShowInDashboard">Whether to show this tool for download on the dashboard</param>
+	public record class GetToolResponse(ToolId Id, string Name, string Description, ToolId? ParentId, List<GetToolDeploymentResponse> Deployments, bool Public, bool ShowInUgs, bool ShowInDashboard);
 
 	/// <summary>
 	/// Summary for a particular tool.
@@ -29,7 +30,8 @@ namespace EpicGames.Horde.Tools
 	/// <param name="Version">Version number of the current deployment of this tool</param>
 	/// <param name="DeploymentId">Identifier for the current deployment</param>
 	/// <param name="ShowInUgs">Whether to show this tool for download inside UGS</param>
-	public record class GetToolSummaryResponse(ToolId Id, string Name, string Description, ToolId? ParentId, string? Version, ToolDeploymentId? DeploymentId, bool ShowInUgs);
+	/// <param name="ShowInDashboard">Whether to show this tool for download on the dashboard</param>
+	public record class GetToolSummaryResponse(ToolId Id, string Name, string Description, ToolId? ParentId, string? Version, ToolDeploymentId? DeploymentId, bool ShowInUgs, bool ShowInDashboard);
 
 	/// <summary>
 	/// Response when querying all tools
