Horde: Remove unused SecretToken parameter from account creation messages.

diff --git a/Engine/Source/Programs/Shared/EpicGames.Horde/Accounts/AccountMessages.cs b/Engine/Source/Programs/Shared/EpicGames.Horde/Accounts/AccountMessages.cs
--- a/Engine/Source/Programs/Shared/EpicGames.Horde/Accounts/AccountMessages.cs
+++ b/Engine/Source/Programs/Shared/EpicGames.Horde/Accounts/AccountMessages.cs
@@ -15,10 +15,9 @@ namespace EpicGames.Horde.Accounts
 	/// <param name="Claims">Claims for the user</param>
 	/// <param name="Description">Description for the account</param>
 	/// <param name="Email">User's email address</param>
-	/// <param name="SecretToken">Optional secret token for API access</param>
 	/// <param name="Password">Password for the user</param>
 	/// <param name="Enabled">Whether the account is enabled</param>
-	public record class CreateAccountRequest(string Name, string Login, List<AccountClaimMessage> Claims, string? Description, string? Email, string? SecretToken, string? Password, bool? Enabled);
+	public record class CreateAccountRequest(string Name, string Login, List<AccountClaimMessage> Claims, string? Description, string? Email, string? Password, bool? Enabled);
 
 	/// <summary>
 	/// Creates the admin account
