Horde: Pool connections in ManagedWorkspace rather than creating a new instance each time.

diff --git a/Engine/Source/Programs/Shared/EpicGames.Perforce.Managed/ManagedWorkspace.cs b/Engine/Source/Programs/Shared/EpicGames.Perforce.Managed/ManagedWorkspace.cs
--- a/Engine/Source/Programs/Shared/EpicGames.Perforce.Managed/ManagedWorkspace.cs
+++ b/Engine/Source/Programs/Shared/EpicGames.Perforce.Managed/ManagedWorkspace.cs
@@ -2010,6 +2010,7 @@ namespace EpicGames.Perforce.Managed
 
 					// Spawn some background threads to sync them
 					Dictionary<Task, int> tasks = new Dictionary<Task, int>();
+					Stack<IPerforceConnection> connectionPool = new Stack<IPerforceConnection>();
 					try
 					{
 						while (tasks.Count > 0 || nextBatchIdx < batches.Count)
@@ -2019,7 +2020,7 @@ namespace EpicGames.Perforce.Managed
 							{
 								(int batchBeginIdx, int batchEndIdx) = batches[nextBatchIdx];
 
-								Task task = Task.Run(() => SyncBatchAsync(client, filesToSync, batchBeginIdx, batchEndIdx, fakeSync, cancellationToken), cancellationToken);
+								Task task = Task.Run(() => SyncBatchAsync(client, filesToSync, batchBeginIdx, batchEndIdx, fakeSync, connectionPool, cancellationToken), cancellationToken);
 								tasks[task] = nextBatchIdx++;
 							}
 
@@ -2061,6 +2062,11 @@ namespace EpicGames.Perforce.Managed
 					finally
 					{
 						await Task.WhenAll(tasks.Keys);
+
+						foreach (IPerforceConnection connection in connectionPool)
+						{
+							connection.Dispose();
+						}
 					}
 				}
 			}
@@ -2078,9 +2084,10 @@ namespace EpicGames.Perforce.Managed
 		/// <param name="beginIdx">First file to sync</param>
 		/// <param name="endIdx">Index of the last file to sync (exclusive)</param>
 		/// <param name="fakeSync">Whether to fake a sync</param>
+		/// <param name="connectionPool">Pool of connection instances</param>
 		/// <param name="cancellationToken">Cancellation token for the request</param>
 		/// <returns>Async task</returns>
-		async Task SyncBatchAsync(IPerforceConnection client, WorkspaceFileToSync[] filesToSync, int beginIdx, int endIdx, bool fakeSync, CancellationToken cancellationToken)
+		async Task SyncBatchAsync(IPerforceConnection client, WorkspaceFileToSync[] filesToSync, int beginIdx, int endIdx, bool fakeSync, Stack<IPerforceConnection> connectionPool, CancellationToken cancellationToken)
 		{
 			if (fakeSync)
 			{
@@ -2109,9 +2116,34 @@ namespace EpicGames.Perforce.Managed
 					options |= SyncOptions.DoNotUpdateHaveList;
 				}
 
-				// Note: Explicitly disable parallel syncing here; the P4 API attempts to shell out to p4.exe, which may not be installed.
-				using IPerforceConnection threadedClient = await PerforceConnection.CreateAsync(client.Settings, client.Logger);
-				await threadedClient.SyncAsync(options, -1, 0, -1, -1, -1, -1, files, cancellationToken).ToListAsync(cancellationToken);
+				// Allocate a connection to use for syncing
+#pragma warning disable CA2000
+				IPerforceConnection? connection = null;
+				try
+				{
+					lock (connectionPool)
+					{
+						connectionPool.TryPop(out connection);
+					}
+					if (connection == null)
+					{
+						connection = await PerforceConnection.CreateAsync(client.Settings, client.Logger);
+					}
+
+					// Note: Explicitly disable parallel syncing here; the P4 API attempts to shell out to p4.exe, which may not be installed.
+					await connection.SyncAsync(options, -1, 0, -1, -1, -1, -1, files, cancellationToken).ToListAsync(cancellationToken);
+				}
+				finally
+				{
+					if (connection != null)
+					{
+						lock (connectionPool)
+						{
+							connectionPool.Push(connection);
+						}
+					}
+				}
+#pragma warning restore CA2000
 			}
 			else
 			{
