Horde: Move BundleCache into EpicGames.Horde.Bundles namespace.

#rnx

diff --git a/Engine/Source/Programs/Horde/Horde.Agent.Tests/WorkerServiceTest.cs b/Engine/Source/Programs/Horde/Horde.Agent.Tests/WorkerServiceTest.cs
--- a/Engine/Source/Programs/Horde/Horde.Agent.Tests/WorkerServiceTest.cs
+++ b/Engine/Source/Programs/Horde/Horde.Agent.Tests/WorkerServiceTest.cs
@@ -14,6 +14,7 @@ using EpicGames.Horde.Jobs;
 using EpicGames.Horde.Logs;
 using EpicGames.Horde.Storage;
 using EpicGames.Horde.Storage.Backends;
+using EpicGames.Horde.Storage.Bundles;
 using EpicGames.Horde.Storage.Clients;
 using EpicGames.Horde.Streams;
 using Google.Protobuf.WellKnownTypes;
diff --git a/Engine/Source/Programs/Horde/Horde.Agent/AgentApp.cs b/Engine/Source/Programs/Horde/Horde.Agent/AgentApp.cs
--- a/Engine/Source/Programs/Horde/Horde.Agent/AgentApp.cs
+++ b/Engine/Source/Programs/Horde/Horde.Agent/AgentApp.cs
@@ -13,6 +13,7 @@ using EpicGames.Core;
 using EpicGames.Horde;
 using EpicGames.Horde.Storage;
 using EpicGames.Horde.Storage.Backends;
+using EpicGames.Horde.Storage.Bundles;
 using EpicGames.Horde.Storage.Clients;
 using Horde.Agent.Execution;
 using Horde.Agent.Leases;
diff --git a/Engine/Source/Programs/Horde/Horde.Server.Tests/TestSetup.cs b/Engine/Source/Programs/Horde/Horde.Server.Tests/TestSetup.cs
--- a/Engine/Source/Programs/Horde/Horde.Server.Tests/TestSetup.cs
+++ b/Engine/Source/Programs/Horde/Horde.Server.Tests/TestSetup.cs
@@ -69,6 +69,7 @@ using Horde.Server.Tests.Server;
 using Horde.Server.Agents.Enrollment;
 using Horde.Server.Accounts;
 using Microsoft.VisualStudio.TestTools.UnitTesting;
+using EpicGames.Horde.Storage.Bundles;
 
 namespace Horde.Server.Tests
 {
diff --git a/Engine/Source/Programs/Horde/Horde.Server/Startup.cs b/Engine/Source/Programs/Horde/Horde.Server/Startup.cs
--- a/Engine/Source/Programs/Horde/Horde.Server/Startup.cs
+++ b/Engine/Source/Programs/Horde/Horde.Server/Startup.cs
@@ -121,6 +121,7 @@ using StackExchange.Redis;
 using Swashbuckle.AspNetCore.SwaggerGen;
 using Status = Grpc.Core.Status;
 using StatusCode = Grpc.Core.StatusCode;
+using EpicGames.Horde.Storage.Bundles;
 
 namespace Horde.Server
 {
diff --git a/Engine/Source/Programs/Horde/Horde/Commands/StorageCommandBase.cs b/Engine/Source/Programs/Horde/Horde/Commands/StorageCommandBase.cs
--- a/Engine/Source/Programs/Horde/Horde/Commands/StorageCommandBase.cs
+++ b/Engine/Source/Programs/Horde/Horde/Commands/StorageCommandBase.cs
@@ -3,6 +3,7 @@
 using System.ComponentModel;
 using EpicGames.Core;
 using EpicGames.Horde.Storage;
+using EpicGames.Horde.Storage.Bundles;
 using EpicGames.Horde.Storage.Clients;
 using Microsoft.Extensions.Options;
 
diff --git a/Engine/Source/Programs/Shared/EpicGames.Horde.Tests/BundleCacheTests.cs b/Engine/Source/Programs/Shared/EpicGames.Horde.Tests/BundleCacheTests.cs
--- a/Engine/Source/Programs/Shared/EpicGames.Horde.Tests/BundleCacheTests.cs
+++ b/Engine/Source/Programs/Shared/EpicGames.Horde.Tests/BundleCacheTests.cs
@@ -4,7 +4,7 @@ using System;
 using System.Buffers;
 using System.Threading.Tasks;
 using EpicGames.Core;
-using EpicGames.Horde.Storage;
+using EpicGames.Horde.Storage.Bundles;
 using Microsoft.VisualStudio.TestTools.UnitTesting;
 
 namespace EpicGames.Horde.Tests
@@ -118,7 +118,7 @@ namespace EpicGames.Horde.Tests
 			IRefCountedHandle<TestData> result = await cache.FindOrAddAsync("test", (key, ctx) => Task.FromResult(new TestData(cache.Allocator.Alloc(20, null))));
 			Assert.AreEqual(20, cache.CurrentSize);
 
-			cache.Trim(); 
+			cache.Trim();
 			Assert.AreEqual(20, cache.CurrentSize);
 			result.Dispose();
 			Assert.AreEqual(20, cache.CurrentSize);
diff --git a/Engine/Source/Programs/Shared/EpicGames.Horde/HordeExtensions.cs b/Engine/Source/Programs/Shared/EpicGames.Horde/HordeExtensions.cs
--- a/Engine/Source/Programs/Shared/EpicGames.Horde/HordeExtensions.cs
+++ b/Engine/Source/Programs/Shared/EpicGames.Horde/HordeExtensions.cs
@@ -3,6 +3,7 @@
 using System;
 using EpicGames.Horde.Storage;
 using EpicGames.Horde.Storage.Backends;
+using EpicGames.Horde.Storage.Bundles;
 using EpicGames.Horde.Storage.Clients;
 using Microsoft.Extensions.DependencyInjection;
 using Microsoft.Extensions.Options;
diff --git a/Engine/Source/Programs/Shared/EpicGames.Horde/HordeOptions.cs b/Engine/Source/Programs/Shared/EpicGames.Horde/HordeOptions.cs
--- a/Engine/Source/Programs/Shared/EpicGames.Horde/HordeOptions.cs
+++ b/Engine/Source/Programs/Shared/EpicGames.Horde/HordeOptions.cs
@@ -5,7 +5,6 @@ using System.Linq;
 using System.Net.Http;
 using System.Runtime.InteropServices;
 using System.Runtime.Versioning;
-using EpicGames.Horde.Storage;
 using EpicGames.Horde.Storage.Bundles;
 using Microsoft.Win32;
 
@@ -54,7 +53,7 @@ namespace EpicGames.Horde
 		{
 			if (RuntimeInformation.IsOSPlatform(OSPlatform.Windows))
 			{
-				string? url = 
+				string? url =
 					(Registry.GetValue(@"HKEY_CURRENT_USER\SOFTWARE\Epic Games\Horde", "Url", null) as string) ??
 					(Registry.GetValue(@"HKEY_LOCAL_MACHINE\SOFTWARE\Epic Games\Horde", "Url", null) as string);
 
diff --git a/Engine/Source/Programs/Shared/EpicGames.Horde/Storage/Bundles/BundleCache.cs b/Engine/Source/Programs/Shared/EpicGames.Horde/Storage/Bundles/BundleCache.cs
--- a/Engine/Source/Programs/Shared/EpicGames.Horde/Storage/Bundles/BundleCache.cs
+++ b/Engine/Source/Programs/Shared/EpicGames.Horde/Storage/Bundles/BundleCache.cs
@@ -10,7 +10,7 @@ using System.Threading.Tasks;
 using EpicGames.Core;
 using Microsoft.Extensions.Caching.Memory;
 
-namespace EpicGames.Horde.Storage
+namespace EpicGames.Horde.Storage.Bundles
 {
 	/// <summary>
 	/// Options for creating a storage cache
