Horde: Use BitFaster.Caching for the LRU bundle cache. Much more efficient than the cache I cobbled together.

diff --git a/Engine/Source/Programs/Shared/EpicGames.Horde.Tests/BundleCacheTests.cs b/Engine/Source/Programs/Shared/EpicGames.Horde.Tests/BundleCacheTests.cs
deleted file mode 100644
diff --git a/Engine/Source/Programs/Shared/EpicGames.Horde/EpicGames.Horde.csproj b/Engine/Source/Programs/Shared/EpicGames.Horde/EpicGames.Horde.csproj
--- a/Engine/Source/Programs/Shared/EpicGames.Horde/EpicGames.Horde.csproj
+++ b/Engine/Source/Programs/Shared/EpicGames.Horde/EpicGames.Horde.csproj
@@ -38,6 +38,7 @@
     </Compile>
   </ItemGroup>
   <ItemGroup>
+    <PackageReference Include="BitFaster.Caching" Version="2.4.1" />
     <PackageReference Include="Dapper" Version="2.1.24" />
     <PackageReference Include="K4os.Compression.LZ4" Version="1.2.16" />
     <PackageReference Include="Microsoft.CodeAnalysis.CSharp.Workspaces" Version="4.2.0" />
diff --git a/Engine/Source/Programs/Shared/EpicGames.Horde/Storage/Bundles/BundleCache.cs b/Engine/Source/Programs/Shared/EpicGames.Horde/Storage/Bundles/BundleCache.cs
--- a/Engine/Source/Programs/Shared/EpicGames.Horde/Storage/Bundles/BundleCache.cs
+++ b/Engine/Source/Programs/Shared/EpicGames.Horde/Storage/Bundles/BundleCache.cs
@@ -1,13 +1,13 @@
 // Copyright Epic Games, Inc. All Rights Reserved.
 
 using System;
-using System.Buffers;
-using System.Collections.Generic;
-using System.Diagnostics;
 using System.Diagnostics.CodeAnalysis;
 using System.Threading;
 using System.Threading.Tasks;
+using BitFaster.Caching;
+using BitFaster.Caching.Lru;
 using EpicGames.Core;
+using EpicGames.Horde.Storage.Bundles.V2;
 using Microsoft.Extensions.Caching.Memory;
 
 namespace EpicGames.Horde.Storage.Bundles
@@ -18,9 +18,14 @@ namespace EpicGames.Horde.Storage.Bundles
 	public class BundleCacheOptions
 	{
 		/// <summary>
-		/// Maximum size of the cache
+		/// Number of packet readers to keep in the cache
 		/// </summary>
-		public long MaxSize { get; set; } = 128 * 1024 * 1024;
+		public int PacketReaderCount { get; set; } = 200;
+
+		/// <summary>
+		/// Number of bundle pages to keep in the cache.
+		/// </summary>
+		public int BundlePageCount { get; set; } = 500;
 
 		/// <summary>
 		/// Size of the header cache
@@ -33,211 +38,37 @@ namespace EpicGames.Horde.Storage.Bundles
 		public long PacketCacheSize { get; set; } = 192 * 1024 * 1024;
 	}
 
+	record struct BundlePageCacheKey(BundleHandle Bundle, int Index)
+	{
+		public override string ToString()
+			=> $"bundle-page:{Bundle}:{Index}";
+	}
+
 	/// <summary>
-	/// Caches items with user-defined keys, where initialization of each value is performed asynchronously. 
-	/// Memory allocations are tracked via calls to Reserve(), and a hard upper limit is enforced as long as data can be freed.
+	/// Cache for reading bundle data.
 	/// </summary>
 	public sealed class BundleCache : IAsyncDisposable
 	{
-		[DebuggerDisplay("{Key}")]
-		class CacheValue : IDisposable
-		{
-			public object Key { get; }
-			public LinkedListNode<CacheValue> Node { get; }
-			public Task<IDisposable> InitTask { get; }
-
-			readonly BundleCache _bundleCache;
-			int _refCount;
-
-			public int RefCount => Interlocked.CompareExchange(ref _refCount, 0, 0);
-
-			public CacheValue(BundleCache bundleCache, object key, IDisposable value)
-			{
-				_bundleCache = bundleCache;
-				_refCount = 1;
-				Key = key;
-				Node = new LinkedListNode<CacheValue>(this);
-				InitTask = Task.FromResult<IDisposable>(value);
-
-				_bundleCache._purgeableItems.AddFirst(Node);
-			}
-
-			public CacheValue(BundleCache bundleCache, object key, Func<Task<IDisposable>> initTask)
-			{
-				_bundleCache = bundleCache;
-				_refCount = 2; // Will be released by RunAndUnlock
-				Key = key;
-				Node = new LinkedListNode<CacheValue>(this);
-				InitTask = RunAndUnlockAsync(initTask);
-
-				_bundleCache._items.AddFirst(Node);
-			}
-
-			public void Dispose()
-			{
-#pragma warning disable VSTHRD002
-				InitTask.Result.Dispose();
-#pragma warning restore VSTHRD002
-			}
-
-			async Task<IDisposable> RunAndUnlockAsync(Func<Task<IDisposable>> initTask)
-			{
-				await Task.Yield();
-				try
-				{
-					return await initTask();
-				}
-				finally
-				{
-					Release();
-				}
-			}
-
-			public void AddRef()
-			{
-				lock (_bundleCache._lockObject)
-				{
-					if (_refCount == 1)
-					{
-						_bundleCache._purgeableItems.Remove(Node);
-						_bundleCache._items.AddFirst(Node);
-					}
-					_refCount++;
-				}
-			}
-
-			public void Release()
-			{
-				lock (_bundleCache._lockObject)
-				{
-					_refCount--;
-					if (_refCount == 1)
-					{
-						_bundleCache._items.Remove(Node);
-						_bundleCache._purgeableItems.AddFirst(Node);
-					}
-				}
-			}
-		}
-
-		class CacheValueHandle<T> : IRefCountedHandle<T> where T : class, IDisposable
-		{
-			CacheValue? _item;
-			T? _target;
-
-			public int RefCount => _item?.RefCount ?? throw new ObjectDisposedException(nameof(CacheValueHandle<T>));
-
-			public T Target => _target ?? throw new ObjectDisposedException(nameof(CacheValueHandle<T>));
-
-			public CacheValueHandle(CacheValue? item, T? value)
-			{
-				_item = item;
-				_target = value;
-			}
-
-			public void Dispose()
-			{
-				if (_item != null)
-				{
-					_item.Release();
-					_item = null;
-					_target = null;
-				}
-			}
-
-			public IRefCountedHandle<T> AddRef()
-			{
-				if (_item == null)
-				{
-					throw new ObjectDisposedException(nameof(CacheValueHandle<T>));
-				}
-
-				_item.AddRef();
-				return new CacheValueHandle<T>(_item, _target);
-			}
-		}
-
-		// Tracks an owned blocks of memory against the cache budged.
-		sealed class MemoryAllocation : IMemoryOwner<byte>
-		{
-			readonly BundleCache _bundleCache;
-			IMemoryOwner<byte> _owner;
-
-			public Memory<byte> Memory => _owner.Memory;
-
-			public MemoryAllocation(BundleCache bundleCache, IMemoryOwner<byte> owner)
-			{
-				_bundleCache = bundleCache;
-				_owner = owner;
-			}
-
-			/// <inheritdoc/>
-			public void Dispose()
-			{
-				if (_owner != null)
-				{
-					long size = _owner.Memory.Length;
-					_owner.Dispose();
-					_bundleCache.ReleaseSpace(size);
-					_owner = null!;
-				}
-			}
-		}
-
-		// Custom memory allocator which tracks allocated blocks against the cache's budget. Older cache entries will be 
-		// disposed to create space for new allocations.
-		class MemoryAllocator : IMemoryAllocator<byte>
-		{
-			readonly BundleCache _bundleCache;
-			readonly IMemoryAllocator<byte> _allocator;
-
-			public MemoryAllocator(BundleCache bundleCache, IMemoryAllocator<byte> allocator)
-			{
-				_bundleCache = bundleCache;
-				_allocator = allocator;
-			}
-
-			public IMemoryOwner<byte> Alloc(int minSize, object? tag)
-			{
-				_bundleCache.CreateSpace(minSize);
-				IMemoryOwner<byte> owner = _allocator.Alloc(minSize, tag);
-				_bundleCache.CreateSpace(owner.Memory.Length - minSize);
-				return new MemoryAllocation(_bundleCache, owner);
-			}
-		}
-
-		readonly object _lockObject = new object();
 		readonly BundleCacheOptions _options;
-		readonly LinkedList<CacheValue> _items = new LinkedList<CacheValue>();
-		readonly Dictionary<object, CacheValue> _itemLookup = new Dictionary<object, CacheValue>();
-		readonly LinkedList<CacheValue> _purgeableItems = new LinkedList<CacheValue>();
 		readonly CancellationTokenSource _cancellationSource = new CancellationTokenSource();
-		readonly MemoryAllocator _allocator;
-
-		int _numCacheAdds;
-		int _numCacheHits;
-		int _numCacheMisses;
-		long _currentSize;
-		int _numPurges;
-		readonly Stopwatch _createSpaceTimer = new Stopwatch();
 
+		readonly IScopedAsyncCache<PacketReaderCacheKey, PacketReader> _packetReaderCache;
+		readonly IScopedAsyncCache<BundlePageCacheKey, IReadOnlyMemoryOwner<byte>> _bundlePageCache;
 		readonly MemoryCache? _headerCache;
 		readonly MemoryCache? _packetCache;
 
+		internal IScopedAsyncCache<PacketReaderCacheKey, PacketReader> PacketReaderCache => _packetReaderCache;
+		internal IScopedAsyncCache<BundlePageCacheKey, IReadOnlyMemoryOwner<byte>> BundlePageCache => _bundlePageCache;
+
 		/// <summary>
 		/// Instance of an empty cache
 		/// </summary>
-		public static BundleCache None { get; } = new BundleCache(new BundleCacheOptions { MaxSize = 0, HeaderCacheSize = 0, PacketCacheSize = 0 });
+		public static BundleCache None { get; } = new BundleCache(new BundleCacheOptions { HeaderCacheSize = 0, PacketCacheSize = 0 });
 
 		/// <summary>
 		/// Accessor for the default allocator
 		/// </summary>
-		public IMemoryAllocator<byte> Allocator => _allocator;
-
-		/// <summary>
-		/// Current size of data allocated or in the cache
-		/// </summary>
-		public long CurrentSize => _currentSize;
+		public IMemoryAllocator<byte> Allocator { get; }
 
 		/// <summary>
 		/// Size of the configured header cache
@@ -257,9 +88,9 @@ namespace EpicGames.Horde.Storage.Bundles
 		/// <summary>
 		/// Constructor
 		/// </summary>
-		public BundleCache() 
-			: this(new BundleCacheOptions()) 
-		{ 
+		public BundleCache()
+			: this(new BundleCacheOptions())
+		{
 		}
 
 		/// <summary>
@@ -282,7 +113,27 @@ namespace EpicGames.Horde.Storage.Bundles
 #if DEBUG
 			innerAllocator = new TrackingMemoryAllocator(innerAllocator);
 #endif
-			_allocator = new MemoryAllocator(this, innerAllocator);
+			Allocator = innerAllocator;
+
+			_packetReaderCache = new ConcurrentLruBuilder<PacketReaderCacheKey, PacketReader>()
+				.WithCapacity(_options.PacketReaderCount)
+				.WithAtomicGetOrAdd()
+#if TRACE
+				.WithMetrics()
+#endif
+				.AsAsyncCache()
+				.AsScopedCache()
+				.Build();
+
+			_bundlePageCache = new ConcurrentLruBuilder<BundlePageCacheKey, IReadOnlyMemoryOwner<byte>>()
+				.WithCapacity(_options.BundlePageCount)
+				.WithAtomicGetOrAdd()
+#if TRACE
+				.WithMetrics()
+#endif
+				.AsAsyncCache()
+				.AsScopedCache()
+				.Build();
 
 			if (options.HeaderCacheSize > 0)
 			{
@@ -295,170 +146,39 @@ namespace EpicGames.Horde.Storage.Bundles
 		}
 
 		/// <inheritdoc/>
-		public async ValueTask DisposeAsync()
+		public ValueTask DisposeAsync()
 		{
+			_packetReaderCache.Clear();
+			_bundlePageCache.Clear();
+
 			_headerCache?.Dispose();
 			_packetCache?.Dispose();
 
-			if (_items.Count > 0)
-			{
-				_cancellationSource.Cancel();
-
-				foreach (CacheValue item in _items)
-				{
-					try
-					{
-						await item.InitTask;
-					}
-					catch (OperationCanceledException)
-					{
-					}
-				}
-
-				_items.Clear();
-			}
-
 			_cancellationSource.Dispose();
-		}
 
-		/// <summary>
-		/// Empty the cache
-		/// </summary>
-		public void Trim()
-		{
-			CreateSpace(_options.MaxSize);
-			ReleaseSpace(_options.MaxSize);
+			return default;
 		}
 
 		/// <summary>
-		/// Find or add a new cached value to the cache
-		/// </summary>
-		/// <param name="key">Key for the lookup</param>
-		/// <returns></returns>
-		public IRefCountedHandle<TValue>? Find<TKey, TValue>(TKey key)
-			where TKey : notnull
-			where TValue : class, IDisposable
-		{
-			lock (_lockObject)
-			{
-				CacheValue? item;
-				if (_itemLookup.TryGetValue(key, out item) && item.InitTask.TryGetResult(out IDisposable result))
-				{
-					Interlocked.Increment(ref _numCacheHits);
-					item.AddRef();
-					return new CacheValueHandle<TValue>(item, (TValue)result);
-				}
-			}
-			return null;
-		}
-
-		/// <summary>
-		/// Attempts to add a value to the cache. 
+		/// Gets stats for the cache
 		/// </summary>
-		/// <param name="key">Key to add the item to the cache with</param>
-		/// <param name="value">Value to be added. If the item is added, ownership is implicitly transferred to the cache.</param>
-		/// <returns>True if the value was added</returns>
-		public bool TryAdd<TKey, TValue>(TKey key, TValue value)
-			where TKey : notnull
-			where TValue : class, IDisposable
+		public void GetStats(StorageStats stats)
 		{
-			lock (_lockObject)
+			Optional<ICacheMetrics> packetMetrics = PacketReaderCache.Metrics;
+			if (packetMetrics.Value != null)
 			{
-				if (!_itemLookup.ContainsKey(key))
-				{
-					Interlocked.Increment(ref _numCacheAdds);
-					CacheValue item = new CacheValue(this, key, value);
-					_itemLookup.Add(key, item);
-					return true;
-				}
+				stats.Add("bundle.packet_cache.hits", packetMetrics.Value.Hits);
+				stats.Add("bundle.packet_cache.misses", packetMetrics.Value.Misses);
 			}
-			return false;
-		}
-
-		/// <summary>
-		/// Find or add a new cached value to the cache
-		/// </summary>
-		/// <param name="key"></param>
-		/// <param name="createAsync"></param>
-		/// <param name="cancellationToken"></param>
-		/// <returns>Handle to the item that was read. Must be disposed by the caller.</returns>
-		public async Task<IRefCountedHandle<TValue>> FindOrAddAsync<TKey, TValue>(TKey key, Func<TKey, CancellationToken, Task<TValue>> createAsync, CancellationToken cancellationToken = default)
-			where TKey : notnull
-			where TValue : class, IDisposable
-		{
-			CacheValue? item = null;
-			try
-			{
-				lock (_lockObject)
-				{
-					CacheValue? untypedItem;
-					if (_itemLookup.TryGetValue(key, out untypedItem))
-					{
-						Interlocked.Increment(ref _numCacheHits);
-						item = untypedItem;
-					}
-					else
-					{
-						Interlocked.Increment(ref _numCacheMisses);
-						item = new CacheValue(this, key, async () => await createAsync(key, _cancellationSource.Token));
-						_itemLookup.Add(key, item);
-					}
-					item.AddRef(); // Don't allow the item to be freed while we wait for it
-				}
 
-				TValue value = (TValue)await item.InitTask.WaitAsync(cancellationToken);
-				return new CacheValueHandle<TValue>(item, value);
-			}
-			catch
+			Optional<ICacheMetrics> bundleMetrics = BundlePageCache.Metrics;
+			if (bundleMetrics.Value != null)
 			{
-				item?.Release();
-				throw;
+				stats.Add("bundle.bundle_cache.hits", bundleMetrics.Value.Hits);
+				stats.Add("bundle.bundle_cache.misses", bundleMetrics.Value.Misses);
 			}
 		}
 
-		void CreateSpace(long size)
-		{
-			if (Interlocked.Add(ref _currentSize, size) > _options.MaxSize)
-			{
-				lock (_lockObject)
-				{
-					Interlocked.Increment(ref _numPurges);
-					_createSpaceTimer.Start();
-
-					while (_purgeableItems.Count > 0 && Interlocked.CompareExchange(ref _currentSize, 0, 0) > _options.MaxSize)
-					{
-						LinkedListNode<CacheValue> lastNode = _purgeableItems.Last!;
-						CacheValue lastItem = lastNode.Value;
-
-						_purgeableItems.Remove(lastNode);
-						_itemLookup.Remove(lastItem.Key);
-						lastItem.Dispose();
-					}
-
-					_createSpaceTimer.Stop();
-				}
-			}
-		}
-
-		void ReleaseSpace(long size)
-		{
-			Interlocked.Add(ref _currentSize, -size);
-		}
-
-		/// <summary>
-		/// Gets stats for the cache
-		/// </summary>
-		public void GetStats(StorageStats stats)
-		{
-			stats.Add("bundle.cache.adds", _numCacheAdds);
-			stats.Add("bundle.cache.hits", _numCacheHits);
-			stats.Add("bundle.cache.misses", _numCacheMisses);
-			stats.Add("bundle.cache.size_count", _items.Count);
-			stats.Add("bundle.cache.size_bytes", _currentSize);
-			stats.Add("bundle.cache.purge.count", _numPurges);
-			stats.Add("bundle.cache.purge.time_ms", _createSpaceTimer.ElapsedMilliseconds);
-		}
-
 		#region V1
 
 		static void AddCachedValue(IMemoryCache? cache, string cacheKey, object value, int size)
diff --git a/Engine/Source/Programs/Shared/EpicGames.Horde/Storage/Bundles/V2/PacketHandle.cs b/Engine/Source/Programs/Shared/EpicGames.Horde/Storage/Bundles/V2/PacketHandle.cs
--- a/Engine/Source/Programs/Shared/EpicGames.Horde/Storage/Bundles/V2/PacketHandle.cs
+++ b/Engine/Source/Programs/Shared/EpicGames.Horde/Storage/Bundles/V2/PacketHandle.cs
@@ -7,6 +7,7 @@ using System.Linq;
 using System.Text;
 using System.Threading;
 using System.Threading.Tasks;
+using BitFaster.Caching;
 using EpicGames.Core;
 
 namespace EpicGames.Horde.Storage.Bundles.V2
@@ -161,8 +162,8 @@ namespace EpicGames.Horde.Storage.Bundles.V2
 		{
 			try
 			{
-				using IRefCountedHandle<PacketReader> packetReaderHandle = await GetPacketReaderAsync(cancellationToken);
-				return packetReaderHandle.Target.ReadExport(exportIdx);
+				using Lifetime<PacketReader> packetReaderHandle = await GetPacketReaderAsync(cancellationToken);
+				return packetReaderHandle.Value.ReadExport(exportIdx);
 			}
 			catch (OperationCanceledException)
 			{
@@ -203,32 +204,23 @@ namespace EpicGames.Horde.Storage.Bundles.V2
 
 		#region Packet reader access
 
-		record struct PacketReaderCacheKey(BundleHandle Bundle, int Offset)
-		{
-			public override string ToString()
-				=> $"packet-reader:{Bundle}@{Offset}";
-		}
-
-		record struct BundlePageCacheKey(BundleHandle Bundle, int Index)
-		{
-			public override string ToString()
-				=> $"bundle-page:{Bundle}:{Index}";
-		}
-
-		async ValueTask<IRefCountedHandle<PacketReader>> GetPacketReaderAsync(CancellationToken cancellationToken = default)
+		async ValueTask<Lifetime<PacketReader>> GetPacketReaderAsync(CancellationToken cancellationToken = default)
 		{
 			PacketReaderCacheKey cacheKey = new PacketReaderCacheKey(_outer, _packetOffset);
-			return await _cache.FindOrAddAsync(cacheKey, CreatePacketReaderAsync, cancellationToken);
+			return await _cache.PacketReaderCache.ScopedGetOrAddAsync(cacheKey, CreatePacketReaderAsync, cancellationToken);
 		}
 
-		async Task<PacketReader> CreatePacketReaderAsync(PacketReaderCacheKey cacheKey, CancellationToken cancellationToken)
+		async Task<Scoped<PacketReader>> CreatePacketReaderAsync(PacketReaderCacheKey cacheKey, CancellationToken cancellationToken)
 		{
 			using IReadOnlyMemoryOwner<byte> encodedData = await ReadEncodedPacketAsync(cancellationToken);
 			IRefCountedHandle<Packet> packet = Packet.Decode(encodedData.Memory, _cache.Allocator, cacheKey);
-			return new PacketReader(_storageClient, _cache, _outer, this, packet.Target, packet);
+#pragma warning disable CA2000
+			PacketReader reader = new PacketReader(_storageClient, _cache, _outer, this, packet.Target, packet);
+			return new Scoped<PacketReader>(reader);
+#pragma warning restore CA2000
 		}
 
-		const int BundlePageSize = 1024 * 1024;
+		public const int BundlePageSize = 1024 * 1024;
 
 		async ValueTask<IReadOnlyMemoryOwner<byte>> ReadEncodedPacketAsync(CancellationToken cancellationToken)
 		{
@@ -238,8 +230,8 @@ namespace EpicGames.Horde.Storage.Bundles.V2
 			// If the whole packet is contained within on page, just return that.
 			if (maxPageIdx == minPageIdx + 1)
 			{
-				IRefCountedHandle<IReadOnlyMemoryOwner<byte>> pageData = await ReadBundlePageAsync(minPageIdx, cancellationToken);
-				return ReadOnlyMemoryOwner.Create(pageData.Target.Memory.Slice(_packetOffset & (BundlePageSize - 1), _packetLength), pageData);
+				Lifetime<IReadOnlyMemoryOwner<byte>> pageData = await ReadBundlePageAsync(minPageIdx, cancellationToken);
+				return ReadOnlyMemoryOwner.Create(pageData.Value.Memory.Slice(_packetOffset & (BundlePageSize - 1), _packetLength), pageData);
 			}
 
 			// Otherwise read sections from each page that contributes to the output and copy into a shared buffer
@@ -256,28 +248,29 @@ namespace EpicGames.Horde.Storage.Bundles.V2
 			}
 		}
 
-		async ValueTask<IRefCountedHandle<IReadOnlyMemoryOwner<byte>>> ReadBundlePageAsync(int pageIdx, CancellationToken cancellationToken)
+		async ValueTask<Lifetime<IReadOnlyMemoryOwner<byte>>> ReadBundlePageAsync(int pageIdx, CancellationToken cancellationToken)
 		{
 			BundlePageCacheKey bundlePageCacheKey = new BundlePageCacheKey(_outer, pageIdx);
-			return await _cache.FindOrAddAsync(bundlePageCacheKey, ReadBundlePageInternalAsync, cancellationToken);
+			return await _cache.BundlePageCache.ScopedGetOrAddAsync(bundlePageCacheKey, ReadBundlePageInternalAsync, cancellationToken);
 		}
 
 		async Task ReadBundlePageAsync(int pageIdx, Memory<byte> targetMemory, CancellationToken cancellationToken)
 		{
-			using IRefCountedHandle<IReadOnlyMemoryOwner<byte>> pageData = await ReadBundlePageAsync(pageIdx, cancellationToken);
+			using Lifetime<IReadOnlyMemoryOwner<byte>> pageData = await ReadBundlePageAsync(pageIdx, cancellationToken);
 
 			int pageBase = pageIdx * BundlePageSize;
 
 			int minOffset = Math.Max(_packetOffset, pageBase);
 			int maxOffset = Math.Min(_packetOffset + _packetLength, pageBase + BundlePageSize);
 
-			ReadOnlyMemory<byte> sourceMemory = pageData.Target.Memory.Slice(minOffset - pageBase, maxOffset - minOffset);
+			ReadOnlyMemory<byte> sourceMemory = pageData.Value.Memory.Slice(minOffset - pageBase, maxOffset - minOffset);
 			sourceMemory.CopyTo(targetMemory.Slice(minOffset - _packetOffset));
 		}
 
-		async Task<IReadOnlyMemoryOwner<byte>> ReadBundlePageInternalAsync(BundlePageCacheKey key, CancellationToken cancellationToken)
+		async Task<Scoped<IReadOnlyMemoryOwner<byte>>> ReadBundlePageInternalAsync(BundlePageCacheKey key, CancellationToken cancellationToken)
 		{
-			return await _outer.ReadAsync(key.Index * BundlePageSize, BundlePageSize, cancellationToken);
+			IReadOnlyMemoryOwner<byte> owner = await _outer.ReadAsync(key.Index * BundlePageSize, BundlePageSize, cancellationToken);
+			return new Scoped<IReadOnlyMemoryOwner<byte>>(owner);
 		}
 		#endregion
 	}
diff --git a/Engine/Source/Programs/Shared/EpicGames.Horde/Storage/Bundles/V2/PacketReader.cs b/Engine/Source/Programs/Shared/EpicGames.Horde/Storage/Bundles/V2/PacketReader.cs
--- a/Engine/Source/Programs/Shared/EpicGames.Horde/Storage/Bundles/V2/PacketReader.cs
+++ b/Engine/Source/Programs/Shared/EpicGames.Horde/Storage/Bundles/V2/PacketReader.cs
@@ -25,6 +25,12 @@ namespace EpicGames.Horde.Storage.Bundles.V2
 		}
 	}
 
+	record struct PacketReaderCacheKey(BundleHandle Bundle, int Offset)
+	{
+		public override string ToString()
+			=> $"packet-reader:{Bundle}@{Offset}";
+	}
+
 	/// <summary>
 	/// Utility class for constructing BlobData objects from a packet, caching any computed handles to other blobs.
 	/// </summary>
