Horde: Expose the bundle page size as a configurable option. Still defaults to 1mb.

diff --git a/Engine/Source/Programs/Shared/EpicGames.Horde/Storage/Bundles/BundleCache.cs b/Engine/Source/Programs/Shared/EpicGames.Horde/Storage/Bundles/BundleCache.cs
--- a/Engine/Source/Programs/Shared/EpicGames.Horde/Storage/Bundles/BundleCache.cs
+++ b/Engine/Source/Programs/Shared/EpicGames.Horde/Storage/Bundles/BundleCache.cs
@@ -22,6 +22,11 @@ namespace EpicGames.Horde.Storage.Bundles
 		/// </summary>
 		public int PacketReaderCount { get; set; } = 200;
 
+		/// <summary>
+		/// Size of a bundle page
+		/// </summary>
+		public int BundlePageSize { get; set; } = 1024 * 1024;
+
 		/// <summary>
 		/// Number of bundle pages to keep in the cache.
 		/// </summary>
@@ -70,6 +75,11 @@ namespace EpicGames.Horde.Storage.Bundles
 		/// </summary>
 		public IMemoryAllocator<byte> Allocator { get; }
 
+		/// <summary>
+		/// Size of a bundle page to keep in the cache
+		/// </summary>
+		public int BundlePageSize => _options.BundlePageSize;
+
 		/// <summary>
 		/// Size of the configured header cache
 		/// </summary>
diff --git a/Engine/Source/Programs/Shared/EpicGames.Horde/Storage/Bundles/V2/PacketHandle.cs b/Engine/Source/Programs/Shared/EpicGames.Horde/Storage/Bundles/V2/PacketHandle.cs
--- a/Engine/Source/Programs/Shared/EpicGames.Horde/Storage/Bundles/V2/PacketHandle.cs
+++ b/Engine/Source/Programs/Shared/EpicGames.Horde/Storage/Bundles/V2/PacketHandle.cs
@@ -220,18 +220,16 @@ namespace EpicGames.Horde.Storage.Bundles.V2
 #pragma warning restore CA2000
 		}
 
-		public const int BundlePageSize = 1024 * 1024;
-
 		async ValueTask<IReadOnlyMemoryOwner<byte>> ReadEncodedPacketAsync(CancellationToken cancellationToken)
 		{
-			int minPageIdx = _packetOffset / BundlePageSize;
-			int maxPageIdx = ((_packetOffset + _packetLength) + (BundlePageSize - 1)) / BundlePageSize;
+			int minPageIdx = _packetOffset / _cache.BundlePageSize;
+			int maxPageIdx = ((_packetOffset + _packetLength) + (_cache.BundlePageSize - 1)) / _cache.BundlePageSize;
 
 			// If the whole packet is contained within on page, just return that.
 			if (maxPageIdx == minPageIdx + 1)
 			{
 				Lifetime<IReadOnlyMemoryOwner<byte>> pageData = await ReadBundlePageAsync(minPageIdx, cancellationToken);
-				return ReadOnlyMemoryOwner.Create(pageData.Value.Memory.Slice(_packetOffset & (BundlePageSize - 1), _packetLength), pageData);
+				return ReadOnlyMemoryOwner.Create(pageData.Value.Memory.Slice(_packetOffset - (minPageIdx * _cache.BundlePageSize), _packetLength), pageData);
 			}
 
 			// Otherwise read sections from each page that contributes to the output and copy into a shared buffer
@@ -258,10 +256,10 @@ namespace EpicGames.Horde.Storage.Bundles.V2
 		{
 			using Lifetime<IReadOnlyMemoryOwner<byte>> pageData = await ReadBundlePageAsync(pageIdx, cancellationToken);
 
-			int pageBase = pageIdx * BundlePageSize;
+			int pageBase = pageIdx * _cache.BundlePageSize;
 
 			int minOffset = Math.Max(_packetOffset, pageBase);
-			int maxOffset = Math.Min(_packetOffset + _packetLength, pageBase + BundlePageSize);
+			int maxOffset = Math.Min(_packetOffset + _packetLength, pageBase + _cache.BundlePageSize);
 
 			ReadOnlyMemory<byte> sourceMemory = pageData.Value.Memory.Slice(minOffset - pageBase, maxOffset - minOffset);
 			sourceMemory.CopyTo(targetMemory.Slice(minOffset - _packetOffset));
@@ -269,9 +267,10 @@ namespace EpicGames.Horde.Storage.Bundles.V2
 
 		async Task<Scoped<IReadOnlyMemoryOwner<byte>>> ReadBundlePageInternalAsync(BundlePageCacheKey key, CancellationToken cancellationToken)
 		{
-			IReadOnlyMemoryOwner<byte> owner = await _outer.ReadAsync(key.Index * BundlePageSize, BundlePageSize, cancellationToken);
+			IReadOnlyMemoryOwner<byte> owner = await _outer.ReadAsync(key.Index * _cache.BundlePageSize, _cache.BundlePageSize, cancellationToken);
 			return new Scoped<IReadOnlyMemoryOwner<byte>>(owner);
 		}
+
 		#endregion
 	}
 }
