Horde: Add separate stats classes for creating and extracting an archive.

diff --git a/Engine/Source/Programs/Horde/Horde.Agent/Execution/JobExecutor.cs b/Engine/Source/Programs/Horde/Horde.Agent/Execution/JobExecutor.cs
--- a/Engine/Source/Programs/Horde/Horde.Agent/Execution/JobExecutor.cs
+++ b/Engine/Source/Programs/Horde/Horde.Agent/Execution/JobExecutor.cs
@@ -871,7 +871,7 @@ namespace Horde.Agent.Execution
 					try
 					{
 						DirectoryNode dir = new DirectoryNode();
-						await dir.AddFilesAsync(baseDir, files, blobWriter, progress: new CopyStatsLogger(logger), cancellationToken: cancellationToken);
+						await dir.AddFilesAsync(baseDir, files, blobWriter, progress: new UpdateStatsLogger(logger), cancellationToken: cancellationToken);
 						rootRef = await blobWriter.WriteBlobAsync(dir, cancellationToken: cancellationToken);
 					}
 					catch (Exception ex)
diff --git a/Engine/Source/Programs/Horde/Horde.Agent/Utility/TempStorage.cs b/Engine/Source/Programs/Horde/Horde.Agent/Utility/TempStorage.cs
--- a/Engine/Source/Programs/Horde/Horde.Agent/Utility/TempStorage.cs
+++ b/Engine/Source/Programs/Horde/Horde.Agent/Utility/TempStorage.cs
@@ -700,7 +700,7 @@ namespace Horde.Storage.Utility
 
 			// Create the file tree
 			DirectoryNode rootNode = new DirectoryNode();
-			await rootNode.AddFilesAsync(workspaceDir, archiveFiles, writer, progress: new CopyStatsLogger(logger), cancellationToken: cancellationToken);
+			await rootNode.AddFilesAsync(workspaceDir, archiveFiles, writer, progress: new UpdateStatsLogger(logger), cancellationToken: cancellationToken);
 
 			IBlobRef<DirectoryNode> rootNodeRef = await writer.WriteBlobAsync(rootNode, cancellationToken: cancellationToken);
 			return new DirectoryEntry(blockDirectoryName, rootNode.Length, rootNodeRef);
@@ -767,7 +767,7 @@ namespace Horde.Storage.Utility
 
 				// Add all the files and flush the ref
 				DirectoryNode rootDirNode = await rootDirEntry.Handle.ReadBlobAsync(cancellationToken: cancellationToken);
-				await rootDirNode.CopyToDirectoryAsync(rootDir.ToDirectoryInfo(), new CopyStatsLogger(logger), logger, cancellationToken);
+				await rootDirNode.CopyToDirectoryAsync(rootDir.ToDirectoryInfo(), new ExtractStatsLogger(logger), logger, cancellationToken);
 
 				StorageStats deltaStats = StorageStats.GetDelta(initialStats, storageClient.GetStats());
 				logger.LogInformation("{Stats}", $"Elapsed: {(int)timer.Elapsed.TotalSeconds}s, {String.Join(", ", deltaStats.Values.Select(x => $"{x.Item1}: {x.Item2:n0}"))}");
diff --git a/Engine/Source/Programs/Horde/Horde/Commands/Bundles/BundleCreate.cs b/Engine/Source/Programs/Horde/Horde/Commands/Bundles/BundleCreate.cs
--- a/Engine/Source/Programs/Horde/Horde/Commands/Bundles/BundleCreate.cs
+++ b/Engine/Source/Programs/Horde/Horde/Commands/Bundles/BundleCreate.cs
@@ -86,9 +86,9 @@ namespace Horde.Commands.Bundles
 				ChunkingOptions options = new ChunkingOptions();
 
 				List<FileInfo> fileInfos = files.ConvertAll(x => x.ToFileInfo());
-				CopyStatsLogger copyStatsLogger = new CopyStatsLogger(files.Count, fileInfos.Sum(x => x.Length), logger);
+				UpdateStatsLogger updateStatsLogger = new UpdateStatsLogger(files.Count, fileInfos.Sum(x => x.Length), logger);
 
-				IBlobRef<DirectoryNode> nodeRef = await writer.WriteFilesAsync(baseDir.ToDirectoryInfo(), fileInfos, options, copyStatsLogger, CancellationToken.None);
+				IBlobRef<DirectoryNode> nodeRef = await writer.WriteFilesAsync(baseDir.ToDirectoryInfo(), fileInfos, options, updateStatsLogger, CancellationToken.None);
 
 				await writer.FlushAsync();
 
diff --git a/Engine/Source/Programs/Horde/Horde/Commands/Bundles/BundleExtract.cs b/Engine/Source/Programs/Horde/Horde/Commands/Bundles/BundleExtract.cs
--- a/Engine/Source/Programs/Horde/Horde/Commands/Bundles/BundleExtract.cs
+++ b/Engine/Source/Programs/Horde/Horde/Commands/Bundles/BundleExtract.cs
@@ -84,7 +84,7 @@ namespace Horde.Commands.Bundles
 			Stopwatch timer = Stopwatch.StartNew();
 
 			DirectoryNode node = await handle.ReadBlobAsync<DirectoryNode>();
-			await node.CopyToDirectoryAsync(OutputDir.ToDirectoryInfo(), new CopyStatsLogger(logger), logger, CancellationToken.None);
+			await node.CopyToDirectoryAsync(OutputDir.ToDirectoryInfo(), new ExtractStatsLogger(logger), logger, CancellationToken.None);
 
 			logger.LogInformation("Elapsed: {Time}s", timer.Elapsed.TotalSeconds);
 
diff --git a/Engine/Source/Programs/Shared/EpicGames.Horde/Storage/Bundles/V2/PacketHandle.cs b/Engine/Source/Programs/Shared/EpicGames.Horde/Storage/Bundles/V2/PacketHandle.cs
--- a/Engine/Source/Programs/Shared/EpicGames.Horde/Storage/Bundles/V2/PacketHandle.cs
+++ b/Engine/Source/Programs/Shared/EpicGames.Horde/Storage/Bundles/V2/PacketHandle.cs
@@ -214,6 +214,7 @@ namespace EpicGames.Horde.Storage.Bundles.V2
 		{
 			using IReadOnlyMemoryOwner<byte> encodedData = await ReadEncodedPacketAsync(cancellationToken);
 			IRefCountedHandle<Packet> packet = Packet.Decode(encodedData.Memory, _cache.Allocator, cacheKey);
+			Interlocked.Add(ref _storageClient.PacketReaderStats._numDecodedBytesRead, packet.Target.Length);
 #pragma warning disable CA2000
 			PacketReader reader = new PacketReader(_storageClient, _cache, _outer, this, packet.Target, packet);
 			return new Scoped<PacketReader>(reader);
@@ -222,6 +223,8 @@ namespace EpicGames.Horde.Storage.Bundles.V2
 
 		async ValueTask<IReadOnlyMemoryOwner<byte>> ReadEncodedPacketAsync(CancellationToken cancellationToken)
 		{
+			Interlocked.Add(ref _storageClient.PacketReaderStats._numEncodedBytesRead, _packetLength);
+
 			int minPageIdx = _packetOffset / _cache.BundlePageSize;
 			int maxPageIdx = ((_packetOffset + _packetLength) + (_cache.BundlePageSize - 1)) / _cache.BundlePageSize;
 
@@ -268,6 +271,7 @@ namespace EpicGames.Horde.Storage.Bundles.V2
 		async Task<Scoped<IReadOnlyMemoryOwner<byte>>> ReadBundlePageInternalAsync(BundlePageCacheKey key, CancellationToken cancellationToken)
 		{
 			IReadOnlyMemoryOwner<byte> owner = await _outer.ReadAsync(key.Index * _cache.BundlePageSize, _cache.BundlePageSize, cancellationToken);
+			Interlocked.Add(ref _storageClient.PacketReaderStats._numBytesRead, owner.Memory.Length);
 			return new Scoped<IReadOnlyMemoryOwner<byte>>(owner);
 		}
 
diff --git a/Engine/Source/Programs/Shared/EpicGames.Horde/Storage/Bundles/V2/PacketReader.cs b/Engine/Source/Programs/Shared/EpicGames.Horde/Storage/Bundles/V2/PacketReader.cs
--- a/Engine/Source/Programs/Shared/EpicGames.Horde/Storage/Bundles/V2/PacketReader.cs
+++ b/Engine/Source/Programs/Shared/EpicGames.Horde/Storage/Bundles/V2/PacketReader.cs
@@ -11,17 +11,15 @@ namespace EpicGames.Horde.Storage.Bundles.V2
 	/// </summary>
 	class PacketReaderStats
 	{
-		public int _numReads;
-		public int _numPacketsRead;
 		public long _numBytesRead;
-		public long _numWastedBytesRead;
+		public long _numEncodedBytesRead;
+		public long _numDecodedBytesRead;
 
 		public void GetStats(StorageStats stats)
 		{
-			stats.Add("bundle.packet_reader.num_reads", _numReads);
-			stats.Add("bundle.packet_reader.num_packets", _numPacketsRead);
-			stats.Add("bundle.packet_reader.num_bytes", _numBytesRead);
-			stats.Add("bundle.packet_reader.num_wasted_bytes", _numWastedBytesRead);
+			stats.Add("bundle.packet_reader.num_bytes_read", _numBytesRead);
+			stats.Add("bundle.packet_reader.num_encoded_bytes_read", _numEncodedBytesRead);
+			stats.Add("bundle.packet_reader.num_decoded_bytes_read", _numDecodedBytesRead);
 		}
 	}
 
diff --git a/Engine/Source/Programs/Shared/EpicGames.Horde/Storage/Nodes/ChunkedDataNode.cs b/Engine/Source/Programs/Shared/EpicGames.Horde/Storage/Nodes/ChunkedDataNode.cs
--- a/Engine/Source/Programs/Shared/EpicGames.Horde/Storage/Nodes/ChunkedDataNode.cs
+++ b/Engine/Source/Programs/Shared/EpicGames.Horde/Storage/Nodes/ChunkedDataNode.cs
@@ -289,10 +289,10 @@ namespace EpicGames.Horde.Storage.Nodes
 		/// <param name="writer">Writer for output nodes</param>
 		/// <param name="stream">Stream to read from</param>
 		/// <param name="options">Options for finding chunk boundaries</param>
-		/// <param name="copyStats">Stats for the copy operation</param>
+		/// <param name="updateStats">Stats for the copy operation</param>
 		/// <param name="cancellationToken">Cancellation token for the operation</param>
 		/// <returns>Hash of the full file data</returns>
-		internal static async Task<LeafChunkedData> CreateFromStreamAsync(IBlobWriter writer, Stream stream, LeafChunkedDataNodeOptions options, CopyStats? copyStats, CancellationToken cancellationToken)
+		internal static async Task<LeafChunkedData> CreateFromStreamAsync(IBlobWriter writer, Stream stream, LeafChunkedDataNodeOptions options, UpdateStats? updateStats, CancellationToken cancellationToken)
 		{
 			using Blake3.Hasher hasher = Blake3.Hasher.New();
 			using IMemoryOwner<byte> readBuffer = MemoryPool<byte>.Shared.Rent(options.MaxSize);
@@ -318,7 +318,7 @@ namespace EpicGames.Horde.Storage.Nodes
 				sizeSinceProgressUpdate += nextLength;
 				if (sizeSinceProgressUpdate > 512 * 1024)
 				{
-					copyStats?.Update(0, sizeSinceProgressUpdate);
+					updateStats?.Update(0, sizeSinceProgressUpdate);
 					sizeSinceProgressUpdate = 0;
 				}
 
@@ -329,7 +329,7 @@ namespace EpicGames.Horde.Storage.Nodes
 				size -= nextLength;
 			}
 
-			copyStats?.Update(1, sizeSinceProgressUpdate);
+			updateStats?.Update(1, sizeSinceProgressUpdate);
 
 			IoHash hash = IoHash.FromBlake3(hasher);
 			return new LeafChunkedData(hash, leafNodeRefs);
diff --git a/Engine/Source/Programs/Shared/EpicGames.Horde/Storage/Nodes/DirectoryNode.Extract.cs b/Engine/Source/Programs/Shared/EpicGames.Horde/Storage/Nodes/DirectoryNode.Extract.cs
--- a/Engine/Source/Programs/Shared/EpicGames.Horde/Storage/Nodes/DirectoryNode.Extract.cs
+++ b/Engine/Source/Programs/Shared/EpicGames.Horde/Storage/Nodes/DirectoryNode.Extract.cs
@@ -20,7 +20,7 @@ namespace EpicGames.Horde.Storage.Nodes
 	/// <summary>
 	/// Stats reported for copy operations
 	/// </summary>
-	public interface ICopyStats
+	public interface IExtractStats
 	{
 		/// <summary>
 		/// Number of files that have been copied
@@ -41,53 +41,63 @@ namespace EpicGames.Horde.Storage.Nodes
 	/// <summary>
 	/// Reports progress info back to callers
 	/// </summary>
-	class CopyStats : ICopyStats
+	class ExtractStats : IExtractStats
 	{
 		readonly object _lockObject = new object();
 		readonly Stopwatch _timer = Stopwatch.StartNew();
-		readonly IProgress<ICopyStats> _progress;
+		readonly IProgress<IExtractStats>? _progress;
 		long _lastTotalSize;
 
 		public int Count { get; set; }
 		public long Size { get; set; }
 		public double Rate { get; set; }
 
-		public CopyStats(IProgress<ICopyStats> progress)
+		public ExtractStats(IProgress<IExtractStats>? progress)
 		{
 			_progress = progress;
 		}
 
 		public void Update(int count, long size)
 		{
-			lock (_lockObject)
+			if (_progress != null)
 			{
-				Count += count;
-				Size += size;
-				if (_timer.Elapsed > TimeSpan.FromSeconds(10.0))
+				lock (_lockObject)
 				{
-					Rate = (Size - _lastTotalSize) / _timer.Elapsed.TotalSeconds;
-					_lastTotalSize = Size;
-
-					_progress.Report(this);
-					_timer.Restart();
+					Count += count;
+					Size += size;
+					if (_timer.Elapsed > TimeSpan.FromSeconds(5.0))
+					{
+						FlushInternal();
+					}
 				}
 			}
 		}
 
 		public void Flush()
 		{
-			lock (_lockObject)
+			if (_progress != null)
 			{
-				_progress.Report(this);
-				_timer.Restart();
+				lock (_lockObject)
+				{
+					FlushInternal();
+				}
 			}
 		}
+
+		void FlushInternal()
+		{
+			Rate = (Size - _lastTotalSize) / _timer.Elapsed.TotalSeconds;
+			_lastTotalSize = Size;
+
+			_progress!.Report(this);
+			_timer.Restart();
+		}
 	}
 
 	/// <summary>
 	/// Progress logger for writing copy stats
 	/// </summary>
-	public class CopyStatsLogger : IProgress<ICopyStats>
+	public class ExtractStatsLogger : IProgress<IExtractStats>
 	{
 		readonly int _totalCount;
 		readonly long _totalSize;
@@ -96,13 +106,13 @@ namespace EpicGames.Horde.Storage.Nodes
 		/// <summary>
 		/// Constructor
 		/// </summary>
-		public CopyStatsLogger(ILogger logger)
+		public ExtractStatsLogger(ILogger logger)
 			=> _logger = logger;
 
 		/// <summary>
 		/// Constructor
 		/// </summary>
-		public CopyStatsLogger(int totalCount, long totalSize, ILogger logger)
+		public ExtractStatsLogger(int totalCount, long totalSize, ILogger logger)
 		{
 			_totalCount = totalCount;
 			_totalSize = totalSize;
@@ -110,7 +120,7 @@ namespace EpicGames.Horde.Storage.Nodes
 		}
 
 		/// <inheritdoc/>
-		public void Report(ICopyStats stats)
+		public void Report(IExtractStats stats)
 		{
 			if (_totalCount > 0 && _totalSize > 0)
 			{
@@ -262,17 +272,12 @@ namespace EpicGames.Horde.Storage.Nodes
 		/// <param name="progress">Sink for progress updates</param>
 		/// <param name="logger">Logger for output</param>
 		/// <param name="cancellationToken">Cancellation token for the operation</param>
-		public static async Task CopyToDirectoryAsync(this DirectoryNode directoryNode, DirectoryInfo directoryInfo, IProgress<ICopyStats>? progress, ILogger logger, CancellationToken cancellationToken)
+		public static async Task CopyToDirectoryAsync(this DirectoryNode directoryNode, DirectoryInfo directoryInfo, IProgress<IExtractStats>? progress, ILogger logger, CancellationToken cancellationToken)
 		{
 			int numTasks = Math.Min(1 + (int)(directoryNode.Length / (16 * 1024 * 1024)), 16);
 			logger.LogInformation("Splitting read into {NumThreads} threads", numTasks);
 
-			CopyStats? copyStats = null;
-			if (progress != null)
-			{
-				copyStats = new CopyStats(progress);
-			}
-
+			ExtractStats extractStats = new ExtractStats(progress);
 			using (CancellationTokenSource cancellationSource = CancellationTokenSource.CreateLinkedTokenSource(cancellationToken))
 			{
 				// Helper method to run a background task and set a cancellation source on error
@@ -280,7 +285,7 @@ namespace EpicGames.Horde.Storage.Nodes
 				{
 					try
 					{
-						await taskFunc(cancellationSource.Token);
+						await Task.Run(() => taskFunc(cancellationSource.Token), cancellationSource.Token);
 					}
 					catch (OperationCanceledException)
 					{
@@ -308,7 +313,7 @@ namespace EpicGames.Horde.Storage.Nodes
 
 					for (int idx = 0; idx < numTasks; idx++)
 					{
-						tasks.Add(RunBackgroundTask(ctx => ExtractAsync(prefetchBatches.Reader, copyStats, logger, ctx)));
+						tasks.Add(RunBackgroundTask(ctx => WriteAsync(prefetchBatches.Reader, extractStats, logger, ctx)));
 					}
 
 					await Task.WhenAll(tasks);
@@ -320,6 +325,7 @@ namespace EpicGames.Horde.Storage.Nodes
 						fetchedBatch.Dispose();
 					}
 				}
+				extractStats.Flush();
 			}
 		}
 
@@ -500,7 +506,7 @@ namespace EpicGames.Horde.Storage.Nodes
 			List<Task> tasks = new List<Task>();
 			for (int idx = 0; idx < numParallel; idx++)
 			{
-				tasks.Add(FetchWorkerAsync(batchReader, batchWriter, cancellationToken));
+				tasks.Add(Task.Run(() => FetchWorkerAsync(batchReader, batchWriter, cancellationToken), cancellationToken));
 			}
 
 			try
@@ -544,7 +550,7 @@ namespace EpicGames.Horde.Storage.Nodes
 
 		#region Write to disk
 
-		static async Task ExtractAsync(ChannelReader<OutputFetchedBatch> batchReader, CopyStats? copyStats, ILogger logger, CancellationToken cancellationToken)
+		static async Task WriteAsync(ChannelReader<OutputFetchedBatch> batchReader, ExtractStats stats, ILogger logger, CancellationToken cancellationToken)
 		{
 			const int WriteBatchSize = 64;
 			while (await batchReader.WaitToReadAsync(cancellationToken))
@@ -557,7 +563,7 @@ namespace EpicGames.Horde.Storage.Nodes
 						List<Task> tasks = new List<Task>();
 						foreach (IReadOnlyList<OutputFetchedChunk> group in batch.Chunks.Batch(WriteBatchSize))
 						{
-							tasks.Add(ExtractChunksAsync(group.ToArray(), copyStats, logger, cancellationToken));
+							tasks.Add(WriteChunksAsync(group.ToArray(), stats, logger, cancellationToken));
 						}
 						await Task.WhenAll(tasks);
 					}
@@ -569,7 +575,7 @@ namespace EpicGames.Horde.Storage.Nodes
 			}
 		}
 
-		static async Task ExtractChunksAsync(ArraySegment<OutputFetchedChunk> chunks, CopyStats? copyStats, ILogger logger, CancellationToken cancellationToken)
+		static async Task WriteChunksAsync(ArraySegment<OutputFetchedChunk> chunks, ExtractStats stats, ILogger logger, CancellationToken cancellationToken)
 		{
 			for (int chunkIdx = 0; chunkIdx < chunks.Count;)
 			{
@@ -583,8 +589,8 @@ namespace EpicGames.Horde.Storage.Nodes
 
 				try
 				{
-					await ExtractChunksToFileAsync(file, chunks.Slice(chunkIdx, maxChunkIdx - chunkIdx), copyStats, logger, cancellationToken);
-					//await ExtractChunksToNullAsync(file, chunks.Slice(chunkIdx, maxChunkIdx - chunkIdx), copyStats, logger, cancellationToken);
+					await ExtractChunksToFileAsync(file, chunks.Slice(chunkIdx, maxChunkIdx - chunkIdx), stats, logger, cancellationToken);
+					//await ExtractChunksToNullAsync(file, chunks.Slice(chunkIdx, maxChunkIdx - chunkIdx), stats, logger, cancellationToken);
 				}
 				catch (OperationCanceledException)
 				{
@@ -599,7 +605,7 @@ namespace EpicGames.Horde.Storage.Nodes
 			}
 		}
 
-		static async Task ExtractChunksToFileAsync(OutputFile file, ArraySegment<OutputFetchedChunk> chunks, CopyStats? copyStats, ILogger logger, CancellationToken cancellationToken)
+		static async Task ExtractChunksToFileAsync(OutputFile file, ArraySegment<OutputFetchedChunk> chunks, ExtractStats stats, ILogger logger, CancellationToken cancellationToken)
 		{
 			// Open the file for the current chunk
 			int remainingChunks = 0;
@@ -630,7 +636,7 @@ namespace EpicGames.Horde.Storage.Nodes
 
 						// Update the stats
 						remainingChunks = file.DecrementRemaining();
-						copyStats?.Update(0, chunk.Length);
+						stats.Update(0, chunk.Length);
 					}
 				}
 			}
@@ -640,13 +646,13 @@ namespace EpicGames.Horde.Storage.Nodes
 			{
 				file.FileInfo.Refresh();
 				FileEntry.SetPermissions(file.FileInfo!, file.FileEntry.Flags);
-				copyStats?.Update(1, 0);
+				stats.Update(1, 0);
 			}
 		}
 
 #pragma warning disable IDE0051
 		// Update counters for extracting chunks without writing any data. Useful for profiling bottlenecks in other stages of the pipeline.
-		static Task ExtractChunksToNullAsync(OutputFile file, ArraySegment<OutputFetchedChunk> chunks, CopyStats? copyStats, ILogger logger, CancellationToken cancellationToken)
+		static Task ExtractChunksToNullAsync(OutputFile file, ArraySegment<OutputFetchedChunk> chunks, ExtractStats stats, ILogger logger, CancellationToken cancellationToken)
 		{
 			_ = logger;
 			_ = cancellationToken;
@@ -656,11 +662,11 @@ namespace EpicGames.Horde.Storage.Nodes
 				int remainingChunks = file.DecrementRemaining();
 				if (remainingChunks == 0)
 				{
-					copyStats?.Update(1, chunk.Length);
+					stats.Update(1, chunk.Length);
 				}
 				else
 				{
-					copyStats?.Update(0, chunk.Length);
+					stats.Update(0, chunk.Length);
 				}
 			}
 			return Task.CompletedTask;
diff --git a/Engine/Source/Programs/Shared/EpicGames.Horde/Storage/Nodes/DirectoryNode.Update.cs b/Engine/Source/Programs/Shared/EpicGames.Horde/Storage/Nodes/DirectoryNode.Update.cs
--- a/Engine/Source/Programs/Shared/EpicGames.Horde/Storage/Nodes/DirectoryNode.Update.cs
+++ b/Engine/Source/Programs/Shared/EpicGames.Horde/Storage/Nodes/DirectoryNode.Update.cs
@@ -2,15 +2,141 @@
 
 using System;
 using System.Collections.Generic;
+using System.Diagnostics;
 using System.IO;
 using System.IO.Compression;
 using System.Linq;
 using System.Threading;
 using System.Threading.Tasks;
 using EpicGames.Core;
+using Microsoft.Extensions.Logging;
 
 namespace EpicGames.Horde.Storage.Nodes
 {
+	/// <summary>
+	/// Stats reported for copy operations
+	/// </summary>
+	public interface IUpdateStats
+	{
+		/// <summary>
+		/// Number of files that have been copied
+		/// </summary>
+		int Count { get; }
+
+		/// <summary>
+		/// Total size of data to be copied
+		/// </summary>
+		long Size { get; }
+
+		/// <summary>
+		/// Processing speed, in bytes per second
+		/// </summary>
+		double Rate { get; }
+	}
+
+	/// <summary>
+	/// Reports progress info back to callers
+	/// </summary>
+	class UpdateStats : IUpdateStats
+	{
+		readonly object _lockObject = new object();
+		readonly Stopwatch _timer = Stopwatch.StartNew();
+		readonly IProgress<IUpdateStats>? _progress;
+		long _lastTotalSize;
+
+		public int Count { get; set; }
+		public long Size { get; set; }
+		public double Rate { get; set; }
+
+		public UpdateStats(IProgress<IUpdateStats>? progress)
+		{
+			_progress = progress;
+		}
+
+		public void Update(int count, long size)
+		{
+			if (_progress != null)
+			{
+				lock (_lockObject)
+				{
+					Count += count;
+					Size += size;
+					if (_timer.Elapsed > TimeSpan.FromSeconds(5.0))
+					{
+						Rate = (Size - _lastTotalSize) / _timer.Elapsed.TotalSeconds;
+						_lastTotalSize = Size;
+
+						FlushInternal();
+					}
+				}
+			}
+		}
+
+		public void Flush()
+		{
+			if (_progress != null)
+			{
+				lock (_lockObject)
+				{
+					FlushInternal();
+				}
+			}
+		}
+
+		void FlushInternal()
+		{
+			_progress!.Report(this);
+			_timer.Restart();
+		}
+	}
+
+	/// <summary>
+	/// Progress logger for writing copy stats
+	/// </summary>
+	public class UpdateStatsLogger : IProgress<IUpdateStats>
+	{
+		readonly int _totalCount;
+		readonly long _totalSize;
+		readonly ILogger _logger;
+
+		/// <summary>
+		/// Constructor
+		/// </summary>
+		public UpdateStatsLogger(ILogger logger)
+			=> _logger = logger;
+
+		/// <summary>
+		/// Constructor
+		/// </summary>
+		public UpdateStatsLogger(int totalCount, long totalSize, ILogger logger)
+		{
+			_totalCount = totalCount;
+			_totalSize = totalSize;
+			_logger = logger;
+		}
+
+		/// <inheritdoc/>
+		public void Report(IUpdateStats stats)
+		{
+			if (_totalCount > 0 && _totalSize > 0)
+			{
+				_logger.LogInformation("Copied {NumFiles:n0}/{TotalFiles:n0} files ({Size:n1}/{TotalSize:n1}mb, {Rate:n1}mb/s, {Pct}%)", stats.Count, _totalCount, stats.Size / (1024.0 * 1024.0), _totalSize / (1024.0 * 1024.0), stats.Rate / (1024.0 * 1024.0), (int)((Math.Max(stats.Size, 1) * 100) / Math.Max(_totalSize, 1)));
+			}
+			else if (_totalCount > 0)
+			{
+				_logger.LogInformation("Copied {NumFiles:n0}/{TotalFiles:n0} files ({Size:n1}mb, {Rate:n1}mb/s)", stats.Count, _totalCount, stats.Size / (1024.0 * 1024.0), stats.Rate / (1024.0 * 1024.0));
+			}
+			else if (_totalSize > 0)
+			{
+				_logger.LogInformation("Copied {NumFiles:n0} files ({Size:n1}/{TotalSize:n1}mb, {Rate:n1}mb/s, {Pct}%)", stats.Count, stats.Size / (1024.0 * 1024.0), _totalSize / (1024.0 * 1024.0), stats.Rate / (1024.0 * 1024.0), (int)((Math.Max(stats.Size, 1) * 100) / Math.Max(_totalSize, 1)));
+			}
+			else
+			{
+				_logger.LogInformation("Copied {NumFiles:n0} files ({Size:n1}mb, {Rate:n1}mb/s)", stats.Count, stats.Size / (1024.0 * 1024.0), stats.Rate / (1024.0 * 1024.0));
+			}
+		}
+	}
+
 	/// <summary>
 	/// Describes an update to a file in a directory tree
 	/// </summary>
@@ -214,7 +340,7 @@ namespace EpicGames.Horde.Storage.Nodes
 		/// <summary>
 		/// Writes a tree of files to a storage writer
 		/// </summary>
-		public static async Task<IBlobRef<DirectoryNode>> WriteFilesAsync(this IBlobWriter writer, DirectoryReference baseDir, ChunkingOptions? options = null, IProgress<ICopyStats>? progress = null, CancellationToken cancellationToken = default)
+		public static async Task<IBlobRef<DirectoryNode>> WriteFilesAsync(this IBlobWriter writer, DirectoryReference baseDir, ChunkingOptions? options = null, IProgress<IUpdateStats>? progress = null, CancellationToken cancellationToken = default)
 		{
 			DirectoryNode outputNode = new DirectoryNode();
 			await outputNode.AddFilesAsync(baseDir, DirectoryReference.EnumerateFiles(baseDir, "*", SearchOption.AllDirectories), writer, options, progress, cancellationToken);
@@ -224,7 +350,7 @@ namespace EpicGames.Horde.Storage.Nodes
 		/// <summary>
 		/// Writes a tree of files to a storage writer
 		/// </summary>
-		public static async Task<IBlobRef<DirectoryNode>> WriteFilesAsync(this IBlobWriter writer, DirectoryInfo baseDir, IReadOnlyList<FileInfo> files, ChunkingOptions? options = null, IProgress<ICopyStats>? progress = null, CancellationToken cancellationToken = default)
+		public static async Task<IBlobRef<DirectoryNode>> WriteFilesAsync(this IBlobWriter writer, DirectoryInfo baseDir, IReadOnlyList<FileInfo> files, ChunkingOptions? options = null, IProgress<IUpdateStats>? progress = null, CancellationToken cancellationToken = default)
 		{
 			DirectoryNode outputNode = new DirectoryNode();
 			await outputNode.AddFilesAsync(baseDir, files, writer, options, progress, cancellationToken);
@@ -234,27 +360,27 @@ namespace EpicGames.Horde.Storage.Nodes
 		/// <summary>
 		/// Writes a tree of files to a storage writer
 		/// </summary>
-		public static async Task<IBlobRef<DirectoryNode>> WriteFilesAsync(this IBlobWriter writer, DirectoryReference baseDir, IReadOnlyList<FileReference> files, ChunkingOptions? options = null, IProgress<ICopyStats>? progress = null, CancellationToken cancellationToken = default)
+		public static async Task<IBlobRef<DirectoryNode>> WriteFilesAsync(this IBlobWriter writer, DirectoryReference baseDir, IReadOnlyList<FileReference> files, ChunkingOptions? options = null, IProgress<IUpdateStats>? progress = null, CancellationToken cancellationToken = default)
 		{
 			DirectoryNode outputNode = new DirectoryNode();
 			await outputNode.AddFilesAsync(baseDir, files, writer, options, progress, cancellationToken);
 			return await writer.WriteBlobAsync(outputNode, cancellationToken: cancellationToken);
 		}
 
-		/// <inheritdoc cref="AddFilesAsync(DirectoryNode, DirectoryReference, IEnumerable{FileInfo}, IBlobWriter, ChunkingOptions?, IProgress{ICopyStats}?, CancellationToken)"/>
-		public static async Task AddFilesAsync(this DirectoryNode directoryNode, DirectoryInfo directoryInfo, IBlobWriter writer, ChunkingOptions? options = null, IProgress<ICopyStats>? progress = null, CancellationToken cancellationToken = default)
+		/// <inheritdoc cref="AddFilesAsync(DirectoryNode, DirectoryReference, IEnumerable{FileInfo}, IBlobWriter, ChunkingOptions?, IProgress{IUpdateStats}?, CancellationToken)"/>
+		public static async Task AddFilesAsync(this DirectoryNode directoryNode, DirectoryInfo directoryInfo, IBlobWriter writer, ChunkingOptions? options = null, IProgress<IUpdateStats>? progress = null, CancellationToken cancellationToken = default)
 		{
 			await directoryNode.AddFilesAsync(new DirectoryReference(directoryInfo), directoryInfo.EnumerateFiles("*", SearchOption.AllDirectories).ToList(), writer, options, progress, cancellationToken);
 		}
 
-		/// <inheritdoc cref="AddFilesAsync(DirectoryNode, DirectoryReference, IEnumerable{FileInfo}, IBlobWriter, ChunkingOptions?, IProgress{ICopyStats}?, CancellationToken)"/>
-		public static Task AddFilesAsync(this DirectoryNode directoryNode, DirectoryReference baseDir, IEnumerable<FileReference> files, IBlobWriter writer, ChunkingOptions? options = null, IProgress<ICopyStats>? progress = null, CancellationToken cancellationToken = default)
+		/// <inheritdoc cref="AddFilesAsync(DirectoryNode, DirectoryReference, IEnumerable{FileInfo}, IBlobWriter, ChunkingOptions?, IProgress{IUpdateStats}?, CancellationToken)"/>
+		public static Task AddFilesAsync(this DirectoryNode directoryNode, DirectoryReference baseDir, IEnumerable<FileReference> files, IBlobWriter writer, ChunkingOptions? options = null, IProgress<IUpdateStats>? progress = null, CancellationToken cancellationToken = default)
 		{
 			return directoryNode.AddFilesAsync(baseDir, files.Select(x => x.ToFileInfo()).ToList(), writer, options, progress, cancellationToken);
 		}
 
-		/// <inheritdoc cref="AddFilesAsync(DirectoryNode, DirectoryReference, IEnumerable{FileInfo}, IBlobWriter, ChunkingOptions?, IProgress{ICopyStats}?, CancellationToken)"/>
-		public static Task AddFilesAsync(this DirectoryNode directoryNode, DirectoryInfo baseDir, IEnumerable<FileInfo> files, IBlobWriter writer, ChunkingOptions? options = null, IProgress<ICopyStats>? progress = null, CancellationToken cancellationToken = default)
+		/// <inheritdoc cref="AddFilesAsync(DirectoryNode, DirectoryReference, IEnumerable{FileInfo}, IBlobWriter, ChunkingOptions?, IProgress{IUpdateStats}?, CancellationToken)"/>
+		public static Task AddFilesAsync(this DirectoryNode directoryNode, DirectoryInfo baseDir, IEnumerable<FileInfo> files, IBlobWriter writer, ChunkingOptions? options = null, IProgress<IUpdateStats>? progress = null, CancellationToken cancellationToken = default)
 		{
 			return directoryNode.AddFilesAsync(new DirectoryReference(baseDir), files.ToList(), writer, options, progress, cancellationToken);
 		}
@@ -319,14 +445,14 @@ namespace EpicGames.Horde.Storage.Nodes
 		/// <param name="writer">Writer for new node data</param>
 		/// <param name="progress">Feedback interface for progress updates</param>
 		/// <param name="cancellationToken">Cancellation token for the operation</param>
-		public static async Task AddFilesAsync(this DirectoryNode directoryNode, DirectoryReference baseDir, IEnumerable<FileInfo> files, IBlobWriter writer, ChunkingOptions? options = null, IProgress<ICopyStats>? progress = null, CancellationToken cancellationToken = default)
+		public static async Task AddFilesAsync(this DirectoryNode directoryNode, DirectoryReference baseDir, IEnumerable<FileInfo> files, IBlobWriter writer, ChunkingOptions? options = null, IProgress<IUpdateStats>? progress = null, CancellationToken cancellationToken = default)
 		{
 			options ??= new ChunkingOptions();
 
-			CopyStats? copyStats = null;
+			UpdateStats? updateStats = null;
 			if (progress != null)
 			{
-				copyStats = new CopyStats(progress);
+				updateStats = new UpdateStats(progress);
 			}
 
 			DirectoryUpdate update = new DirectoryUpdate();
@@ -353,7 +479,7 @@ namespace EpicGames.Horde.Storage.Nodes
 					// Partition them up into parallel writers
 					List<(int Start, int Count)> partitions = ComputePartitions(batch, batchSize);
 					LeafChunkedData[] leafChunkedFiles = new LeafChunkedData[batch.Count];
-					await Parallel.ForEachAsync(partitions, cancellationToken, (filePartition, ctx) => CreateLeafChunkNodesAsync(writer, batch, leafChunkedFiles, filePartition.Start, filePartition.Count, copyStats, options, cancellationToken));
+					await Parallel.ForEachAsync(partitions, cancellationToken, (filePartition, ctx) => CreateLeafChunkNodesAsync(writer, batch, leafChunkedFiles, filePartition.Start, filePartition.Count, updateStats, options, cancellationToken));
 
 					// Write all the interior nodes and generate the directory update
 					for (int idx = 0; idx < batch.Count; idx++)
@@ -407,7 +533,7 @@ namespace EpicGames.Horde.Storage.Nodes
 			return partitions;
 		}
 
-		static async ValueTask CreateLeafChunkNodesAsync(IBlobWriter writer, IReadOnlyList<FileInfo> files, LeafChunkedData[] leafChunks, int start, int count, CopyStats? copyStats, ChunkingOptions options, CancellationToken cancellationToken)
+		static async ValueTask CreateLeafChunkNodesAsync(IBlobWriter writer, IReadOnlyList<FileInfo> files, LeafChunkedData[] leafChunks, int start, int count, UpdateStats? updateStats, ChunkingOptions options, CancellationToken cancellationToken)
 		{
 			await using IBlobWriter writerFork = writer.Fork();
 			for (int idx = start; idx < start + count; idx++)
@@ -415,7 +541,7 @@ namespace EpicGames.Horde.Storage.Nodes
 				FileInfo file = files[idx];
 				using (Stream stream = file.OpenRead())
 				{
-					leafChunks[idx] = await LeafChunkedDataNode.CreateFromStreamAsync(writerFork, stream, options.LeafOptions, copyStats, cancellationToken);
+					leafChunks[idx] = await LeafChunkedDataNode.CreateFromStreamAsync(writerFork, stream, options.LeafOptions, updateStats, cancellationToken);
 				}
 			}
 			await writerFork.FlushAsync(cancellationToken);
