Update description for submitting DotNetPerforceLib binaries.

diff --git a/Engine/Source/Programs/Shared/EpicGames.Perforce.Native/BuildLibrary.xml b/Engine/Source/Programs/Shared/EpicGames.Perforce.Native/BuildLibrary.xml
--- a/Engine/Source/Programs/Shared/EpicGames.Perforce.Native/BuildLibrary.xml
+++ b/Engine/Source/Programs/Shared/EpicGames.Perforce.Native/BuildLibrary.xml
@@ -35,7 +35,7 @@
 	<Agent Name="DotNetPerforceLib Submit" Type="Win64">
 		<Node Name="Submit DotNetPerforceLib Binaries" Requires="#BinariesWin64;#BinariesMac;#BinariesLinux">
 			<Tag Files="#BinariesWin64;#BinariesMac;#BinariesLinux" Filter="Engine/Binaries/DotNET/EpicGames.Perforce.Native/..." With="#FilesToSubmit"/>
-			<Submit Files="#FilesToSubmit" Description="Rebuilt DotNetPerforceLib binaries for CL $(Change)"/>
+			<Submit Files="#FilesToSubmit" Description="Rebuilt DotNetPerforceLib binaries for CL $(Change)\n\n#okforgithub public\n#jira none\n#submittool bypass"/>
 		</Node>
 	</Agent>
 
