Adding more SDK info logging
- ManageAndValidateSDK() was only writing SDK info to the log for valid platforms and wasn't writing invalid SDK info at all.
- This new invalid SDK logging shouldn't also log about Turnkey setup.
- Making invalid SDK logging more readable and also including context when AutoSDK, etc SDK versions are included.

diff --git a/Engine/Source/Programs/Shared/EpicGames.Build/UEBuildPlatformSDK.cs b/Engine/Source/Programs/Shared/EpicGames.Build/UEBuildPlatformSDK.cs
--- a/Engine/Source/Programs/Shared/EpicGames.Build/UEBuildPlatformSDK.cs
+++ b/Engine/Source/Programs/Shared/EpicGames.Build/UEBuildPlatformSDK.cs
@@ -876,7 +876,7 @@ namespace EpicGames.Core
 		public static bool bSuppressSDKWarnings = false;
 
 		public virtual SDKStatus PrintSDKInfoAndReturnValidity(LogEventType Verbosity = LogEventType.Console, LogFormatOptions Options = LogFormatOptions.None,
-			LogEventType ErrorVerbosity = LogEventType.Error, LogFormatOptions ErrorOptions = LogFormatOptions.None)
+			LogEventType ErrorVerbosity = LogEventType.Error, LogFormatOptions ErrorOptions = LogFormatOptions.None, bool bBriefInvalidSDKWarnings = false)
 		{
 			if (SDKInfoValidity != null)
 			{
@@ -905,27 +905,28 @@ namespace EpicGames.Core
 					SDKInfoValidity = SDKStatus.Invalid;
 
 					StringBuilder Msg = new StringBuilder();
-					Msg.AppendFormat("Unable to find valid SDK(s) for {0}:", PlatformName);
+					Msg.AppendLine($"Unable to find valid SDK(s) for {PlatformName}:");
 
 					foreach (SDKDescriptor Desc in SDKInfo.Sdks)
 					{
 						if (Desc.Validity == SDKStatus.Valid)
 						{
-							Msg.Append($"  {Desc.Name} is valid ({Desc}");
+							Msg.AppendLine($"  {Desc.Name} is valid ({Desc}");
 						}
 						else
 						{
+							Msg.Append($"  Found {Desc.Name} Version");
+
 							if (Desc.Current != null)
 							{
-								Msg.AppendFormat($" Found {Desc.Name} Version: {Desc.Current}.");
+								Msg.Append($"={Desc.Current}");
 							}
 
-							Msg.AppendLine($"   {Desc.ToString("Required", null, false)}");
+							Msg.AppendLine($", {Desc.ToString("Required", null, false)}.");
 						}
 					}
 
-
-					if (!bHasShownTurnkey)
+					if (!bBriefInvalidSDKWarnings && !bHasShownTurnkey)
 					{
 						Msg.AppendLine("  If your Studio has it set up, you can run this command to find the SDK to install:");
 						Msg.AppendLine("    RunUAT Turnkey -command=InstallSdk -platform={0} -BestAvailable", PlatformName!);
@@ -2189,7 +2190,7 @@ namespace EpicGames.Core
 			}
 
 			// print all SDKs to log file (errors will print out later for builds and generateprojectfiles)
-			PrintSDKInfoAndReturnValidity(LogEventType.Log, LogFormatOptions.NoConsoleOutput, LogEventType.Verbose, LogFormatOptions.NoConsoleOutput);
+			PrintSDKInfoAndReturnValidity(LogEventType.Log, LogFormatOptions.NoConsoleOutput, LogEventType.Log, LogFormatOptions.NoConsoleOutput, bBriefInvalidSDKWarnings:true);
 		}
 		#endregion
 
