Horde: Use IHordeClient to show Horde connection state. Removes dependency on OIDC configuration in engine config files.

diff --git a/Engine/Source/Programs/Shared/EpicGames.Horde/HordeClient.cs b/Engine/Source/Programs/Shared/EpicGames.Horde/HordeClient.cs
--- a/Engine/Source/Programs/Shared/EpicGames.Horde/HordeClient.cs
+++ b/Engine/Source/Programs/Shared/EpicGames.Horde/HordeClient.cs
@@ -2,6 +2,8 @@
 
 using System;
 using System.Net.Http;
+using System.Threading;
+using System.Threading.Tasks;
 using EpicGames.Horde.Storage;
 using EpicGames.Horde.Storage.Backends;
 using EpicGames.Horde.Storage.Bundles;
@@ -16,6 +18,7 @@ namespace EpicGames.Horde
 	class HordeClient : IHordeClient
 	{
 		readonly IHttpClientFactory _httpClientFactory;
+		readonly HordeHttpAuthHandlerState _authHandlerState;
 		readonly BundleCache _bundleCache;
 		readonly HordeOptions _hordeOptions;
 		readonly ILoggerFactory _loggerFactory;
@@ -23,14 +26,35 @@ namespace EpicGames.Horde
 		/// <summary>
 		/// Constructor
 		/// </summary>
-		public HordeClient(IHttpClientFactory httpClientFactory, BundleCache bundleCache, IOptionsSnapshot<HordeOptions> hordeOptions, ILoggerFactory loggerFactory)
+		public HordeClient(IHttpClientFactory httpClientFactory, HordeHttpAuthHandlerState authHandlerState, BundleCache bundleCache, IOptionsSnapshot<HordeOptions> hordeOptions, ILoggerFactory loggerFactory)
 		{
 			_httpClientFactory = httpClientFactory;
+			_authHandlerState = authHandlerState;
 			_bundleCache = bundleCache;
 			_hordeOptions = hordeOptions.Value;
 			_loggerFactory = loggerFactory;
 		}
 
+		/// <inheritdoc/>
+		public async Task<bool> ConnectAsync(bool allowLogin, CancellationToken cancellationToken)
+		{
+			if (_authHandlerState.TryGetConfiguredAuthHeader() != null)
+			{
+				return true;
+			}
+
+			_authHandlerState.Invalidate();
+			await _authHandlerState.RefreshAsync(allowLogin, cancellationToken);
+
+			return IsConnected();
+		}
+
+		/// <inheritdoc/>
+		public bool IsConnected()
+		{
+			return _authHandlerState.IsAuthenticated();
+		}
+
 		/// <inheritdoc/>
 		public HordeHttpClient CreateHttpClient()
 			=> _httpClientFactory.CreateHordeHttpClient();
diff --git a/Engine/Source/Programs/Shared/EpicGames.Horde/HordeHttpAuthHandler.cs b/Engine/Source/Programs/Shared/EpicGames.Horde/HordeHttpAuthHandler.cs
--- a/Engine/Source/Programs/Shared/EpicGames.Horde/HordeHttpAuthHandler.cs
+++ b/Engine/Source/Programs/Shared/EpicGames.Horde/HordeHttpAuthHandler.cs
@@ -6,14 +6,15 @@ using System.Linq;
 using System.Net;
 using System.Net.Http;
 using System.Net.Http.Headers;
-using System.Threading;
-using System.Threading.Tasks;
-using Microsoft.Extensions.Logging;
-using Microsoft.Extensions.Configuration;
-using EpicGames.OIDC;
 using System.Net.Http.Json;
 using System.Text.Json;
+using System.Threading;
+using System.Threading.Tasks;
+using EpicGames.Core;
 using EpicGames.Horde.Server;
+using EpicGames.OIDC;
+using Microsoft.Extensions.Configuration;
+using Microsoft.Extensions.Logging;
 using Microsoft.Extensions.Options;
 
 namespace EpicGames.Horde
@@ -42,15 +43,11 @@ namespace EpicGames.Horde
 		{
 			if (request.Headers.Authorization == null)
 			{
-				if (_options.Value.AccessToken != null)
-				{
-					// If an explicit access token is specified, just use that
-					request.Headers.Authorization = new AuthenticationHeaderValue("Bearer", _options.Value.AccessToken);
-				}
-				else if (request.RequestUri != null && TryGetAccessTokenFromEnvironment(request.RequestUri, out string? accessToken))
+				AuthenticationHeaderValue? configuredAuthHeader = _authState.TryGetConfiguredAuthHeader();
+				if (configuredAuthHeader != null)
 				{
-					// Use the access token specified in the environment
-					request.Headers.Authorization = new AuthenticationHeaderValue("Bearer", accessToken);
+					// Use the configured header
+					request.Headers.Authorization = configuredAuthHeader;
 				}
 				else if (_options.Value.AllowAuthPrompt)
 				{
@@ -83,28 +80,6 @@ namespace EpicGames.Horde
 			}
 			return await base.SendAsync(request, cancellationToken);
 		}
-
-		static bool TryGetAccessTokenFromEnvironment(Uri requestUri, out string? accessToken)
-		{
-			// Only use the token from the environment if the configured base address is missing or matches the one configured in the environment
-			string? hordeUrlEnvVar = Environment.GetEnvironmentVariable(HordeHttpClient.HordeUrlEnvVarName);
-			if (!String.IsNullOrEmpty(hordeUrlEnvVar))
-			{
-				Uri hordeUrl = new Uri(hordeUrlEnvVar);
-				if (String.Equals(requestUri.Host, hordeUrl.Host, StringComparison.OrdinalIgnoreCase))
-				{
-					string? hordeToken = Environment.GetEnvironmentVariable(HordeHttpClient.HordeTokenEnvVarName);
-					if (!String.IsNullOrEmpty(hordeToken))
-					{
-						accessToken = hordeToken;
-						return true;
-					}
-				}
-			}
-
-			accessToken = null;
-			return false;
-		}
 	}
 
 	/// <summary>
@@ -117,33 +92,81 @@ namespace EpicGames.Horde
 		/// </summary>
 		public const string HttpClientName = "HordeHttpAuthState";
 
+		record class AuthState(AuthMethod Method, OidcTokenInfo? TokenInfo)
+		{
+			public bool IsAuthorized()
+				=> (Method == AuthMethod.Anonymous) || (TokenInfo != null && TokenInfo.IsValid && TokenInfo.TokenExpiry > DateTimeOffset.Now);
+		}
+
 		readonly object _lockObject = new object();
 		readonly CancellationTokenSource _cancellationTokenSource = new CancellationTokenSource();
-		Task<AuthenticationHeaderValue?>? _authHeaderTask;
+		Task<AuthState>? _authStateTask = null;
 		readonly IHttpClientFactory _httpClientFactory;
+		readonly IOptions<HordeOptions> _options;
 		readonly ILogger _logger;
 
 		/// <summary>
 		/// Constructor
 		/// </summary>
-		public HordeHttpAuthHandlerState(IHttpClientFactory httpClientFactory, ILogger<HordeHttpAuthHandler> logger)
+		public HordeHttpAuthHandlerState(IHttpClientFactory httpClientFactory, IOptions<HordeOptions> options, ILogger<HordeHttpAuthHandler> logger)
 		{
 			_httpClientFactory = httpClientFactory;
+			_options = options;
 			_logger = logger;
 		}
 
 		/// <inheritdoc/>
 		public async ValueTask DisposeAsync()
 		{
-			if (_authHeaderTask != null)
+			if (_authStateTask != null && !_authStateTask.IsCompleted)
 			{
 				_cancellationTokenSource.Cancel();
-				await _authHeaderTask;
-				_authHeaderTask = null;
+				try
+				{
+#pragma warning disable VSTHRD003
+					await _authStateTask;
+#pragma warning restore VSTHRD003
+				}
+				catch (OperationCanceledException)
+				{
+				}
 			}
+
 			_cancellationTokenSource.Dispose();
 		}
 
+		/// <summary>
+		/// Checks if we have a valid auth header at the moment
+		/// </summary>
+		public bool IsAuthenticated()
+		{
+			if (TryGetConfiguredAuthHeader() != null)
+			{
+				return true;
+			}
+			if (_authStateTask != null && _authStateTask.TryGetResult(out AuthState? authState) && authState.IsAuthorized())
+			{
+				return true;
+			}
+			return false;
+		}
+
+		/// <summary>
+		/// Invalidate the cached header value
+		/// </summary>
+		public void Invalidate()
+		{
+			lock (_lockObject)
+			{
+#pragma warning disable VSTHRD002
+				if (_authStateTask != null && _authStateTask.IsCompleted)
+				{
+					_authStateTask = null;
+				}
+#pragma warning restore VSTHRD002
+			}
+		}
+
 		/// <summary>
 		/// Invalidate a cached header value
 		/// </summary>
@@ -153,36 +176,101 @@ namespace EpicGames.Horde
 			lock (_lockObject)
 			{
 #pragma warning disable VSTHRD002
-				if (_authHeaderTask != null && _authHeaderTask.IsCompleted && Object.Equals(_authHeaderTask.Result, authHeader))
+				if (_authStateTask != null && _authStateTask.IsCompleted && Object.Equals(_authStateTask.Result?.TokenInfo?.AccessToken, authHeader.Parameter))
 				{
-					_authHeaderTask = null;
+					_authStateTask = null;
 				}
 #pragma warning restore VSTHRD002
 			}
 		}
 
 		/// <summary>
-		/// Gets a new auth header
+		/// Try to get a configured auth header
 		/// </summary>
-		/// <param name="cancellationToken">Cancellation token for the operation</param>
-		public async ValueTask<AuthenticationHeaderValue?> TryGetAuthHeaderAsync(CancellationToken cancellationToken)
+		public AuthenticationHeaderValue? TryGetConfiguredAuthHeader()
 		{
-			Task<AuthenticationHeaderValue?>? authHeaderTask = _authHeaderTask;
-			if (authHeaderTask == null)
+			if (_options.Value.AccessToken != null)
+			{
+				// If an explicit access token is specified, just use that
+				return new AuthenticationHeaderValue("Bearer", _options.Value.AccessToken);
+			}
+			else if (TryGetAccessTokenFromEnvironment(out string? accessToken))
+			{
+				// Use the access token specified in the environment
+				return new AuthenticationHeaderValue("Bearer", accessToken);
+			}
+			else
 			{
-				lock (_lockObject)
+				// Will need to login asynchronously
+				return null;
+			}
+		}
+
+		bool TryGetAccessTokenFromEnvironment(out string? accessToken)
+		{
+			// Only use the token from the environment if the configured base address is missing or matches the one configured in the environment
+			string? hordeUrlEnvVar = Environment.GetEnvironmentVariable(HordeHttpClient.HordeUrlEnvVarName);
+			if (!String.IsNullOrEmpty(hordeUrlEnvVar))
+			{
+				Uri hordeUrl = new Uri(hordeUrlEnvVar);
+				if (String.Equals(_options.Value.ServerUrl?.Host, hordeUrl.Host, StringComparison.OrdinalIgnoreCase))
 				{
-					_authHeaderTask ??= Task.Run(() => GetNewAuthHeaderAsync(_cancellationTokenSource.Token), _cancellationTokenSource.Token);
-					authHeaderTask = _authHeaderTask;
+					string? hordeToken = Environment.GetEnvironmentVariable(HordeHttpClient.HordeTokenEnvVarName);
+					if (!String.IsNullOrEmpty(hordeToken))
+					{
+						accessToken = hordeToken;
+						return true;
+					}
 				}
 			}
-			return await authHeaderTask.WaitAsync(cancellationToken);
+
+			accessToken = null;
+			return false;
 		}
 
 		/// <summary>
-		/// Get an access token for the server specified in a config instance
+		/// Refresh the auth state
 		/// </summary>
-		async Task<AuthenticationHeaderValue?> GetNewAuthHeaderAsync(CancellationToken cancellationToken)
+		/// <param name="allowLogin">Whether to allow logging in</param>
+		/// <param name="cancellationToken"></param>
+		/// <returns></returns>
+		public async Task RefreshAsync(bool allowLogin, CancellationToken cancellationToken)
+		{
+			Invalidate();
+			await GetAuthStateAsync(allowLogin, cancellationToken);
+		}
+
+		/// <summary>
+		/// Gets a new auth header
+		/// </summary>
+		/// <param name="cancellationToken">Cancellation token for the operation</param>
+		public async ValueTask<AuthenticationHeaderValue?> TryGetAuthHeaderAsync(CancellationToken cancellationToken)
+		{
+			AuthState? authState = await GetAuthStateAsync(true, cancellationToken);
+			if (authState?.TokenInfo == null)
+			{
+				return null;
+			}
+			return new AuthenticationHeaderValue("Bearer", authState.TokenInfo.AccessToken);
+		}
+
+		async Task<AuthState?> GetAuthStateAsync(bool allowLogin, CancellationToken cancellationToken)
+		{
+			if (TryGetConfiguredAuthHeader() != null)
+			{
+				return null;
+			}
+
+			Task<AuthState> authStateTask;
+			lock (_lockObject)
+			{
+				_authStateTask ??= Task.Run(() => GetAuthStateInternalAsync(allowLogin, _cancellationTokenSource.Token), _cancellationTokenSource.Token);
+				authStateTask = _authStateTask;
+			}
+			return await authStateTask.WaitAsync(cancellationToken);
+		}
+
+		async Task<AuthState> GetAuthStateInternalAsync(bool allowLogin, CancellationToken cancellationToken)
 		{
 			Uri serverUrl;
 
@@ -207,10 +295,15 @@ namespace EpicGames.Horde
 				}
 			}
 
+			if (authConfig.Method == AuthMethod.Anonymous)
+			{
+				return new AuthState(authConfig.Method, null);
+			}
+
 			string? localRedirectUrl = authConfig.LocalRedirectUrls?.FirstOrDefault();
 			if (String.IsNullOrEmpty(authConfig.ServerUrl) || String.IsNullOrEmpty(localRedirectUrl))
 			{
-				return null;
+				throw new Exception("No auth server configuration found");
 			}
 
 			const string OidcProvider = "Horde";
@@ -241,19 +334,13 @@ namespace EpicGames.Horde
 					_logger.LogTrace(ex, "Unable to get access token; attempting login: {Message}", ex.Message);
 				}
 			}
-			if (result == null)
+			if (result == null && allowLogin)
 			{
 				_logger.LogInformation("Logging in to {Server}...", serverUrl);
 				result = await oidcTokenManager.Login(OidcProvider, cancellationToken);
 			}
 
-			if (result.AccessToken == null)
-			{
-				throw new Exception($"Unable to get access token for {serverUrl}");
-			}
-
-			_logger.LogDebug("Received access token for {Server}", serverUrl);
-			return new AuthenticationHeaderValue("Bearer", result.AccessToken);
+			return new AuthState(authConfig.Method, result);
 		}
 	}
 }
diff --git a/Engine/Source/Programs/Shared/EpicGames.Horde/IHordeClient.cs b/Engine/Source/Programs/Shared/EpicGames.Horde/IHordeClient.cs
--- a/Engine/Source/Programs/Shared/EpicGames.Horde/IHordeClient.cs
+++ b/Engine/Source/Programs/Shared/EpicGames.Horde/IHordeClient.cs
@@ -1,5 +1,7 @@
 // Copyright Epic Games, Inc. All Rights Reserved.
 
+using System.Threading;
+using System.Threading.Tasks;
 using EpicGames.Horde.Storage;
 using EpicGames.Horde.Tools;
 
@@ -10,6 +12,19 @@ namespace EpicGames.Horde
 	/// </summary>
 	public interface IHordeClient
 	{
+		/// <summary>
+		/// Connect to the Horde server
+		/// </summary>
+		/// <param name="allowLogin">Whether to allow interactive logins</param>
+		/// <param name="cancellationToken">Cancellation token for the operation</param>
+		/// <returns>True if the connection succeded</returns>
+		Task<bool> ConnectAsync(bool allowLogin, CancellationToken cancellationToken);
+
+		/// <summary>
+		/// Gets the current connection state
+		/// </summary>
+		bool IsConnected();
+
 		/// <summary>
 		/// Creates a http client 
 		/// </summary>
diff --git a/Engine/Source/Programs/UnrealGameSync/UnrealGameSync/Controls/WorkspaceControl.cs b/Engine/Source/Programs/UnrealGameSync/UnrealGameSync/Controls/WorkspaceControl.cs
--- a/Engine/Source/Programs/UnrealGameSync/UnrealGameSync/Controls/WorkspaceControl.cs
+++ b/Engine/Source/Programs/UnrealGameSync/UnrealGameSync/Controls/WorkspaceControl.cs
@@ -22,6 +22,7 @@ using Microsoft.Extensions.DependencyInjection;
 using System.Diagnostics.CodeAnalysis;
 using EpicGames.OIDC;
 using Microsoft.Win32;
+using EpicGames.Horde;
 
 namespace UnrealGameSync
 {
@@ -177,6 +178,7 @@ namespace UnrealGameSync
 		readonly IServiceProvider _serviceProvider;
 		readonly ILogger _logger;
 		readonly IPerforceSettings _perforceSettings;
+		readonly IHordeClient? _hordeClient;
 		ProjectInfo ProjectInfo { get; }
 
 		readonly UserSettings _settings;
@@ -286,6 +288,7 @@ namespace UnrealGameSync
 			_settings = settings;
 			_workspaceSettings = openProjectInfo.WorkspaceSettings;
 			_projectSettings = settings.FindOrAddProjectSettings(openProjectInfo.ProjectInfo, openProjectInfo.WorkspaceSettings, _logger);
+			_hordeClient = serviceProvider.GetService<IHordeClient>();
 
 			DesiredTaskbarState = Tuple.Create(TaskbarState.NoProgress, 0.0f);
 
@@ -352,8 +355,7 @@ namespace UnrealGameSync
 			FileReference projectLogBaseName = FileReference.Combine(_workspaceDataFolder, "sync.log");
 
 			ILogger perforceLogger = _serviceProvider.GetRequiredService<ILogger<PerforceMonitor>>();
-			_perforceMonitor = new PerforceMonitor(perforceClientSettings, openProjectInfo.ProjectInfo, openProjectInfo.LatestProjectConfigFile, openProjectInfo.ProjectInfo.CacheFolder, openProjectInfo.LocalConfigFiles,
-				openProjectInfo.OidcTokenClient, _serviceProvider);
+			_perforceMonitor = new PerforceMonitor(perforceClientSettings, openProjectInfo.ProjectInfo, openProjectInfo.LatestProjectConfigFile, openProjectInfo.ProjectInfo.CacheFolder, openProjectInfo.LocalConfigFiles, _serviceProvider);
 
 			_perforceMonitor.OnUpdate += UpdateBuildListCallback;
 			_perforceMonitor.OnUpdateMetadata += UpdateBuildMetadataCallback;
@@ -3353,6 +3355,23 @@ namespace UnrealGameSync
 			}, null);
 		}
 
+		private void AddHordeStatus(StatusLine statusLine)
+		{
+			if (_hordeClient != null)
+			{
+				statusLine.AddText("  |  ");
+
+				if (_hordeClient.IsConnected())
+				{
+					statusLine.AddText("Connected to Horde.");
+				}
+				else
+				{
+					statusLine.AddLink("Connect to Horde", FontStyle.Bold | FontStyle.Underline, (p, r) => ConnectToHorde(_hordeClient));
+				}
+			}
+		}
+
 		private void UpdateStatusPanel()
 		{
 			if (_workspace == null)
@@ -3369,22 +3388,7 @@ namespace UnrealGameSync
 				// Project
 				StatusLine projectLine = new StatusLine();
 				projectLine.AddText("Opened " + SelectedFileName.FullName);
-
-				OidcTokenClient? oidcClient = _perforceMonitor.LatestOidcTokenClient;
-				if (oidcClient != null)
-				{
-					projectLine.AddText("  |  ");
-
-					StatusLine hordeLine = new StatusLine();
-					if (oidcClient.GetStatus() == OidcStatus.Connected)
-					{
-						projectLine.AddText("Connected to Horde.");
-					}
-					else
-					{
-						projectLine.AddLink("Connect to Horde", FontStyle.Bold | FontStyle.Underline, (p, r) => DoOidcLogin());
-					}
-				}
+				AddHordeStatus(projectLine);
 				lines.Add(projectLine);
 
 				// Spacer
@@ -3479,23 +3483,7 @@ namespace UnrealGameSync
 				StatusLine projectLine = new StatusLine();
 				projectLine.AddText(String.Format("Opened "));
 				projectLine.AddLink(SelectedFileName.FullName + " \u25BE", FontStyle.Regular, (p, r) => { SelectRecentProject(r); });
-
-				OidcTokenClient? oidcClient = _perforceMonitor.LatestOidcTokenClient;
-				if (oidcClient != null)
-				{
-					projectLine.AddText("  |  ");
-
-					StatusLine hordeLine = new StatusLine();
-					if (oidcClient.GetStatus() == OidcStatus.Connected)
-					{
-						projectLine.AddText("Connected to Horde.");
-					}
-					else
-					{
-						projectLine.AddLink("Connect to Horde", FontStyle.Bold | FontStyle.Underline, (p, r) => DoOidcLogin());
-					}
-				}
-
+				AddHordeStatus(projectLine);
 				projectLine.AddText("  |  ");
 				projectLine.AddLink("Settings...", FontStyle.Regular, (p, r) => { _owner.EditSelectedProject(this); });
 				lines.Add(projectLine);
@@ -3775,24 +3763,20 @@ namespace UnrealGameSync
 			StatusPanel.Set(lines, caption, alert, tintColor);
 		}
 
-		private void DoOidcLogin()
+		private void ConnectToHorde(IHordeClient hordeClient)
 		{
-			OidcTokenClient? oidcClient = _perforceMonitor.LatestOidcTokenClient;
-			if (oidcClient != null)
+			ModalTask.Execute(this, "Horde Login", "Connecting to Horde...", async cancellationToken =>
 			{
-				ModalTask.Execute(this, "OIDC Login", "Opening OIDC login page (check your browser).", async cancellationToken =>
+				try
 				{
-					try
-					{
-						await oidcClient.LoginAsync(cancellationToken);
-					}
-					catch when (cancellationToken.IsCancellationRequested)
-					{
-						// nop
-					}
-				});
-				UpdateStatusPanel();
-			}
+					await hordeClient.ConnectAsync(true, cancellationToken);
+				}
+				catch when (cancellationToken.IsCancellationRequested)
+				{
+					// nop
+				}
+			});
+			UpdateStatusPanel();
 		}
 
 		private void ShowBuildHealthMenu(Rectangle bounds)
diff --git a/Engine/Source/Programs/UnrealGameSync/UnrealGameSync/Forms/MainWindow.cs b/Engine/Source/Programs/UnrealGameSync/UnrealGameSync/Forms/MainWindow.cs
--- a/Engine/Source/Programs/UnrealGameSync/UnrealGameSync/Forms/MainWindow.cs
+++ b/Engine/Source/Programs/UnrealGameSync/UnrealGameSync/Forms/MainWindow.cs
@@ -103,8 +103,6 @@ namespace UnrealGameSync
 
 		bool _allowClose = false;
 
-		readonly OidcTokenManager _oidcTokenManager;
-
 		System.Threading.Timer? _scheduleTimer;
 		System.Threading.Timer? _scheduleSettledTimer;
 
@@ -142,7 +140,6 @@ namespace UnrealGameSync
 			ToolUpdateMonitor = new ToolUpdateMonitor(_defaultPerforceSettings, _dataFolder, settings, _serviceProvider);
 
 			_settings = settings;
-			_oidcTokenManager = _serviceProvider.GetRequiredService<OidcTokenManager>();
 
 			// While creating tab controls during startup, we need to prevent layout calls resulting in the window handle being created too early. Disable layout calls here.
 			SuspendLayout();
@@ -1229,7 +1226,7 @@ namespace UnrealGameSync
 
 			PerforceSettings perforceSettings = Utility.OverridePerforceSettings(_defaultPerforceSettings, project.ServerAndPort, project.UserName);
 
-			ModalTask<OpenProjectInfo>? settingsTask = PerforceModalTask.Execute(this, "Opening Project", "Opening project, please wait...", perforceSettings, (p, c) => OpenProjectWindow.DetectSettingsAsync(p, project, _settings, _oidcTokenManager, false, projectLogger, c), projectLogger, taskFlags);
+			ModalTask<OpenProjectInfo>? settingsTask = PerforceModalTask.Execute(this, "Opening Project", "Opening project, please wait...", perforceSettings, (p, c) => OpenProjectWindow.DetectSettingsAsync(p, project, _settings, false, projectLogger, c), projectLogger, taskFlags);
 			if (settingsTask == null || settingsTask.Failed)
 			{
 				if (settingsTask != null)
diff --git a/Engine/Source/Programs/UnrealGameSync/UnrealGameSync/Forms/OpenProjectWindow.cs b/Engine/Source/Programs/UnrealGameSync/UnrealGameSync/Forms/OpenProjectWindow.cs
--- a/Engine/Source/Programs/UnrealGameSync/UnrealGameSync/Forms/OpenProjectWindow.cs
+++ b/Engine/Source/Programs/UnrealGameSync/UnrealGameSync/Forms/OpenProjectWindow.cs
@@ -295,11 +295,10 @@ namespace UnrealGameSync
 			if(TryGetSelectedProject(out selectedProject))
 			{
 				ILogger<OpenProjectInfo> logger = _serviceProvider.GetRequiredService<ILogger<OpenProjectInfo>>();
-				OidcTokenManager oidcTokenManager = _serviceProvider.GetRequiredService<OidcTokenManager>();
 
 				PerforceSettings newPerforceSettings = Utility.OverridePerforceSettings(Perforce, selectedProject.ServerAndPort, selectedProject.UserName);
 
-				ModalTask<OpenProjectInfo>? newOpenProjectInfo = PerforceModalTask.Execute(this, "Opening project", "Opening project, please wait...", newPerforceSettings, (x, y) => DetectSettingsAsync(x, selectedProject, _settings, oidcTokenManager, GenerateP4ConfigCheckbox.Checked, logger, y), logger);
+				ModalTask<OpenProjectInfo>? newOpenProjectInfo = PerforceModalTask.Execute(this, "Opening project", "Opening project, please wait...", newPerforceSettings, (x, y) => DetectSettingsAsync(x, selectedProject, _settings, GenerateP4ConfigCheckbox.Checked, logger, y), logger);
 				if (newOpenProjectInfo != null && newOpenProjectInfo.Succeeded)
 				{
 					_openProjectInfo = newOpenProjectInfo.Result;
@@ -309,9 +308,9 @@ namespace UnrealGameSync
 			}
 		}
 
-		public static async Task<OpenProjectInfo> DetectSettingsAsync(IPerforceConnection perforce, UserSelectedProjectSettings selectedProject, UserSettings userSettings, OidcTokenManager oidcTokenManager, bool GenerateP4Config, ILogger<OpenProjectInfo> logger, CancellationToken cancellationToken)
+		public static async Task<OpenProjectInfo> DetectSettingsAsync(IPerforceConnection perforce, UserSelectedProjectSettings selectedProject, UserSettings userSettings, bool GenerateP4Config, ILogger<OpenProjectInfo> logger, CancellationToken cancellationToken)
 		{
-			OpenProjectInfo settings = await OpenProjectInfo.CreateAsync(perforce, selectedProject, userSettings, oidcTokenManager, GenerateP4Config, logger, cancellationToken);
+			OpenProjectInfo settings = await OpenProjectInfo.CreateAsync(perforce, selectedProject, userSettings, GenerateP4Config, logger, cancellationToken);
 			if (s_onDetectProjectSettings != null)
 			{
 				string? message;
diff --git a/Engine/Source/Programs/UnrealGameSync/UnrealGameSync/PerforceMonitor.cs b/Engine/Source/Programs/UnrealGameSync/UnrealGameSync/PerforceMonitor.cs
--- a/Engine/Source/Programs/UnrealGameSync/UnrealGameSync/PerforceMonitor.cs
+++ b/Engine/Source/Programs/UnrealGameSync/UnrealGameSync/PerforceMonitor.cs
@@ -53,7 +53,6 @@ namespace UnrealGameSync
 		readonly DirectoryReference _cacheFolder;
 		readonly List<KeyValuePair<FileReference, DateTime>> _localConfigFiles;
 		readonly IAsyncDisposer _asyncDisposeTasks;
-		readonly OidcTokenManager _oidcTokenManager;
 		readonly IHordeClient _hordeClient;
 
 		string[] prevCodeRules = Array.Empty<string>();
@@ -72,7 +71,7 @@ namespace UnrealGameSync
 			private set;
 		}
 
-		public PerforceMonitor(IPerforceSettings perforceSettings, ProjectInfo projectInfo, ConfigFile projectConfigFile, DirectoryReference cacheFolder, List<KeyValuePair<FileReference, DateTime>> localConfigFiles, OidcTokenClient? oidcTokenClient, IServiceProvider serviceProvider)
+		public PerforceMonitor(IPerforceSettings perforceSettings, ProjectInfo projectInfo, ConfigFile projectConfigFile, DirectoryReference cacheFolder, List<KeyValuePair<FileReference, DateTime>> localConfigFiles, IServiceProvider serviceProvider)
 		{
 			_perforceSettings = perforceSettings;
 			_branchClientPath = projectInfo.ClientRootPath;
@@ -89,11 +88,9 @@ namespace UnrealGameSync
 			_asyncDisposeTasks = serviceProvider.GetRequiredService<IAsyncDisposer>();
 			_synchronizationContext = SynchronizationContext.Current!;
 			_cancellationSource = new CancellationTokenSource();
-			_oidcTokenManager = serviceProvider.GetRequiredService<OidcTokenManager>();
 			_hordeClient = serviceProvider.GetRequiredService<IHordeClient>();
 
 			AvailableArchives = (new List<IArchiveInfo>()).AsReadOnly();
-			LatestOidcTokenClient = oidcTokenClient;
 		}
 
 		public void Start()
@@ -491,7 +488,6 @@ namespace UnrealGameSync
 			if(_localConfigFiles.Any(x => FileReference.GetLastWriteTimeUtc(x.Key) != x.Value))
 			{
 				await UpdateProjectConfigFileAsync(perforce, cancellationToken);
-				// TODO: Also check OIDC config
 				_synchronizationContext.Post(_ => OnUpdateMetadata?.Invoke(), null);
 			}
 
@@ -521,7 +517,6 @@ namespace UnrealGameSync
 		{
 			_localConfigFiles.Clear();
 			LatestProjectConfigFile = await ConfigUtils.ReadProjectConfigFileAsync(perforce, _branchClientPath, _selectedClientFileName, _cacheFolder, _localConfigFiles, _logger, cancellationToken);
-			LatestOidcTokenClient = await ConfigUtils.CreateOidcTokenClientAsync(_oidcTokenManager, LatestProjectConfigFile, _selectedProjectIdentifier, perforce, _branchClientPath, _selectedClientFileName, _localConfigFiles, _cacheFolder, _logger, cancellationToken);
 		}
 
 		public List<ChangesRecord> GetChanges()
@@ -571,12 +566,6 @@ namespace UnrealGameSync
 			private set;
 		}
 
-		public OidcTokenClient? LatestOidcTokenClient
-		{
-			get;
-			private set;
-		}
-
 		public IReadOnlyList<IArchiveInfo> AvailableArchives
 		{
 			get;
diff --git a/Engine/Source/Programs/UnrealGameSync/UnrealGameSync/ProgramApplicationContext.cs b/Engine/Source/Programs/UnrealGameSync/UnrealGameSync/ProgramApplicationContext.cs
--- a/Engine/Source/Programs/UnrealGameSync/UnrealGameSync/ProgramApplicationContext.cs
+++ b/Engine/Source/Programs/UnrealGameSync/UnrealGameSync/ProgramApplicationContext.cs
@@ -15,6 +15,7 @@ using System.Threading;
 using System.Threading.Tasks;
 using System.Windows.Forms;
 using EpicGames.OIDC;
+using EpicGames.Horde;
 
 #nullable enable
 
@@ -195,12 +196,13 @@ namespace UnrealGameSync
 				if (IsValidProject(projectSettings, _logger))
 				{
 					ILogger<OpenProjectInfo> logger = serviceProvider.GetRequiredService<ILogger<OpenProjectInfo>>();
-					OidcTokenManager oidcTokenManager = serviceProvider.GetRequiredService<OidcTokenManager>();
-					Task<OpenProjectInfo> startupTask = Task.Run(() => OpenProjectInfo.CreateAsync(defaultPerforceSettings, projectSettings, _settings, oidcTokenManager, false, logger, _startupCancellationSource.Token), _startupCancellationSource.Token);
+					Task<OpenProjectInfo> startupTask = Task.Run(() => OpenProjectInfo.CreateAsync(defaultPerforceSettings, projectSettings, _settings, false, logger, _startupCancellationSource.Token), _startupCancellationSource.Token);
 					startupTasks.Add((projectSettings, new ModalTask<OpenProjectInfo>(startupTask)));
 				}
 			}
-			_startupTask = Task.Run(() => WaitForStartupTasks(startupTasks));
+
+			HordeHttpAuthHandlerState? hordeAuthState = _serviceProvider.GetService<HordeHttpAuthHandlerState>();
+			_startupTask = Task.Run(() => WaitForStartupTasks(startupTasks, hordeAuthState));
 
 			_startupWindow = new ModalTaskWindow("Opening Projects", "Opening projects, please wait...", FormStartPosition.CenterScreen, _startupTask, _startupCancellationSource);
 			_components.Add(_startupWindow);
@@ -234,7 +236,7 @@ namespace UnrealGameSync
 			}
 		}
 
-		static async Task WaitForStartupTasks(List<(UserSelectedProjectSettings, ModalTask<OpenProjectInfo>)> startupTasks)
+		static async Task WaitForStartupTasks(List<(UserSelectedProjectSettings, ModalTask<OpenProjectInfo>)> startupTasks, HordeHttpAuthHandlerState? hordeAuthState)
 		{
 			foreach ((_, ModalTask<OpenProjectInfo> modalTask) in startupTasks)
 			{
@@ -256,6 +258,11 @@ namespace UnrealGameSync
 					Program.CaptureException(ex);
 				}
 			}
+
+			if (hordeAuthState != null)
+			{
+				await hordeAuthState.RefreshAsync(false, CancellationToken.None);
+			}
 		}
 
 		private UserSelectedProjectSettings UpgradeSelectedProjectSettings(UserSelectedProjectSettings project)
diff --git a/Engine/Source/Programs/UnrealGameSync/UnrealGameSyncShared/ConfigUtils.cs b/Engine/Source/Programs/UnrealGameSync/UnrealGameSyncShared/ConfigUtils.cs
--- a/Engine/Source/Programs/UnrealGameSync/UnrealGameSyncShared/ConfigUtils.cs
+++ b/Engine/Source/Programs/UnrealGameSync/UnrealGameSyncShared/ConfigUtils.cs
@@ -196,52 +196,6 @@ namespace UnrealGameSync
 			return contents;
 		}
 
-		public static async Task<OidcTokenClient?> CreateOidcTokenClientAsync(OidcTokenManager oidcTokenManager, ConfigFile projectConfigFile, string projectIdentifier, IPerforceConnection perforce, string clientRoot, string? clientProjectFile, List<KeyValuePair<FileReference, DateTime>> localFiles, DirectoryReference cacheFolder, ILogger logger, CancellationToken cancellationToken)
-		{
-			OidcTokenClient? oidcTokenClient = null;
-
-			string? oidcProvider;
-			if (TryGetProjectSetting(projectConfigFile, projectIdentifier, "OidcProvider", out oidcProvider))
-			{
-				List<string> configFiles = new List<string>();
-				configFiles.AddRange(ProviderConfigurationFactory.ConfigPaths.Select(x => $"{clientRoot.TrimEnd('/')}/Engine/{x}"));
-
-				if (clientProjectFile != null && clientProjectFile.EndsWith(".uproject", StringComparison.OrdinalIgnoreCase))
-				{
-					string clientProjectPath = clientProjectFile.Substring(0, clientProjectFile.LastIndexOf('/'));
-					configFiles.AddRange(ProviderConfigurationFactory.ConfigPaths.Select(x => $"{clientProjectPath.TrimEnd('/')}/{x}"));
-				}
-
-				List<string[]> latestOidcConfigFiles = await ConfigUtils.ReadConfigFiles(perforce, configFiles, localFiles, cacheFolder, logger, cancellationToken);
-
-				ConfigurationBuilder configBuilder = new ConfigurationBuilder();
-				foreach (string[] configFile in latestOidcConfigFiles)
-				{
-#pragma warning disable CA2000
-					MemoryStream configStream = new MemoryStream(Encoding.UTF8.GetBytes(String.Join("\n", configFile)));
-					configBuilder.AddJsonStream(configStream);
-#pragma warning restore CA2000
-				}
-
-				IConfigurationRoot config = configBuilder.Build();
-				OidcTokenOptions oidcOptions = OidcTokenOptions.Bind(config);
-
-				if (oidcOptions.Providers.TryGetValue(oidcProvider, out ProviderInfo? providerInfo))
-				{
-					oidcTokenClient = oidcTokenManager.FindOrAddClient(oidcProvider, providerInfo);
-
-					// Trigger an update of the access token if necessary, just to make sure we have the most current state
-					try
-					{
-						await oidcTokenClient.GetAccessTokenAsync(cancellationToken);
-					}
-					catch { }
-				}
-			}
-
-			return oidcTokenClient;
-		}
-
 		public static FileReference GetEditorTargetFile(ProjectInfo projectInfo, ConfigFile projectConfig)
 		{
 			if (projectInfo.ProjectPath.EndsWith(".uproject", StringComparison.OrdinalIgnoreCase))
diff --git a/Engine/Source/Programs/UnrealGameSync/UnrealGameSyncShared/OpenProjectInfo.cs b/Engine/Source/Programs/UnrealGameSync/UnrealGameSyncShared/OpenProjectInfo.cs
--- a/Engine/Source/Programs/UnrealGameSync/UnrealGameSyncShared/OpenProjectInfo.cs
+++ b/Engine/Source/Programs/UnrealGameSync/UnrealGameSyncShared/OpenProjectInfo.cs
@@ -25,10 +25,9 @@ namespace UnrealGameSync
 		public ConfigFile WorkspaceProjectConfigFile { get; }
 		public IReadOnlyList<string>? WorkspaceProjectStreamFilter { get; }
 		public List<KeyValuePair<FileReference, DateTime>> LocalConfigFiles { get; }
-		public OidcTokenClient? OidcTokenClient { get; }
 		public bool GenerateP4Config { get; }
 
-		public OpenProjectInfo(UserSelectedProjectSettings selectedProject, IPerforceSettings perforceSettings, ProjectInfo projectInfo, UserWorkspaceSettings workspaceSettings, WorkspaceStateWrapper workspaceStateWrapper, ConfigFile latestProjectConfigFile, ConfigFile workspaceProjectConfigFile, IReadOnlyList<string>? workspaceProjectStreamFilter, List<KeyValuePair<FileReference, DateTime>> localConfigFiles, OidcTokenClient? oidcTokenClient, bool generateP4Config)
+		public OpenProjectInfo(UserSelectedProjectSettings selectedProject, IPerforceSettings perforceSettings, ProjectInfo projectInfo, UserWorkspaceSettings workspaceSettings, WorkspaceStateWrapper workspaceStateWrapper, ConfigFile latestProjectConfigFile, ConfigFile workspaceProjectConfigFile, IReadOnlyList<string>? workspaceProjectStreamFilter, List<KeyValuePair<FileReference, DateTime>> localConfigFiles, bool generateP4Config)
 		{
 			SelectedProject = selectedProject;
 
@@ -40,11 +39,10 @@ namespace UnrealGameSync
 			WorkspaceProjectConfigFile = workspaceProjectConfigFile;
 			WorkspaceProjectStreamFilter = workspaceProjectStreamFilter;
 			LocalConfigFiles = localConfigFiles;
-			OidcTokenClient = oidcTokenClient;
 			GenerateP4Config = generateP4Config;
 		}
 
-		public static async Task<OpenProjectInfo> CreateAsync(IPerforceSettings defaultPerforceSettings, UserSelectedProjectSettings selectedProject, UserSettings userSettings, OidcTokenManager oidcTokenManager, bool generateP4Config, ILogger<OpenProjectInfo> logger, CancellationToken cancellationToken)
+		public static async Task<OpenProjectInfo> CreateAsync(IPerforceSettings defaultPerforceSettings, UserSelectedProjectSettings selectedProject, UserSettings userSettings, bool generateP4Config, ILogger<OpenProjectInfo> logger, CancellationToken cancellationToken)
 		{
 			PerforceSettings perforceSettings = Utility.OverridePerforceSettings(defaultPerforceSettings, selectedProject.ServerAndPort, selectedProject.UserName);
 			using IPerforceConnection perforce = await PerforceConnection.CreateAsync(perforceSettings, logger);
@@ -57,10 +55,10 @@ namespace UnrealGameSync
 			}
 
 			// Execute like a regular task
-			return await CreateAsync(perforce, selectedProject, userSettings, oidcTokenManager, generateP4Config, logger, cancellationToken);
+			return await CreateAsync(perforce, selectedProject, userSettings, generateP4Config, logger, cancellationToken);
 		}
 
-		public static async Task<OpenProjectInfo> CreateAsync(IPerforceConnection defaultConnection, UserSelectedProjectSettings selectedProject, UserSettings userSettings, OidcTokenManager oidcTokenManager, bool generateP4Config, ILogger<OpenProjectInfo> logger, CancellationToken cancellationToken)
+		public static async Task<OpenProjectInfo> CreateAsync(IPerforceConnection defaultConnection, UserSelectedProjectSettings selectedProject, UserSettings userSettings, bool generateP4Config, ILogger<OpenProjectInfo> logger, CancellationToken cancellationToken)
 		{
 			using IDisposable loggerScope = logger.BeginScope("Project {SelectedProject}", selectedProject.ToString());
 			logger.LogInformation("Detecting settings for {Project}", selectedProject);
@@ -229,14 +227,11 @@ namespace UnrealGameSync
 				List<KeyValuePair<FileReference, DateTime>> localConfigFiles = new List<KeyValuePair<FileReference, DateTime>>();
 				ConfigFile latestProjectConfigFile = await ConfigUtils.ReadProjectConfigFileAsync(perforceClient, projectInfo, localConfigFiles, logger, cancellationToken);
 
-				// Read the initial OIDC config files
-				OidcTokenClient? oidcTokenClient = await ConfigUtils.CreateOidcTokenClientAsync(oidcTokenManager, latestProjectConfigFile, projectInfo.ProjectIdentifier, perforceClient, branchClientPath, newSelectedClientFileName, localConfigFiles, projectInfo.CacheFolder, logger, cancellationToken);
-
 				// Get the local config file and stream filter
 				ConfigFile workspaceProjectConfigFile = await WorkspaceUpdate.ReadProjectConfigFile(branchDirectoryName, newSelectedFileName, logger);
 				IReadOnlyList<string>? workspaceProjectStreamFilter = await WorkspaceUpdate.ReadProjectStreamFilter(perforceClient, workspaceProjectConfigFile, cancellationToken);
 
-				OpenProjectInfo workspaceSettings = new OpenProjectInfo(selectedProject, perforceSettings, projectInfo, userWorkspaceSettings, workspaceStateWrapper, latestProjectConfigFile, workspaceProjectConfigFile, workspaceProjectStreamFilter, localConfigFiles, oidcTokenClient, generateP4Config);
+				OpenProjectInfo workspaceSettings = new OpenProjectInfo(selectedProject, perforceSettings, projectInfo, userWorkspaceSettings, workspaceStateWrapper, latestProjectConfigFile, workspaceProjectConfigFile, workspaceProjectStreamFilter, localConfigFiles, generateP4Config);
 
 				if (generateP4Config)
 				{
