Horde: Return the name of an OIDC profile to save API access tokens in, so they can be reused by OidcToken.exe calls.

#rnx

diff --git a/Engine/Source/Programs/Horde/Horde.Server/Authentication/ExternalJwtAuthHandler.cs b/Engine/Source/Programs/Horde/Horde.Server/Authentication/ExternalJwtAuthHandler.cs
--- a/Engine/Source/Programs/Horde/Horde.Server/Authentication/ExternalJwtAuthHandler.cs
+++ b/Engine/Source/Programs/Horde/Horde.Server/Authentication/ExternalJwtAuthHandler.cs
@@ -19,7 +19,6 @@ using Microsoft.AspNetCore.Authentication;
 using Microsoft.AspNetCore.Authentication.JwtBearer;
 using Microsoft.Extensions.DependencyInjection;
 using Microsoft.Extensions.Logging;
-using Microsoft.IdentityModel.JsonWebTokens;
 using Microsoft.IdentityModel.Tokens;
 
 namespace Horde.Server.Authentication
diff --git a/Engine/Source/Programs/Horde/Horde.Server/Server/ServerController.cs b/Engine/Source/Programs/Horde/Horde.Server/Server/ServerController.cs
--- a/Engine/Source/Programs/Horde/Horde.Server/Server/ServerController.cs
+++ b/Engine/Source/Programs/Horde/Horde.Server/Server/ServerController.cs
@@ -139,6 +139,7 @@ namespace Horde.Server.Server
 
 			GetAuthConfigResponse response = new GetAuthConfigResponse();
 			response.Method = settings.AuthMethod;
+			response.ProfileName = settings.OidcProfileName;
 			response.ServerUrl = settings.OidcAuthority;
 			response.ClientId = settings.OidcClientId;
 			response.LocalRedirectUrls = settings.OidcLocalRedirectUrls;
diff --git a/Engine/Source/Programs/Horde/Horde.Server/ServerSettings.cs b/Engine/Source/Programs/Horde/Horde.Server/ServerSettings.cs
--- a/Engine/Source/Programs/Horde/Horde.Server/ServerSettings.cs
+++ b/Engine/Source/Programs/Horde/Horde.Server/ServerSettings.cs
@@ -479,6 +479,12 @@ namespace Horde.Server
 		/// </summary>
 		public AuthMethod AuthMethod { get; set; } = AuthMethod.Anonymous;
 
+		/// <summary>
+		/// Optional profile name to report through the /api/v1/server/auth endpoint. Allows sharing auth tokens between providers configured through
+		/// the same profile name in OidcToken.exe config files.
+		/// </summary>
+		public string? OidcProfileName { get; set; }
+
 		/// <summary>
 		/// Issuer for tokens from the auth provider
 		/// </summary>
diff --git a/Engine/Source/Programs/Shared/EpicGames.Horde/HordeHttpAuthHandler.cs b/Engine/Source/Programs/Shared/EpicGames.Horde/HordeHttpAuthHandler.cs
--- a/Engine/Source/Programs/Shared/EpicGames.Horde/HordeHttpAuthHandler.cs
+++ b/Engine/Source/Programs/Shared/EpicGames.Horde/HordeHttpAuthHandler.cs
@@ -306,13 +306,13 @@ namespace EpicGames.Horde
 				throw new Exception("No auth server configuration found");
 			}
 
-			const string OidcProvider = "Horde";
+			string oidcProvider = authConfig.ProfileName ?? "Horde";
 
 			Dictionary<string, string?> values = new Dictionary<string, string?>();
-			values[$"Providers:{OidcProvider}:DisplayName"] = "Horde";
-			values[$"Providers:{OidcProvider}:ServerUri"] = authConfig.ServerUrl;
-			values[$"Providers:{OidcProvider}:ClientId"] = authConfig.ClientId;
-			values[$"Providers:{OidcProvider}:RedirectUri"] = localRedirectUrl;
+			values[$"Providers:{oidcProvider}:DisplayName"] = "Horde";
+			values[$"Providers:{oidcProvider}:ServerUri"] = authConfig.ServerUrl;
+			values[$"Providers:{oidcProvider}:ClientId"] = authConfig.ClientId;
+			values[$"Providers:{oidcProvider}:RedirectUri"] = localRedirectUrl;
 
 			ConfigurationBuilder builder = new ConfigurationBuilder();
 			builder.AddInMemoryCollection(values);
@@ -320,14 +320,14 @@ namespace EpicGames.Horde
 			IConfiguration configuration = builder.Build();
 
 			using ITokenStore tokenStore = TokenStoreFactory.CreateTokenStore();
-			OidcTokenManager oidcTokenManager = OidcTokenManager.CreateTokenManager(configuration, tokenStore, new List<string>() { OidcProvider });
+			OidcTokenManager oidcTokenManager = OidcTokenManager.CreateTokenManager(configuration, tokenStore, new List<string>() { oidcProvider });
 
 			OidcTokenInfo? result = null;
-			if (oidcTokenManager.GetStatusForProvider(OidcProvider) != OidcStatus.NotLoggedIn)
+			if (oidcTokenManager.GetStatusForProvider(oidcProvider) != OidcStatus.NotLoggedIn)
 			{
 				try
 				{
-					result = await oidcTokenManager.TryGetAccessToken(OidcProvider, cancellationToken);
+					result = await oidcTokenManager.TryGetAccessToken(oidcProvider, cancellationToken);
 				}
 				catch (Exception ex)
 				{
@@ -337,7 +337,7 @@ namespace EpicGames.Horde
 			if (result == null && allowLogin)
 			{
 				_logger.LogInformation("Logging in to {Server}...", serverUrl);
-				result = await oidcTokenManager.Login(OidcProvider, cancellationToken);
+				result = await oidcTokenManager.Login(oidcProvider, cancellationToken);
 			}
 
 			return new AuthState(authConfig.Method, result);
diff --git a/Engine/Source/Programs/Shared/EpicGames.Horde/Server/ServerMessages.cs b/Engine/Source/Programs/Shared/EpicGames.Horde/Server/ServerMessages.cs
--- a/Engine/Source/Programs/Shared/EpicGames.Horde/Server/ServerMessages.cs
+++ b/Engine/Source/Programs/Shared/EpicGames.Horde/Server/ServerMessages.cs
@@ -107,6 +107,11 @@ namespace EpicGames.Horde.Server
 		/// </summary>
 		public AuthMethod Method { get; set; }
 
+		/// <summary>
+		/// Optional profile name used by OidcToken
+		/// </summary>
+		public string? ProfileName { get; set; }
+
 		/// <summary>
 		/// Issuer for tokens from the auth provider
 		/// </summary>
