Horde: Run VS code cleanup on shared libraries.

#rnx

diff --git a/Engine/Source/Programs/Shared/EpicGames.AspNet/CompactBinaryFormatter.cs b/Engine/Source/Programs/Shared/EpicGames.AspNet/CompactBinaryFormatter.cs
--- a/Engine/Source/Programs/Shared/EpicGames.AspNet/CompactBinaryFormatter.cs
+++ b/Engine/Source/Programs/Shared/EpicGames.AspNet/CompactBinaryFormatter.cs
@@ -1,16 +1,16 @@
 // Copyright Epic Games, Inc. All Rights Reserved.
 
-using Microsoft.AspNetCore.Mvc.Formatters;
 using System;
-using Microsoft.Net.Http.Headers;
-using System.Threading.Tasks;
 using System.IO;
-using EpicGames.Serialization;
+using System.Threading.Tasks;
 using EpicGames.Core;
+using EpicGames.Serialization;
 using Microsoft.AspNetCore.Http;
+using Microsoft.AspNetCore.Mvc.Formatters;
 using Microsoft.Extensions.DependencyInjection;
-using Microsoft.Extensions.Primitives;
 using Microsoft.Extensions.Logging;
+using Microsoft.Extensions.Primitives;
+using Microsoft.Net.Http.Headers;
 
 namespace EpicGames.AspNet
 {
diff --git a/Engine/Source/Programs/Shared/EpicGames.BuildGraph.Tests/GraphTests.cs b/Engine/Source/Programs/Shared/EpicGames.BuildGraph.Tests/GraphTests.cs
--- a/Engine/Source/Programs/Shared/EpicGames.BuildGraph.Tests/GraphTests.cs
+++ b/Engine/Source/Programs/Shared/EpicGames.BuildGraph.Tests/GraphTests.cs
@@ -1,10 +1,10 @@
 // Copyright Epic Games, Inc. All Rights Reserved.
 
-using Microsoft.VisualStudio.TestTools.UnitTesting;
 using System.Collections.Generic;
 using System.Linq;
-using EpicGames.BuildGraph.Expressions;
 using System.Threading.Tasks;
+using EpicGames.BuildGraph.Expressions;
+using Microsoft.VisualStudio.TestTools.UnitTesting;
 
 namespace EpicGames.BuildGraph.Tests
 {
@@ -47,7 +47,7 @@ namespace EpicGames.BuildGraph.Tests
 			Assert.AreEqual(0, node1.InputDependencies.Count);
 			Assert.AreEqual(0, node1.Inputs.Count);
 		}
-		
+
 		[TestMethod]
 		public void AgentTest()
 		{
@@ -59,12 +59,13 @@ namespace EpicGames.BuildGraph.Tests
 			Assert.AreEqual(2, agent.PossibleTypes.Count);
 			Assert.AreEqual("win64", agent.PossibleTypes[0]);
 			Assert.AreEqual("incremental", agent.PossibleTypes[1]);
-/*
-			Assert.AreEqual(1, agent.Nodes.Count);
+			/*
+						Assert.AreEqual(1, agent.Nodes.Count);
 
-			BgNode node = agent.Nodes[0];
-			Assert.AreEqual("Update Version Files", node.Name);
-*/		}
+						BgNode node = agent.Nodes[0];
+						Assert.AreEqual("Update Version Files", node.Name);
+			*/
+		}
 
 		[TestMethod]
 		public void AggregateTest()
diff --git a/Engine/Source/Programs/Shared/EpicGames.BuildGraph.Tests/OpcodeTests.cs b/Engine/Source/Programs/Shared/EpicGames.BuildGraph.Tests/OpcodeTests.cs
--- a/Engine/Source/Programs/Shared/EpicGames.BuildGraph.Tests/OpcodeTests.cs
+++ b/Engine/Source/Programs/Shared/EpicGames.BuildGraph.Tests/OpcodeTests.cs
@@ -1,12 +1,12 @@
 // Copyright Epic Games, Inc. All Rights Reserved.
 
-using Microsoft.VisualStudio.TestTools.UnitTesting;
 using System;
 using System.Collections.Generic;
 using System.IO;
 using System.Linq;
-using EpicGames.BuildGraph.Expressions;
 using System.Threading.Tasks;
+using EpicGames.BuildGraph.Expressions;
+using Microsoft.VisualStudio.TestTools.UnitTesting;
 
 namespace EpicGames.BuildGraph.Tests
 {
diff --git a/Engine/Source/Programs/Shared/EpicGames.Horde.Tests/BundleStoreTests.cs b/Engine/Source/Programs/Shared/EpicGames.Horde.Tests/BundleStoreTests.cs
--- a/Engine/Source/Programs/Shared/EpicGames.Horde.Tests/BundleStoreTests.cs
+++ b/Engine/Source/Programs/Shared/EpicGames.Horde.Tests/BundleStoreTests.cs
@@ -1,24 +1,24 @@
 // Copyright Epic Games, Inc. All Rights Reserved.
 
-using EpicGames.Horde.Storage;
-using EpicGames.Horde.Storage.Clients;
-using EpicGames.Horde.Storage.Nodes;
-using Microsoft.VisualStudio.TestTools.UnitTesting;
 using System;
 using System.Buffers;
 using System.Collections.Generic;
+using System.IO;
 using System.Linq;
-using System.Threading.Tasks;
-using System.Threading;
 using System.Text;
-using System.IO;
+using System.Threading;
+using System.Threading.Tasks;
 using EpicGames.Core;
-using Microsoft.Extensions.Caching.Memory;
-using Microsoft.Extensions.Logging.Abstractions;
+using EpicGames.Horde.Storage;
+using EpicGames.Horde.Storage.Backends;
 using EpicGames.Horde.Storage.Bundles;
 using EpicGames.Horde.Storage.Bundles.V1;
+using EpicGames.Horde.Storage.Clients;
+using EpicGames.Horde.Storage.Nodes;
 using EpicGames.Serialization;
-using EpicGames.Horde.Storage.Backends;
+using Microsoft.Extensions.Caching.Memory;
+using Microsoft.Extensions.Logging.Abstractions;
+using Microsoft.VisualStudio.TestTools.UnitTesting;
 
 namespace EpicGames.Horde.Tests
 {
@@ -81,7 +81,7 @@ namespace EpicGames.Horde.Tests
 			List<BlobType> types = new List<BlobType>();
 			types.Add(new BlobType("{F63606D4-4061-5DBB-446F-55A69E22654F}", 1));
 
-			List <BundleExport> exports = new List<BundleExport>();
+			List<BundleExport> exports = new List<BundleExport>();
 			exports.Add(new BundleExport(0, 0, 0, payload.Length, Array.Empty<BundleExportRef>()));
 
 			List<BundlePacket> packets = new List<BundlePacket>();
@@ -171,7 +171,7 @@ namespace EpicGames.Horde.Tests
 				await writer.FlushAsync();
 
 				await store.WriteRefAsync(new RefName("test"), rootRef);
-			
+
 				BundleReader reader = new BundleReader(store, BundleCache.None, NullLogger.Instance);
 				await CheckTreeAsync(root);
 			}
diff --git a/Engine/Source/Programs/Shared/EpicGames.Horde.Tests/ChunkingTests.cs b/Engine/Source/Programs/Shared/EpicGames.Horde.Tests/ChunkingTests.cs
--- a/Engine/Source/Programs/Shared/EpicGames.Horde.Tests/ChunkingTests.cs
+++ b/Engine/Source/Programs/Shared/EpicGames.Horde.Tests/ChunkingTests.cs
@@ -1,18 +1,18 @@
 // Copyright Epic Games, Inc. All Rights Reserved.
 
-using EpicGames.Core;
-using EpicGames.Horde.Storage.Nodes;
-using Microsoft.VisualStudio.TestTools.UnitTesting;
 using System;
-using System.Linq;
-using System.Threading.Tasks;
 using System.IO;
+using System.Linq;
+using System.Text;
 using System.Threading;
+using System.Threading.Tasks;
+using EpicGames.Core;
 using EpicGames.Horde.Storage;
 using EpicGames.Horde.Storage.Clients;
+using EpicGames.Horde.Storage.Nodes;
 using Microsoft.Extensions.Caching.Memory;
-using System.Text;
 using Microsoft.Extensions.Logging.Abstractions;
+using Microsoft.VisualStudio.TestTools.UnitTesting;
 
 namespace EpicGames.Horde.Tests
 {
@@ -62,7 +62,7 @@ namespace EpicGames.Horde.Tests
 
 				using MemoryStream emptyStream = new MemoryStream();
 				LeafChunkedData leafChunkedData = await LeafChunkedDataNode.CreateFromStreamAsync(writer, emptyStream, new LeafChunkedDataNodeOptions(64, 64, 64), CancellationToken.None);
-				ChunkedData chunkedData = await InteriorChunkedDataNode.CreateTreeAsync(leafChunkedData, new InteriorChunkedDataNodeOptions(4, 4, 4), writer, CancellationToken.None); 
+				ChunkedData chunkedData = await InteriorChunkedDataNode.CreateTreeAsync(leafChunkedData, new InteriorChunkedDataNodeOptions(4, 4, 4), writer, CancellationToken.None);
 
 				DirectoryNode directory = new DirectoryNode();
 				directory.AddFile("test.foo", FileEntryFlags.None, 0, chunkedData);
diff --git a/Engine/Source/Programs/Shared/EpicGames.Horde/Compute/AgentMessageHandler.cs b/Engine/Source/Programs/Shared/EpicGames.Horde/Compute/AgentMessageHandler.cs
--- a/Engine/Source/Programs/Shared/EpicGames.Horde/Compute/AgentMessageHandler.cs
+++ b/Engine/Source/Programs/Shared/EpicGames.Horde/Compute/AgentMessageHandler.cs
@@ -3,19 +3,19 @@
 using System;
 using System.Collections.Generic;
 using System.Diagnostics;
+using System.IO;
+using System.Linq;
+using System.Runtime.ExceptionServices;
 using System.Runtime.InteropServices;
+using System.Text;
 using System.Threading;
 using System.Threading.Tasks;
 using EpicGames.Core;
 using EpicGames.Horde.Compute.Buffers;
-using EpicGames.Horde.Storage.Nodes;
-using Microsoft.Extensions.Logging;
 using EpicGames.Horde.Storage;
-using System.IO;
-using System.Linq;
-using System.Runtime.ExceptionServices;
-using System.Text;
 using EpicGames.Horde.Storage.Bundles;
+using EpicGames.Horde.Storage.Nodes;
+using Microsoft.Extensions.Logging;
 
 namespace EpicGames.Horde.Compute
 {
@@ -223,7 +223,7 @@ namespace EpicGames.Horde.Compute
 				{
 					await using FileStream stream = FileReference.Open(file, FileMode.Open, FileAccess.Read);
 					IoHash hash = await IoHash.ComputeAsync(stream, cancellationToken);
-					
+
 					if (hash == fileEntry.StreamHash)
 					{
 						_logger.LogInformation("Hash of {File} is correct ({Hash})", file, hash);
@@ -322,7 +322,7 @@ namespace EpicGames.Horde.Compute
 			using CancellationTokenSource cancellationTokenSource = CancellationTokenSource.CreateLinkedTokenSource(cancellationTokens);
 			CancellationToken cancellationToken = cancellationTokenSource.Token;
 
-			List <SharedMemoryBuffer> buffers = new();
+			List<SharedMemoryBuffer> buffers = new();
 			try
 			{
 				List<(int, ComputeBufferWriter)> writers = new List<(int, ComputeBufferWriter)>();
@@ -464,7 +464,7 @@ namespace EpicGames.Horde.Compute
 			finally
 			{
 				Directory.SetCurrentDirectory(prevWorkingDir);
-				foreach((string key, string? value) in prevEnvVars)
+				foreach ((string key, string? value) in prevEnvVars)
 				{
 					Environment.SetEnvironmentVariable(key, value);
 				}
@@ -477,7 +477,7 @@ namespace EpicGames.Horde.Compute
 			{
 				throw new Exception("Only Linux is supported for executing a process inside a container");
 			}
-			
+
 			if (_containerEngineExecutable == null)
 			{
 				throw new Exception("Container execution requested but agent has no container engine configured");
@@ -497,7 +497,7 @@ namespace EpicGames.Horde.Compute
 				sb.AppendLine($"{key}={value}");
 			}
 			await File.WriteAllTextAsync(envFilePath, sb.ToString(), cancellationToken);
-			
+
 			List<string> resolvedArguments = new()
 			{
 				"run",
@@ -507,7 +507,7 @@ namespace EpicGames.Horde.Compute
 				$"--volume={_sandboxDir}:{_sandboxDir}:rw",
 				"--env-file=" + envFilePath,
 			};
-			
+
 			if (flags.HasFlag(ExecuteProcessFlags.ReplaceContainerEntrypoint))
 			{
 				resolvedArguments.Add("--entrypoint=" + resolvedExecutable);
@@ -518,7 +518,7 @@ namespace EpicGames.Horde.Compute
 				resolvedArguments.Add(containerImageUrl);
 				resolvedArguments.Add(resolvedExecutable); // Add executable as first argument and assume the entrypoint inside the container image will handle this
 			}
-			
+
 			resolvedArguments.AddRange(arguments);
 			_logger.LogInformation("Executing {File} {Arguments} in container", _containerEngineExecutable, arguments);
 
@@ -549,12 +549,12 @@ namespace EpicGames.Horde.Compute
 				Dictionary<string, string> resolvedEnvVars = ResolveEnvVars(envVars);
 				if (!File.Exists(resolvedExecutable))
 				{
-					_logger.LogWarning("Executable {Path} does not exist", resolvedExecutable);	
+					_logger.LogWarning("Executable {Path} does not exist", resolvedExecutable);
 				}
-				
+
 				if (!Directory.Exists(resolvedWorkingDir))
 				{
-					_logger.LogWarning("Working dir {Path} does not exist", resolvedWorkingDir);	
+					_logger.LogWarning("Working dir {Path} does not exist", resolvedWorkingDir);
 				}
 
 				using ManagedProcessGroup group = new ManagedProcessGroup();
@@ -580,12 +580,12 @@ namespace EpicGames.Horde.Compute
 				}
 			}
 		}
-		
+
 		private string GetExecutableAbsPath(string relPath)
 		{
 			return FileReference.Combine(_sandboxDir, relPath).FullName;
 		}
-		
+
 		private string GetWorkingDirAbsPath(string? relPath)
 		{
 			return DirectoryReference.Combine(_sandboxDir, relPath ?? String.Empty).FullName;
diff --git a/Engine/Source/Programs/Shared/EpicGames.Horde/Compute/Buffers/SharedMemoryBuffer.cs b/Engine/Source/Programs/Shared/EpicGames.Horde/Compute/Buffers/SharedMemoryBuffer.cs
--- a/Engine/Source/Programs/Shared/EpicGames.Horde/Compute/Buffers/SharedMemoryBuffer.cs
+++ b/Engine/Source/Programs/Shared/EpicGames.Horde/Compute/Buffers/SharedMemoryBuffer.cs
@@ -1,12 +1,12 @@
 // Copyright Epic Games, Inc. All Rights Reserved.
 
 using System;
-using EpicGames.Core;
-using System.IO.MemoryMappedFiles;
 using System.IO;
-using System.Threading.Tasks;
-using System.Threading;
+using System.IO.MemoryMappedFiles;
 using System.Runtime.InteropServices;
+using System.Threading;
+using System.Threading.Tasks;
+using EpicGames.Core;
 
 namespace EpicGames.Horde.Compute.Buffers
 {
diff --git a/Engine/Source/Programs/Shared/EpicGames.Horde/Logs/LogNode.cs b/Engine/Source/Programs/Shared/EpicGames.Horde/Logs/LogNode.cs
--- a/Engine/Source/Programs/Shared/EpicGames.Horde/Logs/LogNode.cs
+++ b/Engine/Source/Programs/Shared/EpicGames.Horde/Logs/LogNode.cs
@@ -2,14 +2,14 @@
 
 using System;
 using System.Collections.Generic;
+using System.Diagnostics;
 using System.Linq;
-using System.Threading.Tasks;
+using System.Runtime.CompilerServices;
 using System.Threading;
+using System.Threading.Tasks;
 using EpicGames.Core;
 using EpicGames.Horde.Storage;
 using Microsoft.Extensions.Logging;
-using System.Diagnostics;
-using System.Runtime.CompilerServices;
 
 namespace EpicGames.Horde.Logs
 {
diff --git a/Engine/Source/Programs/Shared/EpicGames.Horde/Storage/Bundles/V1/BundleReader.cs b/Engine/Source/Programs/Shared/EpicGames.Horde/Storage/Bundles/V1/BundleReader.cs
--- a/Engine/Source/Programs/Shared/EpicGames.Horde/Storage/Bundles/V1/BundleReader.cs
+++ b/Engine/Source/Programs/Shared/EpicGames.Horde/Storage/Bundles/V1/BundleReader.cs
@@ -2,14 +2,14 @@
 
 using System;
 using System.Collections.Generic;
+using System.Diagnostics;
+using System.IO;
 using System.Linq;
 using System.Threading;
 using System.Threading.Tasks;
 using EpicGames.Core;
-using Microsoft.Extensions.Logging;
 using Microsoft.CodeAnalysis;
-using System.Diagnostics;
-using System.IO;
+using Microsoft.Extensions.Logging;
 
 namespace EpicGames.Horde.Storage.Bundles.V1
 {
diff --git a/Engine/Source/Programs/Shared/EpicGames.IoHash/IoHash.cs b/Engine/Source/Programs/Shared/EpicGames.IoHash/IoHash.cs
--- a/Engine/Source/Programs/Shared/EpicGames.IoHash/IoHash.cs
+++ b/Engine/Source/Programs/Shared/EpicGames.IoHash/IoHash.cs
@@ -5,11 +5,11 @@ using System.Buffers;
 using System.Buffers.Binary;
 using System.ComponentModel;
 using System.Globalization;
-using System.Text.Json.Serialization;
-using System.Text.Json;
 using System.IO;
-using System.Threading.Tasks;
+using System.Text.Json;
+using System.Text.Json.Serialization;
 using System.Threading;
+using System.Threading.Tasks;
 
 namespace EpicGames.Core
 {
diff --git a/Engine/Source/Programs/Shared/EpicGames.Serialization/CbSerializer.cs b/Engine/Source/Programs/Shared/EpicGames.Serialization/CbSerializer.cs
--- a/Engine/Source/Programs/Shared/EpicGames.Serialization/CbSerializer.cs
+++ b/Engine/Source/Programs/Shared/EpicGames.Serialization/CbSerializer.cs
@@ -1,6 +1,5 @@
 // Copyright Epic Games, Inc. All Rights Reserved.
 
-using EpicGames.Core;
 using System;
 
 namespace EpicGames.Serialization
diff --git a/Engine/Source/Programs/Shared/EpicGames.Serialization/Converters/CbArrayConverter.cs b/Engine/Source/Programs/Shared/EpicGames.Serialization/Converters/CbArrayConverter.cs
--- a/Engine/Source/Programs/Shared/EpicGames.Serialization/Converters/CbArrayConverter.cs
+++ b/Engine/Source/Programs/Shared/EpicGames.Serialization/Converters/CbArrayConverter.cs
@@ -1,6 +1,5 @@
 // Copyright Epic Games, Inc. All Rights Reserved.
 
-using EpicGames.Core;
 using System;
 using System.Collections.Generic;
 
@@ -64,7 +63,7 @@ namespace EpicGames.Serialization.Converters
 			}
 		}
 	}
-	
+
 	/// <summary>
 	/// Factory for CbListConverter
 	/// </summary>
diff --git a/Engine/Source/Programs/Shared/EpicGames.Serialization/Converters/CbDictionaryConverter.cs b/Engine/Source/Programs/Shared/EpicGames.Serialization/Converters/CbDictionaryConverter.cs
--- a/Engine/Source/Programs/Shared/EpicGames.Serialization/Converters/CbDictionaryConverter.cs
+++ b/Engine/Source/Programs/Shared/EpicGames.Serialization/Converters/CbDictionaryConverter.cs
@@ -1,6 +1,5 @@
 // Copyright Epic Games, Inc. All Rights Reserved.
 
-using EpicGames.Core;
 using System;
 using System.Collections.Generic;
 
diff --git a/Engine/Source/Programs/Shared/EpicGames.Serialization/Converters/CbEnumConverter.cs b/Engine/Source/Programs/Shared/EpicGames.Serialization/Converters/CbEnumConverter.cs
--- a/Engine/Source/Programs/Shared/EpicGames.Serialization/Converters/CbEnumConverter.cs
+++ b/Engine/Source/Programs/Shared/EpicGames.Serialization/Converters/CbEnumConverter.cs
@@ -1,6 +1,5 @@
 // Copyright Epic Games, Inc. All Rights Reserved.
 
-using EpicGames.Core;
 using System;
 using System.Linq.Expressions;
 using System.Reflection;
diff --git a/Engine/Source/Programs/Shared/EpicGames.Serialization/Converters/CbFieldConverter.cs b/Engine/Source/Programs/Shared/EpicGames.Serialization/Converters/CbFieldConverter.cs
--- a/Engine/Source/Programs/Shared/EpicGames.Serialization/Converters/CbFieldConverter.cs
+++ b/Engine/Source/Programs/Shared/EpicGames.Serialization/Converters/CbFieldConverter.cs
@@ -1,7 +1,5 @@
 // Copyright Epic Games, Inc. All Rights Reserved.
 
-using EpicGames.Core;
-
 namespace EpicGames.Serialization.Converters
 {
 	/// <summary>
diff --git a/Engine/Source/Programs/Shared/EpicGames.Serialization/Converters/CbNullableConverter.cs b/Engine/Source/Programs/Shared/EpicGames.Serialization/Converters/CbNullableConverter.cs
--- a/Engine/Source/Programs/Shared/EpicGames.Serialization/Converters/CbNullableConverter.cs
+++ b/Engine/Source/Programs/Shared/EpicGames.Serialization/Converters/CbNullableConverter.cs
@@ -1,6 +1,5 @@
 // Copyright Epic Games, Inc. All Rights Reserved.
 
-using EpicGames.Core;
 using System;
 
 namespace EpicGames.Serialization.Converters
diff --git a/Engine/Source/Programs/Shared/EpicGames.Serialization/Converters/CbPrimitiveConverter.cs b/Engine/Source/Programs/Shared/EpicGames.Serialization/Converters/CbPrimitiveConverter.cs
--- a/Engine/Source/Programs/Shared/EpicGames.Serialization/Converters/CbPrimitiveConverter.cs
+++ b/Engine/Source/Programs/Shared/EpicGames.Serialization/Converters/CbPrimitiveConverter.cs
@@ -1,6 +1,5 @@
 // Copyright Epic Games, Inc. All Rights Reserved.
 
-using EpicGames.Core;
 using System;
 using System.Linq.Expressions;
 using System.Reflection;
diff --git a/Engine/Source/Programs/Shared/EpicGames.Serialization/Converters/CbStringConverter.cs b/Engine/Source/Programs/Shared/EpicGames.Serialization/Converters/CbStringConverter.cs
--- a/Engine/Source/Programs/Shared/EpicGames.Serialization/Converters/CbStringConverter.cs
+++ b/Engine/Source/Programs/Shared/EpicGames.Serialization/Converters/CbStringConverter.cs
@@ -1,6 +1,5 @@
 // Copyright Epic Games, Inc. All Rights Reserved.
 
-using EpicGames.Core;
 using System;
 using System.ComponentModel;
 
@@ -26,7 +25,7 @@ namespace EpicGames.Serialization.Converters
 		}
 
 		public override void Write(CbWriter writer, T value)
-		{ 
+		{
 			writer.WriteStringValue(_typeConverter.ConvertToInvariantString(value)!);
 		}
 	}
