Horde: Add a ToString() method for FileSpecList.

#rnx

diff --git a/Engine/Source/Programs/Shared/EpicGames.Perforce/FileSpecList.cs b/Engine/Source/Programs/Shared/EpicGames.Perforce/FileSpecList.cs
--- a/Engine/Source/Programs/Shared/EpicGames.Perforce/FileSpecList.cs
+++ b/Engine/Source/Programs/Shared/EpicGames.Perforce/FileSpecList.cs
@@ -88,5 +88,18 @@ namespace EpicGames.Perforce
 		{
 			return new FileSpecList(new string[] { fileSpec });
 		}
+
+		/// <inheritdoc/>
+		public override string ToString()
+		{
+			if (List.Count == 0)
+			{
+				return "(none)";
+			}
+			else
+			{
+				return String.Join(", ", List);
+			}
+		}
 	}
 }
