EpicGames.MsBuild: Conform to standards

#rnx

diff --git a/Engine/Source/Programs/Shared/EpicGames.MsBuild/CsProjBuilder.cs b/Engine/Source/Programs/Shared/EpicGames.MsBuild/CsProjBuilder.cs
--- a/Engine/Source/Programs/Shared/EpicGames.MsBuild/CsProjBuilder.cs
+++ b/Engine/Source/Programs/Shared/EpicGames.MsBuild/CsProjBuilder.cs
@@ -12,7 +12,6 @@ using System.Collections.Generic;
 using System.Diagnostics;
 using System.IO;
 using System.Linq;
-using System.Text;
 using System.Text.Json;
 
 namespace EpicGames.MsBuild
@@ -27,42 +26,42 @@ namespace EpicGames.MsBuild
 	{
 		class MLogger : IBuildLogger
 		{
-			ILogger Inner;
+			readonly ILogger _inner;
 
 			LoggerVerbosity IBuildLogger.Verbosity { get => LoggerVerbosity.Normal; set => throw new NotImplementedException(); }
 			string IBuildLogger.Parameters { get => throw new NotImplementedException(); set { } }
 
-			public bool bVeryVerboseLog = false;
+			public bool _bVeryVerboseLog = false;
 
-			bool bFirstError = true;
+			bool _bFirstError = true;
 
-			public MLogger(ILogger InInner)
+			public MLogger(ILogger inInner)
 			{
-				Inner = InInner;
+				_inner = inInner;
 			}
 
-			void IBuildLogger.Initialize(IEventSource EventSource)
+			void IBuildLogger.Initialize(IEventSource eventSource)
 			{
-				EventSource.ProjectStarted += new ProjectStartedEventHandler(eventSource_ProjectStarted);
-				EventSource.TaskStarted += new TaskStartedEventHandler(eventSource_TaskStarted);
-				EventSource.MessageRaised += new BuildMessageEventHandler(eventSource_MessageRaised);
-				EventSource.WarningRaised += new BuildWarningEventHandler(eventSource_WarningRaised);
-				EventSource.ErrorRaised += new BuildErrorEventHandler(eventSource_ErrorRaised);
-				EventSource.ProjectFinished += new ProjectFinishedEventHandler(eventSource_ProjectFinished);
+				eventSource.ProjectStarted += new ProjectStartedEventHandler(EventSource_ProjectStarted);
+				eventSource.TaskStarted += new TaskStartedEventHandler(EventSource_TaskStarted);
+				eventSource.MessageRaised += new BuildMessageEventHandler(EventSource_MessageRaised);
+				eventSource.WarningRaised += new BuildWarningEventHandler(EventSource_WarningRaised);
+				eventSource.ErrorRaised += new BuildErrorEventHandler(EventSource_ErrorRaised);
+				eventSource.ProjectFinished += new ProjectFinishedEventHandler(EventSource_ProjectFinished);
 			}
 
-			void eventSource_ErrorRaised(object Sender, BuildErrorEventArgs e)
+			void EventSource_ErrorRaised(object sender, BuildErrorEventArgs e)
 			{
-				if (bFirstError)
+				if (_bFirstError)
 				{
 					Trace.WriteLine("");
-					Log.Logger.LogInformation("");
-					bFirstError = false;
+					_inner.LogInformation("");
+					_bFirstError = false;
 				}
-				Log.Logger.LogError("{File}({Line},{Column}): error {Code}: {Message} ({ProjectFile})", new FileReference(e.File), new LogValue(LogValueType.LineNumber, e.LineNumber.ToString()), new LogValue(LogValueType.ColumnNumber, e.ColumnNumber.ToString()), new LogValue(LogValueType.ErrorCode, e.Code), e.Message, new FileReference(e.ProjectFile));
+				_inner.LogError("{File}({Line},{Column}): error {Code}: {Message} ({ProjectFile})", new FileReference(e.File), new LogValue(LogValueType.LineNumber, e.LineNumber.ToString()), new LogValue(LogValueType.ColumnNumber, e.ColumnNumber.ToString()), new LogValue(LogValueType.ErrorCode, e.Code), e.Message, new FileReference(e.ProjectFile));
 			}
 
-			void eventSource_WarningRaised(object Sender, BuildWarningEventArgs e)
+			void EventSource_WarningRaised(object sender, BuildWarningEventArgs e)
 			{
 				{
 					// workaround for warnings that appear after revert of net6.0 upgrade. Delete this block when the net6.0 upgrade is done.
@@ -70,53 +69,53 @@ namespace EpicGames.MsBuild
 					// The starting contents (in bytes) are: 06-01-01-00-00-00-01-19-50-72-6F-70-65-72-74-69-65 ... (...\[projectname].csproj)
 					if (String.Equals(e.Code, "MSB3088", StringComparison.Ordinal))
 					{
-						Log.Logger.LogDebug("{File}({Line},{Column}): suppressed warning {Code}: {Message} ({ProjectFile})", new FileReference(e.File), new LogValue(LogValueType.LineNumber, e.LineNumber.ToString()), new LogValue(LogValueType.ColumnNumber, e.ColumnNumber.ToString()), new LogValue(LogValueType.ErrorCode, e.Code), e.Message, new FileReference(e.ProjectFile));
+						_inner.LogDebug("{File}({Line},{Column}): suppressed warning {Code}: {Message} ({ProjectFile})", new FileReference(e.File), new LogValue(LogValueType.LineNumber, e.LineNumber.ToString()), new LogValue(LogValueType.ColumnNumber, e.ColumnNumber.ToString()), new LogValue(LogValueType.ErrorCode, e.Code), e.Message, new FileReference(e.ProjectFile));
 						return;
 					}
 				}
 
-				if (bFirstError)
-                {
-					Log.Logger.LogInformation("");
-					bFirstError = false;
-                }
+				if (_bFirstError)
+				{
+					_inner.LogInformation("");
+					_bFirstError = false;
+				}
 
-				Log.Logger.LogWarning("{File}({Line},{Column}): warning {Code}: {Message} ({ProjectFile})", new FileReference(e.File), new LogValue(LogValueType.LineNumber, e.LineNumber.ToString()), new LogValue(LogValueType.ColumnNumber, e.ColumnNumber.ToString()), new LogValue(LogValueType.ErrorCode, e.Code), e.Message, new FileReference(e.ProjectFile));
+				_inner.LogWarning("{File}({Line},{Column}): warning {Code}: {Message} ({ProjectFile})", new FileReference(e.File), new LogValue(LogValueType.LineNumber, e.LineNumber.ToString()), new LogValue(LogValueType.ColumnNumber, e.ColumnNumber.ToString()), new LogValue(LogValueType.ErrorCode, e.Code), e.Message, new FileReference(e.ProjectFile));
 			}
 
-			void eventSource_MessageRaised(object Sender, BuildMessageEventArgs e)
+			void EventSource_MessageRaised(object sender, BuildMessageEventArgs e)
 			{
-				if (bVeryVerboseLog)
+				if (_bVeryVerboseLog)
 				{
 					//if (!String.Equals(e.SenderName, "ResolveAssemblyReference"))
 					//if (e.Message.Contains("atic"))
 					{
-						Log.Logger.LogDebug("{SenderName}: {Message}", e.SenderName, e.Message);
+						_inner.LogDebug("{SenderName}: {Message}", e.SenderName, e.Message);
 					}
 				}
 			}
 
-			void eventSource_ProjectStarted(object Sender, ProjectStartedEventArgs e)
+			void EventSource_ProjectStarted(object sender, ProjectStartedEventArgs e)
 			{
-				if (bVeryVerboseLog)
+				if (_bVeryVerboseLog)
 				{
-					Log.Logger.LogDebug("{SenderName}: {Message}", e.SenderName, e.Message);
+					_inner.LogDebug("{SenderName}: {Message}", e.SenderName, e.Message);
 				}
 			}
 
-			void eventSource_ProjectFinished(object Sender, ProjectFinishedEventArgs e)
+			void EventSource_ProjectFinished(object sender, ProjectFinishedEventArgs e)
 			{
-				if (bVeryVerboseLog)
+				if (_bVeryVerboseLog)
 				{
-					Log.Logger.LogDebug("{SenderName}: {Message}", e.SenderName, e.Message);
+					_inner.LogDebug("{SenderName}: {Message}", e.SenderName, e.Message);
 				}
 			}
 
-			void eventSource_TaskStarted(object Sender, TaskStartedEventArgs e)
+			void EventSource_TaskStarted(object sender, TaskStartedEventArgs e)
 			{
-				if (bVeryVerboseLog)
+				if (_bVeryVerboseLog)
 				{
-					Log.Logger.LogDebug("{SenderName}: {Message}", e.SenderName, e.Message);
+					_inner.LogDebug("{SenderName}: {Message}", e.SenderName, e.Message);
 				}
 			}
 
@@ -125,71 +124,71 @@ namespace EpicGames.MsBuild
 			}
 		}
 
-		static FileReference ConstructBuildRecordPath(CsProjBuildHook Hook, FileReference ProjectPath, IEnumerable<DirectoryReference> BaseDirectories)
+		static FileReference ConstructBuildRecordPath(CsProjBuildHook hook, FileReference projectPath, IEnumerable<DirectoryReference> baseDirectories)
 		{
-			DirectoryReference BasePath = null;
+			DirectoryReference basePath = null;
 
-			foreach (DirectoryReference ScriptFolder in BaseDirectories)
+			foreach (DirectoryReference scriptFolder in baseDirectories)
 			{
-				if (ProjectPath.IsUnderDirectory(ScriptFolder))
+				if (projectPath.IsUnderDirectory(scriptFolder))
 				{
-					BasePath = ScriptFolder;
+					basePath = scriptFolder;
 					break;
 				}
 			}
 
-			if (BasePath == null)
+			if (basePath == null)
 			{
-				throw new Exception($"Unable to map csproj {ProjectPath} to Engine, game, or an additional script folder. Candidates were:{Environment.NewLine} {String.Join(Environment.NewLine, BaseDirectories)}");
+				throw new Exception($"Unable to map csproj {projectPath} to Engine, game, or an additional script folder. Candidates were:{Environment.NewLine} {String.Join(Environment.NewLine, baseDirectories)}");
 			}
 
-			DirectoryReference BuildRecordDirectory = Hook.GetBuildRecordDirectory(BasePath);
-			DirectoryReference.CreateDirectory(BuildRecordDirectory);
+			DirectoryReference buildRecordDirectory = hook.GetBuildRecordDirectory(basePath);
+			DirectoryReference.CreateDirectory(buildRecordDirectory);
 
-			return FileReference.Combine(BuildRecordDirectory, ProjectPath.GetFileName()).ChangeExtension(".json");
+			return FileReference.Combine(buildRecordDirectory, projectPath.GetFileName()).ChangeExtension(".json");
 		}
 
 		/// <summary>
 		/// Builds multiple projects
 		/// </summary>
-		/// <param name="FoundProjects">Collection of project to be built</param>
+		/// <param name="foundProjects">Collection of project to be built</param>
 		/// <param name="bForceCompile">If true, force the compilation of the projects</param>
 		/// <param name="bBuildSuccess">Set to true/false depending on if all projects compiled or are up-to-date</param>
-		/// <param name="Hook">Interface to fetch data about the building environment</param>
-		/// <param name="BaseDirectories">Base directories of the engine and project</param>
-		/// <param name="DefineConstants">Collection of constants to be defined while building projects</param>
-		/// <param name="OnBuildingProjects">Action invoked to notify caller regarding the number of projects being built</param>
-		/// <param name="Logger">Destination logger</param>
-		public static Dictionary<FileReference, CsProjBuildRecordEntry> Build(HashSet<FileReference> FoundProjects,
-			bool bForceCompile, out bool bBuildSuccess, CsProjBuildHook Hook, IEnumerable<DirectoryReference> BaseDirectories,
-			IEnumerable<string> DefineConstants, Action<int> OnBuildingProjects, ILogger Logger)
+		/// <param name="hook">Interface to fetch data about the building environment</param>
+		/// <param name="baseDirectories">Base directories of the engine and project</param>
+		/// <param name="defineConstants">Collection of constants to be defined while building projects</param>
+		/// <param name="onBuildingProjects">Action invoked to notify caller regarding the number of projects being built</param>
+		/// <param name="logger">Destination logger</param>
+		public static Dictionary<FileReference, CsProjBuildRecordEntry> Build(HashSet<FileReference> foundProjects,
+			bool bForceCompile, out bool bBuildSuccess, CsProjBuildHook hook, IEnumerable<DirectoryReference> baseDirectories,
+			IEnumerable<string> defineConstants, Action<int> onBuildingProjects, ILogger logger)
 		{
 
 			// Register the MS build path prior to invoking the internal routine.  By not having the internal routine
 			// inline, we avoid having the issue of the Microsoft.Build libraries being resolved prior to the build path
 			// being set.
-			RegisterMsBuildPath(Hook);
-			return BuildInternal(FoundProjects, bForceCompile, out bBuildSuccess, Hook, BaseDirectories, DefineConstants, OnBuildingProjects, Logger);
+			RegisterMsBuildPath(hook);
+			return BuildInternal(foundProjects, bForceCompile, out bBuildSuccess, hook, baseDirectories, defineConstants, onBuildingProjects, logger);
 		}
 
 		/// <summary>
 		/// Builds multiple projects.  This is the internal implementation invoked after the MS build path is set
 		/// </summary>
-		/// <param name="FoundProjects">Collection of project to be built</param>
+		/// <param name="foundProjects">Collection of project to be built</param>
 		/// <param name="bForceCompile">If true, force the compilation of the projects</param>
 		/// <param name="bBuildSuccess">Set to true/false depending on if all projects compiled or are up-to-date</param>
-		/// <param name="Hook">Interface to fetch data about the building environment</param>
-		/// <param name="BaseDirectories">Base directories of the engine and project</param>
-		/// <param name="DefineConstants">Collection of constants to be defined while building projects</param>
-		/// <param name="OnBuildingProjects">Action invoked to notify caller regarding the number of projects being built</param>
-		/// <param name="Logger">Destination logger</param>
-		private static Dictionary<FileReference, CsProjBuildRecordEntry> BuildInternal(HashSet<FileReference> FoundProjects,
-			bool bForceCompile, out bool bBuildSuccess, CsProjBuildHook Hook, IEnumerable<DirectoryReference> BaseDirectories, IEnumerable<string> DefineConstants,
-			Action<int> OnBuildingProjects, ILogger Logger)
+		/// <param name="hook">Interface to fetch data about the building environment</param>
+		/// <param name="baseDirectories">Base directories of the engine and project</param>
+		/// <param name="defineConstants">Collection of constants to be defined while building projects</param>
+		/// <param name="onBuildingProjects">Action invoked to notify caller regarding the number of projects being built</param>
+		/// <param name="logger">Destination logger</param>
+		private static Dictionary<FileReference, CsProjBuildRecordEntry> BuildInternal(HashSet<FileReference> foundProjects,
+			bool bForceCompile, out bool bBuildSuccess, CsProjBuildHook hook, IEnumerable<DirectoryReference> baseDirectories, IEnumerable<string> defineConstants,
+			Action<int> onBuildingProjects, ILogger logger)
 		{
-			Dictionary<string, string> GlobalProperties = new Dictionary<string, string>
+			Dictionary<string, string> globalProperties = new Dictionary<string, string>
 			{
-				{ "EngineDir", Hook.EngineDirectory.FullName },
+				{ "EngineDir", hook.EngineDirectory.FullName },
 #if DEBUG
 				{ "Configuration", "Debug" },
 #else
@@ -197,53 +196,53 @@ namespace EpicGames.MsBuild
 #endif
 			};
 
-			if (DefineConstants.Any())
+			if (defineConstants.Any())
 			{
-				GlobalProperties.Add("DefineConstants", String.Join(';', DefineConstants));
+				globalProperties.Add("DefineConstants", String.Join(';', defineConstants));
 			}
 
-			Dictionary<FileReference, CsProjBuildRecordEntry> BuildRecords = new();
+			Dictionary<FileReference, CsProjBuildRecordEntry> buildRecords = new();
 
-			using ProjectCollection ProjectCollection = new ProjectCollection(GlobalProperties);
-			Dictionary<string, Project> Projects = new Dictionary<string, Project>();
-			HashSet<string> SkippedProjects = new HashSet<string>();
+			using ProjectCollection projectCollection = new ProjectCollection(globalProperties);
+			Dictionary<string, Project> projects = new Dictionary<string, Project>();
+			HashSet<string> skippedProjects = new HashSet<string>();
 
 			// Microsoft.Build.Evaluation.Project provides access to information stored in the .csproj xml that is 
 			// not available when using Microsoft.Build.Execution.ProjectInstance (used later in this function and
 			// in BuildProjects) - particularly, to access glob information defined in the source file.
 
 			// Load all found projects, and any other referenced projects.
-			foreach (FileReference ProjectPath in FoundProjects)
+			foreach (FileReference projectPath in foundProjects)
 			{
-				void LoadProjectAndReferences(string ProjectPath, string ReferencedBy)
+				void LoadProjectAndReferences(string projectPath, string referencedBy)
 				{
-					ProjectPath = Path.GetFullPath(ProjectPath);
-					if (!Projects.ContainsKey(ProjectPath) && !SkippedProjects.Contains(ProjectPath))
+					projectPath = Path.GetFullPath(projectPath);
+					if (!projects.ContainsKey(projectPath) && !skippedProjects.Contains(projectPath))
 					{
-						Project Project;
+						Project project;
 
 						// Microsoft.Build.Evaluation.Project doesn't give a lot of useful information if this fails,
 						// so make sure to print our own diagnostic info if something goes wrong
 						try
 						{
-							Project = new Project(ProjectPath, GlobalProperties, toolsVersion: null, projectCollection: ProjectCollection);
+							project = new Project(projectPath, globalProperties, toolsVersion: null, projectCollection: projectCollection);
 						}
-						catch (Microsoft.Build.Exceptions.InvalidProjectFileException IPFEx)
+						catch (Microsoft.Build.Exceptions.InvalidProjectFileException iPFEx)
 						{
-							Logger.LogError("Could not load project file {ProjectPath}", ProjectPath);
-							Logger.LogError("{Message}", IPFEx.BaseMessage);
+							logger.LogError("Could not load project file {ProjectPath}", projectPath);
+							logger.LogError("{Message}", iPFEx.BaseMessage);
 
-							if (!String.IsNullOrEmpty(ReferencedBy))
+							if (!String.IsNullOrEmpty(referencedBy))
 							{
-								Logger.LogError("Referenced by: {ReferencedBy}", ReferencedBy);
+								logger.LogError("Referenced by: {ReferencedBy}", referencedBy);
 							}
-							if (Projects.Count > 0)
+							if (projects.Count > 0)
 							{
-								Logger.LogError("See the log file for the list of previously loaded projects.");
-								Logger.LogError("Loaded projects (most recently loaded first):");
-								foreach (string Path in Projects.Keys.Reverse())
+								logger.LogError("See the log file for the list of previously loaded projects.");
+								logger.LogError("Loaded projects (most recently loaded first):");
+								foreach (string path in projects.Keys.Reverse())
 								{
-									Logger.LogError("  {Path}", Path);
+									logger.LogError("  {Path}", path);
 								}
 							}
 							throw;
@@ -253,107 +252,107 @@ namespace EpicGames.MsBuild
 						{
 							// check the TargetFramework of the project: we can't build Windows-only projects on 
 							// non-Windows platforms.
-							if (Project.GetProperty("TargetFramework").EvaluatedValue.Contains("windows", StringComparison.Ordinal))
+							if (project.GetProperty("TargetFramework").EvaluatedValue.Contains("windows", StringComparison.Ordinal))
 							{
-								SkippedProjects.Add(ProjectPath);
-								Logger.LogInformation("Skipping windows-only project {ProjectPath}", ProjectPath);
+								skippedProjects.Add(projectPath);
+								logger.LogInformation("Skipping windows-only project {ProjectPath}", projectPath);
 								return;
 							}
 						}
 
-						Projects.Add(ProjectPath, Project);
-						ReferencedBy = String.IsNullOrEmpty(ReferencedBy) ? ProjectPath : $"{ProjectPath}{Environment.NewLine}{ReferencedBy}";
-						foreach (string ReferencedProject in Project.GetItems("ProjectReference").
-							Select(I => I.EvaluatedInclude))
+						projects.Add(projectPath, project);
+						referencedBy = String.IsNullOrEmpty(referencedBy) ? projectPath : $"{projectPath}{Environment.NewLine}{referencedBy}";
+						foreach (string referencedProject in project.GetItems("ProjectReference").
+							Select(i => i.EvaluatedInclude))
 						{
-							LoadProjectAndReferences(Path.Combine(Project.DirectoryPath, ReferencedProject), ReferencedBy);
+							LoadProjectAndReferences(Path.Combine(project.DirectoryPath, referencedProject), referencedBy);
 						}
 					}
 				}
-				LoadProjectAndReferences(ProjectPath.FullName, null);
+				LoadProjectAndReferences(projectPath.FullName, null);
 			}
 
 			// generate a BuildRecord for each loaded project - the gathered information will be used to determine if the project is
 			// out of date, and if building this project can be skipped. It is also used to populate Intermediate/ScriptModules after the
 			// build completes
-			foreach (Project Project in Projects.Values)
+			foreach (Project project in projects.Values)
 			{
-				string TargetPath = Path.GetRelativePath(Project.DirectoryPath, Project.GetPropertyValue("TargetPath"));
+				string targetPath = Path.GetRelativePath(project.DirectoryPath, project.GetPropertyValue("TargetPath"));
 
-				FileReference ProjectPath = FileReference.FromString(Project.FullPath);
-				FileReference BuildRecordPath = ConstructBuildRecordPath(Hook, ProjectPath, BaseDirectories);
+				FileReference projectPath = FileReference.FromString(project.FullPath);
+				FileReference buildRecordPath = ConstructBuildRecordPath(hook, projectPath, baseDirectories);
 
-				CsProjBuildRecord BuildRecord = new CsProjBuildRecord()
+				CsProjBuildRecord buildRecord = new CsProjBuildRecord()
 				{
 					Version = CsProjBuildRecord.CurrentVersion,
-					TargetPath = TargetPath,
-					TargetBuildTime = Hook.GetLastWriteTime(Project.DirectoryPath, TargetPath),
-					ProjectPath = Path.GetRelativePath(BuildRecordPath.Directory.FullName, Project.FullPath)
+					TargetPath = targetPath,
+					TargetBuildTime = hook.GetLastWriteTime(project.DirectoryPath, targetPath),
+					ProjectPath = Path.GetRelativePath(buildRecordPath.Directory.FullName, project.FullPath)
 				};
 
 				// the .csproj
-				BuildRecord.Dependencies.Add(Path.GetRelativePath(Project.DirectoryPath, Project.FullPath));
+				buildRecord.Dependencies.Add(Path.GetRelativePath(project.DirectoryPath, project.FullPath));
 
 				// Imports: files included in the xml (typically props, targets, etc)
-				foreach (ResolvedImport Import in Project.Imports)
+				foreach (ResolvedImport import in project.Imports)
 				{
-					string ImportPath = Path.GetRelativePath(Project.DirectoryPath, Import.ImportedProject.FullPath);
+					string importPath = Path.GetRelativePath(project.DirectoryPath, import.ImportedProject.FullPath);
 
 					// nuget.g.props and nuget.g.targets are generated by Restore, and are frequently re-written;
 					// it should be safe to ignore these files - changes to references from a .csproj file will
 					// show up as that file being out of date.
-					if (ImportPath.Contains("nuget.g.", StringComparison.Ordinal))
+					if (importPath.Contains("nuget.g.", StringComparison.Ordinal))
 					{
 						continue;
 					}
 
-					BuildRecord.Dependencies.Add(ImportPath);
+					buildRecord.Dependencies.Add(importPath);
 				}
 
 				// References: e.g. Ionic.Zip.Reduced.dll, fastJSON.dll
-				foreach (ProjectItem Item in Project.GetItems("Reference"))
+				foreach (ProjectItem item in project.GetItems("Reference"))
 				{
-					BuildRecord.Dependencies.Add(Item.GetMetadataValue("HintPath"));
+					buildRecord.Dependencies.Add(item.GetMetadataValue("HintPath"));
 				}
 
-				foreach (ProjectItem ReferencedProjectItem in Project.GetItems("ProjectReference"))
+				foreach (ProjectItem referencedProjectItem in project.GetItems("ProjectReference"))
 				{
-					BuildRecord.ProjectReferencesAndTimes.Add(new CsProjBuildRecordRef { ProjectPath = ReferencedProjectItem.EvaluatedInclude });
+					buildRecord.ProjectReferencesAndTimes.Add(new CsProjBuildRecordRef { ProjectPath = referencedProjectItem.EvaluatedInclude });
 				}
 
-				foreach (ProjectItem CompileItem in Project.GetItems("Compile"))
+				foreach (ProjectItem compileItem in project.GetItems("Compile"))
 				{
-					if (Hook.HasWildcards(CompileItem.UnevaluatedInclude))
+					if (hook.HasWildcards(compileItem.UnevaluatedInclude))
 					{
-						BuildRecord.GlobbedDependencies.Add(CompileItem.EvaluatedInclude);
+						buildRecord.GlobbedDependencies.Add(compileItem.EvaluatedInclude);
 					}
 					else
 					{
-						BuildRecord.Dependencies.Add(CompileItem.EvaluatedInclude);
+						buildRecord.Dependencies.Add(compileItem.EvaluatedInclude);
 					}
 				}
 
-				foreach (ProjectItem ContentItem in Project.GetItems("Content"))
+				foreach (ProjectItem contentItem in project.GetItems("Content"))
 				{
-					if (Hook.HasWildcards(ContentItem.UnevaluatedInclude))
+					if (hook.HasWildcards(contentItem.UnevaluatedInclude))
 					{
-						BuildRecord.GlobbedDependencies.Add(ContentItem.EvaluatedInclude);
+						buildRecord.GlobbedDependencies.Add(contentItem.EvaluatedInclude);
 					}
 					else
 					{
-						BuildRecord.Dependencies.Add(ContentItem.EvaluatedInclude);
+						buildRecord.Dependencies.Add(contentItem.EvaluatedInclude);
 					}
 				}
 
-				foreach (ProjectItem EmbeddedResourceItem in Project.GetItems("EmbeddedResource"))
+				foreach (ProjectItem embeddedResourceItem in project.GetItems("EmbeddedResource"))
 				{
-					if (Hook.HasWildcards(EmbeddedResourceItem.UnevaluatedInclude))
+					if (hook.HasWildcards(embeddedResourceItem.UnevaluatedInclude))
 					{
-						BuildRecord.GlobbedDependencies.Add(EmbeddedResourceItem.EvaluatedInclude);
+						buildRecord.GlobbedDependencies.Add(embeddedResourceItem.EvaluatedInclude);
 					}
 					else
 					{
-						BuildRecord.Dependencies.Add(EmbeddedResourceItem.EvaluatedInclude);
+						buildRecord.Dependencies.Add(embeddedResourceItem.EvaluatedInclude);
 					}
 				}
 
@@ -364,101 +363,101 @@ namespace EpicGames.MsBuild
 				// This also returns a lot more information than we care for - MSBuildGlob objects,
 				// which have a range of precomputed values. It may be possible to take source for
 				// GetAllGlobs() and construct a version that does less.
-				List<GlobResult> Globs = Project.GetAllGlobs();
+				List<GlobResult> globs = project.GetAllGlobs();
 
 				// FileMatcher.IsMatch() requires directory separators in glob strings to match the
 				// local flavor. There's probably a better way.
-				string CleanGlobString(string GlobString)
+				string CleanGlobString(string globString)
 				{
-					char Sep = Path.DirectorySeparatorChar;
-					char NotSep = Sep == '/' ? '\\' : '/'; // AltDirectorySeparatorChar isn't always what we need (it's '/' on Mac)
+					char sep = Path.DirectorySeparatorChar;
+					char notSep = sep == '/' ? '\\' : '/'; // AltDirectorySeparatorChar isn't always what we need (it's '/' on Mac)
 
-					char[] Chars = GlobString.ToCharArray();
-					int P = 0;
-					for (int I = 0; I < GlobString.Length; ++I, ++P)
+					char[] chars = globString.ToCharArray();
+					int p = 0;
+					for (int i = 0; i < globString.Length; ++i, ++p)
 					{
 						// Flip a non-native separator
-						if (Chars[I] == NotSep)
+						if (chars[i] == notSep)
 						{
-							Chars[P] = Sep;
+							chars[p] = sep;
 						}
 						else
 						{
-							Chars[P] = Chars[I];
+							chars[p] = chars[i];
 						}
 
 						// Collapse adjacent separators
-						if (I > 0 && Chars[P] == Sep && Chars[P - 1] == Sep)
+						if (i > 0 && chars[p] == sep && chars[p - 1] == sep)
 						{
-							P -= 1;
+							p -= 1;
 						}
 					}
 
-					return new string(Chars, 0, P);
+					return new string(chars, 0, p);
 				}
 
-				foreach (GlobResult Glob in Globs)
+				foreach (GlobResult glob in globs)
 				{
-					if (String.Equals("None", Glob.ItemElement.ItemType, StringComparison.Ordinal))
+					if (String.Equals("None", glob.ItemElement.ItemType, StringComparison.Ordinal))
 					{
 						// don't record the default "None" glob - it's not (?) a trigger for any rebuild
 						continue;
 					}
 
-					List<string> Include = new List<string>(Glob.IncludeGlobs.Select(F => CleanGlobString(F))).OrderBy(x => x).ToList();
-					List<string> Exclude = new List<string>(Glob.Excludes.Select(F => CleanGlobString(F))).OrderBy(x => x).ToList();
-					List<string> Remove = new List<string>(Glob.Removes.Select(F => CleanGlobString(F))).OrderBy(x => x).ToList();
+					List<string> include = new List<string>(glob.IncludeGlobs.Select(f => CleanGlobString(f))).OrderBy(x => x).ToList();
+					List<string> exclude = new List<string>(glob.Excludes.Select(f => CleanGlobString(f))).OrderBy(x => x).ToList();
+					List<string> remove = new List<string>(glob.Removes.Select(f => CleanGlobString(f))).OrderBy(x => x).ToList();
 
-					BuildRecord.Globs.Add(new CsProjBuildRecord.Glob()
+					buildRecord.Globs.Add(new CsProjBuildRecord.Glob()
 					{
-						ItemType = Glob.ItemElement.ItemType,
-						Include = Include,
-						Exclude = Exclude,
-						Remove = Remove
+						ItemType = glob.ItemElement.ItemType,
+						Include = include,
+						Exclude = exclude,
+						Remove = remove
 					});
 				}
 
-				CsProjBuildRecordEntry Entry = new CsProjBuildRecordEntry(ProjectPath, BuildRecordPath, BuildRecord);
-				BuildRecords.Add(Entry.ProjectFile, Entry);
+				CsProjBuildRecordEntry entry = new CsProjBuildRecordEntry(projectPath, buildRecordPath, buildRecord);
+				buildRecords.Add(entry.ProjectFile, entry);
 			}
 
 			// Potential optimization: Constructing the ProjectGraph here gives the full graph of dependencies - which is nice,
 			// but not strictly necessary, and slower than doing it some other way.
-			ProjectGraph InputProjectGraph;
-			InputProjectGraph = new ProjectGraph(FoundProjects
+			ProjectGraph inputProjectGraph;
+			inputProjectGraph = new ProjectGraph(foundProjects
 				// Build the graph without anything that can't be built on this platform
-				.Where(x => !SkippedProjects.Contains(x.FullName))
-				.Select(P => P.FullName), GlobalProperties, ProjectCollection);
+				.Where(x => !skippedProjects.Contains(x.FullName))
+				.Select(p => p.FullName), globalProperties, projectCollection);
 
 			// A ProjectGraph that will represent the set of projects that we actually want to build
-			ProjectGraph BuildProjectGraph = null;
+			ProjectGraph buildProjectGraph = null;
 
 			if (bForceCompile)
 			{
-				Logger.LogDebug("Script modules will build: '-Compile' on command line");
-				BuildProjectGraph = InputProjectGraph;
+				logger.LogDebug("Script modules will build: '-Compile' on command line");
+				buildProjectGraph = inputProjectGraph;
 			}
 			else
 			{
-				foreach (ProjectGraphNode Project in InputProjectGraph.ProjectNodesTopologicallySorted)
+				foreach (ProjectGraphNode project in inputProjectGraph.ProjectNodesTopologicallySorted)
 				{
-					Hook.ValidateRecursively(BuildRecords, FileReference.FromString(Project.ProjectInstance.FullPath));
+					hook.ValidateRecursively(buildRecords, FileReference.FromString(project.ProjectInstance.FullPath));
 				}
 
 				// Select the projects that have been found to be out of date
-				Dictionary<FileReference, CsProjBuildRecordEntry> InvalidBuildRecords = new(BuildRecords.Where(x => x.Value.Status == CsProjBuildRecordStatus.Invalid));
-				HashSet<ProjectGraphNode> OutOfDateProjects = new HashSet<ProjectGraphNode>(InputProjectGraph.ProjectNodes.Where(x => InvalidBuildRecords.ContainsKey(FileReference.FromString(x.ProjectInstance.FullPath))));
+				Dictionary<FileReference, CsProjBuildRecordEntry> invalidBuildRecords = new(buildRecords.Where(x => x.Value.Status == CsProjBuildRecordStatus.Invalid));
+				HashSet<ProjectGraphNode> outOfDateProjects = new HashSet<ProjectGraphNode>(inputProjectGraph.ProjectNodes.Where(x => invalidBuildRecords.ContainsKey(FileReference.FromString(x.ProjectInstance.FullPath))));
 
-				if (OutOfDateProjects.Count > 0)
+				if (outOfDateProjects.Count > 0)
 				{
-					BuildProjectGraph = new ProjectGraph(OutOfDateProjects.Select(P => P.ProjectInstance.FullPath), GlobalProperties, ProjectCollection);
+					buildProjectGraph = new ProjectGraph(outOfDateProjects.Select(p => p.ProjectInstance.FullPath), globalProperties, projectCollection);
 				}
 			}
 
-			if (BuildProjectGraph != null)
+			if (buildProjectGraph != null)
 			{
-				OnBuildingProjects(BuildProjectGraph.EntryPointNodes.Count);
-				bBuildSuccess = BuildProjects(BuildProjectGraph, GlobalProperties, Logger);
+				onBuildingProjects(buildProjectGraph.EntryPointNodes.Count);
+				bBuildSuccess = BuildProjects(buildProjectGraph, globalProperties, logger);
 			}
 			else
 			{
@@ -466,153 +465,153 @@ namespace EpicGames.MsBuild
 			}
 
 			// Update the target times
-			foreach (ProjectGraphNode ProjectNode in InputProjectGraph.ProjectNodes)
+			foreach (ProjectGraphNode projectNode in inputProjectGraph.ProjectNodes)
 			{
-				FileReference ProjectPath = FileReference.FromString(ProjectNode.ProjectInstance.FullPath);
-				CsProjBuildRecordEntry Entry = BuildRecords[ProjectPath];
-				FileReference FullPath = FileReference.Combine(ProjectPath.Directory, Entry.BuildRecord.TargetPath);
-				Entry.BuildRecord.TargetBuildTime = FileReference.GetLastWriteTime(FullPath);
+				FileReference projectPath = FileReference.FromString(projectNode.ProjectInstance.FullPath);
+				CsProjBuildRecordEntry entry = buildRecords[projectPath];
+				FileReference fullPath = FileReference.Combine(projectPath.Directory, entry.BuildRecord.TargetPath);
+				entry.BuildRecord.TargetBuildTime = FileReference.GetLastWriteTime(fullPath);
 			}
 
 			// Update the project reference target times
-			foreach (ProjectGraphNode ProjectNode in InputProjectGraph.ProjectNodes)
+			foreach (ProjectGraphNode projectNode in inputProjectGraph.ProjectNodes)
 			{
-				FileReference ProjectPath = FileReference.FromString(ProjectNode.ProjectInstance.FullPath);
-				CsProjBuildRecordEntry Entry = BuildRecords[ProjectPath];
-				foreach (CsProjBuildRecordRef ReferencedProject in Entry.BuildRecord.ProjectReferencesAndTimes)
+				FileReference projectPath = FileReference.FromString(projectNode.ProjectInstance.FullPath);
+				CsProjBuildRecordEntry entry = buildRecords[projectPath];
+				foreach (CsProjBuildRecordRef referencedProject in entry.BuildRecord.ProjectReferencesAndTimes)
 				{
-					FileReference RefProjectPath = FileReference.FromString(Path.GetFullPath(ReferencedProject.ProjectPath, ProjectPath.Directory.FullName));
-					if (BuildRecords.TryGetValue(RefProjectPath, out CsProjBuildRecordEntry RefEntry))
+					FileReference refProjectPath = FileReference.FromString(Path.GetFullPath(referencedProject.ProjectPath, projectPath.Directory.FullName));
+					if (buildRecords.TryGetValue(refProjectPath, out CsProjBuildRecordEntry refEntry))
 					{
-						ReferencedProject.TargetBuildTime = RefEntry.BuildRecord.TargetBuildTime;
+						referencedProject.TargetBuildTime = refEntry.BuildRecord.TargetBuildTime;
 					}
 				}
 			}
 
 			// write all build records
-			foreach (ProjectGraphNode ProjectNode in InputProjectGraph.ProjectNodes)
+			foreach (ProjectGraphNode projectNode in inputProjectGraph.ProjectNodes)
 			{
-				FileReference ProjectPath = FileReference.FromString(ProjectNode.ProjectInstance.FullPath);
-				CsProjBuildRecordEntry Entry = BuildRecords[ProjectPath];
-				if (FileReference.WriteAllTextIfDifferent(Entry.BuildRecordFile,
-					JsonSerializer.Serialize<CsProjBuildRecord>(Entry.BuildRecord, new JsonSerializerOptions { WriteIndented = true })))
+				FileReference projectPath = FileReference.FromString(projectNode.ProjectInstance.FullPath);
+				CsProjBuildRecordEntry entry = buildRecords[projectPath];
+				if (FileReference.WriteAllTextIfDifferent(entry.BuildRecordFile,
+					JsonSerializer.Serialize(entry.BuildRecord, new JsonSerializerOptions { WriteIndented = true })))
 				{
-					Logger.LogDebug("Wrote script module build record to {BuildRecordPath}", Entry.BuildRecordFile);
+					logger.LogDebug("Wrote script module build record to {BuildRecordPath}", entry.BuildRecordFile);
 				}
 			}
 
 			// todo: re-verify build records after a build to verify that everything is actually up to date
 
 			// even if only a subset was built, this function returns the full list of target assembly paths
-			Dictionary<FileReference, CsProjBuildRecordEntry> OutDict = new();
-			foreach (ProjectGraphNode EntryPointNode in InputProjectGraph.EntryPointNodes)
+			Dictionary<FileReference, CsProjBuildRecordEntry> outDict = new();
+			foreach (ProjectGraphNode entryPointNode in inputProjectGraph.EntryPointNodes)
 			{
-				FileReference ProjectPath = FileReference.FromString(EntryPointNode.ProjectInstance.FullPath);
-				OutDict.Add(ProjectPath, BuildRecords[ProjectPath]);
+				FileReference projectPath = FileReference.FromString(entryPointNode.ProjectInstance.FullPath);
+				outDict.Add(projectPath, buildRecords[projectPath]);
 			}
-			return OutDict;
+			return outDict;
 		}
 
-		private static bool BuildProjects(ProjectGraph ProjectGraph, Dictionary<string, string> GlobalProperties, ILogger Logger)
+		private static bool BuildProjects(ProjectGraph projectGraph, Dictionary<string, string> globalProperties, ILogger logger)
 		{
-			DateTime StartTime = DateTime.UtcNow;
-			MLogger BuildLogger = new MLogger(Logger);
+			DateTime startTime = DateTime.UtcNow;
+			MLogger buildLogger = new MLogger(logger);
 
-			string[] TargetsToBuild = { "Restore", "Build" };
+			string[] targetsToBuild = { "Restore", "Build" };
 
-			bool Result = true;
+			bool result = true;
 
-			foreach (string TargetToBuild in TargetsToBuild)
+			foreach (string targetToBuild in targetsToBuild)
 			{
-				GraphBuildRequestData GraphRequest = new GraphBuildRequestData(ProjectGraph, new string[] { TargetToBuild });
+				GraphBuildRequestData graphRequest = new GraphBuildRequestData(projectGraph, new string[] { targetToBuild });
 
-				BuildManager BuildMan = BuildManager.DefaultBuildManager;
+				BuildManager buildMan = BuildManager.DefaultBuildManager;
 
-				BuildParameters BuildParameters = new BuildParameters();
-				BuildParameters.AllowFailureWithoutError = false;
-				BuildParameters.DetailedSummary = true;
+				BuildParameters buildParameters = new BuildParameters();
+				buildParameters.AllowFailureWithoutError = false;
+				buildParameters.DetailedSummary = true;
 
-				BuildParameters.Loggers = new List<IBuildLogger> { BuildLogger };
-				BuildParameters.MaxNodeCount = 1; // msbuild bug - more than 1 here and the build stalls. Likely related to [REDACTED_PROFILE_URL]/msbuild/issues/1941
+				buildParameters.Loggers = new List<IBuildLogger> { buildLogger };
+				buildParameters.MaxNodeCount = 1; // msbuild bug - more than 1 here and the build stalls. Likely related to [REDACTED_PROFILE_URL]/msbuild/issues/1941
 
-				BuildParameters.OnlyLogCriticalEvents = false;
-				BuildParameters.ShutdownInProcNodeOnBuildFinish = false;
+				buildParameters.OnlyLogCriticalEvents = false;
+				buildParameters.ShutdownInProcNodeOnBuildFinish = false;
 
-				BuildParameters.GlobalProperties = GlobalProperties;
+				buildParameters.GlobalProperties = globalProperties;
 
-				Logger.LogInformation(" {TargetToBuild}...", TargetToBuild);
+				logger.LogInformation(" {TargetToBuild}...", targetToBuild);
 
-				GraphBuildResult BuildResult = BuildMan.Build(BuildParameters, GraphRequest);
+				GraphBuildResult buildResult = buildMan.Build(buildParameters, graphRequest);
 
-				if (BuildResult.OverallResult == BuildResultCode.Failure)
+				if (buildResult.OverallResult == BuildResultCode.Failure)
 				{
-					Logger.LogInformation("");
-					foreach (KeyValuePair<ProjectGraphNode, BuildResult> NodeResult in BuildResult.ResultsByNode)
+					logger.LogInformation("");
+					foreach (KeyValuePair<ProjectGraphNode, BuildResult> nodeResult in buildResult.ResultsByNode)
 					{
-						if (NodeResult.Value.OverallResult == BuildResultCode.Failure)
+						if (nodeResult.Value.OverallResult == BuildResultCode.Failure)
 						{
-							Logger.LogError("  Failed to build: {ProjectPath}", new FileReference(NodeResult.Key.ProjectInstance.FullPath));
+							logger.LogError("  Failed to build: {ProjectPath}", new FileReference(nodeResult.Key.ProjectInstance.FullPath));
 						}
 					}
-					Result = false;
+					result = false;
 				}
 			}
-			Logger.LogInformation("Build projects time: {TimeSeconds:0.00} s", (DateTime.UtcNow - StartTime).TotalMilliseconds / 1000);
+			logger.LogInformation("Build projects time: {TimeSeconds:0.00} s", (DateTime.UtcNow - startTime).TotalMilliseconds / 1000);
 
-			return Result;
+			return result;
 		}
 
-		static bool _hasRegiteredMsBuildPath = false;
+		static bool s_hasRegiteredMsBuildPath = false;
 
 		/// <summary>
 		/// Register our bundled dotnet installation to be used by Microsoft.Build
 		/// This needs to happen in a function called before the first use of any Microsoft.Build types
 		/// </summary>
-		public static void RegisterMsBuildPath(CsProjBuildHook Hook)
+		public static void RegisterMsBuildPath(CsProjBuildHook hook)
 		{
-			if (_hasRegiteredMsBuildPath)
+			if (s_hasRegiteredMsBuildPath)
 			{
 				return;
 			}
-			_hasRegiteredMsBuildPath = true;
+			s_hasRegiteredMsBuildPath = true;
 
 			// Find our bundled dotnet SDK
-			List<string> ListOfSdks = new List<string>();
-			ProcessStartInfo StartInfo = new ProcessStartInfo
+			List<string> listOfSdks = new List<string>();
+			ProcessStartInfo startInfo = new ProcessStartInfo
 			{
-				FileName = Hook.DotnetPath.FullName,
+				FileName = hook.DotnetPath.FullName,
 				RedirectStandardOutput = true,
 				UseShellExecute = false,
 				ArgumentList = { "--list-sdks" }
 			};
-			StartInfo.EnvironmentVariables["DOTNET_MULTILEVEL_LOOKUP"] = "0"; // use only the bundled dotnet installation - ignore any other/system dotnet install
+			startInfo.EnvironmentVariables["DOTNET_MULTILEVEL_LOOKUP"] = "0"; // use only the bundled dotnet installation - ignore any other/system dotnet install
 
-			Process DotnetProcess = Process.Start(StartInfo);
+			Process dotnetProcess = Process.Start(startInfo);
 			{
-				string Line;
-				while ((Line = DotnetProcess.StandardOutput.ReadLine()) != null)
+				string line;
+				while ((line = dotnetProcess.StandardOutput.ReadLine()) != null)
 				{
-					ListOfSdks.Add(Line);
+					listOfSdks.Add(line);
 				}
 			}
-			DotnetProcess.WaitForExit();
+			dotnetProcess.WaitForExit();
 
-			if (ListOfSdks.Count != 1)
+			if (listOfSdks.Count != 1)
 			{
 				throw new Exception("Expected only one sdk installed for bundled dotnet");
 			}
 
 			// Expected output has this form:
 			// 3.1.403 [D:\UE5_Main\engine\binaries\ThirdParty\DotNet\Windows\sdk]
-			string SdkVersion = ListOfSdks[0].Split(' ')[0];
+			string sdkVersion = listOfSdks[0].Split(' ')[0];
 
-			DirectoryReference DotnetSdkDirectory = DirectoryReference.Combine(Hook.DotnetDirectory, "sdk", SdkVersion);
-			if (!DirectoryReference.Exists(DotnetSdkDirectory))
+			DirectoryReference dotnetSdkDirectory = DirectoryReference.Combine(hook.DotnetDirectory, "sdk", sdkVersion);
+			if (!DirectoryReference.Exists(dotnetSdkDirectory))
 			{
-				throw new Exception("Failed to find .NET SDK directory: " + DotnetSdkDirectory.FullName);
+				throw new Exception("Failed to find .NET SDK directory: " + dotnetSdkDirectory.FullName);
 			}
 
-			MSBuildLocator.RegisterMSBuildPath(DotnetSdkDirectory.FullName);
+			MSBuildLocator.RegisterMSBuildPath(dotnetSdkDirectory.FullName);
 		}
 	}
 }
