Horde: Fix using auth tokens from environment variables when the server URL is also derived from environment variables.

diff --git a/Engine/Source/Programs/Shared/EpicGames.Horde/HordeHttpAuthHandler.cs b/Engine/Source/Programs/Shared/EpicGames.Horde/HordeHttpAuthHandler.cs
--- a/Engine/Source/Programs/Shared/EpicGames.Horde/HordeHttpAuthHandler.cs
+++ b/Engine/Source/Programs/Shared/EpicGames.Horde/HordeHttpAuthHandler.cs
@@ -213,7 +213,7 @@ namespace EpicGames.Horde
 			if (!String.IsNullOrEmpty(hordeUrlEnvVar))
 			{
 				Uri hordeUrl = new Uri(hordeUrlEnvVar);
-				if (String.Equals(_options.Value.ServerUrl?.Host, hordeUrl.Host, StringComparison.OrdinalIgnoreCase))
+				if (_options.Value.ServerUrl == null || String.Equals(_options.Value.ServerUrl.Host, hordeUrl.Host, StringComparison.OrdinalIgnoreCase))
 				{
 					string? hordeToken = Environment.GetEnvironmentVariable(HordeHttpClient.HordeTokenEnvVarName);
 					if (!String.IsNullOrEmpty(hordeToken))
diff --git a/Engine/Source/Programs/Shared/EpicGames.Horde/HordeHttpClient.cs b/Engine/Source/Programs/Shared/EpicGames.Horde/HordeHttpClient.cs
--- a/Engine/Source/Programs/Shared/EpicGames.Horde/HordeHttpClient.cs
+++ b/Engine/Source/Programs/Shared/EpicGames.Horde/HordeHttpClient.cs
@@ -468,7 +468,7 @@ namespace EpicGames.Horde
 				if (!response.IsSuccessStatusCode)
 				{
 					string body = await response.Content.ReadAsStringAsync(cancellationToken);
-					throw new HttpRequestException($"{response.StatusCode} posting to {new Uri(httpClient.BaseAddress!, relativePath)}: {body}", null, response.StatusCode);
+					throw new HttpRequestException($"{(int)response.StatusCode} ({response.StatusCode}) posting to {new Uri(httpClient.BaseAddress!, relativePath)}: {body}", null, response.StatusCode);
 				}
 
 				TResponse? responseValue = await response.Content.ReadFromJsonAsync<TResponse>(s_jsonSerializerOptions, cancellationToken);
@@ -612,11 +612,10 @@ namespace EpicGames.Horde
 
 			// Create the HTTP client for handling Horde requests
 			IHttpClientBuilder builder = services.AddHttpClient<HordeHttpClient>(HordeHttpClient.HttpClientName, ConfigureClientFromEnvironment)
+				.AddHttpMessageHandler<HordeHttpAuthHandler>()
 				.AddPolicyHandler((serviceProvider, request) => CreateDefaultTimeoutRetryPolicy(request, serviceProvider.GetRequiredService<ILogger<HttpStorageBackend>>()))
 				.AddPolicyHandler((serviceProvider, request) => CreateDefaultTransientErrorPolicy(request, serviceProvider.GetRequiredService<ILogger<HttpStorageBackend>>()));
 
-			builder = builder.AddHttpMessageHandler<HordeHttpAuthHandler>();
-
 			return builder;
 		}
 
