Fix UnrealObjectPtrTool by undoing changelist 27536467 and addressing the original issue in a different way.  Also fixed the UnrealObjectPtrTools expectations for formatting to ma tch what UHT produces within UBT today and ensured that newlines and tabs were kept (but escaped) in the UHT logging so that UnrealObjectPtrTool could do replacement of the verbatim text.

diff --git a/Engine/Source/Programs/Shared/EpicGames.UHT/Parsers/UhtPropertyParser.cs b/Engine/Source/Programs/Shared/EpicGames.UHT/Parsers/UhtPropertyParser.cs
--- a/Engine/Source/Programs/Shared/EpicGames.UHT/Parsers/UhtPropertyParser.cs
+++ b/Engine/Source/Programs/Shared/EpicGames.UHT/Parsers/UhtPropertyParser.cs
@@ -1313,6 +1313,7 @@ namespace EpicGames.UHT.Parsers
 		{
 			UhtProperty? property;
 			UhtSession session = propertySettings.Outer.Session;
+			int typeStartPos = tokenReader.PeekToken().InputStartPos;
 
 			bool gotConst = tokenReader.TryOptional("const");
 
@@ -1388,6 +1389,10 @@ namespace EpicGames.UHT.Parsers
 				}
 				tokenReader.Require('*');
 
+				// Optionally emit messages about native pointer members and swallow trailing 'const' after pointer properties
+				UhtObjectPropertyBase.ConditionalLogPointerUsage(propertySettings, session.Config!.EngineNativePointerMemberBehavior,
+					session.Config!.EnginePluginNativePointerMemberBehavior, session.Config!.NonEngineNativePointerMemberBehavior, "Native pointer", tokenReader, typeStartPos, "TObjectPtr");
+
 				if (propertySettings.PropertyCategory == UhtPropertyCategory.Member)
 				{
 					tokenReader.TryOptional("const");
diff --git a/Engine/Source/Programs/Shared/EpicGames.UHT/Types/Properties/UhtObjectProperty.cs b/Engine/Source/Programs/Shared/EpicGames.UHT/Types/Properties/UhtObjectProperty.cs
--- a/Engine/Source/Programs/Shared/EpicGames.UHT/Types/Properties/UhtObjectProperty.cs
+++ b/Engine/Source/Programs/Shared/EpicGames.UHT/Types/Properties/UhtObjectProperty.cs
@@ -103,17 +103,6 @@ namespace EpicGames.UHT.Types
 			return builder;
 		}
 
-		/// <inheritdoc/>
-		protected override void ValidateMember(UhtStruct structObj, UhtValidationOptions options)
-		{
-			base.ValidateMember(structObj, options);
-
-			if(PointerType == UhtPointerType.Native)
-			{
-				ConditionalLogPointerUsage(Session.Config!.EngineNativePointerMemberBehavior, Session.Config!.EnginePluginNativePointerMemberBehavior, Session.Config!.NonEngineNativePointerMemberBehavior, "Native pointer", "TObjectPtr");
-			}	
-		}
-
 		/// <inheritdoc/>
 		public override bool IsSameType(UhtProperty other)
 		{
diff --git a/Engine/Source/Programs/Shared/EpicGames.UHT/Types/Properties/UhtObjectPropertyBase.cs b/Engine/Source/Programs/Shared/EpicGames.UHT/Types/Properties/UhtObjectPropertyBase.cs
--- a/Engine/Source/Programs/Shared/EpicGames.UHT/Types/Properties/UhtObjectPropertyBase.cs
+++ b/Engine/Source/Programs/Shared/EpicGames.UHT/Types/Properties/UhtObjectPropertyBase.cs
@@ -224,21 +224,24 @@ namespace EpicGames.UHT.Types
 		/// <summary>
 		/// Logs message for Object pointers to convert UObject* to TObjectPtr or the reverse
 		/// </summary>
+		/// <param name="propertySettings">Property settings</param>
 		/// <param name="engineBehavior">Expected behavior for engine types</param>
 		/// <param name="enginePluginBehavior">Expected behavior for engine plugin types</param>
 		/// <param name="nonEngineBehavior">Expected behavior for non-engine types</param>
 		/// <param name="pointerTypeDesc">Description of the pointer type</param>
+		/// <param name="tokenReader">Token reader for type being parsed</param>
+		/// <param name="typeStartPos">Starting character position of the type</param>
 		/// <param name="alternativeTypeDesc">Suggested alternate declaration</param>
 		/// <exception cref="UhtIceException">Thrown if the behavior type is unexpected</exception>
-		protected void ConditionalLogPointerUsage(UhtIssueBehavior engineBehavior, UhtIssueBehavior enginePluginBehavior,
-			UhtIssueBehavior nonEngineBehavior, string pointerTypeDesc, string? alternativeTypeDesc)
+		public static void ConditionalLogPointerUsage(UhtPropertySettings propertySettings, UhtIssueBehavior engineBehavior, UhtIssueBehavior enginePluginBehavior,
+			UhtIssueBehavior nonEngineBehavior, string pointerTypeDesc, IUhtTokenReader tokenReader, int typeStartPos, string? alternativeTypeDesc)
 		{
-			UhtPackage? package = Outer?.Package;
-			if(Outer == null || package == null)
+			if (propertySettings.PropertyCategory != UhtPropertyCategory.Member)
 			{
 				return;
 			}
 
+			UhtPackage package = propertySettings.Outer.Package;
 			UhtIssueBehavior behavior = nonEngineBehavior;
 			if (package.IsPartOfEngine)
 			{
@@ -257,27 +260,32 @@ namespace EpicGames.UHT.Types
 				return;
 			}
 
+			string type = tokenReader.GetStringView(typeStartPos, tokenReader.InputPos - typeStartPos).ToString();
+			type = type.Replace("\n", "\\n", StringComparison.Ordinal);
+			type = type.Replace("\r", "\\r", StringComparison.Ordinal);
+			type = type.Replace("\t", "\\t", StringComparison.Ordinal);
+
 			switch (behavior)
 			{
 				case UhtIssueBehavior.Disallow:
 					if (!String.IsNullOrEmpty(alternativeTypeDesc))
 					{
-						this.LogError($"{pointerTypeDesc} usage in member declaration detected [[[{Outer.EngineName}]]].  This is disallowed for the target/module, consider {alternativeTypeDesc} as an alternative.");
+						tokenReader.LogError($"{pointerTypeDesc} usage in member declaration detected [[[{type}]]].  This is disallowed for the target/module, consider {alternativeTypeDesc} as an alternative.");
 					}
 					else
 					{
-						this.LogError($"{pointerTypeDesc} usage in member declaration detected [[[{Outer.EngineName}]]].");
+						tokenReader.LogError($"{pointerTypeDesc} usage in member declaration detected [[[{type}]]].");
 					}
 					break;
 
 				case UhtIssueBehavior.AllowAndLog:
 					if (!String.IsNullOrEmpty(alternativeTypeDesc))
 					{
-						this.LogInfo($"{pointerTypeDesc} usage in member declaration detected [[[{Outer.EngineName}]]].  Consider {alternativeTypeDesc} as an alternative.");
+						tokenReader.LogInfo($"{pointerTypeDesc} usage in member declaration detected [[[{type}]]].  Consider {alternativeTypeDesc} as an alternative.");
 					}
 					else
 					{
-						this.LogInfo($"{pointerTypeDesc} usage in member declaration detected [[[{Outer.EngineName}]]].");
+						tokenReader.LogInfo("{PointerTypeDesc} usage in member declaration detected [[[{Type}]]].");
 					}
 					break;
 
diff --git a/Engine/Source/Programs/Shared/EpicGames.UHT/Types/Properties/UhtObjectPtrProperty.cs b/Engine/Source/Programs/Shared/EpicGames.UHT/Types/Properties/UhtObjectPtrProperty.cs
--- a/Engine/Source/Programs/Shared/EpicGames.UHT/Types/Properties/UhtObjectPtrProperty.cs
+++ b/Engine/Source/Programs/Shared/EpicGames.UHT/Types/Properties/UhtObjectPtrProperty.cs
@@ -116,16 +116,10 @@ namespace EpicGames.UHT.Types
 			{
 				this.LogError($"UPROPERTY pointers cannot be interfaces - did you mean TScriptInterface<{Class.SourceName}>?");
 			}
-			else
-			{
-				// Use to reverse migration from TObjectPtr<UObject> Member; back to UObject* Member;
-				ConditionalLogPointerUsage(Session.Config!.EngineObjectPtrMemberBehavior, Session.Config!.EnginePluginObjectPtrMemberBehavior,
-				Session.Config!.NonEngineObjectPtrMemberBehavior, "ObjectPtr", null);
-			}
 		}
 
-			#region Keyword
-			[UhtPropertyType(Keyword = "TObjectPtr")]
+		#region Keyword
+		[UhtPropertyType(Keyword = "TObjectPtr")]
 		[SuppressMessage("CodeQuality", "IDE0051:Remove unused private members", Justification = "Attribute accessed method")]
 		[SuppressMessage("Style", "IDE0060:Remove unused parameter", Justification = "Attribute accessed method")]
 		private static UhtProperty? ObjectPtrProperty(UhtPropertyResolvePhase resolvePhase, UhtPropertySettings propertySettings, IUhtTokenReader tokenReader, UhtToken matchedToken)
@@ -139,6 +133,9 @@ namespace EpicGames.UHT.Types
 				return null;
 			}
 
+			ConditionalLogPointerUsage(propertySettings, session.Config!.EngineObjectPtrMemberBehavior, session.Config!.EnginePluginObjectPtrMemberBehavior,
+				session.Config!.NonEngineObjectPtrMemberBehavior, "ObjectPtr", tokenReader, typeStartPos, null);
+
 			if (propertyClass.IsChildOf(propertyClass.Session.UClass))
 			{
 				// UObject specifies that there is no limiter
diff --git a/Engine/Source/Programs/UnrealObjectPtrTool/Private/UnrealObjectPtrTool.cpp b/Engine/Source/Programs/UnrealObjectPtrTool/Private/UnrealObjectPtrTool.cpp
--- a/Engine/Source/Programs/UnrealObjectPtrTool/Private/UnrealObjectPtrTool.cpp
+++ b/Engine/Source/Programs/UnrealObjectPtrTool/Private/UnrealObjectPtrTool.cpp
@@ -90,12 +90,18 @@ private:
 
 FStringView ParseBracketedStringSegment(const FString& InString, int32& InOutCurrentIndex, const TCHAR* BracketStart, const TCHAR* BracketEnd)
 {
-	int32 ParsedStartIndex = InString.Find(BracketStart, ESearchCase::IgnoreCase, ESearchDir::FromStart, InOutCurrentIndex);
-	if (ParsedStartIndex == INDEX_NONE)
+	int32 ParsedStartIndex = 0;
+
+	if (BracketStart != nullptr)
 	{
-		return FStringView();
+		ParsedStartIndex = InString.Find(BracketStart, ESearchCase::IgnoreCase, ESearchDir::FromStart, InOutCurrentIndex);
+		if (ParsedStartIndex == INDEX_NONE)
+		{
+			return FStringView();
+		}
+		ParsedStartIndex += FCString::Strlen(BracketStart);
 	}
-	ParsedStartIndex += FCString::Strlen(BracketStart);
+
 	int32 ParsedEndIndex = InString.Find(BracketEnd, ESearchCase::IgnoreCase, ESearchDir::FromStart, ParsedStartIndex);
 	if (ParsedEndIndex == INDEX_NONE)
 	{
@@ -110,7 +116,7 @@ TValueOrError<FPointerUpgrader::FPointerUpgradeList, FString> FPointerUpgrader::
 	FPointerUpgradeList UpgradeList;
 	const TCHAR* Preamble = EnumHasAllFlags(BehaviorFlags, EPointerUpgradeBehaviorFlags::Reverse) ? TEXT("ObjectPtr usage in member declaration detected [[") : TEXT("Native pointer usage in member declaration detected [[");
 	TArray<FString> PointerUpgradeEntries;
-	if (!FFileHelper::LoadFileToStringArrayWithPredicate(PointerUpgradeEntries, *LogFilename, [&Preamble](const FString& Line) { return Line.Contains("LogCompile: ") && Line.Contains(Preamble);  }))
+	if (!FFileHelper::LoadFileToStringArrayWithPredicate(PointerUpgradeEntries, *LogFilename, [&Preamble](const FString& Line) { return Line.Contains("Info: ") && Line.Contains(Preamble);  }))
 	{
 		return MakeError(FString::Printf(TEXT("Unable to load UHT log: %s"), *LogFilename));
 	}
@@ -119,7 +125,7 @@ TValueOrError<FPointerUpgrader::FPointerUpgradeList, FString> FPointerUpgrader::
 	for (const FString& Entry : PointerUpgradeEntries)
 	{
 		int32 CurrentIndex = 0;
-		FStringView FilenameView = ParseBracketedStringSegment(Entry, CurrentIndex, TEXT("LogCompile: "), TEXT("("));
+		FStringView FilenameView = ParseBracketedStringSegment(Entry, CurrentIndex, nullptr, TEXT("("));
 		if (FilenameView.IsEmpty())
 		{
 			return MakeError(FString::Printf(TEXT("Failed to parse filename from pointer member upgrade statement: %s"), *Entry));
