UGS: Link to the Horde dashboard from the status panel.

#rnx

diff --git a/Engine/Source/Programs/Shared/EpicGames.Horde/HordeClient.cs b/Engine/Source/Programs/Shared/EpicGames.Horde/HordeClient.cs
--- a/Engine/Source/Programs/Shared/EpicGames.Horde/HordeClient.cs
+++ b/Engine/Source/Programs/Shared/EpicGames.Horde/HordeClient.cs
@@ -23,6 +23,11 @@ namespace EpicGames.Horde
 		readonly HordeOptions _hordeOptions;
 		readonly ILoggerFactory _loggerFactory;
 
+		Uri? _serverUrl;
+
+		/// <inheritdoc/>
+		public Uri ServerUrl => GetServerUrl();
+
 		/// <summary>
 		/// Constructor
 		/// </summary>
@@ -35,6 +40,16 @@ namespace EpicGames.Horde
 			_loggerFactory = loggerFactory;
 		}
 
+		Uri GetServerUrl()
+		{
+			if (_serverUrl == null)
+			{
+				using HttpClient httpClient = _httpClientFactory.CreateClient(HordeHttpClient.HttpClientName);
+				_serverUrl = httpClient.BaseAddress ?? new Uri("http://horde-server");
+			}
+			return _serverUrl;
+		}
+
 		/// <inheritdoc/>
 		public async Task<bool> ConnectAsync(bool allowLogin, CancellationToken cancellationToken)
 		{
diff --git a/Engine/Source/Programs/Shared/EpicGames.Horde/IHordeClient.cs b/Engine/Source/Programs/Shared/EpicGames.Horde/IHordeClient.cs
--- a/Engine/Source/Programs/Shared/EpicGames.Horde/IHordeClient.cs
+++ b/Engine/Source/Programs/Shared/EpicGames.Horde/IHordeClient.cs
@@ -1,5 +1,6 @@
 // Copyright Epic Games, Inc. All Rights Reserved.
 
+using System;
 using System.Threading;
 using System.Threading.Tasks;
 using EpicGames.Horde.Storage;
@@ -12,6 +13,11 @@ namespace EpicGames.Horde
 	/// </summary>
 	public interface IHordeClient
 	{
+		/// <summary>
+		/// Base URL of the horde server
+		/// </summary>
+		Uri ServerUrl { get; }
+
 		/// <summary>
 		/// Connect to the Horde server
 		/// </summary>
diff --git a/Engine/Source/Programs/UnrealGameSync/UnrealGameSync/Controls/WorkspaceControl.cs b/Engine/Source/Programs/UnrealGameSync/UnrealGameSync/Controls/WorkspaceControl.cs
--- a/Engine/Source/Programs/UnrealGameSync/UnrealGameSync/Controls/WorkspaceControl.cs
+++ b/Engine/Source/Programs/UnrealGameSync/UnrealGameSync/Controls/WorkspaceControl.cs
@@ -3362,11 +3362,11 @@ namespace UnrealGameSync
 
 				if (_hordeClient.IsConnected())
 				{
-					statusLine.AddText("Connected to Horde.");
+					statusLine.AddLink("Connected to Horde", FontStyle.Regular, (p, r) => Utility.OpenUrl(_hordeClient.ServerUrl.ToString()));
 				}
 				else
 				{
-					statusLine.AddLink("Connect to Horde", FontStyle.Bold | FontStyle.Underline, (p, r) => ConnectToHorde(_hordeClient));
+					statusLine.AddLink("Connect to Horde...", FontStyle.Bold | FontStyle.Underline, (p, r) => ConnectToHorde(_hordeClient));
 				}
 			}
 		}
