Add platform checks to SDK file references in Oodle C# wrapper

diff --git a/Engine/Source/Programs/Shared/EpicGames.Oodle/EpicGames.Oodle.csproj b/Engine/Source/Programs/Shared/EpicGames.Oodle/EpicGames.Oodle.csproj
--- a/Engine/Source/Programs/Shared/EpicGames.Oodle/EpicGames.Oodle.csproj
+++ b/Engine/Source/Programs/Shared/EpicGames.Oodle/EpicGames.Oodle.csproj
@@ -45,12 +45,12 @@
     <PackageReference Include="Microsoft.Extensions.Logging" Version="6.0.0" />
   </ItemGroup>
   <ItemGroup>
-    <None Include="Sdk\2.9.10\linux\lib\liboo2corelinux64.so.9" Link="liboo2corelinux64.so.9" CopyToOutputDirectory="PreserveNewest" Visible="false" />
-    <None Include="Sdk\2.9.10\linuxarm\lib\liboo2corelinuxarm32.so.9" Link="liboo2corelinuxarm32.so.9" CopyToOutputDirectory="PreserveNewest" Visible="false" />
-    <None Include="Sdk\2.9.10\linuxarm\lib\liboo2corelinuxarm64.so.9" Link="liboo2corelinuxarm64.so.9" CopyToOutputDirectory="PreserveNewest" Visible="false" />
-    <None Include="Sdk\2.9.10\mac\lib\liboo2coremac64.2.9.10.dylib" Link="liboo2coremac64.2.9.10.dylib" CopyToOutputDirectory="PreserveNewest" Visible="false" />
-    <None Include="Sdk\2.9.10\win\redist\oo2core_9_win64.dll" Link="oo2core_9_win64.dll" CopyToOutputDirectory="PreserveNewest" Visible="false" />
-    <None Include="Sdk\2.9.10\win\redist\oo2core_9_win32.dll" Link="oo2core_9_win32.dll" CopyToOutputDirectory="PreserveNewest" Visible="false" />
-    <None Include="Sdk\2.9.10\winuwp\redist\oo2core_9_winuwparm64.dll" Link="oo2core_9_winuwparm64.dll" CopyToOutputDirectory="PreserveNewest" Visible="false" />
+    <None Condition="'$([System.Runtime.InteropServices.RuntimeInformation]::IsOSPlatform($([System.Runtime.InteropServices.OSPlatform]::Linux)))' == 'true' Or Exists('Sdk\2.9.10\linux\lib\liboo2corelinux64.so.9')" Include="Sdk\2.9.10\linux\lib\liboo2corelinux64.so.9" Link="liboo2corelinux64.so.9" CopyToOutputDirectory="PreserveNewest" Visible="false" />
+    <None Condition="'$([System.Runtime.InteropServices.RuntimeInformation]::IsOSPlatform($([System.Runtime.InteropServices.OSPlatform]::Linux)))' == 'true' Or Exists('Sdk\2.9.10\linuxarm\lib\liboo2corelinuxarm32.so.9')" Include="Sdk\2.9.10\linuxarm\lib\liboo2corelinuxarm32.so.9" Link="liboo2corelinuxarm32.so.9" CopyToOutputDirectory="PreserveNewest" Visible="false" />
+    <None Condition="'$([System.Runtime.InteropServices.RuntimeInformation]::IsOSPlatform($([System.Runtime.InteropServices.OSPlatform]::Linux)))' == 'true' Or Exists('Sdk\2.9.10\linuxarm\lib\liboo2corelinuxarm64.so.9')" Include="Sdk\2.9.10\linuxarm\lib\liboo2corelinuxarm64.so.9" Link="liboo2corelinuxarm64.so.9" CopyToOutputDirectory="PreserveNewest" Visible="false" />
+    <None Condition="'$([System.Runtime.InteropServices.RuntimeInformation]::IsOSPlatform($([System.Runtime.InteropServices.OSPlatform]::OSX)))' == 'true' Or Exists('Sdk\2.9.10\mac\lib\liboo2coremac64.2.9.10.dylib')" Include="Sdk\2.9.10\mac\lib\liboo2coremac64.2.9.10.dylib" Link="liboo2coremac64.2.9.10.dylib" CopyToOutputDirectory="PreserveNewest" Visible="false" />
+    <None Condition="'$([System.Runtime.InteropServices.RuntimeInformation]::IsOSPlatform($([System.Runtime.InteropServices.OSPlatform]::Windows)))' == 'true' Or Exists('Sdk\2.9.10\win\redist\oo2core_9_win64.dll')" Include="Sdk\2.9.10\win\redist\oo2core_9_win64.dll" Link="oo2core_9_win64.dll" CopyToOutputDirectory="PreserveNewest" Visible="false" />
+    <None Condition="'$([System.Runtime.InteropServices.RuntimeInformation]::IsOSPlatform($([System.Runtime.InteropServices.OSPlatform]::Windows)))' == 'true' Or Exists('Sdk\2.9.10\win\redist\oo2core_9_win32.dll')" Include="Sdk\2.9.10\win\redist\oo2core_9_win32.dll" Link="oo2core_9_win32.dll" CopyToOutputDirectory="PreserveNewest" Visible="false" />
+    <None Condition="'$([System.Runtime.InteropServices.RuntimeInformation]::IsOSPlatform($([System.Runtime.InteropServices.OSPlatform]::Windows)))' == 'true' Or Exists('Sdk\2.9.10\winuwp\redist\oo2core_9_winuwparm64.dll')" Include="Sdk\2.9.10\winuwp\redist\oo2core_9_winuwparm64.dll" Link="oo2core_9_winuwparm64.dll" CopyToOutputDirectory="PreserveNewest" Visible="false" />
   </ItemGroup>
 </Project>
\ No newline at end of file
