OidcToken - Harden against OPTIONS verb from the browser, currently a issue from some versions of Chrome

diff --git a/Engine/Source/Programs/OidcToken/Program.cs b/Engine/Source/Programs/OidcToken/Program.cs
--- a/Engine/Source/Programs/OidcToken/Program.cs
+++ b/Engine/Source/Programs/OidcToken/Program.cs
@@ -20,7 +20,7 @@ namespace OidcToken
 			{
 				// print help
 				Console.WriteLine("Usage: OidcToken --Service <serviceName> [options]");
-				Console.WriteLine("Service is a required parameter to indicate which OIDC service you intend to connect to. The connection details of the service is configured in appsettings.json");
+				Console.WriteLine("Service is a required parameter to indicate which OIDC service you intend to connect to. The connection details of the service is configured in appsettings.json/oidc-configuration.json");
 				Console.WriteLine();
 				Console.WriteLine("Options: ");
 				Console.WriteLine(" --Mode [Query/GetToken] - Switch mode to allow you to preview operation without triggering user interaction (result can be used to determine if user interaction is required)");
diff --git a/Engine/Source/Programs/Shared/EpicGames.OIDC/OidcTokenManager.cs b/Engine/Source/Programs/Shared/EpicGames.OIDC/OidcTokenManager.cs
--- a/Engine/Source/Programs/Shared/EpicGames.OIDC/OidcTokenManager.cs
+++ b/Engine/Source/Programs/Shared/EpicGames.OIDC/OidcTokenManager.cs
@@ -342,47 +342,63 @@ namespace EpicGames.OIDC
 
 		private static async Task<LoginResult> ProcessHttpRequest(HttpListener http, AuthorizeState loginState, OidcClient oidcClient)
 		{
-			LoginResult loginResult;
-			HttpListenerContext context = await http.GetContextAsync();
-			string? responseData;
-			switch (context.Request.HttpMethod)
+			LoginResult? loginResult = null;
+			const int MaxAttempts = 5;
+			HttpListenerContext? context = null;
+			for (int i = 0; i < MaxAttempts; i++)
 			{
-				case "GET":
-					responseData = context.Request.RawUrl;
-
-					// parse the returned url for the tokens needed to complete the login
-					loginResult = await oidcClient!.ProcessResponseAsync(responseData, loginState);
-					break;
-				case "POST":
+				context = await http.GetContextAsync();
+				string? responseData;
+				switch (context.Request.HttpMethod)
 				{
-					HttpListenerRequest request = context.Request;
-					if (request.ContentType != null && !request.ContentType.Equals("application/x-www-form-urlencoded", StringComparison.OrdinalIgnoreCase))
-					{
-						// we do not support url encoded return types
-						context.Response.StatusCode = 415;
-						throw new Exception("URL encoded responses not support");
-					}
+					case "GET":
+						responseData = context.Request.RawUrl;
 
-					// attempt to parse the body
+						// parse the returned url for the tokens needed to complete the login
+						loginResult = await oidcClient!.ProcessResponseAsync(responseData, loginState);
+						break;
+					case "POST":
+						{
+							HttpListenerRequest request = context.Request;
+							if (request.ContentType != null && !request.ContentType.Equals("application/x-www-form-urlencoded", StringComparison.OrdinalIgnoreCase))
+							{
+								// we do not support url encoded return types
+								context.Response.StatusCode = 415;
+								throw new Exception("URL encoded responses not support");
+							}
 
-					// if there is no body we can not handle the post
-					if (!context.Request.HasEntityBody)
-					{
-						context.Response.StatusCode = 415;
-						throw new Exception("Empty body not supported");
-					}
+							// attempt to parse the body
+
+							// if there is no body we can not handle the post
+							if (!context.Request.HasEntityBody)
+							{
+								context.Response.StatusCode = 415;
+								throw new Exception("Empty body not supported");
+							}
 
-					await using Stream body = request.InputStream;
-					using StreamReader reader = new StreamReader(body, request.ContentEncoding);
-					responseData = await reader.ReadToEndAsync();
+							await using Stream body = request.InputStream;
+							using StreamReader reader = new StreamReader(body, request.ContentEncoding);
+							responseData = await reader.ReadToEndAsync();
 
-					loginResult = await oidcClient!.ProcessResponseAsync(responseData, loginState);
-					break;
+							loginResult = await oidcClient!.ProcessResponseAsync(responseData, loginState);
+							break;
+						}
+					case "OPTIONS":
+						context.Response.StatusCode = 200;
+						context.Response.Close();
+						continue;
+					default:
+						// if we receive any other http method something is very odd. Tell them to use a different method.
+						context.Response.StatusCode = 415;
+						throw new Exception("Unsupported method used: " + context.Request.HttpMethod);
 				}
-				default:
-					// if we receive any other http method something is very odd. Tell them to use a different method.
-					context.Response.StatusCode = 415;
-					throw new Exception("Unsupported method used: " + context.Request.HttpMethod);
+
+				break;
+			}
+
+			if (context == null || loginResult == null)
+			{
+				throw new Exception("Context or loginResult not set");
 			}
 
 			// generate a simple http page to show the user
