Horde: Run a background task to monitor for hangs when executing Perforce commands.

diff --git a/Engine/Source/Programs/Shared/EpicGames.Perforce/HangMonitor.cs b/Engine/Source/Programs/Shared/EpicGames.Perforce/HangMonitor.cs
new file mode 100644
--- /dev/null
+++ b/Engine/Source/Programs/Shared/EpicGames.Perforce/HangMonitor.cs
@@ -0,0 +1,128 @@
+// Copyright Epic Games, Inc. All Rights Reserved.
+
+using System;
+using System.Diagnostics;
+using System.Threading;
+using System.Threading.Tasks;
+using EpicGames.Core;
+using Microsoft.Extensions.Logging;
+
+namespace EpicGames.Perforce
+{
+	/// <summary>
+	/// Runs a background task that logs warnings if the Tick method isn't called within a certain time period
+	/// </summary>
+	public sealed class HangMonitor : IDisposable
+	{
+		class Scope : IDisposable
+		{
+			readonly HangMonitor _hangMonitor;
+			readonly string _activity;
+
+			public string Activity => _activity;
+
+			public Scope(HangMonitor hangMonitor, string activity)
+			{
+				_hangMonitor = hangMonitor;
+				_activity = activity;
+			}
+
+			public void Dispose()
+			{
+				Interlocked.CompareExchange(ref _hangMonitor._scope, null, this);
+			}
+		}
+
+		readonly AsyncEvent _activeEvent = new AsyncEvent();
+		readonly TimeSpan _interval;
+		readonly string _context;
+		readonly ILogger _logger;
+		long _lastUpdateTicks;
+		Scope? _scope;
+		BackgroundTask? _hangTask;
+
+		/// <summary>
+		/// Constructor
+		/// </summary>
+		/// <param name="interval">Interval after which to log a message</param>
+		/// <param name="context">Context for hang messages</param>
+		/// <param name="logger">Logger to write to </param>
+		public HangMonitor(TimeSpan interval, string context, ILogger logger)
+		{
+			_interval = interval;
+			_context = context;
+			_logger = logger;
+		}
+
+		/// <inheritdoc/>
+		public void Dispose()
+		{
+			if (_hangTask != null)
+			{
+				_hangTask.DisposeAsync().AsTask().Wait();
+				_hangTask = null;
+			}
+		}
+
+		/// <summary>
+		/// Start monitoring for hangs.
+		/// </summary>
+		/// <param name="activity">Activity to log if a hang is detected</param>
+		public IDisposable Start(string activity)
+		{
+			Scope scope = new Scope(this, activity);
+			if (Interlocked.CompareExchange(ref _scope, scope, null) != null)
+			{
+				throw new InvalidOperationException();
+			}
+
+			_lastUpdateTicks = Stopwatch.GetTimestamp();
+			_hangTask ??= BackgroundTask.StartNew(CheckStatusAsync);
+			_activeEvent.Pulse();
+
+			return scope;
+		}
+
+		/// <summary>
+		/// Marks the operation as ongoing
+		/// </summary>
+		public void Tick()
+		{
+			Interlocked.Exchange(ref _lastUpdateTicks, Stopwatch.GetTimestamp());
+		}
+
+		async Task CheckStatusAsync(CancellationToken cancellationToken)
+		{
+			TimeSpan nextInterval = TimeSpan.Zero;
+			for (; ; )
+			{
+				await Task.Delay(nextInterval, cancellationToken);
+				nextInterval = _interval;
+
+				Task activeTask = _activeEvent.Task;
+
+				Scope? scope = Interlocked.CompareExchange(ref _scope, null, null);
+				if (scope == null)
+				{
+					await activeTask.WaitAsync(cancellationToken);
+					continue;
+				}
+
+				long currentTicks = Stopwatch.GetTimestamp();
+				long lastUpdateTicks = Interlocked.CompareExchange(ref _lastUpdateTicks, 0, 0);
+				if (currentTicks > lastUpdateTicks)
+				{
+					TimeSpan time = TimeSpan.FromSeconds((double)(currentTicks - lastUpdateTicks) / Stopwatch.Frequency);
+					if (time >= _interval)
+					{
+						_logger.LogWarning("Hang detected ({Context}): {Activity} ({Time}s)", _context, scope.Activity, (int)time.TotalSeconds);
+					}
+					else
+					{
+						nextInterval = _interval - time;
+					}
+				}
+			}
+		}
+	}
+}
diff --git a/Engine/Source/Programs/Shared/EpicGames.Perforce/NativePerforceConnection.cs b/Engine/Source/Programs/Shared/EpicGames.Perforce/NativePerforceConnection.cs
--- a/Engine/Source/Programs/Shared/EpicGames.Perforce/NativePerforceConnection.cs
+++ b/Engine/Source/Programs/Shared/EpicGames.Perforce/NativePerforceConnection.cs
@@ -272,6 +272,8 @@ namespace EpicGames.Perforce
 		Response? _currentResponse;
 		readonly ManualResetEvent _responseCompleteEvent;
 		readonly Stopwatch _stallTimer = new Stopwatch();
+		readonly HangMonitor _hangMonitor;
+		bool _disposed;
 
 		/// <inheritdoc/>
 		public IPerforceSettings Settings { get; }
@@ -302,6 +304,7 @@ namespace EpicGames.Perforce
 			Logger = logger;
 
 			_uniqueId = Interlocked.Increment(ref s_nextUniqueId);
+			_hangMonitor = new HangMonitor(TimeSpan.FromMinutes(1.0), $"Perforce connection ({_uniqueId})", logger);
 
 			_buffers = new PinnedBuffer[bufferCount];
 			for (int idx = 0; idx < bufferCount; idx++)
@@ -364,6 +367,7 @@ namespace EpicGames.Perforce
 		void GetNextWriteBuffer(NativeWriteBuffer nativeWriteBuffer, int minSize)
 		{
 			_stallTimer.Start();
+			_hangMonitor.Tick();
 
 			PinnedBuffer buffer = _writeBuffers.Take();
 			if (buffer.MaxLength < minSize)
@@ -375,6 +379,7 @@ namespace EpicGames.Perforce
 			nativeWriteBuffer._maxLength = buffer.Data.Length;
 			nativeWriteBuffer._maxCount = Int32.MaxValue;
 
+			_hangMonitor.Tick();
 			_stallTimer.Stop();
 		}
 
@@ -396,24 +401,23 @@ namespace EpicGames.Perforce
 		/// <inheritdoc/>
 		void Dispose(bool disposing)
 		{
-			if (disposing)
+			if (disposing && !_disposed)
 			{
-				Logger.LogTrace("Disposing Perforce connection {ConnectionId}", _uniqueId);
-			}
+				using IDisposable scope = _hangMonitor.Start("Disposing");
 
-			if (_backgroundThread != null)
-			{
-				_requests.Add(null);
-				_requests.CompleteAdding();
+				Logger.LogTrace("Disposing Perforce connection {ConnectionId}", _uniqueId);
 
-				if (disposing)
+				if (_backgroundThread != null)
 				{
+					_requests.Add(null);
+					_requests.CompleteAdding();
+
 					_backgroundThread.Join();
 					_backgroundThread = null!;
 				}
-			}
 
-			_requests.Dispose();
+				_requests.Dispose();
+			}
 
 			if (_client != IntPtr.Zero)
 			{
@@ -421,13 +425,21 @@ namespace EpicGames.Perforce
 				_client = IntPtr.Zero;
 			}
 
-			_writeBuffers.Dispose();
+			if (disposing && !_disposed)
+			{
+				using IDisposable scope = _hangMonitor.Start("Disposing");
 
-			_responseCompleteEvent.Dispose();
+				_writeBuffers.Dispose();
 
-			foreach (PinnedBuffer pinnedBuffer in _buffers)
-			{
-				pinnedBuffer.Dispose();
+				_responseCompleteEvent.Dispose();
+
+				foreach (PinnedBuffer pinnedBuffer in _buffers)
+				{
+					pinnedBuffer.Dispose();
+				}
+
+				_hangMonitor.Dispose();
+				_disposed = true;
 			}
 		}
 
@@ -503,22 +515,30 @@ namespace EpicGames.Perforce
 		{
 			for (; ; )
 			{
-				(Action Action, Response Response)? request = _requests.Take();
-				if (request == null)
+				try
 				{
-					break;
-				}
+					(Action Action, Response Response)? request = _requests.Take();
+					if (request == null)
+					{
+						break;
+					}
 
-				_currentResponse = request.Value.Response;
-				_responseCompleteEvent.Reset();
+					_currentResponse = request.Value.Response;
+					_responseCompleteEvent.Reset();
 
-				_stallTimer.Reset();
-				request.Value.Action();
+					_stallTimer.Reset();
+					request.Value.Action();
 
-				_currentResponse._readBuffers.Writer.TryComplete();
-				_responseCompleteEvent.WaitOne();
+					_currentResponse._readBuffers.Writer.TryComplete();
+					_responseCompleteEvent.WaitOne();
 
-				_currentResponse = null;
+					_currentResponse = null;
+				}
+				catch (Exception ex)
+				{
+					Logger.LogError(ex, "Exception while processing Perforce commands: {Message}", ex.Message);
+					throw;
+				}
 			}
 		}
 
@@ -577,6 +597,8 @@ namespace EpicGames.Perforce
 			Stopwatch timer = Stopwatch.StartNew();
 			Logger.LogTrace("Conn {ConnectionId}: {Command} {Args}", _uniqueId, command, argList.ToString());
 
+			using IDisposable scope = _hangMonitor.Start($"{command} {argList}");
+
 			List<IntPtr> nativeArgs = new List<IntPtr>();
 			try
 			{
