BuildGraph: Add support for setting artifact metadata in buildgraph scripts.

diff --git a/Engine/Source/Programs/AutomationTool/BuildGraph/BgSchema.cs b/Engine/Source/Programs/AutomationTool/BuildGraph/BgSchema.cs
--- a/Engine/Source/Programs/AutomationTool/BuildGraph/BgSchema.cs
+++ b/Engine/Source/Programs/AutomationTool/BuildGraph/BgSchema.cs
@@ -652,6 +652,7 @@ namespace AutomationTool
 			artifactType.Attributes.Add(CreateSchemaAttribute("BasePath", ScriptSchemaStandardType.BalancedString, XmlSchemaUse.Optional));
 			artifactType.Attributes.Add(CreateSchemaAttribute("Tag", ScriptSchemaStandardType.Tag, XmlSchemaUse.Optional));
 			artifactType.Attributes.Add(CreateSchemaAttribute("Keys", ScriptSchemaStandardType.BalancedString, XmlSchemaUse.Optional));
+			artifactType.Attributes.Add(CreateSchemaAttribute("Metadata", ScriptSchemaStandardType.BalancedString, XmlSchemaUse.Optional));
 			artifactType.Attributes.Add(CreateSchemaAttribute("If", ScriptSchemaStandardType.BalancedString, XmlSchemaUse.Optional));
 			return artifactType;
 		}
diff --git a/Engine/Source/Programs/AutomationTool/BuildGraph/BgScriptReader.cs b/Engine/Source/Programs/AutomationTool/BuildGraph/BgScriptReader.cs
--- a/Engine/Source/Programs/AutomationTool/BuildGraph/BgScriptReader.cs
+++ b/Engine/Source/Programs/AutomationTool/BuildGraph/BgScriptReader.cs
@@ -1246,8 +1246,9 @@ namespace AutomationTool
 				}
 
 				string[] keys = ReadListAttribute(element, "Keys");
+				string[] metadata = ReadListAttribute(element, "Metadata");
 
-				BgArtifactDef newArtifact = new BgArtifactDef(name, type, description, basePath, tag, keys);
+				BgArtifactDef newArtifact = new BgArtifactDef(name, type, description, basePath, tag, keys, metadata);
 				_graph.Artifacts.Add(newArtifact);
 			}
 		}
diff --git a/Engine/Source/Programs/Shared/EpicGames.BuildGraph/BgArtifactDef.cs b/Engine/Source/Programs/Shared/EpicGames.BuildGraph/BgArtifactDef.cs
--- a/Engine/Source/Programs/Shared/EpicGames.BuildGraph/BgArtifactDef.cs
+++ b/Engine/Source/Programs/Shared/EpicGames.BuildGraph/BgArtifactDef.cs
@@ -39,6 +39,11 @@ namespace EpicGames.BuildGraph
 		/// </summary>
 		public IReadOnlyList<string> Keys { get; }
 
+		/// <summary>
+		/// Metadata for the artifact
+		/// </summary>
+		public IReadOnlyList<string> Metadata { get; }
+
 		/// <summary>
 		/// Constructor
 		/// </summary>
@@ -48,7 +53,8 @@ namespace EpicGames.BuildGraph
 		/// <param name="basePath">Base path for files included in the artifact</param>
 		/// <param name="tagName">Name of the tag producing this artifact</param>
 		/// <param name="keys">Keys that can be used to find the artifact</param>
-		public BgArtifactDef(string name, string? type, string? description, string? basePath, string tagName, IReadOnlyList<string> keys)
+		/// <param name="metadata">Metadata for the artifact</param>
+		public BgArtifactDef(string name, string? type, string? description, string? basePath, string tagName, IReadOnlyList<string> keys, IReadOnlyList<string> metadata)
 		{
 			Name = name;
 			Type = type;
@@ -56,6 +62,7 @@ namespace EpicGames.BuildGraph
 			BasePath = basePath;
 			TagName = tagName;
 			Keys = keys;
+			Metadata = metadata;
 		}
 
 		/// <summary>
diff --git a/Engine/Source/Programs/Shared/EpicGames.BuildGraph/BgGraphDef.cs b/Engine/Source/Programs/Shared/EpicGames.BuildGraph/BgGraphDef.cs
--- a/Engine/Source/Programs/Shared/EpicGames.BuildGraph/BgGraphDef.cs
+++ b/Engine/Source/Programs/Shared/EpicGames.BuildGraph/BgGraphDef.cs
@@ -453,6 +453,13 @@ namespace EpicGames.BuildGraph
 					}
 					jsonWriter.WriteArrayEnd();
 
+					jsonWriter.WriteArrayStart("Metadata");
+					foreach (string metadata in artifact.Metadata)
+					{
+						jsonWriter.WriteValue(metadata);
+					}
+					jsonWriter.WriteArrayEnd();
+
 					jsonWriter.WriteValue("OutputName", artifact.TagName);
 					jsonWriter.WriteObjectEnd();
 				}
