UGS: Catch exception checking Horde connection status.

diff --git a/Engine/Source/Programs/Shared/EpicGames.Horde/HordeClient.cs b/Engine/Source/Programs/Shared/EpicGames.Horde/HordeClient.cs
--- a/Engine/Source/Programs/Shared/EpicGames.Horde/HordeClient.cs
+++ b/Engine/Source/Programs/Shared/EpicGames.Horde/HordeClient.cs
@@ -67,7 +67,14 @@ namespace EpicGames.Horde
 		/// <inheritdoc/>
 		public bool IsConnected()
 		{
-			return _authHandlerState.IsAuthenticated();
+			try
+			{
+				return _authHandlerState.IsAuthenticated();
+			}
+			catch
+			{
+				return false;
+			}
 		}
 
 		/// <inheritdoc/>
