Remove property setter in response to C# static analysis warning.

diff --git a/Engine/Source/Programs/Shared/EpicGames.ProjectStore/ProjectStoreTypes.cs b/Engine/Source/Programs/Shared/EpicGames.ProjectStore/ProjectStoreTypes.cs
--- a/Engine/Source/Programs/Shared/EpicGames.ProjectStore/ProjectStoreTypes.cs
+++ b/Engine/Source/Programs/Shared/EpicGames.ProjectStore/ProjectStoreTypes.cs
@@ -20,7 +20,7 @@ namespace EpicGames.ProjectStore
 		/// <summary>
 		/// The list of remote host names for zenserver
 		/// </summary>
-		public List<string> RemoteHostNames { get; set; } = new List<string>();
+		public List<string> RemoteHostNames { get; } = new List<string>();
 		/// <summary>
 		/// The port for zenserver on the host
 		/// </summary>
