Fixing error matcher for locked robomerge gates. Adding a cursor function to look ahead a limited number of times.

#rnx

diff --git a/Engine/Source/Programs/AutomationTool/AutomationUtils/Matchers/SystemicEventMatcher.cs b/Engine/Source/Programs/AutomationTool/AutomationUtils/Matchers/SystemicEventMatcher.cs
--- a/Engine/Source/Programs/AutomationTool/AutomationUtils/Matchers/SystemicEventMatcher.cs
+++ b/Engine/Source/Programs/AutomationTool/AutomationUtils/Matchers/SystemicEventMatcher.cs
@@ -36,9 +36,12 @@ namespace AutomationUtils.Matchers
 				}
 				return builder.ToMatch(LogEventPriority.High, LogLevel.Information, KnownLogEvents.Systemic_PdbUtil);
 			}
-			if (cursor.IsMatch(s_roboMerge))
+			if (cursor.Contains("failed to submit"))
 			{
-				return new LogEventBuilder(cursor).ToMatch(LogEventPriority.Low, LogLevel.Information, KnownLogEvents.Systemic_RoboMergeGateLocked);
+				if (cursor.CurrentLineNumber != cursor.MatchForwardsLimited(cursor.CurrentLineNumber, s_roboMerge, 5))
+				{
+					return new LogEventBuilder(cursor).ToMatch(LogEventPriority.Low, LogLevel.Information, KnownLogEvents.Systemic_RoboMergeGateLocked);
+				}
 			}
 			if (cursor.IsMatch(s_hostDown))
 			{
diff --git a/Engine/Source/Programs/Shared/EpicGames.Core/ILogCursor.cs b/Engine/Source/Programs/Shared/EpicGames.Core/ILogCursor.cs
--- a/Engine/Source/Programs/Shared/EpicGames.Core/ILogCursor.cs
+++ b/Engine/Source/Programs/Shared/EpicGames.Core/ILogCursor.cs
@@ -228,6 +228,35 @@ namespace EpicGames.Core
 			return offset;
 		}
 
+		/// <summary>
+		/// Matches lines forwards from the given offset until the given pattern matches
+		/// </summary>
+		/// <param name="cursor">The log cursor</param>
+		/// <param name="offset">Initial offset</param>
+		/// <param name="pattern">Pattern to match</param>
+		/// <param name="maxLines">Max number of lines to look forward</param>
+		/// <returns>Offset of the line that matches the pattern (inclusive), or EOF is encountered</returns>
+		public static int MatchForwardsLimited(this ILogCursor cursor, int offset, Regex pattern, int maxLines)
+		{
+			string? nextLine;
+			int numLines = 0;
+			int nextOffset = offset + 1;
+			cursor.TryGetLine(nextOffset, out nextLine);
+
+			while (numLines < maxLines && nextLine!=null)
+			{
+				
+				if (pattern.IsMatch(nextLine))
+				{
+					return nextOffset;
+				}
+				cursor.TryGetLine(nextOffset, out nextLine);
+				nextOffset++;
+				numLines++;
+			}
+			return offset;
+		}
+
 		/// <summary>
 		/// Tests if a line consists only of whitespace
 		/// </summary>
