OIDCToken: Fix OIDC discovery document URL construction if authority already has a trailing slash.

diff --git a/Engine/Source/Programs/Shared/EpicGames.OIDC/OidcTokenManager.cs b/Engine/Source/Programs/Shared/EpicGames.OIDC/OidcTokenManager.cs
--- a/Engine/Source/Programs/Shared/EpicGames.OIDC/OidcTokenManager.cs
+++ b/Engine/Source/Programs/Shared/EpicGames.OIDC/OidcTokenManager.cs
@@ -528,7 +528,8 @@ namespace EpicGames.OIDC
 
 		private async Task<DiscoveryDocumentResponse> GetDiscoveryDocument(CancellationToken cancellationToken)
 		{
-			string discoUrl = $"{_authorityUri}/.well-known/openid-configuration";
+			string baseUrl = _authorityUri.ToString().TrimEnd('/');
+			string discoUrl = $"{baseUrl}/.well-known/openid-configuration";
 
 			using HttpClient client = new HttpClient();
 			using DiscoveryDocumentRequest doc = new DiscoveryDocumentRequest
