Horde: Allow specifying a category for tools for use on the dashboard. Remove the setting for parent tool id.

diff --git a/Engine/Source/Programs/Horde/Horde.Server/Server/GlobalConfig.cs b/Engine/Source/Programs/Horde/Horde.Server/Server/GlobalConfig.cs
--- a/Engine/Source/Programs/Horde/Horde.Server/Server/GlobalConfig.cs
+++ b/Engine/Source/Programs/Horde/Horde.Server/Server/GlobalConfig.cs
@@ -320,27 +320,6 @@ namespace Horde.Server.Server
 				tool.PostLoad(this);
 			}
 
-			foreach (ToolConfig tool in Tools)
-			{
-				if (tool.ParentId != null)
-				{
-					if (tool.ParentId.Value == tool.Id)
-					{
-						throw new InvalidDataException($"Tool {tool.Id} cannot be parented to itself");
-					}
-
-					ToolConfig? parentToolConfig;
-					if (!_toolLookup.TryGetValue(tool.ParentId.Value, out parentToolConfig))
-					{
-						throw new InvalidDataException($"Tool {tool.Id} has an invalid parent tool id of '{tool.ParentId.Value}'");
-					}
-					else if (parentToolConfig.ParentId != null)
-					{
-						throw new InvalidDataException($"Tool {tool.Id} cannot be nested under two parents ({tool.ParentId.Value} -> {parentToolConfig.ParentId.Value})");
-					}
-				}
-			}
-
 			_computeClusterLookup.Clear();
 			foreach (ComputeClusterConfig computeCluster in Compute)
 			{
diff --git a/Engine/Source/Programs/Horde/Horde.Server/Tools/ToolConfig.cs b/Engine/Source/Programs/Horde/Horde.Server/Tools/ToolConfig.cs
--- a/Engine/Source/Programs/Horde/Horde.Server/Tools/ToolConfig.cs
+++ b/Engine/Source/Programs/Horde/Horde.Server/Tools/ToolConfig.cs
@@ -44,9 +44,9 @@ namespace Horde.Server.Tools
 		public string Description { get; set; }
 
 		/// <summary>
-		/// Tool id to nest this tool under
+		/// Category for the tool. Will cause the tool to be shown in a different tab in the dashboard.
 		/// </summary>
-		public ToolId? ParentId { get; set; }
+		public string? Category { get; set; }
 
 		/// <summary>
 		/// Whether this tool should be exposed for download on a public endpoint without authentication
diff --git a/Engine/Source/Programs/Horde/Horde.Server/Tools/ToolsController.cs b/Engine/Source/Programs/Horde/Horde.Server/Tools/ToolsController.cs
--- a/Engine/Source/Programs/Horde/Horde.Server/Tools/ToolsController.cs
+++ b/Engine/Source/Programs/Horde/Horde.Server/Tools/ToolsController.cs
@@ -217,7 +217,7 @@ namespace Horde.Server.Tools
 		static GetToolSummaryResponse CreateGetToolSummaryResponse(ITool tool)
 		{
 			IToolDeployment? deployment = (tool.Deployments.Count == 0) ? null : tool.Deployments[^1];
-			return new GetToolSummaryResponse(tool.Id, tool.Config.Name, tool.Config.Description, tool.Config.ParentId, deployment?.Version, deployment?.Id, tool.Config.ShowInUgs, tool.Config.ShowInDashboard);
+			return new GetToolSummaryResponse(tool.Id, tool.Config.Name, tool.Config.Description, tool.Config.Category, deployment?.Version, deployment?.Id, tool.Config.ShowInUgs, tool.Config.ShowInDashboard);
 		}
 
 		/// <summary>
@@ -261,7 +261,7 @@ namespace Horde.Server.Tools
 
 		static GetToolResponse CreateGetToolResponse(ITool tool, List<GetToolDeploymentResponse> deployments)
 		{
-			return new GetToolResponse(tool.Id, tool.Config.Name, tool.Config.Description, tool.Config.ParentId, deployments, tool.Config.Public, tool.Config.ShowInUgs, tool.Config.ShowInDashboard);
+			return new GetToolResponse(tool.Id, tool.Config.Name, tool.Config.Description, tool.Config.Category, deployments, tool.Config.Public, tool.Config.ShowInUgs, tool.Config.ShowInDashboard);
 		}
 
 		/// <summary>
diff --git a/Engine/Source/Programs/Shared/EpicGames.Horde/Tools/ToolMessages.cs b/Engine/Source/Programs/Shared/EpicGames.Horde/Tools/ToolMessages.cs
--- a/Engine/Source/Programs/Shared/EpicGames.Horde/Tools/ToolMessages.cs
+++ b/Engine/Source/Programs/Shared/EpicGames.Horde/Tools/ToolMessages.cs
@@ -13,12 +13,12 @@ namespace EpicGames.Horde.Tools
 	/// <param name="Id">Unique identifier for the tool</param>
 	/// <param name="Name">Name of the tool</param>
 	/// <param name="Description">Description for the tool</param>
-	/// <param name="ParentId">The parent tool id</param>
+	/// <param name="Category">Category to display the tool in on the dashboard</param>
 	/// <param name="Deployments">Current deployments of this tool, sorted by time.</param>
 	/// <param name="Public">Whether this tool should be exposed for download on a public endpoint without authentication</param>
 	/// <param name="ShowInUgs">Whether to show this tool for download inside UGS</param>
 	/// <param name="ShowInDashboard">Whether to show this tool for download on the dashboard</param>
-	public record class GetToolResponse(ToolId Id, string Name, string Description, ToolId? ParentId, List<GetToolDeploymentResponse> Deployments, bool Public, bool ShowInUgs, bool ShowInDashboard);
+	public record class GetToolResponse(ToolId Id, string Name, string Description, string? Category, List<GetToolDeploymentResponse> Deployments, bool Public, bool ShowInUgs, bool ShowInDashboard);
 
 	/// <summary>
 	/// Summary for a particular tool.
@@ -26,12 +26,12 @@ namespace EpicGames.Horde.Tools
 	/// <param name="Id">Unique identifier for the tool</param>
 	/// <param name="Name">Name of the tool</param>
 	/// <param name="Description">Description for the tool</param>
-	/// <param name="ParentId">The parent tool id</param>
+	/// <param name="Category">Category to display the tool in on the dashboard</param>
 	/// <param name="Version">Version number of the current deployment of this tool</param>
 	/// <param name="DeploymentId">Identifier for the current deployment</param>
 	/// <param name="ShowInUgs">Whether to show this tool for download inside UGS</param>
 	/// <param name="ShowInDashboard">Whether to show this tool for download on the dashboard</param>
-	public record class GetToolSummaryResponse(ToolId Id, string Name, string Description, ToolId? ParentId, string? Version, ToolDeploymentId? DeploymentId, bool ShowInUgs, bool ShowInDashboard);
+	public record class GetToolSummaryResponse(ToolId Id, string Name, string Description, string? Category, string? Version, ToolDeploymentId? DeploymentId, bool ShowInUgs, bool ShowInDashboard);
 
 	/// <summary>
 	/// Response when querying all tools
