Horde: Fix parsing of true/false literals in condition expressions.

diff --git a/Engine/Source/Programs/Shared/EpicGames.Horde/Common/Condition.cs b/Engine/Source/Programs/Shared/EpicGames.Horde/Common/Condition.cs
--- a/Engine/Source/Programs/Shared/EpicGames.Horde/Common/Condition.cs
+++ b/Engine/Source/Programs/Shared/EpicGames.Horde/Common/Condition.cs
@@ -336,6 +336,9 @@ namespace EpicGames.Horde.Common
 		readonly List<Token> _tokens = new List<Token>();
 		readonly List<string> _strings = new List<string>();
 
+		readonly static IEnumerable<string> s_trueScalar = new[] { "true" };
+		readonly static IEnumerable<string> s_falseScalar = new[] { "false" };
+
 		private Condition(string text)
 		{
 			Text = text;
@@ -495,6 +498,11 @@ namespace EpicGames.Horde.Common
 		{
 			switch (reader.Type)
 			{
+				case TokenType.True:
+				case TokenType.False:
+					_tokens.Add(new Token(reader.Type, 0));
+					reader.MoveNext();
+					return null;
 				case TokenType.Identifier:
 					_strings.Add(reader.Token.ToString());
 					_tokens.Add(new Token(TokenType.Identifier, _strings.Count - 1));
@@ -588,6 +596,8 @@ namespace EpicGames.Horde.Common
 			Token token = _tokens[idx++];
 			return token.Type switch
 			{
+				TokenType.True => s_trueScalar,
+				TokenType.False => s_falseScalar,
 				TokenType.Identifier => getPropertyValues(_strings[token.Index]),
 				TokenType.Scalar => new string[] { _strings[token.Index] },
 				_ => throw new InvalidOperationException("Invalid token type")
