Horde: Fix static analysis warning.

#rnx

diff --git a/Engine/Source/Programs/Shared/EpicGames.Horde/Common/Condition.cs b/Engine/Source/Programs/Shared/EpicGames.Horde/Common/Condition.cs
--- a/Engine/Source/Programs/Shared/EpicGames.Horde/Common/Condition.cs
+++ b/Engine/Source/Programs/Shared/EpicGames.Horde/Common/Condition.cs
@@ -336,8 +336,8 @@ namespace EpicGames.Horde.Common
 		readonly List<Token> _tokens = new List<Token>();
 		readonly List<string> _strings = new List<string>();
 
-		readonly static IEnumerable<string> s_trueScalar = new[] { "true" };
-		readonly static IEnumerable<string> s_falseScalar = new[] { "false" };
+		static readonly IEnumerable<string> s_trueScalar = new[] { "true" };
+		static readonly IEnumerable<string> s_falseScalar = new[] { "false" };
 
 		private Condition(string text)
 		{
