Horde: Log and discard timeout exceptions when waiting for compute resources.

diff --git a/Engine/Source/Programs/Shared/EpicGames.Horde/Compute/Clients/ServerComputeClient.cs b/Engine/Source/Programs/Shared/EpicGames.Horde/Compute/Clients/ServerComputeClient.cs
--- a/Engine/Source/Programs/Shared/EpicGames.Horde/Compute/Clients/ServerComputeClient.cs
+++ b/Engine/Source/Programs/Shared/EpicGames.Horde/Compute/Clients/ServerComputeClient.cs
@@ -220,13 +220,21 @@ namespace EpicGames.Horde.Compute.Clients
 		/// <inheritdoc/>
 		public async Task<IComputeLease?> TryAssignWorkerAsync(ClusterId clusterId, Requirements? requirements, string? requestId, ConnectionMetadataRequest? connection, ILogger logger, CancellationToken cancellationToken)
 		{
-			IAsyncEnumerator<LeaseInfo> source = ConnectAsync(clusterId, requirements, requestId, connection, logger, cancellationToken).GetAsyncEnumerator(cancellationToken);
-			if (!await source.MoveNextAsync())
+			try
+			{
+				IAsyncEnumerator<LeaseInfo> source = ConnectAsync(clusterId, requirements, requestId, connection, logger, cancellationToken).GetAsyncEnumerator(cancellationToken);
+				if (!await source.MoveNextAsync())
+				{
+					await source.DisposeAsync();
+					return null;
+				}
+				return new LeaseImpl(source);
+			}
+			catch (Polly.Timeout.TimeoutRejectedException ex)
 			{
-				await source.DisposeAsync();
+				_logger.LogInformation(ex, "Unable to assign worker from pool {ClusterId} (timeout)", clusterId);
 				return null;
 			}
-			return new LeaseImpl(source);
 		}
 
 		/// <inheritdoc/>
