Horde: Shutdown both ends of the socket when closing a compute connection.

diff --git a/Engine/Source/Programs/Shared/EpicGames.Horde/Compute/ComputeSocket.cs b/Engine/Source/Programs/Shared/EpicGames.Horde/Compute/ComputeSocket.cs
--- a/Engine/Source/Programs/Shared/EpicGames.Horde/Compute/ComputeSocket.cs
+++ b/Engine/Source/Programs/Shared/EpicGames.Horde/Compute/ComputeSocket.cs
@@ -414,12 +414,12 @@ namespace EpicGames.Horde.Compute
 		/// </summary>
 		public async ValueTask CloseAsync(CancellationToken cancellationToken)
 		{
+			// Close the transport layer, freeing the remote end to shutdown.
+			await _transport.MarkCompleteAsync(cancellationToken);
+
 			// Close all the buffers
 			await DetachAllBuffersAsync(true, true, cancellationToken);
 
-			// Send a final message indicating that the lease is done. This will allow the senders on the remote end to terminate.
-			await _transport.MarkCompleteAsync(cancellationToken);
-
 			// Wait for the reader to stop
 			if (_recvTask != null)
 			{
diff --git a/Engine/Source/Programs/Shared/EpicGames.Horde/Compute/ComputeTransport.cs b/Engine/Source/Programs/Shared/EpicGames.Horde/Compute/ComputeTransport.cs
--- a/Engine/Source/Programs/Shared/EpicGames.Horde/Compute/ComputeTransport.cs
+++ b/Engine/Source/Programs/Shared/EpicGames.Horde/Compute/ComputeTransport.cs
@@ -28,7 +28,7 @@ namespace EpicGames.Horde.Compute
 		public abstract ValueTask<int> RecvAsync(Memory<byte> buffer, CancellationToken cancellationToken);
 
 		/// <summary>
-		/// Indicate that all data has been written to the transport layer, and that there will be no more calls to <see cref="SendAsync(ReadOnlySequence{Byte}, CancellationToken)"/>
+		/// Indicate that all data has been read and written to the transport layer, and that there will be no more calls to send/recv
 		/// </summary>
 		/// <param name="cancellationToken">Cancellation token for the operation</param>
 		public abstract ValueTask MarkCompleteAsync(CancellationToken cancellationToken);
diff --git a/Engine/Source/Programs/Shared/EpicGames.Horde/Compute/Transports/PipeTransport.cs b/Engine/Source/Programs/Shared/EpicGames.Horde/Compute/Transports/PipeTransport.cs
--- a/Engine/Source/Programs/Shared/EpicGames.Horde/Compute/Transports/PipeTransport.cs
+++ b/Engine/Source/Programs/Shared/EpicGames.Horde/Compute/Transports/PipeTransport.cs
@@ -95,6 +95,10 @@ namespace EpicGames.Horde.Compute.Transports
 		}
 
 		/// <inheritdoc/>
-		public override ValueTask MarkCompleteAsync(CancellationToken cancellationToken) => _writer.CompleteAsync();
+		public override async ValueTask MarkCompleteAsync(CancellationToken cancellationToken)
+		{
+			await _reader.CompleteAsync();
+			await _writer.CompleteAsync();
+		}
 	}
 }
diff --git a/Engine/Source/Programs/Shared/EpicGames.Horde/Compute/Transports/TcpSslTransport.cs b/Engine/Source/Programs/Shared/EpicGames.Horde/Compute/Transports/TcpSslTransport.cs
--- a/Engine/Source/Programs/Shared/EpicGames.Horde/Compute/Transports/TcpSslTransport.cs
+++ b/Engine/Source/Programs/Shared/EpicGames.Horde/Compute/Transports/TcpSslTransport.cs
@@ -129,7 +129,7 @@ public sealed class TcpSslTransport : ComputeTransport
 	/// <inheritdoc/>
 	public override ValueTask MarkCompleteAsync(CancellationToken cancelationToken)
 	{
-		_socket.Shutdown(SocketShutdown.Send);
+		_socket.Shutdown(SocketShutdown.Both);
 		return new ValueTask();
 	}
 
diff --git a/Engine/Source/Programs/Shared/EpicGames.Horde/Compute/Transports/TcpTransport.cs b/Engine/Source/Programs/Shared/EpicGames.Horde/Compute/Transports/TcpTransport.cs
--- a/Engine/Source/Programs/Shared/EpicGames.Horde/Compute/Transports/TcpTransport.cs
+++ b/Engine/Source/Programs/Shared/EpicGames.Horde/Compute/Transports/TcpTransport.cs
@@ -54,7 +54,7 @@ namespace EpicGames.Horde.Compute.Transports
 		/// <inheritdoc/>
 		public override ValueTask MarkCompleteAsync(CancellationToken cancelationToken)
 		{
-			_socket.Shutdown(SocketShutdown.Send);
+			_socket.Shutdown(SocketShutdown.Both);
 			return new ValueTask();
 		}
 	}
