Horde: Do not close the read end of a C# pipe object used for compute traffic. Fixes test failure.

diff --git a/Engine/Source/Programs/Shared/EpicGames.Horde/Compute/Transports/PipeTransport.cs b/Engine/Source/Programs/Shared/EpicGames.Horde/Compute/Transports/PipeTransport.cs
--- a/Engine/Source/Programs/Shared/EpicGames.Horde/Compute/Transports/PipeTransport.cs
+++ b/Engine/Source/Programs/Shared/EpicGames.Horde/Compute/Transports/PipeTransport.cs
@@ -97,7 +97,6 @@ namespace EpicGames.Horde.Compute.Transports
 		/// <inheritdoc/>
 		public override async ValueTask MarkCompleteAsync(CancellationToken cancellationToken)
 		{
-			await _reader.CompleteAsync();
 			await _writer.CompleteAsync();
 		}
 	}
