Horde: Allow '+' characters within object keys.

diff --git a/Engine/Source/Programs/Shared/EpicGames.Horde/Storage/ObjectKey.cs b/Engine/Source/Programs/Shared/EpicGames.Horde/Storage/ObjectKey.cs
--- a/Engine/Source/Programs/Shared/EpicGames.Horde/Storage/ObjectKey.cs
+++ b/Engine/Source/Programs/Shared/EpicGames.Horde/Storage/ObjectKey.cs
@@ -94,7 +94,7 @@ namespace EpicGames.Horde.Storage
 			for (int idx = 0; idx < path.Length; idx++)
 			{
 				byte character = path[idx];
-				if (!StringId.IsValidCharacter(character) && character != '/')
+				if (!StringId.IsValidCharacter(character) && character != '/' && character != '+')
 				{
 					data ??= path.Memory.ToArray();
 					if (character >= 'A' && character <= 'Z')
