Horde: Fix objectkey validation.

diff --git a/Engine/Source/Programs/Shared/EpicGames.Horde/Storage/ObjectKey.cs b/Engine/Source/Programs/Shared/EpicGames.Horde/Storage/ObjectKey.cs
--- a/Engine/Source/Programs/Shared/EpicGames.Horde/Storage/ObjectKey.cs
+++ b/Engine/Source/Programs/Shared/EpicGames.Horde/Storage/ObjectKey.cs
@@ -61,6 +61,10 @@ namespace EpicGames.Horde.Storage
 						{
 							// Allow uppercase characters for now
 						}
+						else if (character == '+')
+						{
+							// Allow plus characters for now
+						}
 						else if (path[idx] == '/' && path[idx - 1] != '/')
 						{
 							// Non-consecutive path separator; allowed.
@@ -94,7 +98,7 @@ namespace EpicGames.Horde.Storage
 			for (int idx = 0; idx < path.Length; idx++)
 			{
 				byte character = path[idx];
-				if (!StringId.IsValidCharacter(character) && character != '/' && character != '+')
+				if (!StringId.IsValidCharacter(character) && character != '/')
 				{
 					data ??= path.Memory.ToArray();
 					if (character >= 'A' && character <= 'Z')
