Horde: Add an IBlobWriter.WriteBlobHandleDangerous method for use when copying preformatted data to a blob writer.

#rnx

diff --git a/Engine/Source/Programs/Shared/EpicGames.Horde/Storage/IBlobWriter.cs b/Engine/Source/Programs/Shared/EpicGames.Horde/Storage/IBlobWriter.cs
--- a/Engine/Source/Programs/Shared/EpicGames.Horde/Storage/IBlobWriter.cs
+++ b/Engine/Source/Programs/Shared/EpicGames.Horde/Storage/IBlobWriter.cs
@@ -61,6 +61,13 @@ namespace EpicGames.Horde.Storage
 		/// <returns>Handle to the written node</returns>
 		ValueTask<IBlobRef<T>> CompleteAsync<T>(BlobType type, CancellationToken cancellationToken = default);
 
+		/// <summary>
+		/// Writes a reference to another blob. NOTE: This does not write anything to the underlying output stream, which prevents the data forming a Merkle tree
+		/// unless guaranteed uniqueness via a hash being written separately.
+		/// </summary>
+		/// <param name="handle">Referenced blob</param>
+		void WriteBlobHandleDangerous(IBlobHandle handle);
+
 		/// <summary>
 		/// Writes a reference to another blob. The blob's hash is serialized to the output stream.
 		/// </summary>
@@ -110,13 +117,19 @@ namespace EpicGames.Horde.Storage
 		/// </summary>
 		public IoHash ComputeHash() => IoHash.Compute(_memory.Span.Slice(0, _length));
 
+		/// <inheritdoc/>
+		public void WriteBlobHandleDangerous(IBlobHandle target)
+		{
+			_imports.Add(target);
+		}
+
 		/// <summary>
 		/// Writes a handle to another node
 		/// </summary>
 		public void WriteBlobRef(IBlobRef target)
 		{
 			this.WriteIoHash(target.Hash);
-			_imports.Add(target);
+			WriteBlobHandleDangerous(target);
 		}
 
 		/// <inheritdoc/>
