Horde: Improve ordering of nodes when writing large file trees to storage, such that the nodes which are read after each other are adjacent in the storage blobs.

Nodes are read in the reverse order to which they are written, depth-first. Arranging nodes in the blob in this order prevents thrashing of the cache and improves performance.

diff --git a/Engine/Source/Programs/Shared/EpicGames.Core/MemoryWriter.cs b/Engine/Source/Programs/Shared/EpicGames.Core/MemoryWriter.cs
--- a/Engine/Source/Programs/Shared/EpicGames.Core/MemoryWriter.cs
+++ b/Engine/Source/Programs/Shared/EpicGames.Core/MemoryWriter.cs
@@ -350,18 +350,34 @@ namespace EpicGames.Core
 
 		/// <inheritdoc/>
 		public Memory<byte> GetMemory(int sizeHint)
+			=> GetMemory(0, sizeHint);
+
+		/// <summary>
+		/// Gets a block of memory to write to, preserving existing (but uncommitted) data in previously returned buffers.
+		/// </summary>
+		/// <param name="usedSize">Size of data in the buffer which has not been committed via a call to Advance</param>
+		/// <param name="desiredSize">Desired size of the buffer to write to</param>
+		/// <returns>New block of memory to write to</returns>
+		public Memory<byte> GetMemory(int usedSize, int desiredSize)
 		{
-			int requiredSize = _currentChunk.Length + Math.Max(sizeHint, 1);
+			int requiredSize = _currentChunk.Length + Math.Max(desiredSize, 1);
 			if (requiredSize > _currentChunk.Data.Length)
 			{
 				int runningIndex = _currentChunk.RunningIndex + _currentChunk.Length;
+				Chunk nextChunk = CreateChunk(runningIndex, Math.Max(desiredSize, _chunkSize));
+
+				if (usedSize > 0)
+				{
+					ReadOnlyMemory<byte> usedData = _currentChunk.Data.Slice(_currentChunk.Length, usedSize);
+					usedData.CopyTo(nextChunk.Data);
+				}
 				if (_currentChunk.Length == 0)
 				{
 					_currentChunk.Release();
 					_chunks.RemoveAt(_chunks.Count - 1);
 				}
 
-				_currentChunk = CreateChunk(runningIndex, Math.Max(sizeHint, _chunkSize));
+				_currentChunk = nextChunk;
 				_chunks.Add(_currentChunk);
 			}
 			return _currentChunk.Data.Slice(_currentChunk.Length);
diff --git a/Engine/Source/Programs/Shared/EpicGames.Horde.Tests/ChunkingTests.cs b/Engine/Source/Programs/Shared/EpicGames.Horde.Tests/ChunkingTests.cs
--- a/Engine/Source/Programs/Shared/EpicGames.Horde.Tests/ChunkingTests.cs
+++ b/Engine/Source/Programs/Shared/EpicGames.Horde.Tests/ChunkingTests.cs
@@ -1,6 +1,7 @@
 // Copyright Epic Games, Inc. All Rights Reserved.
 
 using System;
+using System.Collections.Generic;
 using System.IO;
 using System.Linq;
 using System.Text;
@@ -205,5 +206,49 @@ namespace EpicGames.Horde.Tests
 				tempDir.Delete(true);
 			}
 		}
+
+		[TestMethod]
+		public async Task ChunkOrderAsync()
+		{
+			InteriorChunkedDataNodeOptions interiorOptions = new InteriorChunkedDataNodeOptions(2, 2, 2);
+
+			BlobSerializerOptions serializerOptions = new BlobSerializerOptions();
+			serializerOptions.Converters.Add(new InteriorChunkedDataNodeConverter());
+
+			await using MemoryBlobWriter blobWriter = new MemoryBlobWriter(new BlobSerializerOptions());
+
+			IBlobRef<LeafChunkedDataNode> leafRef = await blobWriter.WriteBlobAsync(new LeafChunkedDataNode(new byte[] { 1, 2, 3 }));
+			int leafRefIndex = MemoryBlobWriter.GetIndex((IBlobRef)leafRef);
+			ChunkedDataNodeRef leafChunkedRef = new ChunkedDataNodeRef(3, leafRef);
+
+			List<ChunkedDataNodeRef> leafNodeRefs = Enumerable.Repeat(leafChunkedRef, 10000).ToList();
+			ChunkedDataNodeRef root = await InteriorChunkedDataNode.CreateTreeAsync(leafNodeRefs, interiorOptions, blobWriter, CancellationToken.None);
+
+			List<IBlobHandle> list = new List<IBlobHandle>();
+			await GetReadOrderAsync(root.Handle, list);
+
+			int prevIndex = Int32.MaxValue;
+			for (int idx = 0; idx < list.Count; idx++)
+			{
+				int index = MemoryBlobWriter.GetIndex((IBlobRef)list[idx].Innermost);
+				if (index != leafRefIndex)
+				{
+					Console.WriteLine("{0}", index);
+					Assert.IsTrue(index <= prevIndex);
+					prevIndex = index;
+				}
+			}
+		}
+
+		static async Task GetReadOrderAsync(IBlobHandle handle, List<IBlobHandle> list)
+		{
+			list.Add(handle);
+
+			using BlobData blobData = await handle.ReadBlobDataAsync(CancellationToken.None);
+			foreach (IBlobHandle childHandle in blobData.Imports)
+			{
+				await GetReadOrderAsync(childHandle, list);
+			}
+		}
 	}
 }
diff --git a/Engine/Source/Programs/Shared/EpicGames.Horde/Storage/IBlobWriter.cs b/Engine/Source/Programs/Shared/EpicGames.Horde/Storage/IBlobWriter.cs
--- a/Engine/Source/Programs/Shared/EpicGames.Horde/Storage/IBlobWriter.cs
+++ b/Engine/Source/Programs/Shared/EpicGames.Horde/Storage/IBlobWriter.cs
@@ -2,6 +2,8 @@
 
 using System;
 using System.Collections.Generic;
+using System.Diagnostics.CodeAnalysis;
+using System.Linq;
 using System.Threading;
 using System.Threading.Tasks;
 using EpicGames.Core;
@@ -203,4 +205,94 @@ namespace EpicGames.Horde.Storage
 		/// <inheritdoc/>
 		public abstract ValueTask DisposeAsync();
 	}
+
+	/// <summary>
+	/// Implementation of <see cref="IBlobWriter"/> which just buffers data in memory
+	/// </summary>
+	public class MemoryBlobWriter : BlobWriter
+	{
+		class BlobRef : IBlobRef
+		{
+			readonly int _index;
+			readonly IoHash _hash;
+			readonly BlobData _data;
+
+			public BlobRef(int index, BlobData data)
+			{
+				_index = index;
+				_hash = IoHash.Compute(data.Data.Span);
+				_data = data;
+			}
+
+			public IoHash Hash
+				=> _hash;
+
+			public int Index
+				=> _index;
+
+			public IBlobHandle Innermost
+				=> this;
+
+			public ValueTask FlushAsync(CancellationToken cancellationToken = default)
+				=> default;
+
+			public ValueTask<BlobData> ReadBlobDataAsync(CancellationToken cancellationToken = default)
+				=> new ValueTask<BlobData>(_data);
+
+			public bool TryGetLocator([NotNullWhen(true)] out BlobLocator locator)
+				=> throw new NotSupportedException();
+		}
+
+		readonly ChunkedMemoryWriter _memoryWriter;
+		int _nextIndex;
+
+		/// <summary>
+		/// Constructor
+		/// </summary>
+		public MemoryBlobWriter(BlobSerializerOptions options)
+			: base(options)
+		{
+			_memoryWriter = new ChunkedMemoryWriter();
+		}
+
+		/// <inheritdoc/>
+		public override ValueTask DisposeAsync()
+		{
+			GC.SuppressFinalize(this);
+			_memoryWriter.Dispose();
+			return new ValueTask();
+		}
+
+		/// <summary>
+		/// Clears the contents of this writer
+		/// </summary>
+		public void Clear()
+			=> _memoryWriter.Clear();
+
+		/// <inheritdoc/>
+		public override Task FlushAsync(CancellationToken cancellationToken = default)
+			=> Task.CompletedTask;
+
+		/// <inheritdoc/>
+		public override IBlobWriter Fork()
+			=> new MemoryBlobWriter(Options);
+
+		/// <inheritdoc/>
+		public override Memory<byte> GetOutputBuffer(int usedSize, int desiredSize)
+			=> _memoryWriter.GetMemory(usedSize, desiredSize);
+
+		/// <inheritdoc/>
+		public override ValueTask<IBlobRef> WriteBlobAsync(BlobType type, int size, IReadOnlyList<IBlobHandle> imports, IReadOnlyList<AliasInfo> aliases, CancellationToken cancellationToken)
+		{
+			Memory<byte> memory = _memoryWriter.GetMemoryAndAdvance(size);
+			BlobData data = new BlobData(type, memory, imports.ToArray());
+			return new ValueTask<IBlobRef>(new BlobRef(++_nextIndex, data));
+		}
+
+		/// <summary>
+		/// Helper function to get the index of a blob
+		/// </summary>
+		public static int GetIndex(IBlobRef handle)
+			=> ((BlobRef)handle.Innermost).Index;
+	}
 }
diff --git a/Engine/Source/Programs/Shared/EpicGames.Horde/Storage/Nodes/ChunkedDataNode.cs b/Engine/Source/Programs/Shared/EpicGames.Horde/Storage/Nodes/ChunkedDataNode.cs
--- a/Engine/Source/Programs/Shared/EpicGames.Horde/Storage/Nodes/ChunkedDataNode.cs
+++ b/Engine/Source/Programs/Shared/EpicGames.Horde/Storage/Nodes/ChunkedDataNode.cs
@@ -479,6 +479,8 @@ namespace EpicGames.Horde.Storage.Nodes
 		/// <returns>Handle to the root node of the tree</returns>
 		public static async Task<ChunkedDataNodeRef> CreateTreeAsync(List<ChunkedDataNodeRef> nodeRefs, InteriorChunkedDataNodeOptions chunkingOptions, IBlobWriter writer, CancellationToken cancellationToken)
 		{
+			await using MemoryBlobWriter memoryWriter = new MemoryBlobWriter(writer.Options);
+
 			List<ChunkedDataNodeRef> handleBuffer = new List<ChunkedDataNodeRef>();
 
 			List<(InteriorChunkedDataNode, long)> interiorNodes = new List<(InteriorChunkedDataNode, long)>();
@@ -490,14 +492,35 @@ namespace EpicGames.Horde.Storage.Nodes
 				handleBuffer.Clear();
 				foreach ((InteriorChunkedDataNode interiorNode, long interiorLength) in interiorNodes)
 				{
-					IBlobRef<InteriorChunkedDataNode> interiorHandle = await writer.WriteBlobAsync(interiorNode, cancellationToken);
+					IBlobRef<InteriorChunkedDataNode> interiorHandle = await memoryWriter.WriteBlobAsync(interiorNode, cancellationToken);
 					handleBuffer.Add(new ChunkedDataNodeRef(ChunkedDataNodeType.Interior, interiorLength, interiorHandle));
 				}
 
 				nodeRefs = handleBuffer;
 			}
 
-			return nodeRefs[0];
+			return await OrderTreeAsync(nodeRefs[0], writer, cancellationToken);
+		}
+
+		static async Task<ChunkedDataNodeRef> OrderTreeAsync(ChunkedDataNodeRef source, IBlobWriter writer, CancellationToken cancellationToken)
+		{
+			if (source.Type != ChunkedDataNodeType.Interior)
+			{
+				return source;
+			}
+
+			InteriorChunkedDataNode sourceNode = await source.Handle.ReadBlobAsync<InteriorChunkedDataNode>(cancellationToken: cancellationToken);
+
+			List<ChunkedDataNodeRef> children = new List<ChunkedDataNodeRef>(sourceNode.Children);
+			for (int idx = children.Count - 1; idx >= 0; idx--)
+			{
+				children[idx] = await OrderTreeAsync(children[idx], writer, cancellationToken);
+			}
+
+			InteriorChunkedDataNode targetNode = new InteriorChunkedDataNode(children);
+			IBlobRef<InteriorChunkedDataNode> targetHandle = await writer.WriteBlobAsync(targetNode, cancellationToken);
+
+			return new ChunkedDataNodeRef(source.Length, targetHandle);
 		}
 
 		/// <summary>
diff --git a/Engine/Source/Programs/Shared/EpicGames.Horde/Storage/Nodes/DirectoryNode.Update.cs b/Engine/Source/Programs/Shared/EpicGames.Horde/Storage/Nodes/DirectoryNode.Update.cs
--- a/Engine/Source/Programs/Shared/EpicGames.Horde/Storage/Nodes/DirectoryNode.Update.cs
+++ b/Engine/Source/Programs/Shared/EpicGames.Horde/Storage/Nodes/DirectoryNode.Update.cs
@@ -183,15 +183,23 @@ namespace EpicGames.Horde.Storage.Nodes
 	/// </summary>
 	public class DirectoryUpdate
 	{
+		class ReverseStringComparer : IComparer<string>
+		{
+			public static ReverseStringComparer Instance { get; } = new ReverseStringComparer();
+
+			public int Compare(string? x, string? y)
+				=> -StringComparer.Ordinal.Compare(x, y);
+		}
+
 		/// <summary>
 		/// Directories to be updated
 		/// </summary>
-		public SortedDictionary<string, DirectoryUpdate?> Directories { get; } = new SortedDictionary<string, DirectoryUpdate?>(StringComparer.Ordinal);
+		public SortedDictionary<string, DirectoryUpdate?> Directories { get; } = new SortedDictionary<string, DirectoryUpdate?>(ReverseStringComparer.Instance);
 
 		/// <summary>
 		/// Files to be updated
 		/// </summary>
-		public SortedDictionary<string, FileUpdate?> Files { get; } = new SortedDictionary<string, FileUpdate?>(StringComparer.Ordinal);
+		public SortedDictionary<string, FileUpdate?> Files { get; } = new SortedDictionary<string, FileUpdate?>(ReverseStringComparer.Instance);
 
 		/// <summary>
 		/// Reset this instance
