EpicGames.Shared: Resolve some editorconfig warnings and intellisense suggestions

#rnx

diff --git a/Engine/Source/Programs/Horde/Horde.Server/Ddc/RefsController.cs b/Engine/Source/Programs/Horde/Horde.Server/Ddc/RefsController.cs
--- a/Engine/Source/Programs/Horde/Horde.Server/Ddc/RefsController.cs
+++ b/Engine/Source/Programs/Horde/Horde.Server/Ddc/RefsController.cs
@@ -294,7 +294,7 @@ namespace Horde.Server.Ddc
 							byte[] packageBytes;
 							{
 								using TelemetrySpan _ = _tracer.StartActiveSpan("cbpackage.buffer").SetAttribute("operation.name", "cbpackage.buffer");
-								packageBytes = await writer.ToByteArray();
+								packageBytes = await writer.ToByteArrayAsync();
 							}
 							await using BlobContents contents = new BlobContents(packageBytes);
 							await WriteBody(contents, CustomMediaTypeNames.UnrealCompactBinaryPackage);
@@ -741,11 +741,11 @@ namespace Horde.Server.Ddc
 			_diagnosticContext.Set("Content-Length", Request.ContentLength ?? -1);
 
 			byte[] b = await RequestUtil.ReadRawBodyAsync(Request);
-			CbPackageReader packageReader = await CbPackageReader.Create(new MemoryStream(b));
+			CbPackageReader packageReader = await CbPackageReader.CreateAsync(new MemoryStream(b));
 
 			try
 			{
-				await foreach ((CbPackageAttachmentEntry entry, byte[] blob) in packageReader.IterateAttachments())
+				await foreach ((CbPackageAttachmentEntry entry, byte[] blob) in packageReader.IterateAttachmentsAsync())
 				{
 					if (entry.Flags.HasFlag(CbPackageAttachmentFlags.IsError))
 					{
diff --git a/Engine/Source/Programs/OidcToken/TokenService.cs b/Engine/Source/Programs/OidcToken/TokenService.cs
--- a/Engine/Source/Programs/OidcToken/TokenService.cs
+++ b/Engine/Source/Programs/OidcToken/TokenService.cs
@@ -92,7 +92,7 @@ namespace OidcToken
 					if (status == OidcStatus.NotLoggedIn && !Settings.CurrentValue.Unattended)
 					{
 						Logger.LogInformation("Logging in to provider {ProviderName}", providerName);
-						tokenInfo = await TokenManager.Login(providerName);
+						tokenInfo = await TokenManager.LoginAsync(providerName);
 					}
 					else
 					{
@@ -112,7 +112,7 @@ namespace OidcToken
 							else
 							{
 								Logger.LogInformation("Logging in to provider {ProviderName}", providerName);
-								tokenInfo = await TokenManager.Login(providerName);
+								tokenInfo = await TokenManager.LoginAsync(providerName);
 							}
 						}
 					}
diff --git a/Engine/Source/Programs/Shared/EpicGames.Build/Platform/Mac/ApplePlatformSDK.cs b/Engine/Source/Programs/Shared/EpicGames.Build/Platform/Mac/ApplePlatformSDK.cs
--- a/Engine/Source/Programs/Shared/EpicGames.Build/Platform/Mac/ApplePlatformSDK.cs
+++ b/Engine/Source/Programs/Shared/EpicGames.Build/Platform/Mac/ApplePlatformSDK.cs
@@ -147,18 +147,6 @@ namespace UnrealBuildBase
 
 			return InstallPath;
 		}
-		
-		/// <summary>
-		/// Runs a command line process, and returns simple StdOut output. This doesn't handle errors or return codes
-		/// </summary>
-		/// <returns>The entire StdOut generated from the process as a single trimmed string</returns>
-		/// <param name="Command">Command to run</param>
-		/// <param name="Args">Arguments to Command</param>
-		/// <param name="Logger">Logger for output</param>
-		private static string RunLocalProcessAndReturnStdOut(string Command, string Args, ILogger? Logger = null)
-		{
-			return RunLocalProcessAndReturnStdOut(Command, Args, out _, Logger);	
-		}
 
 		/// <summary>
 		/// Runs a command line process, and returns simple StdOut output.
diff --git a/Engine/Source/Programs/Shared/EpicGames.Build/System/DirectoryItem.cs b/Engine/Source/Programs/Shared/EpicGames.Build/System/DirectoryItem.cs
--- a/Engine/Source/Programs/Shared/EpicGames.Build/System/DirectoryItem.cs
+++ b/Engine/Source/Programs/Shared/EpicGames.Build/System/DirectoryItem.cs
@@ -377,12 +377,12 @@ namespace UnrealBuildBase
 				return true;
 			}
 
-			if (ReferenceEquals(obj, null))
+			if (obj is null)
 			{
 				return false;
 			}
 
-			return Equals((DirectoryItem?)obj);
+			return Equals(obj as DirectoryItem);
 		}
 
 		public override int GetHashCode()
@@ -392,9 +392,9 @@ namespace UnrealBuildBase
 
 		public static bool operator ==(DirectoryItem? left, DirectoryItem? right)
 		{
-			if (ReferenceEquals(left, null))
+			if (left is null)
 			{
-				return ReferenceEquals(right, null);
+				return right is null;
 			}
 
 			return left.Equals(right);
@@ -407,22 +407,22 @@ namespace UnrealBuildBase
 
 		public static bool operator <(DirectoryItem? left, DirectoryItem? right)
 		{
-			return ReferenceEquals(left, null) ? !ReferenceEquals(right, null) : left.CompareTo(right) < 0;
+			return left is null ? right is not null : left.CompareTo(right) < 0;
 		}
 
 		public static bool operator <=(DirectoryItem? left, DirectoryItem? right)
 		{
-			return ReferenceEquals(left, null) || left.CompareTo(right) <= 0;
+			return left is null || left.CompareTo(right) <= 0;
 		}
 
 		public static bool operator >(DirectoryItem? left, DirectoryItem? right)
 		{
-			return !ReferenceEquals(left, null) && left.CompareTo(right) > 0;
+			return left is not null && left.CompareTo(right) > 0;
 		}
 
 		public static bool operator >=(DirectoryItem? left, DirectoryItem? right)
 		{
-			return ReferenceEquals(left, null) ? ReferenceEquals(right, null) : left.CompareTo(right) >= 0;
+			return left is null ? right is null : left.CompareTo(right) >= 0;
 		}
 		#endregion
 	}
@@ -474,6 +474,5 @@ namespace UnrealBuildBase
 		{
 			return Reader.ReadDirectoryItem()!.Location;
 		}
-
 	}
 }
diff --git a/Engine/Source/Programs/Shared/EpicGames.Build/System/FileHasher.cs b/Engine/Source/Programs/Shared/EpicGames.Build/System/FileHasher.cs
--- a/Engine/Source/Programs/Shared/EpicGames.Build/System/FileHasher.cs
+++ b/Engine/Source/Programs/Shared/EpicGames.Build/System/FileHasher.cs
@@ -16,7 +16,7 @@ namespace UnrealBuildBase
 {
 	public class FileHasher
 	{
-		private struct CachedDigest
+		private readonly struct CachedDigest
 		{
 			[JsonPropertyName("l")]
 			public long Length { get; init; }
diff --git a/Engine/Source/Programs/Shared/EpicGames.Build/System/FileItem.cs b/Engine/Source/Programs/Shared/EpicGames.Build/System/FileItem.cs
--- a/Engine/Source/Programs/Shared/EpicGames.Build/System/FileItem.cs
+++ b/Engine/Source/Programs/Shared/EpicGames.Build/System/FileItem.cs
@@ -291,12 +291,12 @@ namespace UnrealBuildBase
 				return true;
 			}
 
-			if (ReferenceEquals(obj, null))
+			if (obj is null)
 			{
 				return false;
 			}
 
-			return Equals((FileItem?)obj);
+			return Equals(obj as FileItem);
 		}
 
 		public override int GetHashCode()
@@ -306,9 +306,9 @@ namespace UnrealBuildBase
 
 		public static bool operator ==(FileItem? left, FileItem? right)
 		{
-			if (ReferenceEquals(left, null))
+			if (left is null)
 			{
-				return ReferenceEquals(right, null);
+				return right is null;
 			}
 
 			return left.Equals(right);
@@ -321,22 +321,22 @@ namespace UnrealBuildBase
 
 		public static bool operator <(FileItem? left, FileItem? right)
 		{
-			return ReferenceEquals(left, null) ? !ReferenceEquals(right, null) : left.CompareTo(right) < 0;
+			return left is null ? right is not null : left.CompareTo(right) < 0;
 		}
 
 		public static bool operator <=(FileItem? left, FileItem? right)
 		{
-			return ReferenceEquals(left, null) || left.CompareTo(right) <= 0;
+			return left is null || left.CompareTo(right) <= 0;
 		}
 
 		public static bool operator >(FileItem? left, FileItem? right)
 		{
-			return !ReferenceEquals(left, null) && left.CompareTo(right) > 0;
+			return left is not null && left.CompareTo(right) > 0;
 		}
 
 		public static bool operator >=(FileItem? left, FileItem? right)
 		{
-			return ReferenceEquals(left, null) ? ReferenceEquals(right, null) : left.CompareTo(right) >= 0;
+			return left is null ? right is null : left.CompareTo(right) >= 0;
 		}
 		#endregion
 	}
diff --git a/Engine/Source/Programs/Shared/EpicGames.Build/UEBuildPlatformSDK.cs b/Engine/Source/Programs/Shared/EpicGames.Build/UEBuildPlatformSDK.cs
--- a/Engine/Source/Programs/Shared/EpicGames.Build/UEBuildPlatformSDK.cs
+++ b/Engine/Source/Programs/Shared/EpicGames.Build/UEBuildPlatformSDK.cs
@@ -85,7 +85,7 @@ namespace EpicGames.Core
 		{
 		}
 
-		public void UpdateCurrent(string InCurrent, string? Hint,  SDKCollection? Collection)
+		public void UpdateCurrent(string InCurrent, string? Hint, SDKCollection? Collection)
 		{
 			Current = InCurrent;
 			if (Collection != null)
@@ -102,7 +102,7 @@ namespace EpicGames.Core
 			Validity = (Current != null && Collection.IsValid(Current, this)) ? SDKStatus.Valid : SDKStatus.Invalid;
 		}
 
-		public string ToString(bool bIncludeCurrent=true)
+		public string ToString(bool bIncludeCurrent = true)
 		{
 			return $"MinVersion={Min ?? ""}, MaxVersion={Max ?? ""}" + (bIncludeCurrent ? $", Current={Current ?? ""}" : "");
 		}
@@ -193,20 +193,20 @@ namespace EpicGames.Core
 
 			return bAllUngroupedAreValid && bAllGroupsHaveOneValid;
 		}
-	
+
 		public bool IsAutoSDKValid()
 		{
 			return AutoSdk?.Validity == SDKStatus.Valid;
 		}
 
 
-		public string ToString(bool bIncludeCurrent=true)
+		public string ToString(bool bIncludeCurrent = true)
 		{
 			// by default print something like "MinVersion=1.0, MaxVersion=2.0"
 			return ToString("Version", "Version");
 		}
 
-		public string ToString(string Descriptor, string? CurrentDescriptor=null)
+		public string ToString(string Descriptor, string? CurrentDescriptor = null)
 		{
 			if (Sdks.Count == 1)
 			{
@@ -313,7 +313,7 @@ namespace EpicGames.Core
 			{
 				throw new Exception(string.Format("Re-registering SDK for {0}. All Platforms must have a unique SDK object", PlatformName));
 			}
-			
+
 			SDKRegistry.Add(PlatformName, SDK);
 
 			SDK.Init(PlatformName, bIsSdkAllowedOnHost);
@@ -417,7 +417,7 @@ namespace EpicGames.Core
 
 		// Contains a list of projects that had per-project SDK overrides, used for validating conflicting SDKs
 		public List<FileReference> ProjectsThatOverrodeSDK = new();
-	
+
 		// String name of the platform (will match an UnrealTargetPlatform)
 		public string? PlatformName;
 
@@ -460,17 +460,14 @@ namespace EpicGames.Core
 			return GetAllSDKInfo().Sdks.FirstOrDefault(x => string.Compare(x.Name, SDKName, true) == 0);
 		}
 
-		public SDKDescriptor? GetSoftwareInfo(string? SDKName=null)
+		public SDKDescriptor? GetSoftwareInfo(string? SDKName = null)
 		{
-			if (SDKName == null)
-			{
-				// todo: make this queryable?
-				SDKName = "Software";
-			}
+			// todo: make this queryable?
+			SDKName ??= "Software";
 			return GetAllSoftwareInfo().Sdks.FirstOrDefault(x => string.Compare(x.Name, SDKName, true) == 0);
 		}
 
-		public SDKCollection GetAllSoftwareInfo(string? DeviceType=null, string? Current=null)
+		public SDKCollection GetAllSoftwareInfo(string? DeviceType = null, string? Current = null)
 		{
 			SDKCollection AllSoftwares = new SDKCollection(this);
 
@@ -829,7 +826,7 @@ namespace EpicGames.Core
 		/// <param name="StringValueB">Second Version to compare</param>
 		/// <param name="Hint">A platform specific hint that can help guide conversion (usually SDKName or device type)</param>
 		/// <returns>Comparison integer</returns>
-		public int SdkVersionsCompare( string? StringValueA, string? StringValueB, string? Hint = null )
+		public int SdkVersionsCompare(string? StringValueA, string? StringValueB, string? Hint = null)
 		{
 			UInt64 ValueA, ValueB;
 			TryConvertVersionToInt(StringValueA, out ValueA, Hint);
@@ -1001,7 +998,7 @@ namespace EpicGames.Core
 				Platform = "Windows";
 			}
 
-			Func<DirectoryReference, bool, FileReference?> MakeConfigFilename = (RootDir, bIsRequired) =>
+			FileReference? MakeConfigFilename(DirectoryReference RootDir, bool bIsRequired)
 			{
 				FileReference PlatformExtensionLocation = FileReference.Combine(RootDir, "Platforms", Platform, "Config", $"{Platform}_SDK.json");
 				if (FileReference.Exists(PlatformExtensionLocation))
@@ -1018,7 +1015,7 @@ namespace EpicGames.Core
 					throw new Exception($"Failed to find required SDK.json for {Platform}. Looked in '{StandardLocation}' and '{PlatformExtensionLocation}'.");
 				}
 				return null;
-			};
+			}
 
 			// if the SDK isn't allowed on the host, then allow it to not exist
 			FileReference EngineSDKConfigFile = MakeConfigFilename(Unreal.EngineDirectory, bIsSdkAllowedOnHost)!;
@@ -1059,7 +1056,7 @@ namespace EpicGames.Core
 							// now check if it was already overridden, in which case we have a conflict we can't resolve, so error
 							if (ProjectsThatOverrodeSDK.Count > 0 && !OverrideMainVersion.Equals(ConfigSDKVersions[MainVersionKey], StringComparison.OrdinalIgnoreCase))
 							{
-								throw new Exception($"Project {ProjectFile.GetFileNameWithoutAnyExtensions()} wants to override {Platform} SDK to {OverrideMainVersion}, but it was already overridden to version {ConfigSDKVersions[MainVersionKey]}");							
+								throw new Exception($"Project {ProjectFile.GetFileNameWithoutAnyExtensions()} wants to override {Platform} SDK to {OverrideMainVersion}, but it was already overridden to version {ConfigSDKVersions[MainVersionKey]}");
 							}
 
 							Logger.LogWarning("Project {Project} is overriding {Platform} SDK to {OverrideMainVersion}", ProjectFile.GetFileNameWithoutAnyExtensions(), Platform, OverrideMainVersion);
@@ -1140,10 +1137,10 @@ namespace EpicGames.Core
 		public string GetRequiredVersionFromConfig(string VersionName)
 		{
 			// when bIsRequired is true, then we know it will return non-null
-			return GetVersionFromConfig(VersionName, bIsRequired:true)!;
+			return GetVersionFromConfig(VersionName, bIsRequired: true)!;
 		}
 
-		public string? GetVersionFromConfig(string VersionName, bool bIsRequired=false)
+		public string? GetVersionFromConfig(string VersionName, bool bIsRequired = false)
 		{
 			string? Version;
 			// look up both Version_Host and Version (Host specific version wins)
@@ -1165,7 +1162,7 @@ namespace EpicGames.Core
 			return GetVersionNumberFromConfig(VersionName, true)!;
 		}
 
-		public VersionNumber? GetVersionNumberFromConfig(string VersionName, bool bIsRequired=false)
+		public VersionNumber? GetVersionNumberFromConfig(string VersionName, bool bIsRequired = false)
 		{
 			string? VersionString = GetVersionFromConfig(VersionName, bIsRequired);
 
@@ -1221,7 +1218,7 @@ namespace EpicGames.Core
 			List<VersionNumberRange> Ranges = new();
 			if (ConfigSDKVersionArrays.TryGetValue(GetHostSpecificVersionName(VersionName), out VersionRanges) || ConfigSDKVersionArrays.TryGetValue(VersionName, out VersionRanges))
 			{
-				foreach (string VersionRange in VersionRanges) 
+				foreach (string VersionRange in VersionRanges)
 				{
 					VersionNumberRange? Range = ParseVersionNumberRange(VersionRange);
 					if (Range != null)
@@ -1504,10 +1501,10 @@ namespace EpicGames.Core
 			{
 				String InstalledSDKVersionString = GetAutoSDKDirectoryForMainVersion();
 				String PlatformSDKRoot = GetPathToPlatformAutoSDKs();
-                if (!Directory.Exists(PlatformSDKRoot))
-                {
-                    Directory.CreateDirectory(PlatformSDKRoot);
-                }
+				if (!Directory.Exists(PlatformSDKRoot))
+				{
+					Directory.CreateDirectory(PlatformSDKRoot);
+				}
 
 				{
 					string VersionFilename = Path.Combine(PlatformSDKRoot, CurrentlyInstalledSDKStringManifest);
@@ -1812,18 +1809,18 @@ namespace EpicGames.Core
 					}
 
 
-                    // actually perform the PATH stripping / adding.
-                    String? OrigPathVar = Environment.GetEnvironmentVariable("PATH");
-                    String PathDelimiter = RuntimePlatform.IsWindows ? ";" : ":";
-                    String[] PathVars = { };
-                    if (!String.IsNullOrEmpty(OrigPathVar))
-                    {
-                        PathVars = OrigPathVar.Split(PathDelimiter.ToCharArray());
-                    }
-                    else
-                    {
-                        Logger.LogDebug("Path environment variable is null during AutoSDK");
-                    }
+					// actually perform the PATH stripping / adding.
+					String? OrigPathVar = Environment.GetEnvironmentVariable("PATH");
+					String PathDelimiter = RuntimePlatform.IsWindows ? ";" : ":";
+					String[] PathVars = { };
+					if (!String.IsNullOrEmpty(OrigPathVar))
+					{
+						PathVars = OrigPathVar.Split(PathDelimiter.ToCharArray());
+					}
+					else
+					{
+						Logger.LogDebug("Path environment variable is null during AutoSDK");
+					}
 
 					List<String> ModifiedPathVars = new List<string>();
 					ModifiedPathVars.AddRange(PathVars);
@@ -1986,7 +1983,7 @@ namespace EpicGames.Core
 		{
 			String AutoSDKSetupVarName = GetPlatformAutoSDKSetupEnvVar();
 			OutAutoSDKSetupValue = Environment.GetEnvironmentVariable(AutoSDKSetupVarName);
-			
+
 			if (!String.IsNullOrEmpty(OutAutoSDKSetupValue))
 			{
 				return true;
@@ -1996,12 +1993,12 @@ namespace EpicGames.Core
 
 		public SDKStatus HasRequiredManualSDK()
 		{
-// 			if (HasSetupAutoSDK())
-// 			{
-// 				return SDKStatus.Invalid;
-// 			}
-//
-//			// manual installs are always invalid if we have modified the process's environment for AutoSDKs
+			// 			if (HasSetupAutoSDK())
+			// 			{
+			// 				return SDKStatus.Invalid;
+			// 			}
+			//
+			//			// manual installs are always invalid if we have modified the process's environment for AutoSDKs
 			return HasRequiredManualSDKInternal();
 		}
 
diff --git a/Engine/Source/Programs/Shared/EpicGames.Core/AppContainer.cs b/Engine/Source/Programs/Shared/EpicGames.Core/AppContainer.cs
--- a/Engine/Source/Programs/Shared/EpicGames.Core/AppContainer.cs
+++ b/Engine/Source/Programs/Shared/EpicGames.Core/AppContainer.cs
@@ -132,11 +132,7 @@ public sealed class AppContainer : IDisposable
 			}
 			
 			string? sidString = Marshal.PtrToStringAuto(sidStringPtr);
-			if (sidString == null)
-			{
-				throw new AppContainerException("Failed converting string pointer");
-			}
-			return sidString;
+			return sidString ?? throw new AppContainerException("Failed converting string pointer");
 		}
 		finally
 		{
diff --git a/Engine/Source/Programs/Shared/EpicGames.Core/BinaryArchive.cs b/Engine/Source/Programs/Shared/EpicGames.Core/BinaryArchive.cs
--- a/Engine/Source/Programs/Shared/EpicGames.Core/BinaryArchive.cs
+++ b/Engine/Source/Programs/Shared/EpicGames.Core/BinaryArchive.cs
@@ -307,11 +307,8 @@ namespace EpicGames.Core
 				}
 
 				// Find the type constructor
-				ConstructorInfo? constructorInfo = type.GetConstructor(BindingFlags.Instance | BindingFlags.Public | BindingFlags.NonPublic, null, Array.Empty<Type>(), null);
-				if (constructorInfo == null)
-				{
-					throw new NotImplementedException($"Type '{type.Name}' does not have a parameterless constructor");
-				}
+				ConstructorInfo? constructorInfo = type.GetConstructor(BindingFlags.Instance | BindingFlags.Public | BindingFlags.NonPublic, null, Array.Empty<Type>(), null)
+					?? throw new NotImplementedException($"Type '{type.Name}' does not have a parameterless constructor");
 
 				// Create the methods
 				DynamicMethod writerMethod = new DynamicMethod($"BinaryArchiveWriter_Dynamic_{type.Name}", typeof(void), new[] { typeof(BinaryArchiveWriter), type });
diff --git a/Engine/Source/Programs/Shared/EpicGames.Core/BuzHash.cs b/Engine/Source/Programs/Shared/EpicGames.Core/BuzHash.cs
--- a/Engine/Source/Programs/Shared/EpicGames.Core/BuzHash.cs
+++ b/Engine/Source/Programs/Shared/EpicGames.Core/BuzHash.cs
@@ -148,7 +148,7 @@ namespace EpicGames.Core
 			{
 				for (int idx = 0; idx < data.Length; idx++)
 				{
-					state = state ^ Rol32(s_table[dataPtr[idx]], --count);
+					state ^= Rol32(s_table[dataPtr[idx]], --count);
 				}
 			}
 			return state;
@@ -193,7 +193,7 @@ namespace EpicGames.Core
 						state = hash;
 						return idx;
 					}
-					hash = hash ^ Rol32(table[prevPtr[idx]], prev.Length - 1);
+					hash ^= Rol32(table[prevPtr[idx]], prev.Length - 1);
 				}
 				state = hash;
 			}
diff --git a/Engine/Source/Programs/Shared/EpicGames.Core/ByteString.cs b/Engine/Source/Programs/Shared/EpicGames.Core/ByteString.cs
--- a/Engine/Source/Programs/Shared/EpicGames.Core/ByteString.cs
+++ b/Engine/Source/Programs/Shared/EpicGames.Core/ByteString.cs
@@ -16,7 +16,7 @@ namespace EpicGames.Core
 	[JsonSchemaString]
 	[JsonConverter(typeof(ByteStringJsonConverter))]
 	[TypeConverter(typeof(ByteStringTypeConverter))]
-	public struct ByteString : IEquatable<ByteString>, IComparable<ByteString>
+	public readonly struct ByteString : IEquatable<ByteString>, IComparable<ByteString>
 	{
 		/// <summary>
 		/// Underlying data for this string
@@ -35,28 +35,28 @@ namespace EpicGames.Core
 		public ByteString(ReadOnlyMemory<byte> data) => Data = data;
 
 		/// <summary>
-		/// Parses a byte string from a hexidecimal string
+		/// Parses a byte string from a hexadecimal string
 		/// </summary>
 		/// <param name="text">Text to parse</param>
 		/// <returns>New byte string instance</returns>
 		public static ByteString Parse(ReadOnlySpan<byte> text) => new ByteString(StringUtils.ParseHexString(text));
 
 		/// <summary>
-		/// Parses a byte string from a hexidecimal string
+		/// Parses a byte string from a hexadecimal string
 		/// </summary>
 		/// <param name="text">Text to parse</param>
 		/// <returns>New byte string instance</returns>
 		public static ByteString Parse(Utf8String text) => Parse(text.Span);
 
 		/// <summary>
-		/// Parses a byte string from a hexidecimal string
+		/// Parses a byte string from a hexadecimal string
 		/// </summary>
 		/// <param name="text">Text to parse</param>
 		/// <returns>New byte string instance</returns>
 		public static ByteString Parse(string text) => new ByteString(StringUtils.ParseHexString(text));
 
 		/// <summary>
-		/// Attempts to parse a hexidecimal string as a byte string
+		/// Attempts to parse a hexadecimal string as a byte string
 		/// </summary>
 		/// <param name="text">Text to parse</param>
 		/// <param name="byteString">On success, receives the parsed string</param>
@@ -77,7 +77,7 @@ namespace EpicGames.Core
 		}
 
 		/// <summary>
-		/// Attempts to parse a hexidecimal string as a byte string
+		/// Attempts to parse a hexadecimal string as a byte string
 		/// </summary>
 		/// <param name="text">Text to parse</param>
 		/// <param name="byteString">On success, receives the parsed string</param>
@@ -85,7 +85,7 @@ namespace EpicGames.Core
 		public static bool TryParse(Utf8String text, out ByteString byteString) => TryParse(text.Span, out byteString);
 
 		/// <summary>
-		/// Attempts to parse a hexidecimal string as a byte string
+		/// Attempts to parse a hexadecimal string as a byte string
 		/// </summary>
 		/// <param name="text">Text to parse</param>
 		/// <param name="byteString">On success, receives the parsed string</param>
@@ -135,7 +135,7 @@ namespace EpicGames.Core
 		public override string ToString() => StringUtils.FormatHexString(Data.Span);
 
 		/// <summary>
-		/// Creates a Utf8 string represending the bytes in this string
+		/// Creates a Utf8 string representing the bytes in this string
 		/// </summary>
 		public Utf8String ToUtf8String() => StringUtils.FormatUtf8HexString(Data.Span);
 
diff --git a/Engine/Source/Programs/Shared/EpicGames.Core/Digest.cs b/Engine/Source/Programs/Shared/EpicGames.Core/Digest.cs
--- a/Engine/Source/Programs/Shared/EpicGames.Core/Digest.cs
+++ b/Engine/Source/Programs/Shared/EpicGames.Core/Digest.cs
@@ -14,7 +14,7 @@ namespace EpicGames.Core
 	/// <summary>
 	/// Struct representing a weakly typed hash value. Counterpart to <see cref="Digest{T}"/> - a strongly typed digest.
 	/// </summary>
-	public struct Digest : IEquatable<Digest>
+	public readonly struct Digest : IEquatable<Digest>
 	{
 		/// <summary>
 		/// Memory storing the digest data
@@ -242,7 +242,7 @@ namespace EpicGames.Core
 	/// <summary>
 	/// Generic HashValue implementation
 	/// </summary>
-	public struct Digest<T> : IEquatable<Digest<T>> where T : DigestTraits, new()
+	public readonly struct Digest<T> : IEquatable<Digest<T>> where T : DigestTraits, new()
 	{
 		/// <summary>
 		/// Traits instance
diff --git a/Engine/Source/Programs/Shared/EpicGames.Core/HelpUtils.cs b/Engine/Source/Programs/Shared/EpicGames.Core/HelpUtils.cs
--- a/Engine/Source/Programs/Shared/EpicGames.Core/HelpUtils.cs
+++ b/Engine/Source/Programs/Shared/EpicGames.Core/HelpUtils.cs
@@ -74,7 +74,7 @@ namespace EpicGames.Core
 		public static string GetDescription(Type type)
 		{
 			StringBuilder descriptionText = new StringBuilder();
-			foreach (DescriptionAttribute attribute in type.GetCustomAttributes(typeof(DescriptionAttribute), false))
+			foreach (DescriptionAttribute attribute in type.GetCustomAttributes(typeof(DescriptionAttribute), false).OfType<DescriptionAttribute>())
 			{
 				if (descriptionText.Length > 0)
 				{
@@ -186,7 +186,7 @@ namespace EpicGames.Core
 						lines.Add(paramString + descriptionLine);
 
 						// we replace the param string on subsequent lines with white space of the same length
-						paramString = string.Empty.PadRight(indentString.Length + rightPadding);
+						paramString = String.Empty.PadRight(indentString.Length + rightPadding);
 					}
 				}
 			}
diff --git a/Engine/Source/Programs/Shared/EpicGames.Core/ILoggerIndent.cs b/Engine/Source/Programs/Shared/EpicGames.Core/ILoggerIndent.cs
--- a/Engine/Source/Programs/Shared/EpicGames.Core/ILoggerIndent.cs
+++ b/Engine/Source/Programs/Shared/EpicGames.Core/ILoggerIndent.cs
@@ -69,7 +69,7 @@ namespace EpicGames.Core
 		/// Struct to wrap a formatted set of log values with applied indent
 		/// </summary>
 		/// <typeparam name="TState">Arbitrary type parameter</typeparam>
-		struct FormattedLogValues<TState> : IEnumerable<KeyValuePair<string, object>>
+		readonly struct FormattedLogValues<TState> : IEnumerable<KeyValuePair<string, object>>
 		{
 			/// <summary>
 			/// The indent to apply
diff --git a/Engine/Source/Programs/Shared/EpicGames.Core/JsonKnownTypes.cs b/Engine/Source/Programs/Shared/EpicGames.Core/JsonKnownTypes.cs
--- a/Engine/Source/Programs/Shared/EpicGames.Core/JsonKnownTypes.cs
+++ b/Engine/Source/Programs/Shared/EpicGames.Core/JsonKnownTypes.cs
@@ -83,7 +83,7 @@ namespace EpicGames.Core
 				{
 					throw new InvalidOperationException($"{knownType} is not derived from {typeof(T)}");
 				}
-				foreach (JsonDiscriminatorAttribute discriminatorAttribute in knownType.GetCustomAttributes(typeof(JsonDiscriminatorAttribute), true))
+				foreach (JsonDiscriminatorAttribute discriminatorAttribute in knownType.GetCustomAttributes(typeof(JsonDiscriminatorAttribute), true).OfType<JsonDiscriminatorAttribute>())
 				{
 					_discriminatorToType.Add(discriminatorAttribute.Name, knownType);
 				}
diff --git a/Engine/Source/Programs/Shared/EpicGames.Core/JsonLogEvent.cs b/Engine/Source/Programs/Shared/EpicGames.Core/JsonLogEvent.cs
--- a/Engine/Source/Programs/Shared/EpicGames.Core/JsonLogEvent.cs
+++ b/Engine/Source/Programs/Shared/EpicGames.Core/JsonLogEvent.cs
@@ -12,7 +12,7 @@ namespace EpicGames.Core
 	/// <summary>
 	/// Defines a preformatted Json log event, which can pass through raw Json data directly or format it as a regular string
 	/// </summary>
-	public struct JsonLogEvent : IEnumerable<KeyValuePair<string, object?>>
+	public readonly struct JsonLogEvent : IEnumerable<KeyValuePair<string, object?>>
 	{
 		/// <summary>
 		/// The log level
@@ -70,12 +70,7 @@ namespace EpicGames.Core
 				return jsonLogEvent;
 			}
 
-			LogEvent? logEvent = state as LogEvent;
-			if (logEvent == null)
-			{
-				logEvent = LogEvent.FromState(logLevel, eventId, state, exception, formatter);
-			}
-
+			LogEvent? logEvent = state as LogEvent ?? LogEvent.FromState(logLevel, eventId, state, exception, formatter);
 			return new JsonLogEvent(logLevel, eventId, 0, 1, logEvent.ToJsonBytes());
 		}
 
diff --git a/Engine/Source/Programs/Shared/EpicGames.Core/JsonSchema.cs b/Engine/Source/Programs/Shared/EpicGames.Core/JsonSchema.cs
--- a/Engine/Source/Programs/Shared/EpicGames.Core/JsonSchema.cs
+++ b/Engine/Source/Programs/Shared/EpicGames.Core/JsonSchema.cs
@@ -90,7 +90,7 @@ namespace EpicGames.Core
 			/// <inheritdoc/>
 			public void WriteType(JsonSchemaType type)
 			{
-				if (!(type is JsonSchemaPrimitiveType))
+				if (type is not JsonSchemaPrimitiveType)
 				{
 					TypeRefCount.TryGetValue(type, out int refCount);
 					if (refCount < 2)
diff --git a/Engine/Source/Programs/Shared/EpicGames.Core/ListSegment.cs b/Engine/Source/Programs/Shared/EpicGames.Core/ListSegment.cs
--- a/Engine/Source/Programs/Shared/EpicGames.Core/ListSegment.cs
+++ b/Engine/Source/Programs/Shared/EpicGames.Core/ListSegment.cs
@@ -166,7 +166,7 @@ namespace EpicGames.Core
 			readonly int _maxIndex;
 
 			/// <inheritdoc/>
-			public T Current
+			public readonly T Current
 			{
 				get
 				{
@@ -190,7 +190,7 @@ namespace EpicGames.Core
 			}
 
 			/// <inheritdoc/>
-			void IDisposable.Dispose()
+			readonly void IDisposable.Dispose()
 			{
 			}
 
@@ -201,7 +201,7 @@ namespace EpicGames.Core
 			}
 
 			/// <inheritdoc/>
-			object? IEnumerator.Current => Current;
+			readonly object? IEnumerator.Current => Current;
 
 			/// <inheritdoc/>
 			void IEnumerator.Reset() => throw new InvalidOperationException();
diff --git a/Engine/Source/Programs/Shared/EpicGames.Core/Log.cs b/Engine/Source/Programs/Shared/EpicGames.Core/Log.cs
--- a/Engine/Source/Programs/Shared/EpicGames.Core/Log.cs
+++ b/Engine/Source/Programs/Shared/EpicGames.Core/Log.cs
@@ -1207,7 +1207,7 @@ namespace EpicGames.Core
 
 						string timePrefixActual =
 							IncludeTimestamps &&
-							!(listener is DefaultTraceListener) // no timestamps when writing to the Visual Studio debug window
+							listener is not DefaultTraceListener // no timestamps when writing to the Visual Studio debug window
 							? timePrefix
 							: String.Empty;
 
diff --git a/Engine/Source/Programs/Shared/EpicGames.Core/Md5Hash.cs b/Engine/Source/Programs/Shared/EpicGames.Core/Md5Hash.cs
--- a/Engine/Source/Programs/Shared/EpicGames.Core/Md5Hash.cs
+++ b/Engine/Source/Programs/Shared/EpicGames.Core/Md5Hash.cs
@@ -15,7 +15,7 @@ namespace EpicGames.Core
 	/// Struct representing a strongly typed Md5Hash value
 	/// </summary>
 	[TypeConverter(typeof(Md5HashTypeConverter))]
-	public struct Md5Hash : IEquatable<Md5Hash>, IComparable<Md5Hash>
+	public readonly struct Md5Hash : IEquatable<Md5Hash>, IComparable<Md5Hash>
 	{
 		/// <summary>
 		/// Length of an Md5Hash
diff --git a/Engine/Source/Programs/Shared/EpicGames.Core/MessageTemplate.cs b/Engine/Source/Programs/Shared/EpicGames.Core/MessageTemplate.cs
--- a/Engine/Source/Programs/Shared/EpicGames.Core/MessageTemplate.cs
+++ b/Engine/Source/Programs/Shared/EpicGames.Core/MessageTemplate.cs
@@ -147,10 +147,7 @@ namespace EpicGames.Core
 						{
 							break;
 						}
-						if (names == null)
-						{
-							names = new List<(int, int)>();
-						}
+						names ??= new List<(int, int)>();
 
 						names.Add((startIdx, idx - startIdx));
 					}
diff --git a/Engine/Source/Programs/Shared/EpicGames.Core/ReadOnlyMemoryArray.cs b/Engine/Source/Programs/Shared/EpicGames.Core/ReadOnlyMemoryArray.cs
--- a/Engine/Source/Programs/Shared/EpicGames.Core/ReadOnlyMemoryArray.cs
+++ b/Engine/Source/Programs/Shared/EpicGames.Core/ReadOnlyMemoryArray.cs
@@ -28,7 +28,7 @@ namespace EpicGames.Core
 	/// <summary>
 	/// Represents an array of uniformly sized <see cref="ReadOnlyMemory{Byte}"/> objects in a parent <see cref="ReadOnlyMemory{Byte}"/> instance.
 	/// </summary>
-	public struct UniformReadOnlyMemoryArray : IReadOnlyMemoryArray
+	public readonly struct UniformReadOnlyMemoryArray : IReadOnlyMemoryArray
 	{
 		/// <inheritdoc/>
 		public ReadOnlyMemory<byte> Data { get; }
@@ -85,7 +85,7 @@ namespace EpicGames.Core
 	/// <summary>
 	/// Represents an array of variably sized <see cref="ReadOnlyMemory{Byte}"/> objects in a parent <see cref="ReadOnlyMemory{Byte}"/> instance.
 	/// </summary>
-	public struct JaggedReadOnlyMemoryArray : IReadOnlyMemoryArray
+	public readonly struct JaggedReadOnlyMemoryArray : IReadOnlyMemoryArray
 	{
 		/// <inheritdoc/>
 		public ReadOnlyMemory<byte> Data { get; }
diff --git a/Engine/Source/Programs/Shared/EpicGames.Core/SemVer.cs b/Engine/Source/Programs/Shared/EpicGames.Core/SemVer.cs
--- a/Engine/Source/Programs/Shared/EpicGames.Core/SemVer.cs
+++ b/Engine/Source/Programs/Shared/EpicGames.Core/SemVer.cs
@@ -8,9 +8,9 @@ namespace EpicGames.Core
 	/// <summary>
 	/// Wrapper for a semver version string (https://semver.org/)
 	/// </summary>
-	public struct SemVer : IEquatable<SemVer>, IComparable<SemVer>
+	public readonly struct SemVer : IEquatable<SemVer>, IComparable<SemVer>
 	{
-		ref struct Field
+		readonly ref struct Field
 		{
 			public readonly int Pos;
 			public readonly ReadOnlySpan<char> Text;
diff --git a/Engine/Source/Programs/Shared/EpicGames.Core/Sha1Hash.cs b/Engine/Source/Programs/Shared/EpicGames.Core/Sha1Hash.cs
--- a/Engine/Source/Programs/Shared/EpicGames.Core/Sha1Hash.cs
+++ b/Engine/Source/Programs/Shared/EpicGames.Core/Sha1Hash.cs
@@ -17,7 +17,7 @@ namespace EpicGames.Core
 	/// </summary>
 	[JsonConverter(typeof(Sha1HashJsonConverter))]
 	[TypeConverter(typeof(Sha1HashTypeConverter))]
-	public struct Sha1Hash : IEquatable<Sha1Hash>, IComparable<Sha1Hash>
+	public readonly struct Sha1Hash : IEquatable<Sha1Hash>, IComparable<Sha1Hash>
 	{
 		/// <summary>
 		/// Length of an Sha1Hash
diff --git a/Engine/Source/Programs/Shared/EpicGames.Core/ThreadPoolWorkQueue.cs b/Engine/Source/Programs/Shared/EpicGames.Core/ThreadPoolWorkQueue.cs
--- a/Engine/Source/Programs/Shared/EpicGames.Core/ThreadPoolWorkQueue.cs
+++ b/Engine/Source/Programs/Shared/EpicGames.Core/ThreadPoolWorkQueue.cs
@@ -115,11 +115,7 @@ namespace EpicGames.Core
 		/// </summary>
 		ManualResetEvent GetEmptyEvent()
 		{
-			ManualResetEvent? emptyEvent = _emptyEvent;
-			if (emptyEvent == null)
-			{
-				throw new ObjectDisposedException(typeof(ThreadPoolWorkQueue).Name);
-			}
+			ManualResetEvent? emptyEvent = _emptyEvent ?? throw new ObjectDisposedException(typeof(ThreadPoolWorkQueue).Name);
 			return emptyEvent;
 		}
 
diff --git a/Engine/Source/Programs/Shared/EpicGames.Core/Utf8String.cs b/Engine/Source/Programs/Shared/EpicGames.Core/Utf8String.cs
--- a/Engine/Source/Programs/Shared/EpicGames.Core/Utf8String.cs
+++ b/Engine/Source/Programs/Shared/EpicGames.Core/Utf8String.cs
@@ -9,7 +9,7 @@ namespace EpicGames.Core
 	/// <summary>
 	/// Represents a memory region which can be treated as a utf-8 string.
 	/// </summary>
-	public struct Utf8String : IEquatable<Utf8String>, IComparable<Utf8String>
+	public readonly struct Utf8String : IEquatable<Utf8String>, IComparable<Utf8String>
 	{
 		/// <summary>
 		/// An empty string
diff --git a/Engine/Source/Programs/Shared/EpicGames.Core/VersionNumber.cs b/Engine/Source/Programs/Shared/EpicGames.Core/VersionNumber.cs
--- a/Engine/Source/Programs/Shared/EpicGames.Core/VersionNumber.cs
+++ b/Engine/Source/Programs/Shared/EpicGames.Core/VersionNumber.cs
@@ -1,4 +1,4 @@
-﻿// Copyright Epic Games, Inc. All Rights Reserved.
+// Copyright Epic Games, Inc. All Rights Reserved.
 
 using System;
 using System.Collections.Generic;
@@ -56,7 +56,7 @@ namespace EpicGames.Core
 		public override bool Equals(object? obj)
 		{
 			VersionNumber? version = obj as VersionNumber;
-			return !ReferenceEquals(version, null) && this == version;
+			return version is not null && this == version;
 		}
 
 		/// <summary>
@@ -81,14 +81,7 @@ namespace EpicGames.Core
 		/// <returns>True if the versions are equal.</returns>
 		public static bool operator ==(VersionNumber? lhs, VersionNumber? rhs)
 		{
-			if (Object.ReferenceEquals(lhs, null))
-			{
-				return Object.ReferenceEquals(rhs, null);
-			}
-			else
-			{
-				return !Object.ReferenceEquals(rhs, null) && Compare(lhs, rhs) == 0;
-			}
+			return lhs is null ? rhs is null : rhs is not null && Compare(lhs, rhs) == 0;
 		}
 
 		/// <summary>
@@ -153,7 +146,7 @@ namespace EpicGames.Core
 		/// <returns>A negative value if this version is before Other, a positive value if this version is after Other, and zero otherwise.</returns>
 		public int CompareTo(VersionNumber? other)
 		{
-			return ReferenceEquals(other, null) ? 1 : Compare(this, other);
+			return other is null ? 1 : Compare(this, other);
 		}
 
 		/// <summary>
diff --git a/Engine/Source/Programs/Shared/EpicGames.Core/VersionNumberRange.cs b/Engine/Source/Programs/Shared/EpicGames.Core/VersionNumberRange.cs
--- a/Engine/Source/Programs/Shared/EpicGames.Core/VersionNumberRange.cs
+++ b/Engine/Source/Programs/Shared/EpicGames.Core/VersionNumberRange.cs
@@ -26,8 +26,8 @@ namespace EpicGames.Core
 		/// <param name="max"></param>
 		public VersionNumberRange(VersionNumber min, VersionNumber max)
 		{
-			this.Min = min;
-			this.Max = max;
+			Min = min;
+			Max = max;
 		}
 
 		/// <summary>
diff --git a/Engine/Source/Programs/Shared/EpicGames.Horde/HordeHttpAuthHandler.cs b/Engine/Source/Programs/Shared/EpicGames.Horde/HordeHttpAuthHandler.cs
--- a/Engine/Source/Programs/Shared/EpicGames.Horde/HordeHttpAuthHandler.cs
+++ b/Engine/Source/Programs/Shared/EpicGames.Horde/HordeHttpAuthHandler.cs
@@ -337,7 +337,7 @@ namespace EpicGames.Horde
 			if (result == null && allowLogin)
 			{
 				_logger.LogInformation("Logging in to {Server}...", serverUrl);
-				result = await oidcTokenManager.Login(oidcProvider, cancellationToken);
+				result = await oidcTokenManager.LoginAsync(oidcProvider, cancellationToken);
 			}
 
 			return new AuthState(authConfig.Method, result);
diff --git a/Engine/Source/Programs/Shared/EpicGames.IoHash/Blake3Hash.cs b/Engine/Source/Programs/Shared/EpicGames.IoHash/Blake3Hash.cs
--- a/Engine/Source/Programs/Shared/EpicGames.IoHash/Blake3Hash.cs
+++ b/Engine/Source/Programs/Shared/EpicGames.IoHash/Blake3Hash.cs
@@ -12,7 +12,7 @@ namespace EpicGames.Core
 	/// <summary>
 	/// Struct representing a strongly typed Blake3 hash value (a 32-byte Blake3 hash).
 	/// </summary>
-	public struct Blake3Hash : IEquatable<Blake3Hash>, IComparable<Blake3Hash>
+	public readonly struct Blake3Hash : IEquatable<Blake3Hash>, IComparable<Blake3Hash>
 	{
 		/// <summary>
 		/// Length of an Blake3Hash
diff --git a/Engine/Source/Programs/Shared/EpicGames.IoHash/IoHash.cs b/Engine/Source/Programs/Shared/EpicGames.IoHash/IoHash.cs
--- a/Engine/Source/Programs/Shared/EpicGames.IoHash/IoHash.cs
+++ b/Engine/Source/Programs/Shared/EpicGames.IoHash/IoHash.cs
@@ -18,7 +18,7 @@ namespace EpicGames.Core
 	/// </summary>
 	[JsonConverter(typeof(IoHashJsonConverter))]
 	[TypeConverter(typeof(IoHashTypeConverter))]
-	public struct IoHash : IEquatable<IoHash>, IComparable<IoHash>
+	public readonly struct IoHash : IEquatable<IoHash>, IComparable<IoHash>
 	{
 		/// <summary>
 		/// Length of an IoHash
diff --git a/Engine/Source/Programs/Shared/EpicGames.OIDC/OidcTokenManager.cs b/Engine/Source/Programs/Shared/EpicGames.OIDC/OidcTokenManager.cs
--- a/Engine/Source/Programs/Shared/EpicGames.OIDC/OidcTokenManager.cs
+++ b/Engine/Source/Programs/Shared/EpicGames.OIDC/OidcTokenManager.cs
@@ -4,13 +4,11 @@ using System;
 using System.Collections.Generic;
 using System.ComponentModel.DataAnnotations;
 using System.Diagnostics;
-using System.Diagnostics.CodeAnalysis;
 using System.IO;
 using System.Linq;
 using System.Net;
 using System.Net.Http;
 using System.Runtime.InteropServices;
-using System.Runtime.Versioning;
 using System.Text;
 using System.Threading;
 using System.Threading.Tasks;
@@ -20,7 +18,6 @@ using IdentityModel.OidcClient.Results;
 using Microsoft.Extensions.Configuration;
 using Microsoft.Extensions.DependencyInjection;
 using Microsoft.Extensions.Options;
-using Microsoft.Win32;
 
 #pragma warning disable CS1591 // Missing XML documentation on public types
 #pragma warning disable CA2227 // Remote property setters on collection types
@@ -129,7 +126,7 @@ namespace EpicGames.OIDC
 			return _tokenClients.Any(pair => pair.Value.GetStatus() == OidcStatus.NotLoggedIn);
 		}
 
-		public async Task<OidcTokenInfo> Login(string providerIdentifier, CancellationToken cancellationToken = default)
+		public async Task<OidcTokenInfo> LoginAsync(string providerIdentifier, CancellationToken cancellationToken = default)
 		{
 			OidcTokenClient tokenClient = _tokenClients[providerIdentifier];
 
@@ -220,7 +217,7 @@ namespace EpicGames.OIDC
 			_tokenStore.TryGetRefreshToken(name, out _refreshToken);
 		}
 
-		private async Task<OidcClientOptions> BuildClientOptions(Uri redirectUri, CancellationToken cancellationToken)
+		private async Task<OidcClientOptions> BuildClientOptionsAsync(Uri redirectUri, CancellationToken cancellationToken)
 		{
 			OidcClientOptions options = new OidcClientOptions
 			{
@@ -237,7 +234,7 @@ namespace EpicGames.OIDC
 
 			// we need to fetch the discovery document ourselves to support OIDC Authorities which have a subresource for it
 			// with Okta has for authorization servers for instance.
-			DiscoveryDocumentResponse discoveryDocument = await GetDiscoveryDocument(cancellationToken);
+			DiscoveryDocumentResponse discoveryDocument = await GetDiscoveryDocumentAsync(cancellationToken);
 			options.ProviderInformation = new ProviderInformation
 			{
 				IssuerName = discoveryDocument.Issuer,
@@ -270,7 +267,7 @@ namespace EpicGames.OIDC
 					http.Prefixes.Add(prefix);
 					http.Start();
 
-					OidcClientOptions options = await BuildClientOptions(uri, cancellationToken);
+					OidcClientOptions options = await BuildClientOptionsAsync(uri, cancellationToken);
 					OidcClient oidcClient = new OidcClient(options);
 					// generate the appropriate codes we need to login
 					AuthorizeState loginState = await oidcClient!.PrepareLoginAsync(cancellationToken: cancellationToken);
@@ -283,7 +280,7 @@ namespace EpicGames.OIDC
 						{
 							try
 							{
-								Task<LoginResult> processHttpTask = ProcessHttpRequest(http, loginState, oidcClient);
+								Task<LoginResult> processHttpTask = ProcessHttpRequestAsync(http, loginState, oidcClient);
 								Task finishedTask = await Task.WhenAny(Task.Delay(_loginTimeout, cancellationToken), processHttpTask);
 								if (finishedTask == processHttpTask)
 								{
@@ -340,7 +337,7 @@ namespace EpicGames.OIDC
 			};
 		}
 
-		private static async Task<LoginResult> ProcessHttpRequest(HttpListener http, AuthorizeState loginState, OidcClient oidcClient)
+		private static async Task<LoginResult> ProcessHttpRequestAsync(HttpListener http, AuthorizeState loginState, OidcClient oidcClient)
 		{
 			LoginResult? loginResult = null;
 			const int MaxAttempts = 5;
@@ -471,26 +468,10 @@ namespace EpicGames.OIDC
 			return version.Contains("microsoft", StringComparison.InvariantCultureIgnoreCase) || version.Contains("wsl2", StringComparison.InvariantCultureIgnoreCase);
 		}
 
-		[SupportedOSPlatform("windows")]
-		static bool TryReadRegistryString(RegistryKey rootKey, string path, string? name, [NotNullWhen(true)] out string? value)
-		{
-			using (RegistryKey? registryKey = rootKey.OpenSubKey(path))
-			{
-				if (registryKey != null)
-				{
-					value = registryKey.GetValue(name) as string;
-					return value != null;
-				}
-			}
-
-			value = null;
-			return false;
-		}
-
-		private async Task<OidcTokenInfo?> TryDoRefreshToken(string inRefreshToken, CancellationToken cancellationToken)
+		private async Task<OidcTokenInfo?> TryDoRefreshTokenAsync(string inRefreshToken, CancellationToken cancellationToken)
 		{
 			// redirect uri is not used for refrehs tokens so we can just pick one of them to configure the client
-			OidcClientOptions options = await BuildClientOptions(_redirectUris.First(), cancellationToken);
+			OidcClientOptions options = await BuildClientOptionsAsync(_redirectUris.First(), cancellationToken);
 			OidcClient oidcClient = new OidcClient(options);
 
 			// use the refresh token to acquire a new access token
@@ -526,7 +507,7 @@ namespace EpicGames.OIDC
 			};
 		}
 
-		private async Task<DiscoveryDocumentResponse> GetDiscoveryDocument(CancellationToken cancellationToken)
+		private async Task<DiscoveryDocumentResponse> GetDiscoveryDocumentAsync(CancellationToken cancellationToken)
 		{
 			string baseUrl = _authorityUri.ToString().TrimEnd('/');
 			string discoUrl = $"{baseUrl}/.well-known/openid-configuration";
@@ -576,7 +557,7 @@ namespace EpicGames.OIDC
 				};
 			}
 
-			return await TryDoRefreshToken(_refreshToken, cancellationToken);
+			return await TryDoRefreshTokenAsync(_refreshToken, cancellationToken);
 		}
 
 		public OidcStatus GetStatus()
diff --git a/Engine/Source/Programs/Shared/EpicGames.Serialization.Tests/CbPackageTests.cs b/Engine/Source/Programs/Shared/EpicGames.Serialization.Tests/CbPackageTests.cs
--- a/Engine/Source/Programs/Shared/EpicGames.Serialization.Tests/CbPackageTests.cs
+++ b/Engine/Source/Programs/Shared/EpicGames.Serialization.Tests/CbPackageTests.cs
@@ -41,15 +41,15 @@ namespace EpicGames.Serialization.Tests
 
 			using CbPackageBuilder builder = new();
 			builder.AddAttachment(rootHash, CbPackageAttachmentFlags.IsObject, rootObject.GetView().ToArray());
-			byte[] bytes = await builder.ToByteArray();
+			byte[] bytes = await builder.ToByteArrayAsync();
 
 			await using MemoryStream ms = new MemoryStream(bytes);
-			CbPackageReader reader = await CbPackageReader.Create(ms);
+			CbPackageReader reader = await CbPackageReader.CreateAsync(ms);
 
 			Assert.AreEqual(rootHash, reader.RootHash);
 			Assert.AreEqual(rootObject, reader.RootObject);
 
-			await foreach ((CbPackageAttachmentEntry, byte[]) _ in reader.IterateAttachments())
+			await foreach ((CbPackageAttachmentEntry, byte[]) _ in reader.IterateAttachmentsAsync())
 			{
 				Assert.Fail("No Attachments expected");
 			}
@@ -78,16 +78,16 @@ namespace EpicGames.Serialization.Tests
 			builder.AddAttachment(simpleHash, CbPackageAttachmentFlags.IsObject, simpleObject.GetView().ToArray());
 			builder.AddAttachment(blobHash, 0, blob);
 
-			byte[] bytes = await builder.ToByteArray();
+			byte[] bytes = await builder.ToByteArrayAsync();
 
 			await using MemoryStream ms = new MemoryStream(bytes);
-			CbPackageReader reader = await CbPackageReader.Create(ms);
+			CbPackageReader reader = await CbPackageReader.CreateAsync(ms);
 
 			Assert.AreEqual(rootHash, reader.RootHash);
 			Assert.AreEqual(rootObject, reader.RootObject);
 
 			int countOfAttachments = 0;
-			await foreach ((CbPackageAttachmentEntry entry, byte[] attachmentBytes) in reader.IterateAttachments())
+			await foreach ((CbPackageAttachmentEntry entry, byte[] attachmentBytes) in reader.IterateAttachmentsAsync())
 			{
 				countOfAttachments++;
 				if (entry.AttachmentHash == blobHash)
@@ -120,12 +120,12 @@ namespace EpicGames.Serialization.Tests
 			IoHash simpleHash = IoHash.Compute(simpleObject.GetView().Span);
 			packageBuilder.AddAttachment(simpleHash, CbPackageAttachmentFlags.IsObject, simpleObject.GetView().ToArray());
 
-			byte[] buf = await packageBuilder.ToByteArray();
+			byte[] buf = await packageBuilder.ToByteArrayAsync();
 			await using MemoryStream ms = new MemoryStream(buf);
 
 			try
 			{
-				await CbPackageReader.Create(ms);
+				await CbPackageReader.CreateAsync(ms);
 				Assert.Fail("Exception should be thrown");
 			}
 			catch (Exception)
diff --git a/Engine/Source/Programs/Shared/EpicGames.Serialization/CbPackage.cs b/Engine/Source/Programs/Shared/EpicGames.Serialization/CbPackage.cs
--- a/Engine/Source/Programs/Shared/EpicGames.Serialization/CbPackage.cs
+++ b/Engine/Source/Programs/Shared/EpicGames.Serialization/CbPackage.cs
@@ -42,7 +42,7 @@ namespace EpicGames.Serialization
 
 		private const int Length = sizeof(uint) * 4;
 
-		public static async Task<CbPackageHeader> Read(Stream s)
+		public static async Task<CbPackageHeader> ReadAsync(Stream s)
 		{
 			// using async reading of stream because Asp.net requires this
 			byte[] buf = new byte[Length];
@@ -124,7 +124,7 @@ namespace EpicGames.Serialization
 		/// </summary>
 		/// <param name="s">The stream to read from</param>
 		/// <returns>The read package attachment</returns>
-		internal static async Task<CbPackageAttachmentEntry> Read(Stream s)
+		internal static async Task<CbPackageAttachmentEntry> ReadAsync(Stream s)
 		{
 			byte[] buf = new byte[Length];
 			int readBytes = await s.ReadAsync(buf, 0, Length);
@@ -206,9 +206,9 @@ namespace EpicGames.Serialization
 		/// </summary>
 		/// <param name="s">A stream</param>
 		/// <returns>A package reader instance</returns>
-		public static async Task<CbPackageReader> Create(Stream s)
+		public static async Task<CbPackageReader> CreateAsync(Stream s)
 		{
-			CbPackageHeader header = await CbPackageHeader.Read(s);
+			CbPackageHeader header = await CbPackageHeader.ReadAsync(s);
 
 			List<CbPackageAttachmentEntry> entries = new List<CbPackageAttachmentEntry>();
 			// we expect a extra entry for the root object
@@ -216,7 +216,7 @@ namespace EpicGames.Serialization
 
 			for (int i = 0; i < attachmentsCount; i++)
 			{
-				CbPackageAttachmentEntry entry = await CbPackageAttachmentEntry.Read(s);
+				CbPackageAttachmentEntry entry = await CbPackageAttachmentEntry.ReadAsync(s);
 				entries.Add(entry);
 			}
 
@@ -255,7 +255,7 @@ namespace EpicGames.Serialization
 		/// Iterates over the attachments, returning the attachment entry and the attachment (in memory)
 		/// </summary>
 		/// <returns></returns>
-		public async IAsyncEnumerable<(CbPackageAttachmentEntry, byte[])> IterateAttachments()
+		public async IAsyncEnumerable<(CbPackageAttachmentEntry, byte[])> IterateAttachmentsAsync()
 		{
 			// close the stream after we have iterated the attachments as there should be nothing left in it
 			await using Stream s = _attachmentsStream;
@@ -348,7 +348,7 @@ namespace EpicGames.Serialization
 		/// Generate a contiguous buffer of the cb package
 		/// </summary>
 		/// <returns></returns>
-		public async Task<byte[]> ToByteArray()
+		public async Task<byte[]> ToByteArrayAsync()
 		{
 			MemoryStream packageBuffer = new MemoryStream();
 
diff --git a/Engine/Source/Programs/Shared/EpicGames.UHT/Exporters/CodeGen/UhtCodeGenerator.cs b/Engine/Source/Programs/Shared/EpicGames.UHT/Exporters/CodeGen/UhtCodeGenerator.cs
--- a/Engine/Source/Programs/Shared/EpicGames.UHT/Exporters/CodeGen/UhtCodeGenerator.cs
+++ b/Engine/Source/Programs/Shared/EpicGames.UHT/Exporters/CodeGen/UhtCodeGenerator.cs
@@ -2,7 +2,6 @@
 
 using System;
 using System.Collections.Generic;
-using System.Diagnostics;
 using System.IO;
 using System.Linq;
 using System.Text;
diff --git a/Engine/Source/Programs/Shared/EpicGames.UHT/Exporters/CodeGen/UhtHeaderCodeGenerator.cs b/Engine/Source/Programs/Shared/EpicGames.UHT/Exporters/CodeGen/UhtHeaderCodeGenerator.cs
--- a/Engine/Source/Programs/Shared/EpicGames.UHT/Exporters/CodeGen/UhtHeaderCodeGenerator.cs
+++ b/Engine/Source/Programs/Shared/EpicGames.UHT/Exporters/CodeGen/UhtHeaderCodeGenerator.cs
@@ -618,7 +618,7 @@ namespace EpicGames.UHT.Exporters.CodeGen
 			}));
 
 			// We want properties followed by functions
-			notifyTypes.Instances = notifyTypes.Instances.OrderBy(x => (x is UhtProperty ? 0 : 1) * (int)UhtDefineScope.ScopeCount + x.DefineScope).ToList();
+			notifyTypes.Instances.SortBy(x => (x is UhtProperty ? 0 : 1) * (int)UhtDefineScope.ScopeCount + x.DefineScope);
 			return notifyTypes;
 		}
 		#endregion
diff --git a/Engine/Source/Programs/Shared/EpicGames.UHT/Exporters/CodeGen/UhtHeaderCodeGeneratorCppFile.cs b/Engine/Source/Programs/Shared/EpicGames.UHT/Exporters/CodeGen/UhtHeaderCodeGeneratorCppFile.cs
--- a/Engine/Source/Programs/Shared/EpicGames.UHT/Exporters/CodeGen/UhtHeaderCodeGeneratorCppFile.cs
+++ b/Engine/Source/Programs/Shared/EpicGames.UHT/Exporters/CodeGen/UhtHeaderCodeGeneratorCppFile.cs
@@ -10,7 +10,7 @@ using EpicGames.UHT.Utils;
 
 namespace EpicGames.UHT.Exporters.CodeGen
 {
-	internal struct UhtCodeBlockComment : IDisposable
+	internal readonly struct UhtCodeBlockComment : IDisposable
 	{
 		private readonly StringBuilder _builder;
 		private readonly UhtType? _primaryType = null;
@@ -1565,7 +1565,6 @@ namespace EpicGames.UHT.Exporters.CodeGen
 						builder.Append("\tstatic_assert(UE_ARRAY_COUNT(").Append("FuncInfo) < 2048);\r\n");
 					});
 
-
 				if (hasInterfaces)
 				{
 					builder.Append("\tstatic const UECodeGen_Private::FImplementedInterfaceParams InterfaceParams[];\r\n");
diff --git a/Engine/Source/Programs/Shared/EpicGames.UHT/Exporters/CodeGen/UhtHeaderCodeGeneratorHFile.cs b/Engine/Source/Programs/Shared/EpicGames.UHT/Exporters/CodeGen/UhtHeaderCodeGeneratorHFile.cs
--- a/Engine/Source/Programs/Shared/EpicGames.UHT/Exporters/CodeGen/UhtHeaderCodeGeneratorHFile.cs
+++ b/Engine/Source/Programs/Shared/EpicGames.UHT/Exporters/CodeGen/UhtHeaderCodeGeneratorHFile.cs
@@ -815,18 +815,16 @@ namespace EpicGames.UHT.Exporters.CodeGen
 		{
 			if (getterSetterProperties.Any())
 			{
-				using (UhtMacroCreator macro = new(builder, this, classObj, AccessorsMacroSuffix))
+				using UhtMacroCreator macro = new(builder, this, classObj, AccessorsMacroSuffix);
+				foreach (UhtProperty property in getterSetterProperties)
 				{
-					foreach (UhtProperty property in getterSetterProperties)
+					if (property.PropertyExportFlags.HasAnyFlags(UhtPropertyExportFlags.GetterFound))
 					{
-						if (property.PropertyExportFlags.HasAnyFlags(UhtPropertyExportFlags.GetterFound))
-						{
-							builder.Append("static void ").AppendPropertyGetterWrapperName(property).Append("(const void* Object, void* OutValue); \\\r\n");
-						}
-						if (property.PropertyExportFlags.HasAnyFlags(UhtPropertyExportFlags.SetterFound))
-						{
-							builder.Append("static void ").AppendPropertySetterWrapperName(property).Append("(void* Object, const void* InValue); \\\r\n");
-						}
+						builder.Append("static void ").AppendPropertyGetterWrapperName(property).Append("(const void* Object, void* OutValue); \\\r\n");
+					}
+					if (property.PropertyExportFlags.HasAnyFlags(UhtPropertyExportFlags.SetterFound))
+					{
+						builder.Append("static void ").AppendPropertySetterWrapperName(property).Append("(void* Object, const void* InValue); \\\r\n");
 					}
 				}
 			}
diff --git a/Engine/Source/Programs/Shared/EpicGames.UHT/Exporters/CodeGen/UhtMacroCreator.cs b/Engine/Source/Programs/Shared/EpicGames.UHT/Exporters/CodeGen/UhtMacroCreator.cs
--- a/Engine/Source/Programs/Shared/EpicGames.UHT/Exporters/CodeGen/UhtMacroCreator.cs
+++ b/Engine/Source/Programs/Shared/EpicGames.UHT/Exporters/CodeGen/UhtMacroCreator.cs
@@ -7,7 +7,7 @@ using EpicGames.UHT.Utils;
 
 namespace EpicGames.UHT.Exporters.CodeGen
 {
-	internal struct UhtMacroCreator : IDisposable
+	internal readonly struct UhtMacroCreator : IDisposable
 	{
 		private readonly StringBuilder _builder;
 		private readonly int _startingLength;
diff --git a/Engine/Source/Programs/Shared/EpicGames.UHT/Exporters/CodeGen/UhtUsedDefineScopes.cs b/Engine/Source/Programs/Shared/EpicGames.UHT/Exporters/CodeGen/UhtUsedDefineScopes.cs
--- a/Engine/Source/Programs/Shared/EpicGames.UHT/Exporters/CodeGen/UhtUsedDefineScopes.cs
+++ b/Engine/Source/Programs/Shared/EpicGames.UHT/Exporters/CodeGen/UhtUsedDefineScopes.cs
@@ -4,6 +4,7 @@ using System;
 using System.Collections.Generic;
 using System.Linq;
 using System.Text;
+using EpicGames.Core;
 using EpicGames.UHT.Types;
 using EpicGames.UHT.Utils;
 
@@ -24,7 +25,7 @@ namespace EpicGames.UHT.Exporters.CodeGen
 		/// <summary>
 		/// Collection of instances
 		/// </summary>
-		public List<T> Instances { get; set; } = new();
+		public List<T> Instances { get; } = new();
 
 		/// <summary>
 		/// If true, there are no instances
@@ -116,7 +117,7 @@ namespace EpicGames.UHT.Exporters.CodeGen
 		/// </summary>
 		public void OrderByDefineScope()
 		{
-			Instances = Instances.OrderBy(x => x.DefineScope).ToList();
+			Instances.SortBy(x => x.DefineScope);
 		}
 
 		/// <summary>
diff --git a/Engine/Source/Programs/Shared/EpicGames.UHT/Parsers/UhtEnumParser.cs b/Engine/Source/Programs/Shared/EpicGames.UHT/Parsers/UhtEnumParser.cs
--- a/Engine/Source/Programs/Shared/EpicGames.UHT/Parsers/UhtEnumParser.cs
+++ b/Engine/Source/Programs/Shared/EpicGames.UHT/Parsers/UhtEnumParser.cs
@@ -63,10 +63,7 @@ namespace EpicGames.UHT.Parsers
 					specifiers.ParseFieldMetaData();
 					specifiers.ParseDeferred();
 
-					if (enumObject.Outer != null)
-					{
-						enumObject.Outer.AddChild(enumObject);
-					}
+					enumObject.Outer?.AddChild(enumObject);
 
 					if ((topScope.HeaderParser.GetCurrentCompositeCompilerDirective() & UhtCompilerDirective.WithEditorOnlyData) != 0)
 					{
diff --git a/Engine/Source/Programs/Shared/EpicGames.UHT/Parsers/UhtFunctionParser.cs b/Engine/Source/Programs/Shared/EpicGames.UHT/Parsers/UhtFunctionParser.cs
--- a/Engine/Source/Programs/Shared/EpicGames.UHT/Parsers/UhtFunctionParser.cs
+++ b/Engine/Source/Programs/Shared/EpicGames.UHT/Parsers/UhtFunctionParser.cs
@@ -83,7 +83,7 @@ namespace EpicGames.UHT.Parsers
 		}
 
 		/** return if more parameters can be marked */
-		public bool CanMarkMore()
+		public readonly bool CanMarkMore()
 		{
 			return _bUseNumber ? _numberLeaveUnmarked > 0 : (_parameterNames != null && _parameterNames.Length > 0);
 		}
@@ -624,10 +624,7 @@ namespace EpicGames.UHT.Parsers
 
 		private static void AddFunction(UhtFunction function)
 		{
-			if (function.Outer != null)
-			{
-				function.Outer.AddChild(function);
-			}
+			function.Outer?.AddChild(function);
 		}
 
 		private static void SetFunctionNames(UhtFunction function)
diff --git a/Engine/Source/Programs/Shared/EpicGames.UHT/Parsers/UhtHeaderFileParser.cs b/Engine/Source/Programs/Shared/EpicGames.UHT/Parsers/UhtHeaderFileParser.cs
--- a/Engine/Source/Programs/Shared/EpicGames.UHT/Parsers/UhtHeaderFileParser.cs
+++ b/Engine/Source/Programs/Shared/EpicGames.UHT/Parsers/UhtHeaderFileParser.cs
@@ -39,7 +39,7 @@ namespace EpicGames.UHT.Parsers
 	/// Compiler directives
 	/// </summary>
 	[Flags]
-	public enum UhtCompilerDirective : uint
+	public enum UhtCompilerDirective
 	{
 		/// <summary>
 		/// No compile directives
@@ -440,10 +440,7 @@ namespace EpicGames.UHT.Parsers
 		/// <returns>Property parser</returns>
 		public UhtPropertyParser GetCachedPropertyParser()
 		{
-			if (PropertyParser == null)
-			{
-				PropertyParser = new();
-			}
+			PropertyParser ??= new();
 			return PropertyParser;
 		}
 
@@ -644,7 +641,7 @@ namespace EpicGames.UHT.Parsers
 						{
 							List<string> strings = (keywordInfo.AllowedCompilerDirectives).ToStringList();
 							string directives = UhtUtilities.MergeTypeNames(strings, "or", false);
-							topScope.TokenReader.LogError($"'{token.Value.ToString()}' must not be inside preprocessor blocks, except for {directives}");
+							topScope.TokenReader.LogError($"'{token.Value}' must not be inside preprocessor blocks, except for {directives}");
 						}
 						parseResult = keywordInfo.Delegate(topScope, currentScope, ref token);
 					}
diff --git a/Engine/Source/Programs/Shared/EpicGames.UHT/Parsers/UhtParserHelpers.cs b/Engine/Source/Programs/Shared/EpicGames.UHT/Parsers/UhtParserHelpers.cs
--- a/Engine/Source/Programs/Shared/EpicGames.UHT/Parsers/UhtParserHelpers.cs
+++ b/Engine/Source/Programs/Shared/EpicGames.UHT/Parsers/UhtParserHelpers.cs
@@ -41,10 +41,7 @@ namespace EpicGames.UHT.Parsers
 					},
 					(UhtTokenList identifier) =>
 					{
-						if (baseIdentifiersTemp == null)
-						{
-							baseIdentifiersTemp = new List<UhtToken[]>();
-						}
+						baseIdentifiersTemp ??= new List<UhtToken[]>();
 						baseIdentifiersTemp.Add(identifier.ToArray());
 					});
 				superIdentifier = superIdentifierTemp;
diff --git a/Engine/Source/Programs/Shared/EpicGames.UHT/Parsers/UhtPropertyParser.cs b/Engine/Source/Programs/Shared/EpicGames.UHT/Parsers/UhtPropertyParser.cs
--- a/Engine/Source/Programs/Shared/EpicGames.UHT/Parsers/UhtPropertyParser.cs
+++ b/Engine/Source/Programs/Shared/EpicGames.UHT/Parsers/UhtPropertyParser.cs
@@ -325,7 +325,7 @@ namespace EpicGames.UHT.Parsers
 	/// <summary>
 	/// Helper class thread specified object cache
 	/// </summary>
-	public struct UhtThreadBorrower<T> : IDisposable where T : new()
+	public readonly struct UhtThreadBorrower<T> : IDisposable where T : new()
 	{
 		private static readonly ThreadLocal<List<T>> s_tls = new(() => new());
 		private readonly T _instance;
@@ -1115,11 +1115,7 @@ namespace EpicGames.UHT.Parsers
 			propertySettings.SourceName = propertySettings.PropertyCategory == UhtPropertyCategory.Return ? "ReturnValue" : nameToken.Value.ToString();
 
 			// Try to resolve the property using any immediate mode property types
-			UhtProperty? newProperty = ResolveProperty(UhtPropertyResolvePhase.Parsing, propertySettings, propertySettings.Outer.HeaderFile.Data.Memory, typeTokens);
-			if (newProperty == null)
-			{
-				newProperty = new UhtPreResolveProperty(propertySettings, typeTokens);
-			}
+			UhtProperty newProperty = ResolveProperty(UhtPropertyResolvePhase.Parsing, propertySettings, propertySettings.Outer.HeaderFile.Data.Memory, typeTokens) ?? new UhtPreResolveProperty(propertySettings, typeTokens);
 
 			// Force the category in non-engine projects
 			if (newProperty.PropertyCategory == UhtPropertyCategory.Member)
diff --git a/Engine/Source/Programs/Shared/EpicGames.UHT/Parsers/UhtScriptStructParser.cs b/Engine/Source/Programs/Shared/EpicGames.UHT/Parsers/UhtScriptStructParser.cs
--- a/Engine/Source/Programs/Shared/EpicGames.UHT/Parsers/UhtScriptStructParser.cs
+++ b/Engine/Source/Programs/Shared/EpicGames.UHT/Parsers/UhtScriptStructParser.cs
@@ -132,10 +132,7 @@ namespace EpicGames.UHT.Parsers
 					// Process the deferred specifiers
 					specifiers.ParseDeferred();
 
-					if (scriptStruct.Outer != null)
-					{
-						scriptStruct.Outer.AddChild(scriptStruct);
-					}
+					scriptStruct.Outer?.AddChild(scriptStruct);
 
 					topScope.HeaderParser.ParseStatements('{', '}', true);
 
diff --git a/Engine/Source/Programs/Shared/EpicGames.UHT/Specifiers/UhtPropertyArgumentSpecifiers.cs b/Engine/Source/Programs/Shared/EpicGames.UHT/Specifiers/UhtPropertyArgumentSpecifiers.cs
--- a/Engine/Source/Programs/Shared/EpicGames.UHT/Specifiers/UhtPropertyArgumentSpecifiers.cs
+++ b/Engine/Source/Programs/Shared/EpicGames.UHT/Specifiers/UhtPropertyArgumentSpecifiers.cs
@@ -47,9 +47,9 @@ namespace EpicGames.UHT.Parsers
 		}
 		
 		[UhtSpecifier(Extends = UhtTableNames.PropertyArgument, ValueType = UhtSpecifierValueType.None)]
-		private static void RequiredSpecifier(UhtSpecifierContext SpecifierContext)
+		private static void RequiredSpecifier(UhtSpecifierContext specifierContext)
 		{
-			UhtPropertySpecifierContext context = (UhtPropertySpecifierContext)SpecifierContext;
+			UhtPropertySpecifierContext context = (UhtPropertySpecifierContext)specifierContext;
 			context.PropertySettings.PropertyFlags |= EPropertyFlags.RequiredParm;
 		}
 		#endregion
diff --git a/Engine/Source/Programs/Shared/EpicGames.UHT/Tables/UhtExporterTable.cs b/Engine/Source/Programs/Shared/EpicGames.UHT/Tables/UhtExporterTable.cs
--- a/Engine/Source/Programs/Shared/EpicGames.UHT/Tables/UhtExporterTable.cs
+++ b/Engine/Source/Programs/Shared/EpicGames.UHT/Tables/UhtExporterTable.cs
@@ -121,7 +121,7 @@ namespace EpicGames.UHT.Tables
 	/// <summary>
 	/// Defines an exporter in the table
 	/// </summary>
-	public struct UhtExporter
+	public readonly struct UhtExporter
 	{
 
 		/// <summary>
diff --git a/Engine/Source/Programs/Shared/EpicGames.UHT/Tables/UhtKeywordTable.cs b/Engine/Source/Programs/Shared/EpicGames.UHT/Tables/UhtKeywordTable.cs
--- a/Engine/Source/Programs/Shared/EpicGames.UHT/Tables/UhtKeywordTable.cs
+++ b/Engine/Source/Programs/Shared/EpicGames.UHT/Tables/UhtKeywordTable.cs
@@ -82,7 +82,7 @@ namespace EpicGames.UHT.Tables
 	/// <summary>
 	/// Defines a keyword
 	/// </summary>
-	public struct UhtKeyword
+	public readonly struct UhtKeyword
 	{
 
 		/// <summary>
@@ -265,10 +265,7 @@ namespace EpicGames.UHT.Tables
 					}
 					else
 					{
-						if (tables == null)
-						{
-							tables = new List<string>();
-						}
+						tables ??= new List<string>();
 						if (info.AllowText != null)
 						{
 							tables.Add($"{keywordTable.UserName} {info.AllowText}");
diff --git a/Engine/Source/Programs/Shared/EpicGames.UHT/Tokenizer/UhtToken.cs b/Engine/Source/Programs/Shared/EpicGames.UHT/Tokenizer/UhtToken.cs
--- a/Engine/Source/Programs/Shared/EpicGames.UHT/Tokenizer/UhtToken.cs
+++ b/Engine/Source/Programs/Shared/EpicGames.UHT/Tokenizer/UhtToken.cs
@@ -142,7 +142,7 @@ namespace EpicGames.UHT.Tokenizer
 		/// <summary>
 		/// End position of the token value
 		/// </summary>
-		public int InputEndPos => InputStartPos + Value.Span.Length;
+		public readonly int InputEndPos => InputStartPos + Value.Span.Length;
 
 		/// <summary>
 		/// Line containing the token
@@ -200,7 +200,7 @@ namespace EpicGames.UHT.Tokenizer
 		/// Return true if the token is an end token
 		/// </summary>
 		/// <returns>True if the token is an end token</returns>
-		public bool IsEndType()
+		public readonly bool IsEndType()
 		{
 			return TokenType.IsEndType();
 		}
@@ -210,7 +210,7 @@ namespace EpicGames.UHT.Tokenizer
 		/// </summary>
 		/// <param name="value">Value to test</param>
 		/// <returns>True if the token value matches the given value</returns>
-		public bool IsValue(char value)
+		public readonly bool IsValue(char value)
 		{
 			return Value.Span.Length == 1 && Value.Span[0] == value;
 		}
@@ -221,7 +221,7 @@ namespace EpicGames.UHT.Tokenizer
 		/// <param name="value">Value to test</param>
 		/// <param name="ignoreCase">If true, ignore case</param>
 		/// <returns>True if the value matches</returns>
-		public bool IsValue(string value, bool ignoreCase = false)
+		public readonly bool IsValue(string value, bool ignoreCase = false)
 		{
 			return Value.Span.Equals(value, ignoreCase ? StringComparison.OrdinalIgnoreCase : StringComparison.Ordinal);
 		}
@@ -232,7 +232,7 @@ namespace EpicGames.UHT.Tokenizer
 		/// <param name="value">Value to test</param>
 		/// <param name="ignoreCase">If true, ignore case</param>
 		/// <returns>True if the value matches</returns>
-		public bool IsValue(StringView value, bool ignoreCase = false)
+		public readonly bool IsValue(StringView value, bool ignoreCase = false)
 		{
 			return Value.Span.Equals(value.Span, ignoreCase ? StringComparison.OrdinalIgnoreCase : StringComparison.Ordinal);
 		}
@@ -243,7 +243,7 @@ namespace EpicGames.UHT.Tokenizer
 		/// <param name="value">Value to test</param>
 		/// <param name="ignoreCase">If true, ignore case</param>
 		/// <returns>True is the value starts with the given string</returns>
-		public bool ValueStartsWith(string value, bool ignoreCase = false)
+		public readonly bool ValueStartsWith(string value, bool ignoreCase = false)
 		{
 			return Value.Span.StartsWith(value, ignoreCase ? StringComparison.OrdinalIgnoreCase : StringComparison.Ordinal);
 		}
@@ -252,7 +252,7 @@ namespace EpicGames.UHT.Tokenizer
 		/// Return true if the token is an identifier
 		/// </summary>
 		/// <returns>True if the token is an identifier</returns>
-		public bool IsIdentifier()
+		public readonly bool IsIdentifier()
 		{
 			return TokenType == UhtTokenType.Identifier;
 		}
@@ -263,7 +263,7 @@ namespace EpicGames.UHT.Tokenizer
 		/// <param name="identifier">Identifier to test</param>
 		/// <param name="ignoreCase">If true, ignore case</param>
 		/// <returns>True if the identifier matches</returns>
-		public bool IsIdentifier(string identifier, bool ignoreCase = false)
+		public readonly bool IsIdentifier(string identifier, bool ignoreCase = false)
 		{
 			return IsIdentifier() && IsValue(identifier, ignoreCase);
 		}
@@ -274,7 +274,7 @@ namespace EpicGames.UHT.Tokenizer
 		/// <param name="identifier">Identifier to test</param>
 		/// <param name="ignoreCase">If true, ignore case</param>
 		/// <returns>True if the identifier matches</returns>
-		public bool IsIdentifier(StringView identifier, bool ignoreCase = false)
+		public readonly bool IsIdentifier(StringView identifier, bool ignoreCase = false)
 		{
 			return IsIdentifier() && IsValue(identifier, ignoreCase);
 		}
@@ -283,7 +283,7 @@ namespace EpicGames.UHT.Tokenizer
 		/// Return true if the token is a symbol
 		/// </summary>
 		/// <returns>True if the token is a symbol</returns>
-		public bool IsSymbol()
+		public readonly bool IsSymbol()
 		{
 			return TokenType == UhtTokenType.Symbol;
 		}
@@ -293,7 +293,7 @@ namespace EpicGames.UHT.Tokenizer
 		/// </summary>
 		/// <param name="symbol">Symbol to test</param>
 		/// <returns>True if the symbol matches</returns>
-		public bool IsSymbol(char symbol)
+		public readonly bool IsSymbol(char symbol)
 		{
 			return IsSymbol() && IsValue(symbol);
 		}
@@ -303,7 +303,7 @@ namespace EpicGames.UHT.Tokenizer
 		/// </summary>
 		/// <param name="symbol">Symbol to test</param>
 		/// <returns>True if the symbol matches</returns>
-		public bool IsSymbol(string symbol)
+		public readonly bool IsSymbol(string symbol)
 		{
 			return IsSymbol() && IsValue(symbol);
 		}
@@ -313,7 +313,7 @@ namespace EpicGames.UHT.Tokenizer
 		/// </summary>
 		/// <param name="symbol">Symbol to test</param>
 		/// <returns>True if the symbol matches</returns>
-		public bool IsSymbol(StringView symbol)
+		public readonly bool IsSymbol(StringView symbol)
 		{
 			return IsSymbol() && IsValue(symbol);
 		}
@@ -322,7 +322,7 @@ namespace EpicGames.UHT.Tokenizer
 		/// Return true if the token is a constant integer
 		/// </summary>
 		/// <returns>True if constant integer</returns>
-		public bool IsConstInt()
+		public readonly bool IsConstInt()
 		{
 			return TokenType == UhtTokenType.DecimalConst || TokenType == UhtTokenType.HexConst;
 		}
@@ -331,7 +331,7 @@ namespace EpicGames.UHT.Tokenizer
 		/// Return true if the token is a constant floag
 		/// </summary>
 		/// <returns>True if constant float</returns>
-		public bool IsConstFloat()
+		public readonly bool IsConstFloat()
 		{
 			return TokenType == UhtTokenType.FloatConst;
 		}
@@ -341,7 +341,7 @@ namespace EpicGames.UHT.Tokenizer
 		/// </summary>
 		/// <param name="value">Resulting value</param>
 		/// <returns>True if the value was set</returns>
-		public bool GetConstInt(out int value)
+		public readonly bool GetConstInt(out int value)
 		{
 			switch (TokenType)
 			{
@@ -368,7 +368,7 @@ namespace EpicGames.UHT.Tokenizer
 		/// </summary>
 		/// <param name="value">Resulting value</param>
 		/// <returns>True if the value was set</returns>
-		public bool GetConstLong(out long value)
+		public readonly bool GetConstLong(out long value)
 		{
 			switch (TokenType)
 			{
@@ -395,7 +395,7 @@ namespace EpicGames.UHT.Tokenizer
 		/// </summary>
 		/// <param name="value">Resulting value</param>
 		/// <returns>True if the value was set</returns>
-		public bool GetConstFloat(out float value)
+		public readonly bool GetConstFloat(out float value)
 		{
 			switch (TokenType)
 			{
@@ -419,7 +419,7 @@ namespace EpicGames.UHT.Tokenizer
 		/// </summary>
 		/// <param name="value">Resulting value</param>
 		/// <returns>True if the value was set</returns>
-		public bool GetConstDouble(out double value)
+		public readonly bool GetConstDouble(out double value)
 		{
 			switch (TokenType)
 			{
@@ -442,7 +442,7 @@ namespace EpicGames.UHT.Tokenizer
 		/// Return true if the token is a constant string (or a char constant)
 		/// </summary>
 		/// <returns>True if the token is a string or character constant</returns>
-		public bool IsConstString()
+		public readonly bool IsConstString()
 		{
 			return TokenType == UhtTokenType.StringConst || TokenType == UhtTokenType.CharConst;
 		}
@@ -455,7 +455,7 @@ namespace EpicGames.UHT.Tokenizer
 		/// <param name="messageSite"></param>
 		/// <returns>Resulting string</returns>
 		/// <exception cref="UhtException">Thrown if the token type is not a string or character constant</exception>
-		public StringView GetUnescapedString(IUhtMessageSite messageSite)
+		public readonly StringView GetUnescapedString(IUhtMessageSite messageSite)
 		{
 			switch (TokenType)
 			{
@@ -512,7 +512,7 @@ namespace EpicGames.UHT.Tokenizer
 		/// </summary>
 		/// <param name="respectQuotes">If true, embedded quotes will be respected</param>
 		/// <returns>Resulting string</returns>
-		public StringView GetConstantValue(bool respectQuotes = false)
+		public readonly StringView GetConstantValue(bool respectQuotes = false)
 		{
 			switch (TokenType)
 			{
@@ -535,7 +535,7 @@ namespace EpicGames.UHT.Tokenizer
 		/// </summary>
 		/// <param name="respectQuotes">If true, respect embedded quotes</param>
 		/// <returns>Resulting string</returns>
-		public StringView GetTokenString(bool respectQuotes = false)
+		public readonly StringView GetTokenString(bool respectQuotes = false)
 		{
 			StringViewBuilder builder = new();
 			switch (TokenType)
@@ -672,7 +672,7 @@ namespace EpicGames.UHT.Tokenizer
 		/// Convert the token to a string.  This will be the value.
 		/// </summary>
 		/// <returns>Value of the token</returns>
-		public override string ToString()
+		public override readonly string ToString()
 		{
 			if (IsEndType())
 			{
@@ -684,34 +684,34 @@ namespace EpicGames.UHT.Tokenizer
 			}
 		}
 
-		private const NumberStyles s_defaultNumberStyles = NumberStyles.AllowLeadingWhite | NumberStyles.AllowTrailingWhite |
+		private const NumberStyles DefaultNumberStyles = NumberStyles.AllowLeadingWhite | NumberStyles.AllowTrailingWhite |
 			NumberStyles.AllowLeadingSign | NumberStyles.AllowDecimalPoint | NumberStyles.AllowThousands | NumberStyles.AllowExponent;
 
-		private float GetFloatValue()
+		private readonly float GetFloatValue()
 		{
 			if (Value.Span.Length > 0)
 			{
 				if (UhtFCString.IsFloatMarker(Value.Span[^1]))
 				{
-					return Single.Parse(Value.Span[0..^1], s_defaultNumberStyles, CultureInfo.InvariantCulture);
+					return Single.Parse(Value.Span[0..^1], DefaultNumberStyles, CultureInfo.InvariantCulture);
 				}
 			}
-			return Single.Parse(Value.Span, s_defaultNumberStyles, CultureInfo.InvariantCulture);
+			return Single.Parse(Value.Span, DefaultNumberStyles, CultureInfo.InvariantCulture);
 		}
 
-		private double GetDoubleValue()
+		private readonly double GetDoubleValue()
 		{
 			if (Value.Span.Length > 0)
 			{
 				if (UhtFCString.IsFloatMarker(Value.Span[^1]))
 				{
-					return Double.Parse(Value.Span[0..^1], s_defaultNumberStyles, CultureInfo.InvariantCulture);
+					return Double.Parse(Value.Span[0..^1], DefaultNumberStyles, CultureInfo.InvariantCulture);
 				}
 			}
-			return Double.Parse(Value.Span, s_defaultNumberStyles, CultureInfo.InvariantCulture);
+			return Double.Parse(Value.Span, DefaultNumberStyles, CultureInfo.InvariantCulture);
 		}
 
-		long GetDecimalValue()
+		readonly long GetDecimalValue()
 		{
 			ReadOnlySpan<char> span = Value.Span;
 			bool isUnsigned = false;
@@ -735,7 +735,7 @@ namespace EpicGames.UHT.Tokenizer
 			return isUnsigned ? (long)Convert.ToUInt64(span.ToString(), 10) : Convert.ToInt64(span.ToString(), 10);
 		}
 
-		long GetHexValue()
+		readonly long GetHexValue()
 		{
 			return Convert.ToInt64(Value.ToString(), 16);
 		}
diff --git a/Engine/Source/Programs/Shared/EpicGames.UHT/Tokenizer/UhtTokenBufferReader.cs b/Engine/Source/Programs/Shared/EpicGames.UHT/Tokenizer/UhtTokenBufferReader.cs
--- a/Engine/Source/Programs/Shared/EpicGames.UHT/Tokenizer/UhtTokenBufferReader.cs
+++ b/Engine/Source/Programs/Shared/EpicGames.UHT/Tokenizer/UhtTokenBufferReader.cs
@@ -336,17 +336,11 @@ namespace EpicGames.UHT.Tokenizer
 			_savedInputPos = InputPos;
 			if (_comments != null)
 			{
-				if (_savedComments == null)
-				{
-					_savedComments = new List<StringView>();
-				}
+				_savedComments ??= new List<StringView>();
 				_savedComments.Clear();
 				_savedComments.AddRange(_comments);
 			}
-			if (_tokenPreprocessor != null)
-			{
-				_tokenPreprocessor.SaveState();
-			}
+			_tokenPreprocessor?.SaveState();
 		}
 
 		/// <inheritdoc/>
@@ -365,10 +359,7 @@ namespace EpicGames.UHT.Tokenizer
 				_comments.Clear();
 				_comments.AddRange(_savedComments);
 			}
-			if (_tokenPreprocessor != null)
-			{
-				_tokenPreprocessor.RestoreState();
-			}
+			_tokenPreprocessor?.RestoreState();
 		}
 
 		/// <inheritdoc/>
@@ -1043,10 +1034,7 @@ ResetStateMachineCheck:
 
 		private void AddComment(StringView comment)
 		{
-			if (_comments == null)
-			{
-				_comments = new List<StringView>(4);
-			}
+			_comments ??= new List<StringView>(4);
 			_comments.Add(comment);
 		}
 		#endregion
diff --git a/Engine/Source/Programs/Shared/EpicGames.UHT/Tokenizer/UhtTokenReader.cs b/Engine/Source/Programs/Shared/EpicGames.UHT/Tokenizer/UhtTokenReader.cs
--- a/Engine/Source/Programs/Shared/EpicGames.UHT/Tokenizer/UhtTokenReader.cs
+++ b/Engine/Source/Programs/Shared/EpicGames.UHT/Tokenizer/UhtTokenReader.cs
@@ -407,7 +407,7 @@ namespace EpicGames.UHT.Tokenizer
 	/// <summary>
 	/// Helper struct to disable comment parsing.  Should be used in a using block
 	/// </summary>
-	public struct UhtTokenDisableComments : IDisposable
+	public readonly struct UhtTokenDisableComments : IDisposable
 	{
 		private readonly IUhtTokenReader _tokenReader;
 
diff --git a/Engine/Source/Programs/Shared/EpicGames.UHT/Tokenizer/UhtTokenReaderIdentifierExtensions.cs b/Engine/Source/Programs/Shared/EpicGames.UHT/Tokenizer/UhtTokenReaderIdentifierExtensions.cs
--- a/Engine/Source/Programs/Shared/EpicGames.UHT/Tokenizer/UhtTokenReaderIdentifierExtensions.cs
+++ b/Engine/Source/Programs/Shared/EpicGames.UHT/Tokenizer/UhtTokenReaderIdentifierExtensions.cs
@@ -1,8 +1,6 @@
 // Copyright Epic Games, Inc. All Rights Reserved.
 
 using System;
-using System.Collections.Generic;
-using System.Diagnostics.CodeAnalysis;
 using EpicGames.UHT.Utils;
 
 namespace EpicGames.UHT.Tokenizer
diff --git a/Engine/Source/Programs/Shared/EpicGames.UHT/Types/Properties/UhtClassProperty.cs b/Engine/Source/Programs/Shared/EpicGames.UHT/Types/Properties/UhtClassProperty.cs
--- a/Engine/Source/Programs/Shared/EpicGames.UHT/Types/Properties/UhtClassProperty.cs
+++ b/Engine/Source/Programs/Shared/EpicGames.UHT/Types/Properties/UhtClassProperty.cs
@@ -49,7 +49,7 @@ namespace EpicGames.UHT.Types
 		/// <inheritdoc/>
 		private StringBuilder AppendSubClassText(StringBuilder builder)
 		{
-			builder.Append("TSubclassOf<").Append(MetaClass?.SourceName).Append(">");
+			builder.Append("TSubclassOf<").Append(MetaClass?.SourceName).Append('>');
 			return builder;
 		}
 
diff --git a/Engine/Source/Programs/Shared/EpicGames.UHT/Types/UhtClass.cs b/Engine/Source/Programs/Shared/EpicGames.UHT/Types/UhtClass.cs
--- a/Engine/Source/Programs/Shared/EpicGames.UHT/Types/UhtClass.cs
+++ b/Engine/Source/Programs/Shared/EpicGames.UHT/Types/UhtClass.cs
@@ -564,10 +564,7 @@ namespace EpicGames.UHT.Types
 		/// <param name="function">If parsed as part of a UFUNCTION, this will reference it</param>
 		public void AddDeclaration(UhtCompilerDirective compilerDirectives, List<UhtToken> tokens, UhtFunction? function)
 		{
-			if (_declarations == null)
-			{
-				_declarations = new List<UhtDeclaration>();
-			}
+			_declarations ??= new List<UhtDeclaration>();
 			_declarations.Add(new UhtDeclaration { CompilerDirectives = compilerDirectives, Tokens = tokens.ToArray(), Function = function });
 		}
 
@@ -1445,10 +1442,7 @@ namespace EpicGames.UHT.Types
 		/// <returns>Resulting dictionary</returns>
 		private static Dictionary<string, List<GetterSetterToResolve>> AddGetterSetter(Dictionary<string, List<GetterSetterToResolve>>? gsToResolve, string name, UhtProperty property, bool setter)
 		{
-			if (gsToResolve == null)
-			{
-				gsToResolve = new Dictionary<string, List<GetterSetterToResolve>>();
-			}
+			gsToResolve ??= new Dictionary<string, List<GetterSetterToResolve>>();
 			if (!gsToResolve.TryGetValue(name, out List<GetterSetterToResolve>? gsList))
 			{
 				gsList = new List<GetterSetterToResolve>();
diff --git a/Engine/Source/Programs/Shared/EpicGames.UHT/Types/UhtFunction.cs b/Engine/Source/Programs/Shared/EpicGames.UHT/Types/UhtFunction.cs
--- a/Engine/Source/Programs/Shared/EpicGames.UHT/Types/UhtFunction.cs
+++ b/Engine/Source/Programs/Shared/EpicGames.UHT/Types/UhtFunction.cs
@@ -859,14 +859,18 @@ namespace EpicGames.UHT.Types
 			builder.Append(Outer!.SourceName).Append("::").Append(functionName).Append('(');
 
 			bool first = true;
-			foreach (UhtProperty arg in ParameterProperties.Span)
+			foreach (UhtType property in ParameterProperties.Span)
 			{
-				if (!first)
+				UhtProperty? arg = property as UhtProperty;
+				if (arg != null)
 				{
-					builder.Append(", ");
+					if (!first)
+					{
+						builder.Append(", ");
+					}
+					first = false;
+					builder.AppendFullDecl(arg, UhtPropertyTextType.EventFunctionArgOrRetVal, true);
 				}
-				first = false;
-				builder.AppendFullDecl(arg, UhtPropertyTextType.EventFunctionArgOrRetVal, true);
 			}
 
 			builder.Append(')');
diff --git a/Engine/Source/Programs/Shared/EpicGames.UHT/Types/UhtMetaData.cs b/Engine/Source/Programs/Shared/EpicGames.UHT/Types/UhtMetaData.cs
--- a/Engine/Source/Programs/Shared/EpicGames.UHT/Types/UhtMetaData.cs
+++ b/Engine/Source/Programs/Shared/EpicGames.UHT/Types/UhtMetaData.cs
@@ -56,7 +56,7 @@ namespace EpicGames.UHT.Types
 		/// Convert the key to a string
 		/// </summary>
 		/// <returns>String representation</returns>
-		public override string ToString()
+		public override readonly string ToString()
 		{
 			if (Index == UhtMetaData.IndexNone)
 			{
@@ -198,10 +198,7 @@ namespace EpicGames.UHT.Types
 		/// </summary>
 		public void Clear()
 		{
-			if (Dictionary != null)
-			{
-				Dictionary.Clear();
-			}
+			Dictionary?.Clear();
 		}
 
 		/// <summary>
@@ -500,10 +497,7 @@ namespace EpicGames.UHT.Types
 			{
 				if (Config.RedirectMetaDataKey(name, out string remappedName))
 				{
-					if (MessageSite != null)
-					{
-						MessageSite.LogWarning(LineNumber, $"Remapping old metadata key '{name}' to new key '{remappedName}', please update the declaration.");
-					}
+					MessageSite?.LogWarning(LineNumber, $"Remapping old metadata key '{name}' to new key '{remappedName}', please update the declaration.");
 				}
 			}
 			AddInternal(GetDictionary(), new UhtMetaDataKey(name, nameIndex), value, checkValueChange);
@@ -522,10 +516,7 @@ namespace EpicGames.UHT.Types
 			{
 				if (!oldValue.Equals(value, StringComparison.OrdinalIgnoreCase))
 				{
-					if (MessageSite != null)
-					{
-						MessageSite.LogError(LineNumber, $"Metadata key '{key.Name}' first seen with value '{oldValue}' then '{value}'");
-					}
+					MessageSite?.LogError(LineNumber, $"Metadata key '{key.Name}' first seen with value '{oldValue}' then '{value}'");
 				}
 			}
 			else
@@ -541,10 +532,7 @@ namespace EpicGames.UHT.Types
 		/// <param name="nameIndex">Index of the meta data key</param>
 		public void Remove(string name, int nameIndex = IndexNone)
 		{
-			if (Dictionary != null)
-			{
-				Dictionary.Remove(new UhtMetaDataKey(name, nameIndex));
-			}
+			Dictionary?.Remove(new UhtMetaDataKey(name, nameIndex));
 		}
 
 		/// <summary>
@@ -604,10 +592,7 @@ namespace EpicGames.UHT.Types
 
 		private SortedList<UhtMetaDataKey, string> GetDictionary()
 		{
-			if (Dictionary == null)
-			{
-				Dictionary = new SortedList<UhtMetaDataKey, string>(s_comparer);
-			}
+			Dictionary ??= new SortedList<UhtMetaDataKey, string>(s_comparer);
 			return Dictionary;
 		}
 	}
diff --git a/Engine/Source/Programs/Shared/EpicGames.UHT/Types/UhtProperty.cs b/Engine/Source/Programs/Shared/EpicGames.UHT/Types/UhtProperty.cs
--- a/Engine/Source/Programs/Shared/EpicGames.UHT/Types/UhtProperty.cs
+++ b/Engine/Source/Programs/Shared/EpicGames.UHT/Types/UhtProperty.cs
@@ -1982,7 +1982,7 @@ namespace EpicGames.UHT.Types
 		/// Customization point for subclasses for NeedsGCBarrierWhenPassedToFunction
 		/// </summary>
 		/// <returns>True if GC barriers need to run</returns>		
-		protected virtual bool NeedsGCBarrierWhenPassedToFunctionImpl(UhtFunction function)
+		protected virtual bool NeedsGCBarrierWhenPassedToFunctionImpl(UhtFunction uhtFunction)
 		{
 			return false;
 		}
diff --git a/Engine/Source/Programs/Shared/EpicGames.UHT/Types/UhtType.cs b/Engine/Source/Programs/Shared/EpicGames.UHT/Types/UhtType.cs
--- a/Engine/Source/Programs/Shared/EpicGames.UHT/Types/UhtType.cs
+++ b/Engine/Source/Programs/Shared/EpicGames.UHT/Types/UhtType.cs
@@ -933,10 +933,7 @@ namespace EpicGames.UHT.Types
 		/// <param name="child">The child to be added.</param>
 		public virtual void AddChild(UhtType child)
 		{
-			if (_children == null)
-			{
-				_children = new List<UhtType>();
-			}
+			_children ??= new();
 			_children.Add(child);
 		}
 
diff --git a/Engine/Source/Programs/Shared/EpicGames.UHT/Utils/UhtBuffer.cs b/Engine/Source/Programs/Shared/EpicGames.UHT/Utils/UhtBuffer.cs
--- a/Engine/Source/Programs/Shared/EpicGames.UHT/Utils/UhtBuffer.cs
+++ b/Engine/Source/Programs/Shared/EpicGames.UHT/Utils/UhtBuffer.cs
@@ -57,7 +57,7 @@ namespace EpicGames.UHT.Utils
 		/// <summary>
 		/// Return true if we have a pool block
 		/// </summary>
-		public bool IsSet => _block != null;
+		public readonly bool IsSet => _block != null;
 
 		/// <summary>
 		/// Create a pool buffer
@@ -73,7 +73,7 @@ namespace EpicGames.UHT.Utils
 		/// <summary>
 		/// The backing array.  The size of the array will normally be larger than the requested size.
 		/// </summary>
-		public T[] GetBlock()
+		public readonly T[] GetBlock()
 		{
 			return _block ?? Array.Empty<T>();
 		}
@@ -111,7 +111,7 @@ namespace EpicGames.UHT.Utils
 		/// <summary>
 		/// Return the borrowed buffer to the cache
 		/// </summary>
-		public void Dispose()
+		public readonly void Dispose()
 		{
 			UhtPoolBuffers.Return<T>(Buffer);
 		}
@@ -359,7 +359,7 @@ namespace EpicGames.UHT.Utils
 		/// <summary>
 		/// Return the borrowed buffer to the cache
 		/// </summary>
-		public void Dispose()
+		public readonly void Dispose()
 		{
 			UhtBuffer.Return(Buffer);
 		}
@@ -527,7 +527,7 @@ namespace EpicGames.UHT.Utils
 		/// <summary>
 		/// Return the borrowed buffer to the cache
 		/// </summary>
-		public void Dispose()
+		public readonly void Dispose()
 		{
 			UhtByteBuffer.Return(Buffer);
 		}
diff --git a/Engine/Source/Programs/Shared/EpicGames.UHT/Utils/UhtFileManager.cs b/Engine/Source/Programs/Shared/EpicGames.UHT/Utils/UhtFileManager.cs
--- a/Engine/Source/Programs/Shared/EpicGames.UHT/Utils/UhtFileManager.cs
+++ b/Engine/Source/Programs/Shared/EpicGames.UHT/Utils/UhtFileManager.cs
@@ -1,7 +1,6 @@
 // Copyright Epic Games, Inc. All Rights Reserved.
 
 using System;
-using System.Buffers;
 using System.IO;
 using System.Text;
 using EpicGames.Core;
@@ -108,10 +107,10 @@ namespace EpicGames.UHT.Utils
 				using UhtBorrowByteBuffer byteBuffer = new((int)fs.Length);
 				Span<byte> byteSpan = byteBuffer.Buffer.Memory.Span;
 				int readLength = fs.Read(byteSpan);
-				byteSpan = byteSpan.Slice(0, readLength);
+				byteSpan = byteSpan[..readLength];
 
 				Encoding encoding = GetEncoding(byteSpan, out int skipBytes);
-				byteSpan = byteSpan.Slice(skipBytes);
+				byteSpan = byteSpan[skipBytes..];
 
 				int charCount = encoding.GetCharCount(byteBuffer.Buffer.Memory.Span);
 				UhtBuffer initialBuffer = UhtBuffer.Borrow(charCount);
@@ -140,10 +139,10 @@ namespace EpicGames.UHT.Utils
 				using UhtRentedPoolBuffer<byte> byteBuffer = new((int)fs.Length);
 				Span<byte> byteSpan = byteBuffer.Buffer.Memory.Span;
 				int readLength = fs.Read(byteSpan);
-				byteSpan = byteSpan.Slice(0, readLength);
+				byteSpan = byteSpan[..readLength];
 
 				Encoding encoding = GetEncoding(byteSpan, out int skipBytes);
-				byteSpan = byteSpan.Slice(skipBytes);
+				byteSpan = byteSpan[skipBytes..];
 
 				int charCount = encoding.GetCharCount(byteBuffer.Buffer.Memory.Span);
 				output = UhtPoolBuffers.Rent<char>(charCount);
@@ -212,10 +211,10 @@ namespace EpicGames.UHT.Utils
 				using UhtRentedPoolBuffer<byte> byteBuffer = new((int)fs.Length);
 				Span<byte> byteSpan = byteBuffer.Buffer.Memory.Span;
 				int readLength = fs.Read(byteSpan);
-				byteSpan = byteSpan.Slice(0, readLength);
+				byteSpan = byteSpan[..readLength];
 
 				Encoding encoding = GetEncoding(byteSpan, out int skipBytes);
-				byteSpan = byteSpan.Slice(skipBytes);
+				byteSpan = byteSpan[skipBytes..];
 
 				int charCount = encoding.GetCharCount(byteSpan);
 				char[] initialBuffer = new char[charCount];
diff --git a/Engine/Source/Programs/Shared/EpicGames.UHT/Utils/UhtHash.cs b/Engine/Source/Programs/Shared/EpicGames.UHT/Utils/UhtHash.cs
--- a/Engine/Source/Programs/Shared/EpicGames.UHT/Utils/UhtHash.cs
+++ b/Engine/Source/Programs/Shared/EpicGames.UHT/Utils/UhtHash.cs
@@ -47,7 +47,7 @@ namespace EpicGames.UHT.Utils
 		/// Return the final hash value
 		/// </summary>
 		/// <returns>Final hash value</returns>
-		public uint End()
+		public readonly uint End()
 		{
 			return (uint)(_hash + (_hash >> 32));
 		}
diff --git a/Engine/Source/Programs/Shared/EpicGames.UHT/Utils/UhtMessage.cs b/Engine/Source/Programs/Shared/EpicGames.UHT/Utils/UhtMessage.cs
--- a/Engine/Source/Programs/Shared/EpicGames.UHT/Utils/UhtMessage.cs
+++ b/Engine/Source/Programs/Shared/EpicGames.UHT/Utils/UhtMessage.cs
@@ -527,10 +527,7 @@ namespace EpicGames.UHT.Utils
 		/// <param name="exceptionContext"></param>
 		public void PushExtraContext(object? exceptionContext)
 		{
-			if (_extraContexts == null)
-			{
-				_extraContexts = new Stack<object?>(8);
-			}
+			_extraContexts ??= new Stack<object?>(8);
 			_extraContexts.Push(exceptionContext);
 		}
 
@@ -539,10 +536,7 @@ namespace EpicGames.UHT.Utils
 		/// </summary>
 		public void PopExtraContext()
 		{
-			if (_extraContexts != null)
-			{
-				_extraContexts.Pop();
-			}
+			_extraContexts?.Pop();
 		}
 
 		/// <summary>
@@ -561,7 +555,7 @@ namespace EpicGames.UHT.Utils
 	/// <summary>
 	/// A "using" object to automate the push/pop of extra context to the thread's current extra context
 	/// </summary>
-	public struct UhtMessageContext : IDisposable
+	public readonly struct UhtMessageContext : IDisposable
 	{
 		private readonly UhtTlsMessageExtraContext? _stack;
 
@@ -572,10 +566,7 @@ namespace EpicGames.UHT.Utils
 		public UhtMessageContext(object? extraContext)
 		{
 			_stack = UhtTlsMessageExtraContext.GetTls();
-			if (_stack != null)
-			{
-				_stack.PushExtraContext(extraContext);
-			}
+			_stack?.PushExtraContext(extraContext);
 		}
 
 		/// <summary>
@@ -597,10 +588,7 @@ namespace EpicGames.UHT.Utils
 		/// </summary>
 		public void Dispose()
 		{
-			if (_stack != null)
-			{
-				_stack.PopExtraContext();
-			}
+			_stack?.PopExtraContext();
 		}
 	}
 }
diff --git a/Engine/Source/Programs/Shared/EpicGames.UHT/Utils/UhtStringBuilderExtensions.cs b/Engine/Source/Programs/Shared/EpicGames.UHT/Utils/UhtStringBuilderExtensions.cs
--- a/Engine/Source/Programs/Shared/EpicGames.UHT/Utils/UhtStringBuilderExtensions.cs
+++ b/Engine/Source/Programs/Shared/EpicGames.UHT/Utils/UhtStringBuilderExtensions.cs
@@ -144,7 +144,7 @@ namespace EpicGames.UHT.Utils
 					.AppendNameDef(staticsName, namePrefix, name, nameSuffix).Append(metaNameSuffix)
 					.Append("), ")
 					.AppendNameDef(staticsName, namePrefix, name, nameSuffix).Append(metaNameSuffix)
-					.Append(")");
+					.Append(')');
 			}
 			else
 			{
@@ -514,7 +514,7 @@ namespace EpicGames.UHT.Utils
 	/// Structure to automate the borrowing and returning of a StringBuilder.
 	/// Use some form of a "using" pattern.
 	/// </summary>
-	public struct BorrowStringBuilder : IDisposable
+	public readonly struct BorrowStringBuilder : IDisposable
 	{
 
 		/// <summary>
diff --git a/Engine/Source/Programs/Shared/EpicGames.UHT/Utils/UhtTables.cs b/Engine/Source/Programs/Shared/EpicGames.UHT/Utils/UhtTables.cs
--- a/Engine/Source/Programs/Shared/EpicGames.UHT/Utils/UhtTables.cs
+++ b/Engine/Source/Programs/Shared/EpicGames.UHT/Utils/UhtTables.cs
@@ -4,6 +4,7 @@ using System;
 using System.Collections.Generic;
 using System.Diagnostics.CodeAnalysis;
 using System.IO;
+using System.Linq;
 using System.Reflection;
 using EpicGames.Core;
 using EpicGames.UHT.Tables;
@@ -525,7 +526,7 @@ namespace EpicGames.UHT.Utils
 				{
 					if (type.IsClass)
 					{
-						foreach (Attribute classAttribute in type.GetCustomAttributes(false))
+						foreach (Attribute classAttribute in type.GetCustomAttributes(false).Cast<Attribute>())
 						{
 							if (classAttribute is UnrealHeaderToolAttribute || classAttribute is UhtEngineClassAttribute)
 							{
@@ -597,9 +598,8 @@ namespace EpicGames.UHT.Utils
 		{
 			if (type.IsClass)
 			{
-
 				// Loop through the attributes
-				foreach (Attribute classAttribute in type.GetCustomAttributes(false))
+				foreach (Attribute classAttribute in type.GetCustomAttributes(false).Cast<Attribute>())
 				{
 					if (classAttribute is UnrealHeaderToolAttribute parserAttribute)
 					{
diff --git a/Engine/Source/Programs/UnrealBuildTool/Executors/UnrealBuildAccelerator/UBAAgentCoordinatorHorde.cs b/Engine/Source/Programs/UnrealBuildTool/Executors/UnrealBuildAccelerator/UBAAgentCoordinatorHorde.cs
--- a/Engine/Source/Programs/UnrealBuildTool/Executors/UnrealBuildAccelerator/UBAAgentCoordinatorHorde.cs
+++ b/Engine/Source/Programs/UnrealBuildTool/Executors/UnrealBuildAccelerator/UBAAgentCoordinatorHorde.cs
@@ -415,7 +415,7 @@ namespace UnrealBuildTool
 			}
 			catch (NotLoggedInException)
 			{
-				result = await oidcTokenManager.Login(oidcProvider, cancellationToken);
+				result = await oidcTokenManager.LoginAsync(oidcProvider, cancellationToken);
 			}
 
 			if (result.AccessToken == null)
