Shared C#: code cleanup

#rnx

diff --git a/Engine/Source/Programs/Shared/EpicGames.Build/Automation/ProcessSingleton.cs b/Engine/Source/Programs/Shared/EpicGames.Build/Automation/ProcessSingleton.cs
--- a/Engine/Source/Programs/Shared/EpicGames.Build/Automation/ProcessSingleton.cs
+++ b/Engine/Source/Programs/Shared/EpicGames.Build/Automation/ProcessSingleton.cs
@@ -17,7 +17,7 @@ namespace UnrealBuildBase
 		public static string GetUniqueMutexForPath(string Name, string UniquePath)
 		{
 			// generate a md5 hash of the path, as GetHashCode is not guaranteed to generate a stable hash
-			return string.Format("Global\\{0}_{1}", Name, ContentHash.MD5(UniquePath.ToUpperInvariant()));
+			return String.Format("Global\\{0}_{1}", Name, ContentHash.MD5(UniquePath.ToUpperInvariant()));
 		}
 
 		/// <summary>
diff --git a/Engine/Source/Programs/Shared/EpicGames.Build/Platform/Mac/ApplePlatformSDK.cs b/Engine/Source/Programs/Shared/EpicGames.Build/Platform/Mac/ApplePlatformSDK.cs
--- a/Engine/Source/Programs/Shared/EpicGames.Build/Platform/Mac/ApplePlatformSDK.cs
+++ b/Engine/Source/Programs/Shared/EpicGames.Build/Platform/Mac/ApplePlatformSDK.cs
@@ -16,7 +16,7 @@ namespace UnrealBuildBase
 	{
 		public static readonly string? InstalledSDKVersion = GetInstalledSDKVersion();
 
-		public static bool TryConvertVersionToInt(string? StringValue, out UInt64 OutValue)
+		public static bool TryConvertVersionToInt(string? StringValue, out ulong OutValue)
 		{
 			OutValue = 0;
 
@@ -77,12 +77,12 @@ namespace UnrealBuildBase
 					Registry.GetValue(
 						"HKEY_LOCAL_MACHINE\\SOFTWARE\\Wow6432Node\\Apple Inc.\\Apple Mobile Device Support\\Shared",
 						"iTunesMobileDeviceDLL", null) as string;
-				if (string.IsNullOrEmpty(DllPath) || !File.Exists(DllPath))
+				if (String.IsNullOrEmpty(DllPath) || !File.Exists(DllPath))
 				{
 					DllPath = Registry.GetValue(
 						"HKEY_LOCAL_MACHINE\\SOFTWARE\\Wow6432Node\\Apple Inc.\\Apple Mobile Device Support\\Shared",
 						"MobileDeviceDLL", null) as string;
-					if (string.IsNullOrEmpty(DllPath) || !File.Exists(DllPath))
+					if (String.IsNullOrEmpty(DllPath) || !File.Exists(DllPath))
 					{
 						// iTunes >= 12.7 doesn't have a key specifying the 32-bit DLL but it does have a ASMapiInterfaceDLL key and MobileDevice.dll is in usually in the same directory
 						DllPath = Registry.GetValue(
@@ -92,14 +92,14 @@ namespace UnrealBuildBase
 							? null
 							: DllPath.Substring(0, DllPath.LastIndexOf('\\') + 1) + "MobileDevice.dll";
 
-						if (string.IsNullOrEmpty(DllPath) || !File.Exists(DllPath))
+						if (String.IsNullOrEmpty(DllPath) || !File.Exists(DllPath))
 						{
 							DllPath = FindWindowsStoreITunesDLL();
 						}
 					}
 				}
 
-				if (!string.IsNullOrEmpty(DllPath) && File.Exists(DllPath))
+				if (!String.IsNullOrEmpty(DllPath) && File.Exists(DllPath))
 				{
 					string? DllVersion = FileVersionInfo.GetVersionInfo(DllPath).FileVersion;
 					// Only return the DLL version as the SDK version if we can correctly parse it
diff --git a/Engine/Source/Programs/Shared/EpicGames.Build/System/NativeProjectsBase.cs b/Engine/Source/Programs/Shared/EpicGames.Build/System/NativeProjectsBase.cs
--- a/Engine/Source/Programs/Shared/EpicGames.Build/System/NativeProjectsBase.cs
+++ b/Engine/Source/Programs/Shared/EpicGames.Build/System/NativeProjectsBase.cs
@@ -174,7 +174,7 @@ namespace UnrealBuildBase
 			string ProjectName = Path.GetFileNameWithoutExtension(Project);
 
 			// Search known .uprojects by name, then as relative path, then relative path for Project/Project.uproject
-			return EnumerateProjectFiles(Logger).FirstOrDefault(x => string.Equals(x.GetFileNameWithoutExtension(), ProjectName, StringComparison.OrdinalIgnoreCase)) ??
+			return EnumerateProjectFiles(Logger).FirstOrDefault(x => String.Equals(x.GetFileNameWithoutExtension(), ProjectName, StringComparison.OrdinalIgnoreCase)) ??
 				FindRelativeFileReference(Project, Logger) ??
 				FindRelativeFileReference(Path.Combine(ProjectName, $"{ProjectName}.uproject"), Logger);
 		}
diff --git a/Engine/Source/Programs/Shared/EpicGames.Build/UEBuildPlatformSDK.cs b/Engine/Source/Programs/Shared/EpicGames.Build/UEBuildPlatformSDK.cs
--- a/Engine/Source/Programs/Shared/EpicGames.Build/UEBuildPlatformSDK.cs
+++ b/Engine/Source/Programs/Shared/EpicGames.Build/UEBuildPlatformSDK.cs
@@ -45,9 +45,9 @@ namespace EpicGames.Core
 		// and Turnkey will make sure that at least one of them is installed for the SDK setup to be valid. Leave null otherwise
 		public string? GroupName = null;
 
-		public UInt64 MinInt = 0;
-		public UInt64 MaxInt = UInt64.MaxValue;
-		public UInt64 CurrentInt = 0;
+		public ulong MinInt = 0;
+		public ulong MaxInt = UInt64.MaxValue;
+		public ulong CurrentInt = 0;
 
 		public SDKStatus Validity = SDKStatus.Invalid;
 
@@ -123,7 +123,7 @@ namespace EpicGames.Core
 	}
 	public class SDKCollection
 	{
-		public virtual string DefaultName { get { return "Sdk"; } }
+		public virtual string DefaultName => "Sdk";
 		public virtual bool IsValid(string Version, SDKDescriptor Info)
 		{
 			return PlatformSDK.IsVersionValid(Version, Info);
@@ -133,10 +133,8 @@ namespace EpicGames.Core
 
 		public UEBuildPlatformSDK PlatformSDK;
 		public List<SDKDescriptor> Sdks = new List<SDKDescriptor>();
-		public List<SDKDescriptor> FullSdks
-		{ get { return Sdks.Where(x => string.Compare(x.Name, "AutoSdk", true) != 0).ToList(); } }
-		public SDKDescriptor? AutoSdk
-		{ get { return Sdks.FirstOrDefault(x => string.Compare(x.Name, "AutoSdk", true) == 0); } }
+		public List<SDKDescriptor> FullSdks => Sdks.Where(x => String.Compare(x.Name, "AutoSdk", true) != 0).ToList();
+		public SDKDescriptor? AutoSdk => Sdks.FirstOrDefault(x => String.Compare(x.Name, "AutoSdk", true) == 0);
 
 
 		public SDKCollection(UEBuildPlatformSDK PlatformSDK)
@@ -212,7 +210,7 @@ namespace EpicGames.Core
 			{
 				return Sdks[0].ToString(Descriptor, CurrentDescriptor, false);
 			}
-			return string.Join(", ", Sdks.Select(x => x.ToString(Descriptor, CurrentDescriptor, true)));
+			return String.Join(", ", Sdks.Select(x => x.ToString(Descriptor, CurrentDescriptor, true)));
 		}
 
 		public string ToMultilineString(bool bIncludeCurrent = true)
@@ -247,7 +245,7 @@ namespace EpicGames.Core
 			}
 			else
 			{
-				Info = Sdks.FirstOrDefault(x => string.Compare(x.Name, Name, true) == 0);
+				Info = Sdks.FirstOrDefault(x => String.Compare(x.Name, Name, true) == 0);
 				if (Info == null)
 				{
 					return false;
@@ -275,7 +273,7 @@ namespace EpicGames.Core
 		}
 
 
-		public override string DefaultName { get { return "Software"; } }
+		public override string DefaultName => "Software";
 
 		public override bool IsValid(string Version, SDKDescriptor Info)
 		{
@@ -311,7 +309,7 @@ namespace EpicGames.Core
 			// verify that neither platform or sdk were added before
 			if (SDKRegistry.Count(x => x.Key == PlatformName || x.Value == SDK) > 0)
 			{
-				throw new Exception(string.Format("Re-registering SDK for {0}. All Platforms must have a unique SDK object", PlatformName));
+				throw new Exception(String.Format("Re-registering SDK for {0}. All Platforms must have a unique SDK object", PlatformName));
 			}
 
 			SDKRegistry.Add(PlatformName, SDK);
@@ -332,7 +330,7 @@ namespace EpicGames.Core
 
 			// if the parent set up autosdk, the env vars will be wrong, but we can still get the manual SDK version from before it was setup
 			string? ParentManualSDKVersions = Environment.GetEnvironmentVariable(GetPlatformManualSDKSetupEnvVar());
-			if (!string.IsNullOrEmpty(ParentManualSDKVersions))
+			if (!String.IsNullOrEmpty(ParentManualSDKVersions))
 			{
 
 				// we pass along __None to indicate the parent didn't have a manual sdk installed
@@ -404,10 +402,7 @@ namespace EpicGames.Core
 		/// <summary>
 		/// Gets the set of all known SDKs
 		/// </summary>
-		public static UEBuildPlatformSDK[] AllPlatformSDKObjects
-		{
-			get { return SDKRegistry.Values.ToArray(); }
-		}
+		public static UEBuildPlatformSDK[] AllPlatformSDKObjects => SDKRegistry.Values.ToArray();
 
 		// True if at least one platform has had its version changed from the standard SDK version, which means it may not be compatible with other projects
 		public static bool bHasAnySDKOverride = false;
@@ -457,14 +452,14 @@ namespace EpicGames.Core
 
 		public SDKDescriptor? GetSDKInfo(string SDKName)
 		{
-			return GetAllSDKInfo().Sdks.FirstOrDefault(x => string.Compare(x.Name, SDKName, true) == 0);
+			return GetAllSDKInfo().Sdks.FirstOrDefault(x => String.Compare(x.Name, SDKName, true) == 0);
 		}
 
 		public SDKDescriptor? GetSoftwareInfo(string? SDKName = null)
 		{
 			// todo: make this queryable?
 			SDKName ??= "Software";
-			return GetAllSoftwareInfo().Sdks.FirstOrDefault(x => string.Compare(x.Name, SDKName, true) == 0);
+			return GetAllSoftwareInfo().Sdks.FirstOrDefault(x => String.Compare(x.Name, SDKName, true) == 0);
 		}
 
 		public SDKCollection GetAllSoftwareInfo(string? DeviceType = null, string? Current = null)
@@ -475,7 +470,7 @@ namespace EpicGames.Core
 			foreach (SDKDescriptor Desc in GetValidSoftwareVersions().Sdks)
 			{
 				// are we restricting this to one type?
-				if (DeviceType == null || string.Compare(Desc.Name, DeviceType, true) == 0)
+				if (DeviceType == null || String.Compare(Desc.Name, DeviceType, true) == 0)
 				{
 					AllSoftwares.SetupSDK(Desc.Name, Desc.Min, Desc.Max, null, Desc.GroupName);
 				}
@@ -726,10 +721,10 @@ namespace EpicGames.Core
 			}
 
 			// AutoSDK must match the desired version exactly, since that is the only one we will use
-			if (string.Compare(SDKType, "AutoSdk", true) == 0)
+			if (String.Compare(SDKType, "AutoSdk", true) == 0)
 			{
 				// if integer version checking failed, then we can detect valid autosdk if the version matches the autosdk directory by name
-				return string.Compare(Version, GetAutoSDKDirectoryForMainVersion(), true) == 0;
+				return String.Compare(Version, GetAutoSDKDirectoryForMainVersion(), true) == 0;
 			}
 
 			// look for a range for this hinted type of SDK
@@ -742,7 +737,7 @@ namespace EpicGames.Core
 			}
 
 			// convert it to an integer
-			UInt64 IntVersion;
+			ulong IntVersion;
 			if (!TryConvertVersionToInt(Version, out IntVersion))
 			{
 				return false;
@@ -767,7 +762,7 @@ namespace EpicGames.Core
 			}
 
 			// convert it to an integer
-			UInt64 IntVersion;
+			ulong IntVersion;
 			if (!TryConvertVersionToInt(Version, out IntVersion))
 			{
 				return false;
@@ -967,7 +962,7 @@ namespace EpicGames.Core
 		private SDKDescriptor? GetSDKVersionForHint(SDKCollection Collection, string? Hint)
 		{
 			// if the hint is found, use it always
-			SDKDescriptor? SDKDesc = Collection.Sdks.FirstOrDefault(x => string.Compare(x.Name, Hint, true) == 0);
+			SDKDescriptor? SDKDesc = Collection.Sdks.FirstOrDefault(x => String.Compare(x.Name, Hint, true) == 0);
 			if (SDKDesc != null)
 			{
 				return SDKDesc;
@@ -1413,7 +1408,7 @@ namespace EpicGames.Core
 						string? Version = Reader.ReadLine();
 						string? Type = Reader.ReadLine();
 						string? Level = Reader.ReadLine();
-						if (string.IsNullOrEmpty(Level))
+						if (String.IsNullOrEmpty(Level))
 						{
 							Level = "FULL";
 						}
@@ -1627,7 +1622,7 @@ namespace EpicGames.Core
 			// parse the envvar
 			Dictionary<string, string> PlatformSpecificLevels = new Dictionary<string, string>();
 			string? DetailedSettings = Environment.GetEnvironmentVariable("UE_AUTOSDK_SPECIFIC_LEVELS"); // "Android=PACKAGE;GDk=RuN;PS5=BUILD";
-			if (!string.IsNullOrEmpty(DetailedSettings) && DetailedSettings.Contains(PlatformName, StringComparison.InvariantCultureIgnoreCase))
+			if (!String.IsNullOrEmpty(DetailedSettings) && DetailedSettings.Contains(PlatformName, StringComparison.InvariantCultureIgnoreCase))
 			{
 				foreach (string Detail in DetailedSettings.ToUpper().Split(';'))
 				{
@@ -1648,7 +1643,7 @@ namespace EpicGames.Core
 			else
 			{
 				string? DefaultLevel = Environment.GetEnvironmentVariable("UE_AUTOSDK_DEFAULT_LEVEL");
-				if (!string.IsNullOrEmpty(DefaultLevel) && AutoSDKLevels.Contains(DefaultLevel, StringComparer.InvariantCultureIgnoreCase))
+				if (!String.IsNullOrEmpty(DefaultLevel) && AutoSDKLevels.Contains(DefaultLevel, StringComparer.InvariantCultureIgnoreCase))
 				{
 					FinalAutoSDKLevel = DefaultLevel.ToUpper();
 				}
@@ -1936,7 +1931,7 @@ namespace EpicGames.Core
 				{
 					string DesiredSDKLevel = GetAutoSDKLevelForPlatform(AutoSDKRoot);
 					// if the user doesn't want AutoSDK for this platform, then return that it is not installed, even if it actually is
-					if (string.Compare(DesiredSDKLevel, "NONE", true) == 0)
+					if (String.Compare(DesiredSDKLevel, "NONE", true) == 0)
 					{
 						return SDKStatus.Invalid;
 					}
@@ -2054,7 +2049,7 @@ namespace EpicGames.Core
 
 					string DesiredSDKLevel = GetAutoSDKLevelForPlatform(AutoSDKRoot);
 					// if the user doesn't want AutoSDK for this platform, then do nothing
-					if (string.Compare(DesiredSDKLevel, "NONE", true) == 0)
+					if (String.Compare(DesiredSDKLevel, "NONE", true) == 0)
 					{
 						Logger.LogDebug("Skipping AutoSDK for {PlatformName} because NONE was specified as the desired AutoSDK level", PlatformName);
 						return;
diff --git a/Engine/Source/Programs/Shared/EpicGames.Horde/Accounts/AccountId.cs b/Engine/Source/Programs/Shared/EpicGames.Horde/Accounts/AccountId.cs
--- a/Engine/Source/Programs/Shared/EpicGames.Horde/Accounts/AccountId.cs
+++ b/Engine/Source/Programs/Shared/EpicGames.Horde/Accounts/AccountId.cs
@@ -19,7 +19,7 @@ namespace EpicGames.Horde.Accounts
 		/// <inheritdoc cref="BinaryId.Parse(System.String)"/>
 		public static AccountId Parse(string text) => new AccountId(BinaryId.Parse(text));
 
-		/// <inheritdoc cref="BinaryId.TryParse(ReadOnlySpan{char}, out BinaryId)"/>
+		/// <inheritdoc cref="BinaryId.TryParse(ReadOnlySpan{Char}, out BinaryId)"/>
 		public static bool TryParse(ReadOnlySpan<char> text, out AccountId result)
 		{
 			BinaryId binaryId;
diff --git a/Engine/Source/Programs/Shared/EpicGames.Horde/Jobs/JobStepId.cs b/Engine/Source/Programs/Shared/EpicGames.Horde/Jobs/JobStepId.cs
--- a/Engine/Source/Programs/Shared/EpicGames.Horde/Jobs/JobStepId.cs
+++ b/Engine/Source/Programs/Shared/EpicGames.Horde/Jobs/JobStepId.cs
@@ -29,7 +29,7 @@ namespace EpicGames.Horde.Jobs
 		/// <inheritdoc cref="SubResourceId.Parse(System.String)"/>
 		public static JobStepId Parse(string text) => new JobStepId(SubResourceId.Parse(text));
 
-		/// <inheritdoc cref="SubResourceId.TryParse(string, out SubResourceId)"/>
+		/// <inheritdoc cref="SubResourceId.TryParse(System.String, out SubResourceId)"/>
 		public static bool TryParse(string text, out JobStepId stepId)
 		{
 			if (SubResourceId.TryParse(text, out SubResourceId subResourceId))
diff --git a/Engine/Source/Programs/Shared/EpicGames.Jupiter/JupiterTree.cs b/Engine/Source/Programs/Shared/EpicGames.Jupiter/JupiterTree.cs
--- a/Engine/Source/Programs/Shared/EpicGames.Jupiter/JupiterTree.cs
+++ b/Engine/Source/Programs/Shared/EpicGames.Jupiter/JupiterTree.cs
@@ -32,10 +32,7 @@ namespace EpicGames.Jupiter
 
 		public List<JupiterTree> Trees { get; } = new List<JupiterTree>();
 
-		public IEnumerable<string> ContentHashes
-		{
-			get { return ContentProvider.GetHashes(); }
-		}
+		public IEnumerable<string> ContentHashes => ContentProvider.GetHashes();
 
 		public Task<byte[]> GetContent(string BlobHash)
 		{
@@ -100,10 +97,7 @@ namespace EpicGames.Jupiter
 			{
 				Log.Logger.LogInformation("Creating a new tree root in Jupiter with id: {Key}", JupiterTreeKey);
 
-				if (Metadata == null)
-				{
-					Metadata = new Dictionary<string, object>();
-				}
+				Metadata ??= new Dictionary<string, object>();
 
 				// add a source field if it doesn't already exist to make it known were this tree came from
 				if (!Metadata.ContainsKey("source"))
@@ -116,16 +110,16 @@ namespace EpicGames.Jupiter
 					{
 						TreeReferenceKey = JupiterTreeKey,
 						TreeHash = TreeHashString,
-						Metadata = Metadata,
+						Metadata,
 					};
 					string PutTreeRootString = JsonSerializer.Serialize(TreeRoot);
 					using StringContent PutTreeRootContent = new StringContent(PutTreeRootString, Encoding.UTF8, "application/json");
 					// upload tree information
-					HttpResponseMessage PutTreeResult = await JupiterClient.PutAsync(string.Format("api/v1/c/tree-root/{0}", JupiterNamespace), PutTreeRootContent);
+					HttpResponseMessage PutTreeResult = await JupiterClient.PutAsync(String.Format("api/v1/c/tree-root/{0}", JupiterNamespace), PutTreeRootContent);
 					if (!PutTreeResult.IsSuccessStatusCode)
 					{
 						string ErrorMsg = await PutTreeResult.Content.ReadAsStringAsync();
-						throw new Exception(string.Format("Failed to create tree root in Jupiter. Response: {0}", ErrorMsg));
+						throw new Exception(String.Format("Failed to create tree root in Jupiter. Response: {0}", ErrorMsg));
 					}
 				}
 
@@ -143,11 +137,11 @@ namespace EpicGames.Jupiter
 						};
 					string PutTreeString = JsonSerializer.Serialize(PutTreesRequest);
 					using StringContent PutTreeContent = new StringContent(PutTreeString, Encoding.UTF8, "application/json");
-					HttpResponseMessage PutTreeContentResult = await JupiterClient.PutAsync(string.Format("api/v1/c/tree/{0}", JupiterNamespace), PutTreeContent);
+					HttpResponseMessage PutTreeContentResult = await JupiterClient.PutAsync(String.Format("api/v1/c/tree/{0}", JupiterNamespace), PutTreeContent);
 					if (!PutTreeContentResult.IsSuccessStatusCode)
 					{
 						string ErrorMsg = await PutTreeContentResult.Content.ReadAsStringAsync();
-						throw new Exception(string.Format("Failed to upload one or more trees to Jupiter. Response: {0}", ErrorMsg));
+						throw new Exception(String.Format("Failed to upload one or more trees to Jupiter. Response: {0}", ErrorMsg));
 					}
 				}
 
@@ -173,7 +167,7 @@ namespace EpicGames.Jupiter
 					if (!FilterBlobResponse.IsSuccessStatusCode)
 					{
 						string ErrorMsg = await FilterBlobResponse.Content.ReadAsStringAsync();
-						throw new Exception(string.Format("Failed to determine which blobs were already present in Jupiter. Response: {0}", ErrorMsg));
+						throw new Exception(String.Format("Failed to determine which blobs were already present in Jupiter. Response: {0}", ErrorMsg));
 					}
 
 					string ResponseString = await FilterBlobResponse.Content.ReadAsStringAsync();
@@ -201,12 +195,12 @@ namespace EpicGames.Jupiter
 									ByteArrayContent Content = new ByteArrayContent(await JupiterTree.GetContent(BlobHash));
 									Content.Headers.Remove("Content-Type");
 									Content.Headers.Add("Content-Type", "application/octet-stream");
-									HttpResponseMessage PutBlobResultResult = await JupiterClient.PutAsync(string.Format("api/v1/s/{0}/{1}", JupiterNamespace, BlobHash), Content);
+									HttpResponseMessage PutBlobResultResult = await JupiterClient.PutAsync(String.Format("api/v1/s/{0}/{1}", JupiterNamespace, BlobHash), Content);
 
 									if (!PutBlobResultResult.IsSuccessStatusCode)
 									{
 										string ErrorMsg = await PutBlobResultResult.Content.ReadAsStringAsync();
-										throw new Exception(string.Format(
+										throw new Exception(String.Format(
 											"Failed to upload content blob to Jupiter. Response: {0}", ErrorMsg));
 									}
 								}
@@ -219,11 +213,11 @@ namespace EpicGames.Jupiter
 
 				// verify the build upload
 				{
-					HttpResponseMessage FinalizeTreeResult = await JupiterClient.PutAsync(string.Format("api/v1/c/tree-root/{0}/{1}/finalize", JupiterNamespace, JupiterTreeKey), null);
+					HttpResponseMessage FinalizeTreeResult = await JupiterClient.PutAsync(String.Format("api/v1/c/tree-root/{0}/{1}/finalize", JupiterNamespace, JupiterTreeKey), null);
 					if (!FinalizeTreeResult.IsSuccessStatusCode)
 					{
 						string ErrorMsg = await FinalizeTreeResult.Content.ReadAsStringAsync();
-						throw new Exception(string.Format("Jupiter tree upload verification failed. Response: {0}",
+						throw new Exception(String.Format("Jupiter tree upload verification failed. Response: {0}",
 							ErrorMsg));
 					}
 				}
@@ -242,12 +236,12 @@ namespace EpicGames.Jupiter
 
 				{
 					// download the tree root
-					HttpResponseMessage GetTreeRootResult = await JupiterClient.GetAsync(string.Format("api/v1/c/tree-root/{0}/{1}", JupiterNamespace, JupiterTreeKey));
+					HttpResponseMessage GetTreeRootResult = await JupiterClient.GetAsync(String.Format("api/v1/c/tree-root/{0}/{1}", JupiterNamespace, JupiterTreeKey));
 
 					string GetTreeRootResultString = await GetTreeRootResult.Content.ReadAsStringAsync();
 					if (!GetTreeRootResult.IsSuccessStatusCode)
 					{
-						throw new Exception(string.Format("Failed to download tree root with key {0}. Response: {1}", JupiterTreeKey, GetTreeRootResultString));
+						throw new Exception(String.Format("Failed to download tree root with key {0}. Response: {1}", JupiterTreeKey, GetTreeRootResultString));
 					}
 
 					TreeRootContents TreeRoot = JsonSerializer.Deserialize<TreeRootContents>(GetTreeRootResultString);
@@ -258,12 +252,12 @@ namespace EpicGames.Jupiter
 				Dictionary<string, TreeContents> TreeMapping;
 				{
 					// Flatten this tree, and download the tree descriptions of all those trees
-					HttpResponseMessage GetFlattenTreeResult = await JupiterClient.GetAsync(string.Format("api/v1/c/tree/{0}/{1}/flattend", JupiterNamespace, TopTreeHash));
+					HttpResponseMessage GetFlattenTreeResult = await JupiterClient.GetAsync(String.Format("api/v1/c/tree/{0}/{1}/flattend", JupiterNamespace, TopTreeHash));
 
 					string GetFlattenTreeResultString = await GetFlattenTreeResult.Content.ReadAsStringAsync();
 					if (!GetFlattenTreeResult.IsSuccessStatusCode)
 					{
-						throw new Exception(string.Format("Failed to download flattened tree {0}. Response: {1}", TopTreeHash, GetFlattenTreeResultString));
+						throw new Exception(String.Format("Failed to download flattened tree {0}. Response: {1}", TopTreeHash, GetFlattenTreeResultString));
 					}
 
 					FlattendTreeContents FlattenedTree = JsonSerializer.Deserialize<FlattendTreeContents>(GetFlattenTreeResultString);
@@ -274,12 +268,12 @@ namespace EpicGames.Jupiter
 					foreach (string Tree in TreesToDownload)
 					{
 						// Flatten this tree, and download the tree descriptions of all those trees
-						HttpResponseMessage GetTreeResult = await JupiterClient.GetAsync(string.Format("api/v1/c/tree/{0}/{1}", JupiterNamespace, Tree));
+						HttpResponseMessage GetTreeResult = await JupiterClient.GetAsync(String.Format("api/v1/c/tree/{0}/{1}", JupiterNamespace, Tree));
 
 						string GetTreeResultString = await GetTreeResult.Content.ReadAsStringAsync();
 						if (!GetTreeResult.IsSuccessStatusCode)
 						{
-							throw new Exception(string.Format("Failed to download tree {0}. Response: {1}", Tree, GetTreeResultString));
+							throw new Exception(String.Format("Failed to download tree {0}. Response: {1}", Tree, GetTreeResultString));
 						}
 
 						TreeContents TreeContents = JsonSerializer.Deserialize<TreeContents>(GetTreeResultString);
@@ -339,12 +333,12 @@ namespace EpicGames.Jupiter
 		{
 			using (HttpClient JupiterClient = new HttpClient {BaseAddress = new Uri(BaseUrl)})
 			{
-				HttpResponseMessage GetContentsResult = await JupiterClient.GetAsync(string.Format("api/v1/s/{0}/{1}", JupiterNamespace, Sha1));
+				HttpResponseMessage GetContentsResult = await JupiterClient.GetAsync(String.Format("api/v1/s/{0}/{1}", JupiterNamespace, Sha1));
 
 				if (!GetContentsResult.IsSuccessStatusCode)
 				{
 					string ErrorMsg = await GetContentsResult.Content.ReadAsStringAsync();
-					throw new Exception(string.Format("Failed to download blob {0}. Response: {1}", Sha1, ErrorMsg));
+					throw new Exception(String.Format("Failed to download blob {0}. Response: {1}", Sha1, ErrorMsg));
 				}
 
 				byte[] Data = await GetContentsResult.Content.ReadAsByteArrayAsync();
@@ -607,10 +601,7 @@ namespace EpicGames.Jupiter
 			}
 			private readonly ManifestData Data;
 
-			public Dictionary<string, string> Files
-			{
-				get { return Data.Files; }
-			}
+			public Dictionary<string, string> Files => Data.Files;
 
 			public Manifest(DirectoryReference InBaseDir)
 			{
diff --git a/Engine/Source/Programs/Shared/EpicGames.MongoDB/MongoExtensions.cs b/Engine/Source/Programs/Shared/EpicGames.MongoDB/MongoExtensions.cs
--- a/Engine/Source/Programs/Shared/EpicGames.MongoDB/MongoExtensions.cs
+++ b/Engine/Source/Programs/Shared/EpicGames.MongoDB/MongoExtensions.cs
@@ -1,10 +1,8 @@
 // Copyright Epic Games, Inc. All Rights Reserved.
 
-using MongoDB.Bson;
 using MongoDB.Driver;
 using System;
 using System.Collections.Generic;
-using System.Linq;
 using System.Linq.Expressions;
 using System.Threading.Tasks;
 
diff --git a/Engine/Source/Programs/Shared/EpicGames.MongoDB/SingletonDocument.cs b/Engine/Source/Programs/Shared/EpicGames.MongoDB/SingletonDocument.cs
--- a/Engine/Source/Programs/Shared/EpicGames.MongoDB/SingletonDocument.cs
+++ b/Engine/Source/Programs/Shared/EpicGames.MongoDB/SingletonDocument.cs
@@ -5,8 +5,6 @@ using MongoDB.Bson.Serialization;
 using MongoDB.Bson.Serialization.Attributes;
 using MongoDB.Driver;
 using System;
-using System.Collections.Generic;
-using System.Linq;
 using System.Reflection;
 using System.Threading.Tasks;
 
@@ -75,11 +73,7 @@ namespace EpicGames.MongoDB
 
 			static ObjectId GetSingletonId()
 			{
-				SingletonDocumentAttribute? Attribute = typeof(T).GetCustomAttribute<SingletonDocumentAttribute>();
-				if (Attribute == null)
-				{
-					throw new Exception($"Type {typeof(T).Name} is missing a {nameof(SingletonDocumentAttribute)} annotation");
-				}
+				SingletonDocumentAttribute? Attribute = typeof(T).GetCustomAttribute<SingletonDocumentAttribute>() ?? throw new Exception($"Type {typeof(T).Name} is missing a {nameof(SingletonDocumentAttribute)} annotation");
 				return ObjectId.Parse(Attribute.Id);
 			}
 		}
@@ -124,12 +118,12 @@ namespace EpicGames.MongoDB
 		/// <summary>
 		/// The database service instance
 		/// </summary>
-		IMongoCollection<T> Collection;
+		readonly IMongoCollection<T> Collection;
 
 		/// <summary>
 		/// Unique id for the singleton document
 		/// </summary>
-		ObjectId ObjectId;
+		readonly ObjectId ObjectId;
 
 		/// <summary>
 		/// Static constructor. Registers the document using the automapper.
@@ -147,12 +141,7 @@ namespace EpicGames.MongoDB
 		{
 			this.Collection = Collection.OfType<T>();
 
-			SingletonDocumentAttribute? Attribute = typeof(T).GetCustomAttribute<SingletonDocumentAttribute>();
-			if (Attribute == null)
-			{
-				throw new Exception($"Type {typeof(T).Name} is missing a {nameof(SingletonDocumentAttribute)} annotation");
-			}
-
+			SingletonDocumentAttribute? Attribute = typeof(T).GetCustomAttribute<SingletonDocumentAttribute>() ?? throw new Exception($"Type {typeof(T).Name} is missing a {nameof(SingletonDocumentAttribute)} annotation");
 			ObjectId = new ObjectId(Attribute.Id);
 		}
 
diff --git a/Engine/Source/Programs/Shared/EpicGames.MongoDB/TransactionBuilder.cs b/Engine/Source/Programs/Shared/EpicGames.MongoDB/TransactionBuilder.cs
--- a/Engine/Source/Programs/Shared/EpicGames.MongoDB/TransactionBuilder.cs
+++ b/Engine/Source/Programs/Shared/EpicGames.MongoDB/TransactionBuilder.cs
@@ -1,16 +1,11 @@
 // Copyright Epic Games, Inc. All Rights Reserved.
 
-using MongoDB.Bson.Serialization;
-using MongoDB.Bson.Serialization.Serializers;
 using MongoDB.Driver;
 using System;
 using System.Collections.Generic;
 using System.Linq;
 using System.Linq.Expressions;
 using System.Reflection;
-using System.Runtime.CompilerServices;
-using System.Threading.Tasks;
-using System.Transactions;
 
 namespace EpicGames.MongoDB
 {
@@ -78,15 +73,12 @@ namespace EpicGames.MongoDB
 		/// <summary>
 		/// List of field updates
 		/// </summary>
-		List<FieldUpdate> FieldUpdates = new List<FieldUpdate>();
+		readonly List<FieldUpdate> FieldUpdates = new List<FieldUpdate>();
 
 		/// <summary>
 		/// Whether the transaction is currenty empty
 		/// </summary>
-		public bool IsEmpty
-		{
-			get { return FieldUpdates.Count == 0; }
-		}
+		public bool IsEmpty => FieldUpdates.Count == 0;
 
 		/// <summary>
 		/// Adds a setting to this transaction
@@ -97,7 +89,7 @@ namespace EpicGames.MongoDB
 		public void Set<TField>(Expression<Func<TDocument, TField>> Expr, TField Value)
 		{
 			UpdateDefinition<TDocument> Update = Builders<TDocument>.Update.Set(Expr, Value);
-			Action<object> Apply = Target => Assign(Target, Expr, Value);
+			void Apply(object Target) => Assign(Target, Expr, Value);
 			FieldUpdates.Add(new FieldUpdate(Update, Apply));
 		}
 
@@ -108,7 +100,7 @@ namespace EpicGames.MongoDB
 		public void Unset(Expression<Func<TDocument, object>> Expr)
 		{
 			UpdateDefinition<TDocument> Update = Builders<TDocument>.Update.Unset(Expr);
-			Action<object> Apply = Target => Unassign(Target, Expr.Body);
+			void Apply(object Target) => Unassign(Target, Expr.Body);
 			FieldUpdates.Add(new FieldUpdate(Update, Apply));
 		}
 
diff --git a/Engine/Source/Programs/Shared/EpicGames.Perforce.Managed/ManagedWorkspace.cs b/Engine/Source/Programs/Shared/EpicGames.Perforce.Managed/ManagedWorkspace.cs
--- a/Engine/Source/Programs/Shared/EpicGames.Perforce.Managed/ManagedWorkspace.cs
+++ b/Engine/Source/Programs/Shared/EpicGames.Perforce.Managed/ManagedWorkspace.cs
@@ -886,10 +886,7 @@ namespace EpicGames.Perforce.Managed
 				else
 				{
 					contents = await TryLoadClientContentsAsync(cacheFile, new Utf8String(streamName), cancellationToken);
-					if (contents == null)
-					{
-						contents = await FindAndSaveClientContentsAsync(perforce, new Utf8String(streamName), view, changeNumber, cacheFile, cancellationToken);
-					}
+					contents ??= await FindAndSaveClientContentsAsync(perforce, new Utf8String(streamName), view, changeNumber, cacheFile, cancellationToken);
 				}
 
 				// Sync all the appropriate files
@@ -2125,10 +2122,7 @@ namespace EpicGames.Perforce.Managed
 					{
 						connectionPool.TryPop(out connection);
 					}
-					if (connection == null)
-					{
-						connection = await PerforceConnection.CreateAsync(client.Settings, client.Logger);
-					}
+					connection ??= await PerforceConnection.CreateAsync(client.Settings, client.Logger);
 
 					// Note: Explicitly disable parallel syncing here; the P4 API attempts to shell out to p4.exe, which may not be installed.
 					await connection.SyncAsync(options, -1, 0, -1, -1, -1, -1, files, cancellationToken).ToListAsync(cancellationToken);
diff --git a/Engine/Source/Programs/Shared/EpicGames.Perforce/Records/MergeRecord.cs b/Engine/Source/Programs/Shared/EpicGames.Perforce/Records/MergeRecord.cs
--- a/Engine/Source/Programs/Shared/EpicGames.Perforce/Records/MergeRecord.cs
+++ b/Engine/Source/Programs/Shared/EpicGames.Perforce/Records/MergeRecord.cs
@@ -18,7 +18,7 @@ namespace EpicGames.Perforce
 		/// <summary>
 		/// The starting revision, as an integer (null if 'none')
 		/// </summary>
-		public int? StartFromRev => String.Equals(StartFromRevText, "none", StringComparison.Ordinal) ? (int?)null : int.Parse(StartFromRevText);
+		public int? StartFromRev => String.Equals(StartFromRevText, "none", StringComparison.Ordinal) ? (int?)null : Int32.Parse(StartFromRevText);
 
 		/// <summary>
 		/// The first revision being merged (or 'none')
diff --git a/Engine/Source/Programs/Shared/EpicGames.Tracing.Tests/UnrealInsights/EventTest.cs b/Engine/Source/Programs/Shared/EpicGames.Tracing.Tests/UnrealInsights/EventTest.cs
--- a/Engine/Source/Programs/Shared/EpicGames.Tracing.Tests/UnrealInsights/EventTest.cs
+++ b/Engine/Source/Programs/Shared/EpicGames.Tracing.Tests/UnrealInsights/EventTest.cs
@@ -3,7 +3,6 @@
 using System.IO;
 using Microsoft.VisualStudio.TestTools.UnitTesting;
 using EpicGames.Tracing.UnrealInsights;
-using EpicGames.Tracing.UnrealInsights.Events;
 
 namespace EpicGames.Tracing.Tests.UnrealInsights
 {
diff --git a/Engine/Source/Programs/Shared/EpicGames.Tracing.Tests/UnrealInsights/GenericEventTest.cs b/Engine/Source/Programs/Shared/EpicGames.Tracing.Tests/UnrealInsights/GenericEventTest.cs
--- a/Engine/Source/Programs/Shared/EpicGames.Tracing.Tests/UnrealInsights/GenericEventTest.cs
+++ b/Engine/Source/Programs/Shared/EpicGames.Tracing.Tests/UnrealInsights/GenericEventTest.cs
@@ -226,7 +226,7 @@ namespace EpicGames.Tracing.Tests.UnrealInsights
 		[TestMethod]
 		public void CpuProfiler_EventBatch_EncodeDecodeTimestamp()
 		{
-			void AssertTimestamp(ulong ExpectedTimestamp, bool ExpectedIsScopeEnter)
+			static void AssertTimestamp(ulong ExpectedTimestamp, bool ExpectedIsScopeEnter)
 			{
 				ulong Temp = CpuProfilerSerializer.EncodeTimestamp(ExpectedTimestamp, ExpectedIsScopeEnter);
 				(ulong ActualTimestamp, bool ActualIsScopeEnter) = CpuProfilerSerializer.DecodeTimestamp(Temp);
@@ -292,7 +292,7 @@ namespace EpicGames.Tracing.Tests.UnrealInsights
 			using MemoryStream Ms = new MemoryStream(new byte[] {0x01, 0xC1, 0x0D, 0x00});
 			using BinaryReader Reader = new BinaryReader(Ms);
 
-			var (Uid, FieldIndex, Size) = GenericEvent.DeserializeAuxHeader(Reader.ReadUInt32());
+			(ushort Uid, int FieldIndex, int Size) = GenericEvent.DeserializeAuxHeader(Reader.ReadUInt32());
 			Assert.AreEqual(PredefinedEventUid.AuxData, Uid);
 			Assert.AreEqual(1, FieldIndex);
 			Assert.AreEqual(110, Size);
@@ -311,7 +311,7 @@ namespace EpicGames.Tracing.Tests.UnrealInsights
 			Assert.AreEqual(0x0D, Data[2]);
 			Assert.AreEqual(0x00, Data[3]);
 			
-			var (Uid, FieldIndex, Size) = GenericEvent.DeserializeAuxHeader(Header);
+			(ushort Uid, int FieldIndex, int Size) = GenericEvent.DeserializeAuxHeader(Header);
 			Assert.AreEqual(PredefinedEventUid.AuxData, Uid);
 			Assert.AreEqual(1, FieldIndex);
 			Assert.AreEqual(110, Size);
diff --git a/Engine/Source/Programs/Shared/EpicGames.Tracing.Tests/UnrealInsights/TransportPacketTest.cs b/Engine/Source/Programs/Shared/EpicGames.Tracing.Tests/UnrealInsights/TransportPacketTest.cs
--- a/Engine/Source/Programs/Shared/EpicGames.Tracing.Tests/UnrealInsights/TransportPacketTest.cs
+++ b/Engine/Source/Programs/Shared/EpicGames.Tracing.Tests/UnrealInsights/TransportPacketTest.cs
@@ -29,7 +29,7 @@ namespace EpicGames.Tracing.Tests.UnrealInsights
 		[TestMethod]
 		public void GetPackedUid()
 		{
-			void AssertGetPackedUid(ushort ExpectedUid, bool ExpectedIsTwoByteUid, byte UidLow, byte UidHigh)
+			static void AssertGetPackedUid(ushort ExpectedUid, bool ExpectedIsTwoByteUid, byte UidLow, byte UidHigh)
 			{
 				Assert.AreEqual(ExpectedUid, BinaryReaderExtensions.GetPackedUid(UidLow, UidHigh, out bool IsTwoByteUid));
 				Assert.AreEqual(ExpectedIsTwoByteUid, IsTwoByteUid);
@@ -57,7 +57,7 @@ namespace EpicGames.Tracing.Tests.UnrealInsights
 		[TestMethod]
 		public void WritePackedUid()
 		{
-			void AssertWriteUid(ushort Uid, bool ExpectTwoByteUid)
+			static void AssertWriteUid(ushort Uid, bool ExpectTwoByteUid)
 			{
 				using MemoryStream Ms = new MemoryStream();
 				using BinaryWriter Writer = new BinaryWriter(Ms);
@@ -86,7 +86,7 @@ namespace EpicGames.Tracing.Tests.UnrealInsights
 		[TestMethod]
 		public void Deserialize7Bit()
 		{
-			void Assert7Bit(ulong ExpectedValue, string ActualBytesHex)
+			static void Assert7Bit(ulong ExpectedValue, string ActualBytesHex)
 			{
 				byte[] ActualBytes = GenericEventTest.StringToByteArray(ActualBytesHex.Replace(" ", "", StringComparison.Ordinal));
 				using MemoryStream Ms = new MemoryStream(ActualBytes);
@@ -124,7 +124,7 @@ namespace EpicGames.Tracing.Tests.UnrealInsights
 		[TestMethod]
 		public void Serialize7Bit()
 		{
-			void Assert7Bit(ulong Value, string ExpectedBytesHex)
+			static void Assert7Bit(ulong Value, string ExpectedBytesHex)
 			{
 				byte[] ExpectedBytes = GenericEventTest.StringToByteArray(ExpectedBytesHex.Replace(" ", "", StringComparison.Ordinal));
 				using MemoryStream Ms = new MemoryStream();
diff --git a/Engine/Source/Programs/Shared/EpicGames.Tracing/UnrealInsights/Events/GenericEvent.cs b/Engine/Source/Programs/Shared/EpicGames.Tracing/UnrealInsights/Events/GenericEvent.cs
--- a/Engine/Source/Programs/Shared/EpicGames.Tracing/UnrealInsights/Events/GenericEvent.cs
+++ b/Engine/Source/Programs/Shared/EpicGames.Tracing/UnrealInsights/Events/GenericEvent.cs
@@ -44,7 +44,7 @@ namespace EpicGames.Tracing.UnrealInsights.Events
 		readonly Field[] Fields;
 		public Field[] GetFields() => Fields;
 		public EventType Type => EventType;
-		private EventType EventType;
+		private readonly EventType EventType;
 
 		public GenericEvent(uint Serial, Field[] Fields, EventType EventType)
 		{
@@ -58,7 +58,8 @@ namespace EpicGames.Tracing.UnrealInsights.Events
 			get
 			{
 				ushort TotalSize = 0;
-				if (EventType.HasSerial) TotalSize += 3; // 24-bit serial
+				if (EventType.HasSerial)
+					TotalSize += 3; // 24-bit serial
 				TotalSize += EventType.GetEventSize();
 				if (EventType.MaybeHasAux())
 				{
@@ -105,20 +106,53 @@ namespace EpicGames.Tracing.UnrealInsights.Events
 			{
 				EventTypeField FieldType = EventType.Fields[i];
 
-				if (FieldType.TypeInfo == EventTypeField.TypeInt8) { Writer.Write((byte)Fields[i].Int!); }
-				else if (FieldType.TypeInfo == EventTypeField.TypeInt16) { Writer.Write((ushort)Fields[i].Int!); }
-				else if (FieldType.TypeInfo == EventTypeField.TypeInt32) { Writer.Write((uint)Fields[i].Int!); }
-				else if (FieldType.TypeInfo == EventTypeField.TypeInt64) { Writer.Write((ulong)Fields[i].Long!); }
-				else if (FieldType.TypeInfo == EventTypeField.TypePointer) { Writer.Write((ulong)Fields[i].Long!); }
-				else if (FieldType.TypeInfo == EventTypeField.TypeFloat32) { Writer.Write((float)Fields[i].Float!); }
-				else if (FieldType.TypeInfo == EventTypeField.TypeFloat64) { Writer.Write((double)Fields[i].Float!); }
-				else if (FieldType.TypeInfo == EventTypeField.TypeAnsiString) {  } // Write later as aux-data
-				else if (FieldType.TypeInfo == EventTypeField.TypeWideString) {  } // Write later as aux-data
-				else if (FieldType.TypeInfo == EventTypeField.TypeArray) {  } // Write later as aux-data
-				else if (FieldType.TypeInfo == EventTypeField.TypeBool) { Writer.Write(Fields[i].Bool!.Value ? (byte)1 : (byte)0); }
-				else { throw new Exception($"Found unknown TypeInfo {FieldType.TypeInfo}"); }
+				if (FieldType.TypeInfo == EventTypeField.TypeInt8)
+				{
+					Writer.Write((byte)Fields[i].Int!);
+				}
+				else if (FieldType.TypeInfo == EventTypeField.TypeInt16)
+				{
+					Writer.Write((ushort)Fields[i].Int!);
+				}
+				else if (FieldType.TypeInfo == EventTypeField.TypeInt32)
+				{
+					Writer.Write((uint)Fields[i].Int!);
+				}
+				else if (FieldType.TypeInfo == EventTypeField.TypeInt64)
+				{
+					Writer.Write((ulong)Fields[i].Long!);
+				}
+				else if (FieldType.TypeInfo == EventTypeField.TypePointer)
+				{
+					Writer.Write((ulong)Fields[i].Long!);
+				}
+				else if (FieldType.TypeInfo == EventTypeField.TypeFloat32)
+				{
+					Writer.Write((float)Fields[i].Float!);
+				}
+				else if (FieldType.TypeInfo == EventTypeField.TypeFloat64)
+				{
+					Writer.Write((double)Fields[i].Float!);
+				}
+				else if (FieldType.TypeInfo == EventTypeField.TypeAnsiString)
+				{
+				} // Write later as aux-data
+				else if (FieldType.TypeInfo == EventTypeField.TypeWideString)
+				{
+				} // Write later as aux-data
+				else if (FieldType.TypeInfo == EventTypeField.TypeArray)
+				{
+				} // Write later as aux-data
+				else if (FieldType.TypeInfo == EventTypeField.TypeBool)
+				{
+					Writer.Write(Fields[i].Bool!.Value ? (byte)1 : (byte)0);
+				}
+				else
+				{
+					throw new Exception($"Found unknown TypeInfo {FieldType.TypeInfo}");
+				}
 			}
-			
+
 			if (EventType.MaybeHasAux())
 			{
 				for (int i = 0; i < EventType.Fields.Count; i++)
@@ -186,18 +220,51 @@ namespace EpicGames.Tracing.UnrealInsights.Events
 			{
 				EventTypeField FieldType = EventType.Fields[i];
 
-				if (FieldType.TypeInfo == EventTypeField.TypeInt8) { Fields[i] = Field.FromInt(Reader.ReadByte()); }
-				else if (FieldType.TypeInfo == EventTypeField.TypeInt16) { Fields[i] = Field.FromInt(Reader.ReadUInt16()); }
-				else if (FieldType.TypeInfo == EventTypeField.TypeInt32) { Fields[i] = Field.FromInt((int)Reader.ReadUInt32()); }
-				else if (FieldType.TypeInfo == EventTypeField.TypeInt64) { Fields[i] = Field.FromLong((long) Reader.ReadUInt64()); }
-				else if (FieldType.TypeInfo == EventTypeField.TypePointer) { Fields[i] = Field.FromLong((long) Reader.ReadUInt64()); }
-				else if (FieldType.TypeInfo == EventTypeField.TypeFloat32) { Fields[i] = Field.FromFloat(Reader.ReadSingle()); }
-				else if (FieldType.TypeInfo == EventTypeField.TypeFloat64) { Fields[i] = Field.FromDouble(Reader.ReadDouble()); }
-				else if (FieldType.TypeInfo == EventTypeField.TypeAnsiString) {  } // Read later as aux-data
-				else if (FieldType.TypeInfo == EventTypeField.TypeWideString) {  } // Read later as aux-data
-				else if (FieldType.TypeInfo == EventTypeField.TypeArray) {  } // Read later as aux-data
-				else if (FieldType.TypeInfo == EventTypeField.TypeBool) { Fields[i] = Field.FromBool(Reader.ReadByte() == 1); }
-				else { throw new Exception($"Found unknown TypeInfo {FieldType.TypeInfo}"); }
+				if (FieldType.TypeInfo == EventTypeField.TypeInt8)
+				{
+					Fields[i] = Field.FromInt(Reader.ReadByte());
+				}
+				else if (FieldType.TypeInfo == EventTypeField.TypeInt16)
+				{
+					Fields[i] = Field.FromInt(Reader.ReadUInt16());
+				}
+				else if (FieldType.TypeInfo == EventTypeField.TypeInt32)
+				{
+					Fields[i] = Field.FromInt((int)Reader.ReadUInt32());
+				}
+				else if (FieldType.TypeInfo == EventTypeField.TypeInt64)
+				{
+					Fields[i] = Field.FromLong((long) Reader.ReadUInt64());
+				}
+				else if (FieldType.TypeInfo == EventTypeField.TypePointer)
+				{
+					Fields[i] = Field.FromLong((long) Reader.ReadUInt64());
+				}
+				else if (FieldType.TypeInfo == EventTypeField.TypeFloat32)
+				{
+					Fields[i] = Field.FromFloat(Reader.ReadSingle());
+				}
+				else if (FieldType.TypeInfo == EventTypeField.TypeFloat64)
+				{
+					Fields[i] = Field.FromDouble(Reader.ReadDouble());
+				}
+				else if (FieldType.TypeInfo == EventTypeField.TypeAnsiString)
+				{
+				} // Read later as aux-data
+				else if (FieldType.TypeInfo == EventTypeField.TypeWideString)
+				{
+				} // Read later as aux-data
+				else if (FieldType.TypeInfo == EventTypeField.TypeArray)
+				{
+				} // Read later as aux-data
+				else if (FieldType.TypeInfo == EventTypeField.TypeBool)
+				{
+					Fields[i] = Field.FromBool(Reader.ReadByte() == 1);
+				}
+				else
+				{
+					throw new Exception($"Found unknown TypeInfo {FieldType.TypeInfo}");
+				}
 			}
 
 			if (EventType.MaybeHasAux())
@@ -247,7 +314,7 @@ namespace EpicGames.Tracing.UnrealInsights.Events
 
 			return new GenericEvent(Serial, Fields, EventType);
 		}
-		
+
 		private static uint SetBits(uint Word, uint Value, int Pos, int Size)
 		{
 			uint mask = (((uint)1 << Size) - 1) << Pos;
diff --git a/Engine/Source/Programs/Shared/EpicGames.Tracing/UnrealInsights/Events/NewEvent.cs b/Engine/Source/Programs/Shared/EpicGames.Tracing/UnrealInsights/Events/NewEvent.cs
--- a/Engine/Source/Programs/Shared/EpicGames.Tracing/UnrealInsights/Events/NewEvent.cs
+++ b/Engine/Source/Programs/Shared/EpicGames.Tracing/UnrealInsights/Events/NewEvent.cs
@@ -9,6 +9,7 @@ using System.Text;
 namespace EpicGames.Tracing.UnrealInsights.Events
 {
 	[System.Diagnostics.CodeAnalysis.SuppressMessage("Performance", "CA1823:Avoid unused private fields")]
+	[System.Diagnostics.CodeAnalysis.SuppressMessage("Style", "IDE1006:Naming Styles")]
 	public class EventTypeField
 	{
 		public ushort Offset { get; private set; }
@@ -52,17 +53,28 @@ namespace EpicGames.Tracing.UnrealInsights.Events
 
 		public static string TypeInfoToString(byte ByteInfo)
 		{
-			if (ByteInfo == TypeInt8) return "int8";
-			if (ByteInfo == TypeInt16) return "int16";
-			if (ByteInfo == TypeInt32) return "int32";
-			if (ByteInfo == TypeInt64) return "int64";
-			if (ByteInfo == TypePointer) return "ptr";
-			if (ByteInfo == TypeFloat32) return "float32";
-			if (ByteInfo == TypeFloat64) return "float64";
-			if (ByteInfo == TypeAnsiString) return "ansi_str";
-			if (ByteInfo == TypeWideString) return "wide_str";
-			if (ByteInfo == TypeArray) return "array";
-			if (ByteInfo == TypeBool) return "bool";
+			if (ByteInfo == TypeInt8)
+				return "int8";
+			if (ByteInfo == TypeInt16)
+				return "int16";
+			if (ByteInfo == TypeInt32)
+				return "int32";
+			if (ByteInfo == TypeInt64)
+				return "int64";
+			if (ByteInfo == TypePointer)
+				return "ptr";
+			if (ByteInfo == TypeFloat32)
+				return "float32";
+			if (ByteInfo == TypeFloat64)
+				return "float64";
+			if (ByteInfo == TypeAnsiString)
+				return "ansi_str";
+			if (ByteInfo == TypeWideString)
+				return "wide_str";
+			if (ByteInfo == TypeArray)
+				return "array";
+			if (ByteInfo == TypeBool)
+				return "bool";
 			throw new Exception($"Unable to convert type info {ByteInfo} to string");
 		}
 
@@ -96,7 +108,7 @@ namespace EpicGames.Tracing.UnrealInsights.Events
 
 		public void SetName(string Name)
 		{
-			this.NameSize = (byte) Encoding.UTF8.GetBytes(Name).Length;
+			NameSize = (byte) Encoding.UTF8.GetBytes(Name).Length;
 			this.Name = Name;
 		}
 
@@ -147,11 +159,11 @@ namespace EpicGames.Tracing.UnrealInsights.Events
 
 		public EventType(string LoggerName, string EventName, byte Flags)
 		{
-			this.NewEventUid = 0;
+			NewEventUid = 0;
 			this.LoggerName = LoggerName;
 			this.EventName = EventName;
 			this.Flags = Flags;
-			this.Fields = new List<EventTypeField>();
+			Fields = new List<EventTypeField>();
 		}
 
 		public static EventType WellKnown(ushort Uid, string Name)
@@ -212,11 +224,14 @@ namespace EpicGames.Tracing.UnrealInsights.Events
 		
 		public string ToStringDetailed()
 		{
-			string FieldText = string.Join(',', Fields.Select(x => x.ToString()));
+			string FieldText = String.Join(',', Fields.Select(x => x.ToString()));
 			string FlagText = "";
-			if (IsImportant()) FlagText += ",Important";
-			if (MaybeHasAux()) FlagText += ",MaybeHasAux";
-			if (IsNoSync()) FlagText += ",NoSync";
+			if (IsImportant())
+				FlagText += ",Important";
+			if (MaybeHasAux())
+				FlagText += ",MaybeHasAux";
+			if (IsNoSync())
+				FlagText += ",NoSync";
 			FlagText = FlagText.Trim(',');
 			return $"EventType(Uid={NewEventUid} {LoggerName} {EventName} Flags={FlagText} FlagRaw={Flags} Fields={FieldText})";
 		}
@@ -251,9 +266,12 @@ namespace EpicGames.Tracing.UnrealInsights.Events
 
 		public override bool Equals(object? obj)
 		{
-			if (ReferenceEquals(null, obj)) return false;
-			if (ReferenceEquals(this, obj)) return true;
-			if (obj.GetType() != this.GetType()) return false;
+			if (obj is null)
+				return false;
+			if (ReferenceEquals(this, obj))
+				return true;
+			if (obj.GetType() != GetType())
+				return false;
 			return Equals((EventType) obj);
 		}
 
diff --git a/Engine/Source/Programs/Shared/EpicGames.Tracing/UnrealInsights/Writer.cs b/Engine/Source/Programs/Shared/EpicGames.Tracing/UnrealInsights/Writer.cs
--- a/Engine/Source/Programs/Shared/EpicGames.Tracing/UnrealInsights/Writer.cs
+++ b/Engine/Source/Programs/Shared/EpicGames.Tracing/UnrealInsights/Writer.cs
@@ -15,7 +15,7 @@ namespace EpicGames.Tracing.UnrealInsights
 		readonly Dictionary<ushort, List<ITraceEvent>> ThreadToEvents = new Dictionary<ushort, List<ITraceEvent>>();
 		
 		private ushort UidCounter = PredefinedEventUid._WellKnownNum; // IDs below 16 are reserved for well-known events
-		private object Lock = new object();
+		private readonly object Lock = new object();
 
 		public UnrealInsightsWriter()
 		{
