Shared C#: Add solution to compile all shared projects and update .editorconfig to match Horde's version

#rnx

diff --git a/Engine/Source/Programs/Shared/.editorconfig b/Engine/Source/Programs/Shared/.editorconfig
--- a/Engine/Source/Programs/Shared/.editorconfig
+++ b/Engine/Source/Programs/Shared/.editorconfig
@@ -5,6 +5,8 @@
 # Dotnet code style settings:
 [*.cs]
 
+
+
 csharp_style_prefer_switch_expression = false
 csharp_style_implicit_object_creation_when_type_is_apparent = false
 dotnet_remove_unnecessary_suppression_exclusions = all
@@ -392,7 +394,7 @@ dotnet_diagnostic.CA1054.severity = warning
 dotnet_diagnostic.CA1055.severity = warning
 
 # CA1056: URI-like properties should not be strings
-dotnet_diagnostic.CA1056.severity = silent
+dotnet_diagnostic.CA1056.severity = warning
 
 # CA1058: Types should not extend certain base types
 dotnet_diagnostic.CA1058.severity = warning
@@ -1096,6 +1098,68 @@ dotnet_diagnostic.CA5404.severity = none
 # CA5405: Do not always skip token validation in delegates
 dotnet_diagnostic.CA5405.severity = warning
 
+# VSTHRD003: Avoid awaiting or returning a Task representing work that was not started within your context as that can lead to deadlocks.
+dotnet_diagnostic.VSTHRD003.severity = none
+
+# MA0026: TODO comment
+dotnet_diagnostic.MA0026.severity = none
+
+# MA0048: File name must match type name
+dotnet_diagnostic.MA0048.severity = none
+
+# MA0004: Use Task.ConfigureAwait(false) if the current SynchronizationContext is not needed
+dotnet_diagnostic.MA0004.severity = none
+
+# MA0051: Method is too long
+dotnet_diagnostic.MA0051.severity = none
+
+# MA0002: Use an overload that has a IEqualityComparer<string> or IComparer<string> parameter
+dotnet_diagnostic.MA0002.severity = none
+
+# MA0016: Prefer using collection abstraction instead of implementation
+dotnet_diagnostic.MA0016.severity = none
+
+# MA0006: Use string.Equals instead of Equals operator
+dotnet_diagnostic.MA0006.severity = none
+
+# MA0061: Method overrides should not change default values
+dotnet_diagnostic.MA0061.severity = none
+
+# MA0134: Observe result of async calls
+dotnet_diagnostic.MA0134.severity = none
+
+# MA0011: Use an overload of 'ToString' that has a 'System.IFormatProvider' parameter
+dotnet_diagnostic.MA0011.severity = none
+
+# MA0009: Regular expressions should not be vulnerable to Denial of Service attacks
+dotnet_diagnostic.MA0009.severity = none
+
+# MA0099: Use Explicit enum value for 'SlackNotificationSink.ReactionFlags' instead of 0
+dotnet_diagnostic.MA0099.severity = none
+
+# MA0025: Implement the functionality (or raise NotSupportedException or PlatformNotSupportedException)
+dotnet_diagnostic.MA0025.severity = none
+
+# MA0015: Use an overload of 'System.ArgumentException' with the parameter name
+dotnet_diagnostic.MA0015.severity = none
+
+# MA0132: Do not convert implicitly to DateTimeOffset
+dotnet_diagnostic.MA0132.severity = none
+
+# MA0023: Add RegexOptions.ExplicitCapture to prevent capturing unneeded groups
+dotnet_diagnostic.MA0023.severity = none
+
+# MA0056: Do not call overridable members in constructor
+dotnet_diagnostic.MA0056.severity = none
+
+# MA0144: Use System.OperatingSystem to check the current OS
+dotnet_diagnostic.MA0144.severity = none
+
+# MA0046 - Use EventHandler<T> to declare events
+dotnet_diagnostic.MA0046.severity = none
+
+# MA0008: Add StructLayoutAttribute
+dotnet_diagnostic.MA0008.severity = none
 
 [*Test*.cs]
 
