EpicGames.MongoDB: Fix variable naming conventions

#rnx

diff --git a/Engine/Source/Programs/Shared/EpicGames.MongoDB/MongoExtensions.cs b/Engine/Source/Programs/Shared/EpicGames.MongoDB/MongoExtensions.cs
--- a/Engine/Source/Programs/Shared/EpicGames.MongoDB/MongoExtensions.cs
+++ b/Engine/Source/Programs/Shared/EpicGames.MongoDB/MongoExtensions.cs
@@ -16,61 +16,61 @@ namespace EpicGames.MongoDB
 		/// <summary>
 		/// Filters the documents returned from a search
 		/// </summary>
-		/// <param name="Query">The query to filter</param>
-		/// <param name="Index">Index of the first document to return</param>
-		/// <param name="Count">Number of documents to return</param>
+		/// <param name="query">The query to filter</param>
+		/// <param name="index">Index of the first document to return</param>
+		/// <param name="count">Number of documents to return</param>
 		/// <returns>New query</returns>
-		public static IFindFluent<TDocument, TProjection> Range<TDocument, TProjection>(this IFindFluent<TDocument, TProjection> Query, int? Index, int? Count)
+		public static IFindFluent<TDocument, TProjection> Range<TDocument, TProjection>(this IFindFluent<TDocument, TProjection> query, int? index, int? count)
 		{
-			if (Index != null)
+			if (index != null)
 			{
-				Query = Query.Skip(Index.Value);
+				query = query.Skip(index.Value);
 			}
-			if(Count != null)
+			if(count != null)
 			{
-				Query = Query.Limit(Count.Value);
+				query = query.Limit(count.Value);
 			}
-			return Query;
+			return query;
 		}
 
 		/// <summary>
 		/// Filters the documents returned from a search
 		/// </summary>
-		/// <param name="Query">The query to filter</param>
+		/// <param name="query">The query to filter</param>
 		/// <returns>New query</returns>
-		public static async Task<List<TResult>> ToListAsync<TDocument, TResult>(this IAsyncCursorSource<TDocument> Query) where TDocument : TResult
+		public static async Task<List<TResult>> ToListAsync<TDocument, TResult>(this IAsyncCursorSource<TDocument> query) where TDocument : TResult
 		{
-			List<TResult> Results = new List<TResult>();
-			using (IAsyncCursor<TDocument> Cursor = await Query.ToCursorAsync())
+			List<TResult> results = new List<TResult>();
+			using (IAsyncCursor<TDocument> cursor = await query.ToCursorAsync())
 			{
-				while (await Cursor.MoveNextAsync())
+				while (await cursor.MoveNextAsync())
 				{
-					foreach(TDocument Document in Cursor.Current)
+					foreach (TDocument document in cursor.Current)
 					{
-						Results.Add(Document);
+						results.Add(document);
 					}
 				}
 			}
-			return Results;
+			return results;
 		}
 
 		/// <summary>
 		/// Attempts to insert a document into a collection, handling the error case that a document with the given key already exists
 		/// </summary>
 		/// <typeparam name="TDocument"></typeparam>
-		/// <param name="Collection">Collection to insert into</param>
-		/// <param name="NewDocument">The document to insert</param>
+		/// <param name="collection">Collection to insert into</param>
+		/// <param name="newDocument">The document to insert</param>
 		/// <returns>True if the document was inserted, false if it already exists</returns>
-		public static async Task<bool> InsertOneIgnoreDuplicatesAsync<TDocument>(this IMongoCollection<TDocument> Collection, TDocument NewDocument)
+		public static async Task<bool> InsertOneIgnoreDuplicatesAsync<TDocument>(this IMongoCollection<TDocument> collection, TDocument newDocument)
 		{
 			try
 			{
-				await Collection.InsertOneAsync(NewDocument);
+				await collection.InsertOneAsync(newDocument);
 				return true;
 			}
-			catch (MongoWriteException Ex)
+			catch (MongoWriteException ex)
 			{
-				if (Ex.WriteError.Category == ServerErrorCategory.DuplicateKey)
+				if (ex.WriteError.Category == ServerErrorCategory.DuplicateKey)
 				{
 					return false;
 				}
@@ -86,19 +86,19 @@ namespace EpicGames.MongoDB
 		/// </summary>
 		/// <typeparam name="TDocument">The document type</typeparam>
 		/// <typeparam name="TField">Type of the field to set</typeparam>
-		/// <param name="UpdateBuilder">Update builder</param>
-		/// <param name="Field">Expression for the field to set</param>
-		/// <param name="Value">New value to set</param>
+		/// <param name="updateBuilder">Update builder</param>
+		/// <param name="field">Expression for the field to set</param>
+		/// <param name="value">New value to set</param>
 		/// <returns>Update defintiion</returns>
-		public static UpdateDefinition<TDocument> SetOrUnsetNull<TDocument, TField>(this UpdateDefinitionBuilder<TDocument> UpdateBuilder, Expression<Func<TDocument, TField?>> Field, TField? Value) where TField : struct
+		public static UpdateDefinition<TDocument> SetOrUnsetNull<TDocument, TField>(this UpdateDefinitionBuilder<TDocument> updateBuilder, Expression<Func<TDocument, TField?>> field, TField? value) where TField : struct
 		{
-			if (Value.HasValue)
+			if (value.HasValue)
 			{
-				return UpdateBuilder.Set(Field, Value.Value);
+				return updateBuilder.Set(field, value.Value);
 			}
 			else
 			{
-				return UpdateBuilder.Unset(new ExpressionFieldDefinition<TDocument, TField?>(Field));
+				return updateBuilder.Unset(new ExpressionFieldDefinition<TDocument, TField?>(field));
 			}
 		}
 
@@ -107,19 +107,19 @@ namespace EpicGames.MongoDB
 		/// </summary>
 		/// <typeparam name="TDocument">The document type</typeparam>
 		/// <typeparam name="TField">Type of the field to set</typeparam>
-		/// <param name="UpdateBuilder">Update builder</param>
-		/// <param name="Field">Expression for the field to set</param>
-		/// <param name="Value">New value to set</param>
+		/// <param name="updateBuilder">Update builder</param>
+		/// <param name="field">Expression for the field to set</param>
+		/// <param name="value">New value to set</param>
 		/// <returns>Update defintiion</returns>
-		public static UpdateDefinition<TDocument> SetOrUnsetNullRef<TDocument, TField>(this UpdateDefinitionBuilder<TDocument> UpdateBuilder, Expression<Func<TDocument, TField?>> Field, TField? Value) where TField : class
+		public static UpdateDefinition<TDocument> SetOrUnsetNullRef<TDocument, TField>(this UpdateDefinitionBuilder<TDocument> updateBuilder, Expression<Func<TDocument, TField?>> field, TField? value) where TField : class
 		{
-			if (Value != null)
+			if (value != null)
 			{
-				return UpdateBuilder.Set(Field, Value);
+				return updateBuilder.Set(field, value);
 			}
 			else
 			{
-				return UpdateBuilder.Unset(new ExpressionFieldDefinition<TDocument, TField?>(Field));
+				return updateBuilder.Unset(new ExpressionFieldDefinition<TDocument, TField?>(field));
 			}
 		}
 
@@ -127,12 +127,13 @@ namespace EpicGames.MongoDB
 		/// Creates a filter definition from a linq expression. This is not generally explicitly castable, so expose it as a FilterDefinitionBuilder method.
 		/// </summary>
 		/// <typeparam name="TDocument">The document type</typeparam>
-		/// <param name="Filter">The filter builder</param>
-		/// <param name="Expression">Expression to parse</param>
+		/// <param name="filter">The filter builder</param>
+		/// <param name="expression">Expression to parse</param>
 		/// <returns>New filter definition</returns>
-		public static FilterDefinition<TDocument> Expr<TDocument>(this FilterDefinitionBuilder<TDocument> Filter, Expression<Func<TDocument, bool>> Expression)
+		[System.Diagnostics.CodeAnalysis.SuppressMessage("Style", "IDE0060:Remove unused parameter")]
+		public static FilterDefinition<TDocument> Expr<TDocument>(this FilterDefinitionBuilder<TDocument> filter, Expression<Func<TDocument, bool>> expression)
 		{
-			return Expression;
+			return expression;
 		}
 	}
 }
diff --git a/Engine/Source/Programs/Shared/EpicGames.MongoDB/ObjectIdHelpers.cs b/Engine/Source/Programs/Shared/EpicGames.MongoDB/ObjectIdHelpers.cs
--- a/Engine/Source/Programs/Shared/EpicGames.MongoDB/ObjectIdHelpers.cs
+++ b/Engine/Source/Programs/Shared/EpicGames.MongoDB/ObjectIdHelpers.cs
@@ -6,15 +6,15 @@ namespace EpicGames.MongoDB
 {
 	static class ObjectIdHelpers
 	{
-		public static ObjectId ToObjectId(this string Text)
+		public static ObjectId ToObjectId(this string text)
 		{
-			if (Text.Length == 0)
+			if (text.Length == 0)
 			{
 				return ObjectId.Empty;
 			}
 			else
 			{
-				return ObjectId.Parse(Text);
+				return ObjectId.Parse(text);
 			}
 		}
 	}
diff --git a/Engine/Source/Programs/Shared/EpicGames.MongoDB/SingletonDocument.cs b/Engine/Source/Programs/Shared/EpicGames.MongoDB/SingletonDocument.cs
--- a/Engine/Source/Programs/Shared/EpicGames.MongoDB/SingletonDocument.cs
+++ b/Engine/Source/Programs/Shared/EpicGames.MongoDB/SingletonDocument.cs
@@ -24,10 +24,10 @@ namespace EpicGames.MongoDB
 		/// <summary>
 		/// Constructor
 		/// </summary>
-		/// <param name="Id">Unique id for the singleton document</param>
-		public SingletonDocumentAttribute(string Id)
+		/// <param name="id">Unique id for the singleton document</param>
+		public SingletonDocumentAttribute(string id)
 		{
-			this.Id = Id;
+			Id = id;
 		}
 	}
 
@@ -58,10 +58,10 @@ namespace EpicGames.MongoDB
 		/// <summary>
 		/// Constructor
 		/// </summary>
-		/// <param name="Id">Unique id for the singleton</param>
-		public SingletonBase(ObjectId Id)
+		/// <param name="id">Unique id for the singleton</param>
+		public SingletonBase(ObjectId id)
 		{
-			this.Id = Id;
+			Id = id;
 		}
 
 		/// <summary>
@@ -73,8 +73,8 @@ namespace EpicGames.MongoDB
 
 			static ObjectId GetSingletonId()
 			{
-				SingletonDocumentAttribute? Attribute = typeof(T).GetCustomAttribute<SingletonDocumentAttribute>() ?? throw new Exception($"Type {typeof(T).Name} is missing a {nameof(SingletonDocumentAttribute)} annotation");
-				return ObjectId.Parse(Attribute.Id);
+				SingletonDocumentAttribute? attribute = typeof(T).GetCustomAttribute<SingletonDocumentAttribute>() ?? throw new Exception($"Type {typeof(T).Name} is missing a {nameof(SingletonDocumentAttribute)} annotation");
+				return ObjectId.Parse(attribute.Id);
 			}
 		}
 
@@ -104,9 +104,9 @@ namespace EpicGames.MongoDB
 		/// <summary>
 		/// Attempts to update the document
 		/// </summary>
-		/// <param name="Value">New state of the document</param>
+		/// <param name="value">New state of the document</param>
 		/// <returns>True if the document was updated, false otherwise</returns>
-		Task<bool> TryUpdateAsync(T Value);
+		Task<bool> TryUpdateAsync(T value);
 	}
 
 	/// <summary>
@@ -118,12 +118,12 @@ namespace EpicGames.MongoDB
 		/// <summary>
 		/// The database service instance
 		/// </summary>
-		readonly IMongoCollection<T> Collection;
+		readonly IMongoCollection<T> _collection;
 
 		/// <summary>
 		/// Unique id for the singleton document
 		/// </summary>
-		readonly ObjectId ObjectId;
+		readonly ObjectId _objectId;
 
 		/// <summary>
 		/// Static constructor. Registers the document using the automapper.
@@ -137,12 +137,12 @@ namespace EpicGames.MongoDB
 		/// Constructor
 		/// </summary>
 		/// <param name="DatabaseService">The database service instance</param>
-		public SingletonDocument(IMongoCollection<SingletonBase> Collection)
+		public SingletonDocument(IMongoCollection<SingletonBase> collection)
 		{
-			this.Collection = Collection.OfType<T>();
+			_collection = collection.OfType<T>();
 
-			SingletonDocumentAttribute? Attribute = typeof(T).GetCustomAttribute<SingletonDocumentAttribute>() ?? throw new Exception($"Type {typeof(T).Name} is missing a {nameof(SingletonDocumentAttribute)} annotation");
-			ObjectId = new ObjectId(Attribute.Id);
+			SingletonDocumentAttribute? attribute = typeof(T).GetCustomAttribute<SingletonDocumentAttribute>() ?? throw new Exception($"Type {typeof(T).Name} is missing a {nameof(SingletonDocumentAttribute)} annotation");
+			_objectId = new ObjectId(attribute.Id);
 		}
 
 		/// <inheritdoc/>
@@ -150,31 +150,31 @@ namespace EpicGames.MongoDB
 		{
 			for (; ; )
 			{
-				T? Object = await Collection.Find<T>(x => x.Id == ObjectId).FirstOrDefaultAsync();
-				if (Object != null)
+				T? @object = await _collection.Find<T>(x => x.Id == _objectId).FirstOrDefaultAsync();
+				if (@object != null)
 				{
-					return Object;
+					return @object;
 				}
 
-				T NewItem = new T();
-				NewItem.Id = ObjectId;
-				await Collection.InsertOneAsync(NewItem);
+				T newItem = new T();
+				newItem.Id = _objectId;
+				await _collection.InsertOneAsync(newItem);
 			}
 		}
 
 		/// <inheritdoc/>
-		public async Task<bool> TryUpdateAsync(T Value)
+		public async Task<bool> TryUpdateAsync(T value)
 		{
-			int PrevRevision = Value.Revision++;
+			int prevRevision = value.Revision++;
 			try
 			{
-				ReplaceOneResult Result = await Collection.ReplaceOneAsync(x => x.Id == ObjectId && x.Revision == PrevRevision, Value, new ReplaceOptions { IsUpsert = true });
-				return Result.MatchedCount > 0;
+				ReplaceOneResult result = await _collection.ReplaceOneAsync(x => x.Id == _objectId && x.Revision == prevRevision, value, new ReplaceOptions { IsUpsert = true });
+				return result.MatchedCount > 0;
 			}
-			catch (MongoWriteException Ex)
+			catch (MongoWriteException ex)
 			{
 				// Duplicate key error occurs if filter fails to match because revision is not the same.
-				if (Ex.WriteError != null && Ex.WriteError.Category == ServerErrorCategory.DuplicateKey)
+				if (ex.WriteError != null && ex.WriteError.Category == ServerErrorCategory.DuplicateKey)
 				{
 					return false;
 				}
diff --git a/Engine/Source/Programs/Shared/EpicGames.MongoDB/TransactionBuilder.cs b/Engine/Source/Programs/Shared/EpicGames.MongoDB/TransactionBuilder.cs
--- a/Engine/Source/Programs/Shared/EpicGames.MongoDB/TransactionBuilder.cs
+++ b/Engine/Source/Programs/Shared/EpicGames.MongoDB/TransactionBuilder.cs
@@ -23,22 +23,22 @@ namespace EpicGames.MongoDB
 			/// <summary>
 			/// The update definition
 			/// </summary>
-			public UpdateDefinition<TDocument> Update;
+			public UpdateDefinition<TDocument> _update;
 
 			/// <summary>
 			/// Applies an update to the given object
 			/// </summary>
-			public Action<object> Apply;
+			public Action<object> _apply;
 
 			/// <summary>
 			/// Constructor
 			/// </summary>
-			/// <param name="Update">The update definition</param>
-			/// <param name="Apply">Applies the update to a given object</param>
-			public FieldUpdate(UpdateDefinition<TDocument> Update, Action<object> Apply)
+			/// <param name="update">The update definition</param>
+			/// <param name="apply">Applies the update to a given object</param>
+			public FieldUpdate(UpdateDefinition<TDocument> update, Action<object> apply)
 			{
-				this.Update = Update;
-				this.Apply = Apply;
+				_update = update;
+				_apply = apply;
 			}
 		}
 
@@ -59,11 +59,11 @@ namespace EpicGames.MongoDB
 			/// <returns>The property info</returns>
 			static PropertyInfo GetIndexerProperty()
 			{
-				foreach (PropertyInfo Property in typeof(T).GetProperties())
+				foreach (PropertyInfo property in typeof(T).GetProperties())
 				{
-					if (Property.GetIndexParameters().Length > 0)
+					if (property.GetIndexParameters().Length > 0)
 					{
-						return Property;
+						return property;
 					}
 				}
 				return null!;
@@ -73,57 +73,57 @@ namespace EpicGames.MongoDB
 		/// <summary>
 		/// List of field updates
 		/// </summary>
-		readonly List<FieldUpdate> FieldUpdates = new List<FieldUpdate>();
+		readonly List<FieldUpdate> _fieldUpdates = new List<FieldUpdate>();
 
 		/// <summary>
 		/// Whether the transaction is currenty empty
 		/// </summary>
-		public bool IsEmpty => FieldUpdates.Count == 0;
+		public bool IsEmpty => _fieldUpdates.Count == 0;
 
 		/// <summary>
 		/// Adds a setting to this transaction
 		/// </summary>
 		/// <typeparam name="TField">Type of the field to be updated</typeparam>
-		/// <param name="Expr">The expresion defining the field to update</param>
-		/// <param name="Value">New value for the field</param>
-		public void Set<TField>(Expression<Func<TDocument, TField>> Expr, TField Value)
+		/// <param name="expr">The expresion defining the field to update</param>
+		/// <param name="value">New value for the field</param>
+		public void Set<TField>(Expression<Func<TDocument, TField>> expr, TField value)
 		{
-			UpdateDefinition<TDocument> Update = Builders<TDocument>.Update.Set(Expr, Value);
-			void Apply(object Target) => Assign(Target, Expr, Value);
-			FieldUpdates.Add(new FieldUpdate(Update, Apply));
+			UpdateDefinition<TDocument> update = Builders<TDocument>.Update.Set(expr, value);
+			void Apply(object target) => Assign(target, expr, value);
+			_fieldUpdates.Add(new FieldUpdate(update, Apply));
 		}
 
 		/// <summary>
 		/// Adds an update to remove a field
 		/// </summary>
-		/// <param name="Expr">The expresion defining the field to update</param>
-		public void Unset(Expression<Func<TDocument, object>> Expr)
+		/// <param name="expr">The expresion defining the field to update</param>
+		public void Unset(Expression<Func<TDocument, object>> expr)
 		{
-			UpdateDefinition<TDocument> Update = Builders<TDocument>.Update.Unset(Expr);
-			void Apply(object Target) => Unassign(Target, Expr.Body);
-			FieldUpdates.Add(new FieldUpdate(Update, Apply));
+			UpdateDefinition<TDocument> update = Builders<TDocument>.Update.Unset(expr);
+			void Apply(object target) => Unassign(target, expr.Body);
+			_fieldUpdates.Add(new FieldUpdate(update, Apply));
 		}
 
 		/// <summary>
 		/// Updates a property dictionary with a set of adds and removes
 		/// </summary>
-		/// <param name="Expr">Lambda expression defining a field to update</param>
-		/// <param name="Updates">List of updates</param>
-		public void UpdateDictionary<TKey, TValue>(Expression<Func<TDocument, Dictionary<TKey, TValue>>> Expr, IEnumerable<KeyValuePair<TKey, TValue?>> Updates) where TKey : class where TValue : class
+		/// <param name="expr">Lambda expression defining a field to update</param>
+		/// <param name="updates">List of updates</param>
+		public void UpdateDictionary<TKey, TValue>(Expression<Func<TDocument, Dictionary<TKey, TValue>>> expr, IEnumerable<KeyValuePair<TKey, TValue?>> updates) where TKey : class where TValue : class
 		{
-			PropertyInfo IndexerProperty = Reflection<Dictionary<TKey, TValue>>.Indexer;
-			foreach (KeyValuePair<TKey, TValue?> Update in Updates)
+			PropertyInfo indexerProperty = Reflection<Dictionary<TKey, TValue>>.Indexer;
+			foreach (KeyValuePair<TKey, TValue?> update in updates)
 			{
-				MethodCallExpression Index = Expression.Call(Expr.Body, IndexerProperty.GetMethod!, new[] { Expression.Constant(Update.Key) });
-				if (Update.Value == null)
+				MethodCallExpression index = Expression.Call(expr.Body, indexerProperty.GetMethod!, new[] { Expression.Constant(update.Key) });
+				if (update.Value == null)
 				{
-					Expression<Func<TDocument, object>> IndexExpression = Expression.Lambda<Func<TDocument, object>>(Index, Expr.Parameters[0]);
-					Unset(IndexExpression);
+					Expression<Func<TDocument, object>> indexExpression = Expression.Lambda<Func<TDocument, object>>(index, expr.Parameters[0]);
+					Unset(indexExpression);
 				}
 				else
 				{
-					Expression<Func<TDocument, TValue>> IndexExpression = Expression.Lambda<Func<TDocument, TValue>>(Index, Expr.Parameters[0]);
-					Set(IndexExpression, Update.Value);
+					Expression<Func<TDocument, TValue>> indexExpression = Expression.Lambda<Func<TDocument, TValue>>(index, expr.Parameters[0]);
+					Set(indexExpression, update.Value);
 				}
 			}
 		}
@@ -131,57 +131,57 @@ namespace EpicGames.MongoDB
 		/// <summary>
 		/// Applies this transaction to the given update
 		/// </summary>
-		/// <param name="Target">The object to update</param>
-		public void ApplyTo(TDocument Target)
+		/// <param name="target">The object to update</param>
+		public void ApplyTo(TDocument target)
 		{
-			foreach (FieldUpdate Setting in FieldUpdates)
+			foreach (FieldUpdate setting in _fieldUpdates)
 			{
-				Setting.Apply(Target);
+				setting._apply(target);
 			}
 		}
 
 		/// <summary>
 		/// Assign a value to a field
 		/// </summary>
-		/// <param name="Target">The target object to be updated</param>
-		/// <param name="FieldExpression">Lambda indicating the field to update</param>
-		/// <param name="Value">New value for the field</param>
-		static void Assign(object? Target, LambdaExpression FieldExpression, object? Value)
+		/// <param name="target">The target object to be updated</param>
+		/// <param name="fieldExpression">Lambda indicating the field to update</param>
+		/// <param name="value">New value for the field</param>
+		static void Assign(object? target, LambdaExpression fieldExpression, object? value)
 		{
-			Expression Body = FieldExpression.Body;
-			if (Body.NodeType == ExpressionType.MemberAccess)
+			Expression body = fieldExpression.Body;
+			if (body.NodeType == ExpressionType.MemberAccess)
 			{
-				MemberExpression MemberExpression = (MemberExpression)Body;
-				object? Object = Evaluate(MemberExpression.Expression!, Target);
-				PropertyInfo PropertyInfo = (PropertyInfo)MemberExpression.Member;
-				PropertyInfo.SetValue(Object, Value);
+				MemberExpression memberExpression = (MemberExpression)body;
+				object? @object = Evaluate(memberExpression.Expression!, target);
+				PropertyInfo propertyInfo = (PropertyInfo)memberExpression.Member;
+				propertyInfo.SetValue(@object, value);
 			}
-			else if (Body.NodeType == ExpressionType.ArrayIndex)
+			else if (body.NodeType == ExpressionType.ArrayIndex)
 			{
-				BinaryExpression BinaryExpression = (BinaryExpression)Body;
-				System.Collections.IList List = (System.Collections.IList)Evaluate(BinaryExpression.Left, Target)!;
-				int Index = (int)Evaluate(BinaryExpression.Right, Target)!;
-				List[Index] = Value;
+				BinaryExpression binaryExpression = (BinaryExpression)body;
+				System.Collections.IList list = (System.Collections.IList)Evaluate(binaryExpression.Left, target)!;
+				int index = (int)Evaluate(binaryExpression.Right, target)!;
+				list[index] = value;
 			}
-			else if (Body.NodeType == ExpressionType.Call)
+			else if (body.NodeType == ExpressionType.Call)
 			{
-				MethodCallExpression CallExpression = (MethodCallExpression)Body;
-				System.Collections.IDictionary Dictionary = (System.Collections.IDictionary)Evaluate(CallExpression.Object!, Target)!;
-				object? Key = Evaluate(CallExpression.Arguments[0], null);
-				Dictionary[Key!] = Value;
+				MethodCallExpression callExpression = (MethodCallExpression)body;
+				System.Collections.IDictionary dictionary = (System.Collections.IDictionary)Evaluate(callExpression.Object!, target)!;
+				object? key = Evaluate(callExpression.Arguments[0], null);
+				dictionary[key!] = value;
 			}
-			else if (Body.NodeType == ExpressionType.Index)
+			else if (body.NodeType == ExpressionType.Index)
 			{
-				IndexExpression IndexExpression = (IndexExpression)Body;
-				object? Object = Evaluate(IndexExpression.Object!, Target);
+				IndexExpression indexExpression = (IndexExpression)body;
+				object? @object = Evaluate(indexExpression.Object!, target);
 
-				object?[] Arguments = new object?[IndexExpression.Arguments.Count];
-				for (int Idx = 0; Idx < IndexExpression.Arguments.Count; Idx++)
+				object?[] arguments = new object?[indexExpression.Arguments.Count];
+				for (int idx = 0; idx < indexExpression.Arguments.Count; idx++)
 				{
-					Arguments[Idx] = Evaluate(IndexExpression.Arguments[Idx], null);
+					arguments[idx] = Evaluate(indexExpression.Arguments[idx], null);
 				}
 
-				IndexExpression.Indexer!.SetValue(Object, Value, Arguments);
+				indexExpression.Indexer!.SetValue(@object, value, arguments);
 			}
 			else
 			{
@@ -192,43 +192,43 @@ namespace EpicGames.MongoDB
 		/// <summary>
 		/// Removes an entry from a dictionary
 		/// </summary>
-		/// <param name="Target">The target object to be updated</param>
-		/// <param name="Body">Body of the expression indicating the field to update</param>
-		static void Unassign(object? Target, Expression Body)
+		/// <param name="target">The target object to be updated</param>
+		/// <param name="body">Body of the expression indicating the field to update</param>
+		static void Unassign(object? target, Expression body)
 		{
-			if (Body.NodeType == ExpressionType.MemberAccess)
+			if (body.NodeType == ExpressionType.MemberAccess)
 			{
-				MemberExpression MemberExpression = (MemberExpression)Body;
-				switch (MemberExpression.Member.MemberType)
+				MemberExpression memberExpression = (MemberExpression)body;
+				switch (memberExpression.Member.MemberType)
 				{
 					case MemberTypes.Field:
-						((FieldInfo)MemberExpression.Member).SetValue(Target, null);
+						((FieldInfo)memberExpression.Member).SetValue(target, null);
 						break;
 					case MemberTypes.Property:
-						((PropertyInfo)MemberExpression.Member).SetValue(Target, null);
+						((PropertyInfo)memberExpression.Member).SetValue(target, null);
 						break;
 					default:
 						throw new NotImplementedException();
 				}
 			}
-			else if (Body.NodeType == ExpressionType.Call)
+			else if (body.NodeType == ExpressionType.Call)
 			{
-				MethodCallExpression CallExpression = (MethodCallExpression)Body;
-				System.Collections.IDictionary? Dictionary = (System.Collections.IDictionary?)Evaluate(CallExpression.Object!, Target);
-				object? Key = Evaluate(CallExpression.Arguments[0], null);
-				Dictionary!.Remove(Key!);
+				MethodCallExpression callExpression = (MethodCallExpression)body;
+				System.Collections.IDictionary? dictionary = (System.Collections.IDictionary?)Evaluate(callExpression.Object!, target);
+				object? key = Evaluate(callExpression.Arguments[0], null);
+				dictionary!.Remove(key!);
 			}
-			else if (Body.NodeType == ExpressionType.Index)
+			else if (body.NodeType == ExpressionType.Index)
 			{
-				IndexExpression IndexExpression = (IndexExpression)Body;
-				System.Collections.IDictionary? Dictionary = (System.Collections.IDictionary?)Evaluate(IndexExpression.Object!, Target);
-				object? Key = Evaluate(IndexExpression.Arguments[0], null);
-				Dictionary!.Remove(Key!);
+				IndexExpression indexExpression = (IndexExpression)body;
+				System.Collections.IDictionary? dictionary = (System.Collections.IDictionary?)Evaluate(indexExpression.Object!, target);
+				object? key = Evaluate(indexExpression.Arguments[0], null);
+				dictionary!.Remove(key!);
 			}
-			else if (Body.NodeType == ExpressionType.Convert)
+			else if (body.NodeType == ExpressionType.Convert)
 			{
-				UnaryExpression UnaryExpression = (UnaryExpression)Body;
-				Unassign(Target, UnaryExpression.Operand);
+				UnaryExpression unaryExpression = (UnaryExpression)body;
+				Unassign(target, unaryExpression.Operand);
 			}
 			else
 			{
@@ -239,49 +239,49 @@ namespace EpicGames.MongoDB
 		/// <summary>
 		/// Evaluates an expression
 		/// </summary>
-		/// <param name="Expression">The expression to evaluate</param>
-		/// <param name="Parameter">Parameter to the unary lambda expression</param>
+		/// <param name="expression">The expression to evaluate</param>
+		/// <param name="parameter">Parameter to the unary lambda expression</param>
 		/// <returns>Value of the expression</returns>
-		static object? Evaluate(Expression Expression, object? Parameter)
+		static object? Evaluate(Expression expression, object? parameter)
 		{
-			if (Expression.NodeType == ExpressionType.Call)
+			if (expression.NodeType == ExpressionType.Call)
 			{
-				MethodCallExpression CallExpression = (MethodCallExpression)Expression;
-				object? Object = Evaluate(CallExpression.Object!, Parameter);
-				object?[] Arguments = CallExpression.Arguments.Select(x => Evaluate(x, Parameter)).ToArray();
-				return CallExpression.Method.Invoke(Object, Arguments);
+				MethodCallExpression callExpression = (MethodCallExpression)expression;
+				object? @object = Evaluate(callExpression.Object!, parameter);
+				object?[] arguments = callExpression.Arguments.Select(x => Evaluate(x, parameter)).ToArray();
+				return callExpression.Method.Invoke(@object, arguments);
 			}
-			else if (Expression.NodeType == ExpressionType.Constant)
+			else if (expression.NodeType == ExpressionType.Constant)
 			{
-				ConstantExpression ConstantExpression = (ConstantExpression)Expression;
-				return ConstantExpression.Value;
+				ConstantExpression constantExpression = (ConstantExpression)expression;
+				return constantExpression.Value;
 			}
-			else if (Expression.NodeType == ExpressionType.MemberAccess)
+			else if (expression.NodeType == ExpressionType.MemberAccess)
 			{
-				MemberExpression MemberExpression = (MemberExpression)Expression;
-				object? Target = Evaluate(MemberExpression.Expression!, Parameter);
+				MemberExpression memberExpression = (MemberExpression)expression;
+				object? target = Evaluate(memberExpression.Expression!, parameter);
 
-				MemberInfo Member = MemberExpression.Member;
-				switch (Member.MemberType)
+				MemberInfo member = memberExpression.Member;
+				switch (member.MemberType)
 				{
 					case MemberTypes.Property:
-						return ((PropertyInfo)Member).GetValue(Target);
+						return ((PropertyInfo)member).GetValue(target);
 					case MemberTypes.Field:
-						return ((FieldInfo)Member).GetValue(Target);
+						return ((FieldInfo)member).GetValue(target);
 					default:
 						throw new NotImplementedException("Unsupported expression type");
 				}
 			}
-			else if (Expression.NodeType == ExpressionType.ArrayIndex)
+			else if (expression.NodeType == ExpressionType.ArrayIndex)
 			{
-				BinaryExpression BinaryExpression = (BinaryExpression)Expression;
-				System.Collections.IList List = (System.Collections.IList)Evaluate(BinaryExpression.Left, Parameter)!;
-				int Index = (int)Evaluate(BinaryExpression.Right, null)!;
-				return List[Index];
+				BinaryExpression binaryExpression = (BinaryExpression)expression;
+				System.Collections.IList list = (System.Collections.IList)Evaluate(binaryExpression.Left, parameter)!;
+				int index = (int)Evaluate(binaryExpression.Right, null)!;
+				return list[index];
 			}
-			else if (Expression.NodeType == ExpressionType.Parameter)
+			else if (expression.NodeType == ExpressionType.Parameter)
 			{
-				return Parameter;
+				return parameter;
 			}
 			else
 			{
@@ -294,7 +294,7 @@ namespace EpicGames.MongoDB
 		/// </summary>
 		public UpdateDefinition<TDocument> ToUpdateDefinition()
 		{
-			return Builders<TDocument>.Update.Combine(FieldUpdates.Select(x => x.Update));
+			return Builders<TDocument>.Update.Combine(_fieldUpdates.Select(x => x._update));
 		}
 	}
 }
