EpicgGames.Tracing: Fix variable naming conventions

#rnx

diff --git a/Engine/Source/Programs/Shared/EpicGames.Tracing.Tests/UnrealInsights/EventTest.cs b/Engine/Source/Programs/Shared/EpicGames.Tracing.Tests/UnrealInsights/EventTest.cs
--- a/Engine/Source/Programs/Shared/EpicGames.Tracing.Tests/UnrealInsights/EventTest.cs
+++ b/Engine/Source/Programs/Shared/EpicGames.Tracing.Tests/UnrealInsights/EventTest.cs
@@ -13,17 +13,17 @@ namespace EpicGames.Tracing.Tests.UnrealInsights
 		public void EnterScopeEventTimestampDeserialize()
 		{
 			{
-				using MemoryStream Ms = new MemoryStream(new byte[] {0x10, 0x05, 0x00, 0x00, 0x00, 0x00, 0x00, 0x00});
-				using BinaryReader Reader = new BinaryReader(Ms);
-				EnterScopeEventTimestamp Event = EnterScopeEventTimestamp.Deserialize(Reader);
-				Assert.AreEqual((ulong)0x05, Event.Timestamp);
+				using MemoryStream ms = new MemoryStream(new byte[] { 0x10, 0x05, 0x00, 0x00, 0x00, 0x00, 0x00, 0x00 });
+				using BinaryReader reader = new BinaryReader(ms);
+				EnterScopeEventTimestamp @event = EnterScopeEventTimestamp.Deserialize(reader);
+				Assert.AreEqual((ulong)0x05, @event.Timestamp);
 			}
 
 			{
-				using MemoryStream Ms = new MemoryStream(new byte[] {0x10, 0x00, 0x00, 0x00, 0x00, 0x00, 0x00, 0xFF});
-				using BinaryReader Reader = new BinaryReader(Ms);
-				EnterScopeEventTimestamp Event = EnterScopeEventTimestamp.Deserialize(Reader);
-				Assert.AreEqual((ulong)0x00_FF_00_00_00_00_00_00, Event.Timestamp);
+				using MemoryStream ms = new MemoryStream(new byte[] { 0x10, 0x00, 0x00, 0x00, 0x00, 0x00, 0x00, 0xFF });
+				using BinaryReader reader = new BinaryReader(ms);
+				EnterScopeEventTimestamp @event = EnterScopeEventTimestamp.Deserialize(reader);
+				Assert.AreEqual((ulong)0x00_FF_00_00_00_00_00_00, @event.Timestamp);
 			}
 		}
 	}
diff --git a/Engine/Source/Programs/Shared/EpicGames.Tracing.Tests/UnrealInsights/EventTypeTest.cs b/Engine/Source/Programs/Shared/EpicGames.Tracing.Tests/UnrealInsights/EventTypeTest.cs
--- a/Engine/Source/Programs/Shared/EpicGames.Tracing.Tests/UnrealInsights/EventTypeTest.cs
+++ b/Engine/Source/Programs/Shared/EpicGames.Tracing.Tests/UnrealInsights/EventTypeTest.cs
@@ -11,100 +11,100 @@ namespace EpicGames.Tracing.Tests.UnrealInsights
 	[System.Diagnostics.CodeAnalysis.SuppressMessage("Naming", "CA1707:Identifiers should not contain underscores")]
 	public class EventTypeTest
 	{
-		private readonly string TraceThreadInfoHex = "00 00 4A 00 24 00 04 07 06 0A 00 00 04 00 02 08 04 00 04 00 02 08 08 00 04 00 02 08 0C 00 00 00 88 04 24 54 72 61 63 65 54 68 72 65 61 64 49 6E 66 6F 54 68 72 65 61 64 49 64 53 79 73 74 65 6D 49 64 53 6F 72 74 48 69 6E 74 4E 61 6D 65";
-		private readonly string TraceNewTraceHex = "00 00 55 00 10 00 04 05 06 08 00 00 08 00 03 0A 08 00 08 00 03 0E 10 00 02 00 01 06 12 00 01 00 00 0B 24 54 72 61 63 65 4E 65 77 54 72 61 63 65 53 74 61 72 74 43 79 63 6C 65 43 79 63 6C 65 46 72 65 71 75 65 6E 63 79 45 6E 64 69 61 6E 50 6F 69 6E 74 65 72 53 69 7A 65";
-		private readonly string DiagnosticsSession2Hex = "00 00 9A 00 22 00 08 07 0B 08 00 00 00 00 88 08 00 00 00 00 88 07 00 00 00 00 89 0B 00 00 00 00 89 06 00 00 00 00 89 0C 00 00 04 00 02 0A 04 00 01 00 00 11 05 00 01 00 00 0A 44 69 61 67 6E 6F 73 74 69 63 73 53 65 73 73 69 6F 6E 32 50 6C 61 74 66 6F 72 6D 41 70 70 4E 61 6D 65 43 6F 6D 6D 61 6E 64 4C 69 6E 65 42 72 61 6E 63 68 42 75 69 6C 64 56 65 72 73 69 6F 6E 43 68 61 6E 67 65 6C 69 73 74 43 6F 6E 66 69 67 75 72 61 74 69 6F 6E 54 79 70 65 54 61 72 67 65 74 54 79 70 65";
-		private readonly string LoggingLogMessageHex = "00 00 40 00 1A 00 03 06 07 0A 00 00 08 00 03 08 08 00 08 00 03 05 10 00 00 00 80 0A 4C 6F 67 67 69 6E 67 4C 6F 67 4D 65 73 73 61 67 65 4C 6F 67 50 6F 69 6E 74 43 79 63 6C 65 46 6F 72 6D 61 74 41 72 67 73";
+		private readonly string _traceThreadInfoHex = "00 00 4A 00 24 00 04 07 06 0A 00 00 04 00 02 08 04 00 04 00 02 08 08 00 04 00 02 08 0C 00 00 00 88 04 24 54 72 61 63 65 54 68 72 65 61 64 49 6E 66 6F 54 68 72 65 61 64 49 64 53 79 73 74 65 6D 49 64 53 6F 72 74 48 69 6E 74 4E 61 6D 65";
+		private readonly string _traceNewTraceHex = "00 00 55 00 10 00 04 05 06 08 00 00 08 00 03 0A 08 00 08 00 03 0E 10 00 02 00 01 06 12 00 01 00 00 0B 24 54 72 61 63 65 4E 65 77 54 72 61 63 65 53 74 61 72 74 43 79 63 6C 65 43 79 63 6C 65 46 72 65 71 75 65 6E 63 79 45 6E 64 69 61 6E 50 6F 69 6E 74 65 72 53 69 7A 65";
+		private readonly string _diagnosticsSession2Hex = "00 00 9A 00 22 00 08 07 0B 08 00 00 00 00 88 08 00 00 00 00 88 07 00 00 00 00 89 0B 00 00 00 00 89 06 00 00 00 00 89 0C 00 00 04 00 02 0A 04 00 01 00 00 11 05 00 01 00 00 0A 44 69 61 67 6E 6F 73 74 69 63 73 53 65 73 73 69 6F 6E 32 50 6C 61 74 66 6F 72 6D 41 70 70 4E 61 6D 65 43 6F 6D 6D 61 6E 64 4C 69 6E 65 42 72 61 6E 63 68 42 75 69 6C 64 56 65 72 73 69 6F 6E 43 68 61 6E 67 65 6C 69 73 74 43 6F 6E 66 69 67 75 72 61 74 69 6F 6E 54 79 70 65 54 61 72 67 65 74 54 79 70 65";
+		private readonly string _loggingLogMessageHex = "00 00 40 00 1A 00 03 06 07 0A 00 00 08 00 03 08 08 00 08 00 03 05 10 00 00 00 80 0A 4C 6F 67 67 69 6E 67 4C 6F 67 4D 65 73 73 61 67 65 4C 6F 67 50 6F 69 6E 74 43 79 63 6C 65 46 6F 72 6D 61 74 41 72 67 73";
 		
 		[TestMethod]
 		public void Deserialize_Trace_NewTrace()
 		{
-			DeserializeAndAssertEventType(TraceNewTraceHex, TraceNewTraceEvent.EventType, 16);
+			DeserializeAndAssertEventType(_traceNewTraceHex, TraceNewTraceEvent.EventType, 16);
 		}
 		
 		[TestMethod]
 		public void Serialize_Trace_NewTrace()
 		{
-			SerializeAndAssertEventType(TraceNewTraceHex, TraceNewTraceEvent.EventType, 16);
+			SerializeAndAssertEventType(_traceNewTraceHex, TraceNewTraceEvent.EventType, 16);
 		}
 		
 		[TestMethod]
 		public void Deserialize_Trace_ThreadInfo()
 		{
-			DeserializeAndAssertEventType(TraceThreadInfoHex, TraceThreadInfoEvent.EventType, 36);
+			DeserializeAndAssertEventType(_traceThreadInfoHex, TraceThreadInfoEvent.EventType, 36);
 		}
 		
 		[TestMethod]
 		public void Serialize_Trace_ThreadInfo()
 		{
-			SerializeAndAssertEventType(TraceThreadInfoHex, TraceThreadInfoEvent.EventType, 36);
+			SerializeAndAssertEventType(_traceThreadInfoHex, TraceThreadInfoEvent.EventType, 36);
 		}
 		
 		[TestMethod]
 		public void Deserialize_Diagnostics_Session2()
 		{
-			DeserializeAndAssertEventType(DiagnosticsSession2Hex, DiagnosticsSession2Event.EventType, 34);
+			DeserializeAndAssertEventType(_diagnosticsSession2Hex, DiagnosticsSession2Event.EventType, 34);
 		}
 		
 		[TestMethod]
 		public void Serialize_Diagnostics_Session2()
 		{
-			SerializeAndAssertEventType(DiagnosticsSession2Hex, DiagnosticsSession2Event.EventType, 34);
+			SerializeAndAssertEventType(_diagnosticsSession2Hex, DiagnosticsSession2Event.EventType, 34);
 		}
 		
 		[TestMethod]
 		public void Deserialize_Logging_LogMessage()
 		{
-			DeserializeAndAssertEventType(LoggingLogMessageHex, LoggingLogMessage.EventType, 26);
+			DeserializeAndAssertEventType(_loggingLogMessageHex, LoggingLogMessage.EventType, 26);
 		}
 		
 		[TestMethod]
 		public void Serialize_Logging_LogMessage()
 		{
-			SerializeAndAssertEventType(LoggingLogMessageHex, LoggingLogMessage.EventType, 26);
+			SerializeAndAssertEventType(_loggingLogMessageHex, LoggingLogMessage.EventType, 26);
 		}
 
-		private static EventType DeserializeAndAssertEventType(string HexData, EventType ExpectedType, ushort ExpectedUid)
+		private static EventType DeserializeAndAssertEventType(string hexData, EventType expectedType, ushort expectedUid)
 		{
-			byte[] Block = GenericEventTest.HexStringToBytes(HexData);
-			using MemoryStream Ms = new MemoryStream(Block);
-			using BinaryReader Reader = new BinaryReader(Ms);
+			byte[] block = GenericEventTest.HexStringToBytes(hexData);
+			using MemoryStream ms = new MemoryStream(block);
+			using BinaryReader reader = new BinaryReader(ms);
 
-			TraceImportantEventHeader EventHeader = TraceImportantEventHeader.Deserialize(Reader);
-			Assert.AreEqual(PredefinedEventUid.NewEvent, EventHeader.Uid);
-			
-			(ushort NewEventUid, EventType ActualType)  = EventType.Deserialize(Reader);
-			Reader.EnsureEntireStreamIsConsumed();
+			TraceImportantEventHeader eventHeader = TraceImportantEventHeader.Deserialize(reader);
+			Assert.AreEqual(PredefinedEventUid.NewEvent, eventHeader.Uid);
+
+			(ushort newEventUid, EventType actualType) = EventType.Deserialize(reader);
+			reader.EnsureEntireStreamIsConsumed();
 
-			Assert.AreEqual(ExpectedUid, ActualType.NewEventUid);
-			Assert.AreEqual(ExpectedUid, NewEventUid);
-			Assert.AreEqual(ExpectedType.Name, ActualType.Name);
-			Assert.AreEqual(ExpectedType.Size, ActualType.Size);
-			Assert.AreEqual(ExpectedType.Flags, ActualType.Flags);
-			Assert.AreEqual(ExpectedType.Fields.Count, ActualType.Fields.Count);
-			for (int i = 0; i < ExpectedType.Fields.Count; i++)
+			Assert.AreEqual(expectedUid, actualType.NewEventUid);
+			Assert.AreEqual(expectedUid, newEventUid);
+			Assert.AreEqual(expectedType.Name, actualType.Name);
+			Assert.AreEqual(expectedType.Size, actualType.Size);
+			Assert.AreEqual(expectedType.Flags, actualType.Flags);
+			Assert.AreEqual(expectedType.Fields.Count, actualType.Fields.Count);
+			for (int i = 0; i < expectedType.Fields.Count; i++)
 			{
 				
-				Assert.AreEqual(ExpectedType.Fields[i].Offset, ActualType.Fields[i].Offset);
-				Assert.AreEqual(ExpectedType.Fields[i].Size, ActualType.Fields[i].Size);
-				Assert.AreEqual(ExpectedType.Fields[i].TypeInfo, ActualType.Fields[i].TypeInfo);
-				Assert.AreEqual(ExpectedType.Fields[i].Name, ActualType.Fields[i].Name);
-				Assert.AreEqual(ExpectedType.Fields[i].NameSize, ActualType.Fields[i].NameSize);
+				Assert.AreEqual(expectedType.Fields[i].Offset, actualType.Fields[i].Offset);
+				Assert.AreEqual(expectedType.Fields[i].Size, actualType.Fields[i].Size);
+				Assert.AreEqual(expectedType.Fields[i].TypeInfo, actualType.Fields[i].TypeInfo);
+				Assert.AreEqual(expectedType.Fields[i].Name, actualType.Fields[i].Name);
+				Assert.AreEqual(expectedType.Fields[i].NameSize, actualType.Fields[i].NameSize);
 			}
 
-			return ActualType;
+			return actualType;
 		}
 
-		private static void SerializeAndAssertEventType(string ExpectedHexData, EventType EventType, ushort Uid)
+		private static void SerializeAndAssertEventType(string expectedHexData, EventType eventType, ushort uid)
 		{
-			using MemoryStream Ms = new MemoryStream();
-			using BinaryWriter Writer = new BinaryWriter(Ms);
-			byte[] ExpectedData = GenericEventTest.HexStringToBytes(ExpectedHexData);
-			int ExpectedEventSize = ExpectedData.Length;
-			Assert.AreEqual(ExpectedEventSize, EventType.Size);
-			EventType.Serialize(Uid, Writer);
+			using MemoryStream ms = new MemoryStream();
+			using BinaryWriter writer = new BinaryWriter(ms);
+			byte[] expectedData = GenericEventTest.HexStringToBytes(expectedHexData);
+			int expectedEventSize = expectedData.Length;
+			Assert.AreEqual(expectedEventSize, eventType.Size);
+			eventType.Serialize(uid, writer);
 			
-			GenericEventTest.AssertHexString(ExpectedHexData, Ms.ToArray());
+			GenericEventTest.AssertHexString(expectedHexData, ms.ToArray());
 		}
 	}
 }
\ No newline at end of file
diff --git a/Engine/Source/Programs/Shared/EpicGames.Tracing.Tests/UnrealInsights/GenericEventTest.cs b/Engine/Source/Programs/Shared/EpicGames.Tracing.Tests/UnrealInsights/GenericEventTest.cs
--- a/Engine/Source/Programs/Shared/EpicGames.Tracing.Tests/UnrealInsights/GenericEventTest.cs
+++ b/Engine/Source/Programs/Shared/EpicGames.Tracing.Tests/UnrealInsights/GenericEventTest.cs
@@ -12,46 +12,46 @@ namespace EpicGames.Tracing.Tests.UnrealInsights
 	[System.Diagnostics.CodeAnalysis.SuppressMessage("Naming", "CA1707:Identifiers should not contain underscores")]
 	public class GenericEventTest
 	{
-		private readonly string CpuProfilerEventSpecHex = "21 00 77 00 34 25 00 00 01 C1 0D 00 4B 32 4E 6F 64 65 5F 56 61 72 69 61 62 6C 65 47 65 74 20 2F 45 6E 67 69 6E 65 2F 45 6E 67 69 6E 65 53 6B 79 2F 42 50 5F 53 6B 79 5F 53 70 68 65 72 65 2E 42 50 5F 53 6B 79 5F 53 70 68 65 72 65 3A 55 73 65 72 43 6F 6E 73 74 72 75 63 74 69 6F 6E 53 63 72 69 70 74 2E 4B 32 4E 6F 64 65 5F 56 61 72 69 61 62 6C 65 47 65 74 5F 39 38 33 03";
-		private readonly string CpuProfilerEventBatchHex1 = "4F 00 02 20 1E 00 87 A7 EB AE 8A D1 0B 01 A7 02 02 04 FB 23 03 A4 0A 61 04 0B 05 F2 80 42 AA 1C EB 5C 06 C0 BC 03 17 07 02 47 08 C2 02 AB 15 09 AA 03 07 0A AE 01 05 0B D8 0D 05 0C B7 01 0D D8 03 00 03 0E A4 03 6F 0F 89 FF 03 10 F3 C4 04 11 60 06 A7 02 10 AB 19 12 B3 04 10 9F 11 13 0A 04 02 04 ED 03 10 91 11 14 80 01 02 97 02 10 A9 D6 0B 15 10 06 DD 03 10 CB 85 04 16 12 06 D3 01 10 89 16 17 0C 02 D6 02 07 18 DA 84 28 99 01 19 18 05 1A 0A A5 10 1B CD 01 1C 06 88 6A C5 49 1D B8 EB 24 A9 17 20 AF 01 10 CD F9 01 21 71 22 F2 AE 12 85 04 23 D1 0D 24 85 74 25 EA 58 91 0D 25 80 E6 02 A9 1C 25 88 CD 06 A7 27 25 D0 DA 09 8B 20 25 DE AD 0B AF 33 25 C0 A9 09 ED 23 25 C4 81 0A F5 2D 25 9A EC 0A F3 35 25 D4 CF 03 FF 19 25 BE 92 01 B1 0D 25 98 22 06";
+		private readonly string _cpuProfilerEventSpecHex = "21 00 77 00 34 25 00 00 01 C1 0D 00 4B 32 4E 6F 64 65 5F 56 61 72 69 61 62 6C 65 47 65 74 20 2F 45 6E 67 69 6E 65 2F 45 6E 67 69 6E 65 53 6B 79 2F 42 50 5F 53 6B 79 5F 53 70 68 65 72 65 2E 42 50 5F 53 6B 79 5F 53 70 68 65 72 65 3A 55 73 65 72 43 6F 6E 73 74 72 75 63 74 69 6F 6E 53 63 72 69 70 74 2E 4B 32 4E 6F 64 65 5F 56 61 72 69 61 62 6C 65 47 65 74 5F 39 38 33 03";
+		private readonly string _cpuProfilerEventBatchHex1 = "4F 00 02 20 1E 00 87 A7 EB AE 8A D1 0B 01 A7 02 02 04 FB 23 03 A4 0A 61 04 0B 05 F2 80 42 AA 1C EB 5C 06 C0 BC 03 17 07 02 47 08 C2 02 AB 15 09 AA 03 07 0A AE 01 05 0B D8 0D 05 0C B7 01 0D D8 03 00 03 0E A4 03 6F 0F 89 FF 03 10 F3 C4 04 11 60 06 A7 02 10 AB 19 12 B3 04 10 9F 11 13 0A 04 02 04 ED 03 10 91 11 14 80 01 02 97 02 10 A9 D6 0B 15 10 06 DD 03 10 CB 85 04 16 12 06 D3 01 10 89 16 17 0C 02 D6 02 07 18 DA 84 28 99 01 19 18 05 1A 0A A5 10 1B CD 01 1C 06 88 6A C5 49 1D B8 EB 24 A9 17 20 AF 01 10 CD F9 01 21 71 22 F2 AE 12 85 04 23 D1 0D 24 85 74 25 EA 58 91 0D 25 80 E6 02 A9 1C 25 88 CD 06 A7 27 25 D0 DA 09 8B 20 25 DE AD 0B AF 33 25 C0 A9 09 ED 23 25 C4 81 0A F5 2D 25 9A EC 0A F3 35 25 D4 CF 03 FF 19 25 BE 92 01 B1 0D 25 98 22 06";
 		// Unused at the moment -- private readonly string CpuProfilerEventBatchHex2 = "4F 00 02 40 1E 00 E4 AD A2 B2 8A D1 0B 17 41 12 0B 41 12 0B 41 04 11 41 2C 0F 41 0E 17 41 14 13 41 10 11 41 26 0F 41 12 B1 01 41 28 0D 41 04 0F 41 18 09 41 04 0F 41 16 0B 41 16 19 41 24 0D 41 14 17 41 0E 0B 41 2A 0B 41 48 31 41 06 15 41 06 5B 41 14 0B 41 0C 25 41 14 25 41 14 25 41 10 15 41 12 2D 41 0E 15 41 04 25 41 10 19 41 0C 17 41 16 17 41 1E 17 41 14 0B 41 0E 2D 41 0E 45 41 04 15 41 14 17 41 14 13 41 18 19 41 06 0D 41 10 11 41 10 0B 41 0E 0B 41 0E 0D 41 10 0B 41 10 0F 41 18 0F 41 12 0B 41 18 19 41 0E 1F 41 16 0D 41 14 1F 41 16 33 41 10 0D 41 24 25 41 18 15 41 04 23 41 06 25 41 04 0B 41 14 0F 41 0E 15 41 12 77 41 12 0D 41 14 23 41 18 0F 41 16 3B 41 18 19 41 20 81 01 41 18 4D 41 16 2B 41 16 1F 41 1C 11 41 10 11 41 10 0B 41 0E 15 41 06";
-		private readonly string LoggingLogMessageHex = "35 00 B4 FD CA E1 FB 7F 00 00 A4 AB DB 52 44 17 00 00 02 A2 03 00 02 C2 44 61 00 71 00 50 00 72 00 6F 00 66 00 2E 00 64 00 6C 00 6C 00 00 00 7E 00 00 00 06";
-		private readonly string MemoryMemoryScopeHex1 = "25 00 2D 00 00 27 00 00 00";
-		private readonly string MemoryMemoryScopeHex2 = "25 00 2E 00 00 27 00 00 00";
-		private readonly string MemoryMemoryScopeHex3 = "25 00 3D 59 00 33 00 00 00";
-		private readonly string TraceNewTraceHex = "10 00 13 00 D9 6E DA 52 44 17 00 00 80 96 98 00 00 00 00 00 4D 52 08";
-		private readonly string DiagnosticsSession2Hex = "22 00 EA 00 DC 26 0A 01 03 04 01 A0 00 00 57 69 6E 36 34 01 81 01 00 55 6E 72 65 61 6C 45 64 69 74 6F 72 01 C2 0F 00 20 00 44 00 3A 00 5C 00 64 00 65 00 70 00 6F 00 74 00 5C 00 73 00 74 00 61 00 72 00 73 00 68 00 69 00 70 00 2D 00 6D 00 61 00 69 00 6E 00 5C 00 51 00 41 00 47 00 61 00 6D 00 65 00 5C 00 51 00 41 00 47 00 61 00 6D 00 65 00 2E 00 75 00 70 00 72 00 6F 00 6A 00 65 00 63 00 74 00 20 00 2D 00 74 00 72 00 61 00 63 00 65 00 3D 00 63 00 70 00 75 00 20 00 2D 00 67 00 61 00 6D 00 65 00 01 83 02 00 2B 00 2B 00 55 00 45 00 35 00 2B 00 4D 00 61 00 69 00 6E 00 01 84 05 00 2B 00 2B 00 55 00 45 00 35 00 2B 00 4D 00 61 00 69 00 6E 00 2D 00 43 00 4C 00 2D 00 31 00 37 00 34 00 34 00 32 00 35 00 32 00 34 00 03";
-		private readonly string TraceThreadInfoHex1 = "24 00 1B 00 02 00 00 00 74 9D 02 00 FF FF FF FF 01 43 01 00 47 61 6D 65 54 68 72 65 61 64 03";
-		private readonly string TraceThreadInfoHex2 = "24 00 16 00 03 00 00 00 00 00 00 00 FF FF FF 7F 01 A3 00 00 54 72 61 63 65 03";
-		private readonly string TraceThreadInfoHex3 = "24 00 25 00 04 00 00 00 D8 BB 06 00 50 00 00 00 01 83 02 00 46 6F 72 65 67 72 6F 75 6E 64 20 57 6F 72 6B 65 72 20 23 30 03";
+		private readonly string _loggingLogMessageHex = "35 00 B4 FD CA E1 FB 7F 00 00 A4 AB DB 52 44 17 00 00 02 A2 03 00 02 C2 44 61 00 71 00 50 00 72 00 6F 00 66 00 2E 00 64 00 6C 00 6C 00 00 00 7E 00 00 00 06";
+		private readonly string _memoryMemoryScopeHex1 = "25 00 2D 00 00 27 00 00 00";
+		private readonly string _memoryMemoryScopeHex2 = "25 00 2E 00 00 27 00 00 00";
+		private readonly string _memoryMemoryScopeHex3 = "25 00 3D 59 00 33 00 00 00";
+		private readonly string _traceNewTraceHex = "10 00 13 00 D9 6E DA 52 44 17 00 00 80 96 98 00 00 00 00 00 4D 52 08";
+		private readonly string _diagnosticsSession2Hex = "22 00 EA 00 DC 26 0A 01 03 04 01 A0 00 00 57 69 6E 36 34 01 81 01 00 55 6E 72 65 61 6C 45 64 69 74 6F 72 01 C2 0F 00 20 00 44 00 3A 00 5C 00 64 00 65 00 70 00 6F 00 74 00 5C 00 73 00 74 00 61 00 72 00 73 00 68 00 69 00 70 00 2D 00 6D 00 61 00 69 00 6E 00 5C 00 51 00 41 00 47 00 61 00 6D 00 65 00 5C 00 51 00 41 00 47 00 61 00 6D 00 65 00 2E 00 75 00 70 00 72 00 6F 00 6A 00 65 00 63 00 74 00 20 00 2D 00 74 00 72 00 61 00 63 00 65 00 3D 00 63 00 70 00 75 00 20 00 2D 00 67 00 61 00 6D 00 65 00 01 83 02 00 2B 00 2B 00 55 00 45 00 35 00 2B 00 4D 00 61 00 69 00 6E 00 01 84 05 00 2B 00 2B 00 55 00 45 00 35 00 2B 00 4D 00 61 00 69 00 6E 00 2D 00 43 00 4C 00 2D 00 31 00 37 00 34 00 34 00 32 00 35 00 32 00 34 00 03";
+		private readonly string _traceThreadInfoHex1 = "24 00 1B 00 02 00 00 00 74 9D 02 00 FF FF FF FF 01 43 01 00 47 61 6D 65 54 68 72 65 61 64 03";
+		private readonly string _traceThreadInfoHex2 = "24 00 16 00 03 00 00 00 00 00 00 00 FF FF FF 7F 01 A3 00 00 54 72 61 63 65 03";
+		private readonly string _traceThreadInfoHex3 = "24 00 25 00 04 00 00 00 D8 BB 06 00 50 00 00 00 01 83 02 00 46 6F 72 65 67 72 6F 75 6E 64 20 57 6F 72 6B 65 72 20 23 30 03";
 		
 		[TestMethod]
 		public void Deserialize_Memory_MemoryScope()
 		{
-			GenericEvent Event1 = DeserializeAndAssertGenericEvent(MemoryMemoryScopeHex1, MemoryMemoryScopeEvent.EventType, false, 18, 45);
-			Assert.AreEqual(0x27, Event1.GetFields()[0].Int!.Value);
+			GenericEvent event1 = DeserializeAndAssertGenericEvent(_memoryMemoryScopeHex1, MemoryMemoryScopeEvent.EventType, false, 18, 45);
+			Assert.AreEqual(0x27, event1.GetFields()[0].Int!.Value);
 			
-			GenericEvent Event2 = DeserializeAndAssertGenericEvent(MemoryMemoryScopeHex2, MemoryMemoryScopeEvent.EventType, false, 18, 46);
-			Assert.AreEqual(0x27, Event2.GetFields()[0].Int!.Value);
+			GenericEvent event2 = DeserializeAndAssertGenericEvent(_memoryMemoryScopeHex2, MemoryMemoryScopeEvent.EventType, false, 18, 46);
+			Assert.AreEqual(0x27, event2.GetFields()[0].Int!.Value);
 			
-			GenericEvent Event3 = DeserializeAndAssertGenericEvent(MemoryMemoryScopeHex3, MemoryMemoryScopeEvent.EventType, false, 18, 5832765);
-			Assert.AreEqual(0x33, Event3.GetFields()[0].Int!.Value);
+			GenericEvent event3 = DeserializeAndAssertGenericEvent(_memoryMemoryScopeHex3, MemoryMemoryScopeEvent.EventType, false, 18, 5832765);
+			Assert.AreEqual(0x33, event3.GetFields()[0].Int!.Value);
 		}
 		
 		[TestMethod]
 		public void Serialize_Memory_MemoryScope()
 		{
-			SerializeAndAssertGenericEvent(MemoryMemoryScopeHex1, MemoryMemoryScopeEvent.EventType, false, 18, new []
+			SerializeAndAssertGenericEvent(_memoryMemoryScopeHex1, MemoryMemoryScopeEvent.EventType, false, 18, new []
 			{
 				Field.FromInt(0x27),
 			}, 45);
 			
-			SerializeAndAssertGenericEvent(MemoryMemoryScopeHex2, MemoryMemoryScopeEvent.EventType, false, 18, new []
+			SerializeAndAssertGenericEvent(_memoryMemoryScopeHex2, MemoryMemoryScopeEvent.EventType, false, 18, new []
 			{
 				Field.FromInt(0x27),
 			}, 46);
 			
-			SerializeAndAssertGenericEvent(MemoryMemoryScopeHex3, MemoryMemoryScopeEvent.EventType, false, 18, new []
+			SerializeAndAssertGenericEvent(_memoryMemoryScopeHex3, MemoryMemoryScopeEvent.EventType, false, 18, new []
 			{
 				Field.FromInt(0x33),
 			}, 5832765);
@@ -60,18 +60,18 @@ namespace EpicGames.Tracing.Tests.UnrealInsights
 		[TestMethod]
 		public void Deserialize_Trace_NewTrace()
 		{
-			GenericEvent Event = DeserializeAndAssertGenericEvent(TraceNewTraceHex, TraceNewTraceEvent.EventType, true, 16);
-			Field[] Fields = Event.GetFields();
-			Assert.AreEqual(25582215261913, Fields[0].Long!.Value);
-			Assert.AreEqual(10000000, Fields[1].Long!.Value);
-			Assert.AreEqual(21069, Fields[2].Int!.Value);
-			Assert.AreEqual(8, Fields[3].Int!.Value);
+			GenericEvent @event = DeserializeAndAssertGenericEvent(_traceNewTraceHex, TraceNewTraceEvent.EventType, true, 16);
+			Field[] fields = @event.GetFields();
+			Assert.AreEqual(25582215261913, fields[0].Long!.Value);
+			Assert.AreEqual(10000000, fields[1].Long!.Value);
+			Assert.AreEqual(21069, fields[2].Int!.Value);
+			Assert.AreEqual(8, fields[3].Int!.Value);
 		}
 		
 		[TestMethod]
 		public void Serialize_Trace_NewTrace_Generic()
 		{
-			SerializeAndAssertGenericEvent(TraceNewTraceHex, TraceNewTraceEvent.EventType, true, 16, new []
+			SerializeAndAssertGenericEvent(_traceNewTraceHex, TraceNewTraceEvent.EventType, true, 16, new []
 			{
 				Field.FromLong(25582215261913),
 				Field.FromLong(10000000),
@@ -83,38 +83,38 @@ namespace EpicGames.Tracing.Tests.UnrealInsights
 		[TestMethod]
 		public void Serialize_Trace_NewTrace()
 		{
-			TraceNewTraceEvent Event = new TraceNewTraceEvent(25582215261913, 10000000, 21069, 8);
-			SerializeAndAssertEvent(TraceNewTraceHex, Event, true, 16);
+			TraceNewTraceEvent @event = new TraceNewTraceEvent(25582215261913, 10000000, 21069, 8);
+			SerializeAndAssertEvent(_traceNewTraceHex, @event, true, 16);
 		}
 		
 		[TestMethod]
 		public void Deserialize_Trace_ThreadInfo()
 		{
 			{
-				GenericEvent Event = DeserializeAndAssertGenericEvent(TraceThreadInfoHex1, TraceThreadInfoEvent.EventType, true, 36);
-				Field[] Fields = Event.GetFields();
-				Assert.AreEqual(2, Fields[0].Int!.Value);
-				Assert.AreEqual(171380, Fields[1].Int!.Value);
-				Assert.AreEqual(-1, Fields[2].Int!.Value);
-				Assert.AreEqual("GameThread", Fields[3].String!);
+				GenericEvent @event = DeserializeAndAssertGenericEvent(_traceThreadInfoHex1, TraceThreadInfoEvent.EventType, true, 36);
+				Field[] fields = @event.GetFields();
+				Assert.AreEqual(2, fields[0].Int!.Value);
+				Assert.AreEqual(171380, fields[1].Int!.Value);
+				Assert.AreEqual(-1, fields[2].Int!.Value);
+				Assert.AreEqual("GameThread", fields[3].String!);
 			}
 			
 			{
-				GenericEvent Event = DeserializeAndAssertGenericEvent(TraceThreadInfoHex2, TraceThreadInfoEvent.EventType, true, 36);
-				Field[] Fields = Event.GetFields();
-				Assert.AreEqual(3, Fields[0].Int!.Value);
-				Assert.AreEqual(0, Fields[1].Int!.Value);
-				Assert.AreEqual(2147483647, Fields[2].Int!.Value);
-				Assert.AreEqual("Trace", Fields[3].String!);
+				GenericEvent @event = DeserializeAndAssertGenericEvent(_traceThreadInfoHex2, TraceThreadInfoEvent.EventType, true, 36);
+				Field[] fields = @event.GetFields();
+				Assert.AreEqual(3, fields[0].Int!.Value);
+				Assert.AreEqual(0, fields[1].Int!.Value);
+				Assert.AreEqual(2147483647, fields[2].Int!.Value);
+				Assert.AreEqual("Trace", fields[3].String!);
 			}
 			
 			{
-				GenericEvent Event = DeserializeAndAssertGenericEvent(TraceThreadInfoHex3, TraceThreadInfoEvent.EventType, true, 36);
-				Field[] Fields = Event.GetFields();
-				Assert.AreEqual(4, Fields[0].Int!.Value);
-				Assert.AreEqual(441304, Fields[1].Int!.Value);
-				Assert.AreEqual(80, Fields[2].Int!.Value);
-				Assert.AreEqual("Foreground Worker #0", Fields[3].String!);
+				GenericEvent @event = DeserializeAndAssertGenericEvent(_traceThreadInfoHex3, TraceThreadInfoEvent.EventType, true, 36);
+				Field[] fields = @event.GetFields();
+				Assert.AreEqual(4, fields[0].Int!.Value);
+				Assert.AreEqual(441304, fields[1].Int!.Value);
+				Assert.AreEqual(80, fields[2].Int!.Value);
+				Assert.AreEqual("Foreground Worker #0", fields[3].String!);
 			}
 		}
 		
@@ -122,40 +122,40 @@ namespace EpicGames.Tracing.Tests.UnrealInsights
 		public void Serialize_Trace_ThreadInfo()
 		{
 			{
-				TraceThreadInfoEvent Event = new TraceThreadInfoEvent(2, 171380, -1, "GameThread");
-				SerializeAndAssertEvent(TraceThreadInfoHex1, Event, true, 36);
+				TraceThreadInfoEvent @event = new TraceThreadInfoEvent(2, 171380, -1, "GameThread");
+				SerializeAndAssertEvent(_traceThreadInfoHex1, @event, true, 36);
 			}
 			
 			{
-				TraceThreadInfoEvent Event = new TraceThreadInfoEvent(3, 0, 2147483647, "Trace");
-				SerializeAndAssertEvent(TraceThreadInfoHex2, Event, true, 36);
+				TraceThreadInfoEvent @event = new TraceThreadInfoEvent(3, 0, 2147483647, "Trace");
+				SerializeAndAssertEvent(_traceThreadInfoHex2, @event, true, 36);
 			}
 			
 			{
-				TraceThreadInfoEvent Event = new TraceThreadInfoEvent(4, 441304, 80, "Foreground Worker #0");
-				SerializeAndAssertEvent(TraceThreadInfoHex3, Event, true, 36);
+				TraceThreadInfoEvent @event = new TraceThreadInfoEvent(4, 441304, 80, "Foreground Worker #0");
+				SerializeAndAssertEvent(_traceThreadInfoHex3, @event, true, 36);
 			}
 		}
 		
 		[TestMethod]
 		public void Deserialize_Diagnostics_Session2()
 		{
-			GenericEvent Event = DeserializeAndAssertGenericEvent(DiagnosticsSession2Hex, DiagnosticsSession2Event.EventType, true, 34);
-			Field[] Fields = Event.GetFields();
-			Assert.AreEqual("Win64", Fields[0].String!);
-			Assert.AreEqual("UnrealEditor", Fields[1].String!);
-			Assert.AreEqual(@" D:\depot\starship-main\QAGame\QAGame.uproject -trace=cpu -game", Fields[2].String!);
-			Assert.AreEqual("++UE5+Main", Fields[3].String!);
-			Assert.AreEqual("++UE5+Main-CL-17442524", Fields[4].String!);
-			Assert.AreEqual(17442524, Fields[5].Int!.Value);
-			Assert.AreEqual(3, Fields[6].Int!.Value);
-			Assert.AreEqual(4, Fields[7].Int!.Value);
+			GenericEvent @event = DeserializeAndAssertGenericEvent(_diagnosticsSession2Hex, DiagnosticsSession2Event.EventType, true, 34);
+			Field[] fields = @event.GetFields();
+			Assert.AreEqual("Win64", fields[0].String!);
+			Assert.AreEqual("UnrealEditor", fields[1].String!);
+			Assert.AreEqual(@" D:\depot\starship-main\QAGame\QAGame.uproject -trace=cpu -game", fields[2].String!);
+			Assert.AreEqual("++UE5+Main", fields[3].String!);
+			Assert.AreEqual("++UE5+Main-CL-17442524", fields[4].String!);
+			Assert.AreEqual(17442524, fields[5].Int!.Value);
+			Assert.AreEqual(3, fields[6].Int!.Value);
+			Assert.AreEqual(4, fields[7].Int!.Value);
 		}
 		
 		[TestMethod]
 		public void Serialize_Diagnostics_Session2_Generic()
 		{
-			SerializeAndAssertGenericEvent(DiagnosticsSession2Hex, DiagnosticsSession2Event.EventType, true, 34, new []
+			SerializeAndAssertGenericEvent(_diagnosticsSession2Hex, DiagnosticsSession2Event.EventType, true, 34, new []
 			{
 				Field.FromString("Win64"),
 				Field.FromString("UnrealEditor"),
@@ -171,22 +171,22 @@ namespace EpicGames.Tracing.Tests.UnrealInsights
 		[TestMethod]
 		public void Serialize_Diagnostics_Session2()
 		{
-			DiagnosticsSession2Event Event = new DiagnosticsSession2Event("Win64", "UnrealEditor", @" D:\depot\starship-main\QAGame\QAGame.uproject -trace=cpu -game", "++UE5+Main", "++UE5+Main-CL-17442524", 17442524, 3, 4);
-			SerializeAndAssertEvent(DiagnosticsSession2Hex, Event, true, 34);
+			DiagnosticsSession2Event @event = new DiagnosticsSession2Event("Win64", "UnrealEditor", @" D:\depot\starship-main\QAGame\QAGame.uproject -trace=cpu -game", "++UE5+Main", "++UE5+Main-CL-17442524", 17442524, 3, 4);
+			SerializeAndAssertEvent(_diagnosticsSession2Hex, @event, true, 34);
 		}
 		
 		[TestMethod]
 		public void CpuProfiler_EventSpec_Deserialize()
 		{
-			GenericEvent Event = DeserializeAndAssertGenericEvent(CpuProfilerEventSpecHex, CpuProfilerEventSpecEvent.EventType, true, 33);
-			Assert.AreEqual(9524, Event.GetFields()[0].Int!.Value);
-			Assert.AreEqual("K2Node_VariableGet /Engine/EngineSky/BP_Sky_Sphere.BP_Sky_Sphere:UserConstructionScript.K2Node_VariableGet_983", Event.GetFields()[1].String!);
+			GenericEvent @event = DeserializeAndAssertGenericEvent(_cpuProfilerEventSpecHex, CpuProfilerEventSpecEvent.EventType, true, 33);
+			Assert.AreEqual(9524, @event.GetFields()[0].Int!.Value);
+			Assert.AreEqual("K2Node_VariableGet /Engine/EngineSky/BP_Sky_Sphere.BP_Sky_Sphere:UserConstructionScript.K2Node_VariableGet_983", @event.GetFields()[1].String!);
 		}
 		
 		[TestMethod]
 		public void CpuProfiler_EventSpec_Serialize_Generic()
 		{
-			SerializeAndAssertGenericEvent(CpuProfilerEventSpecHex, CpuProfilerEventSpecEvent.EventType, true, 33, new Field[]
+			SerializeAndAssertGenericEvent(_cpuProfilerEventSpecHex, CpuProfilerEventSpecEvent.EventType, true, 33, new Field[]
 			{
 				Field.FromInt(9524),
 				Field.FromString("K2Node_VariableGet /Engine/EngineSky/BP_Sky_Sphere.BP_Sky_Sphere:UserConstructionScript.K2Node_VariableGet_983")
@@ -196,42 +196,42 @@ namespace EpicGames.Tracing.Tests.UnrealInsights
 		[TestMethod]
 		public void CpuProfiler_EventSpec_Serialize()
 		{
-			CpuProfilerEventSpecEvent Event = new CpuProfilerEventSpecEvent(9524, "K2Node_VariableGet /Engine/EngineSky/BP_Sky_Sphere.BP_Sky_Sphere:UserConstructionScript.K2Node_VariableGet_983");
-			SerializeAndAssertEvent(CpuProfilerEventSpecHex, Event, true, 33);
+			CpuProfilerEventSpecEvent @event = new CpuProfilerEventSpecEvent(9524, "K2Node_VariableGet /Engine/EngineSky/BP_Sky_Sphere.BP_Sky_Sphere:UserConstructionScript.K2Node_VariableGet_983");
+			SerializeAndAssertEvent(_cpuProfilerEventSpecHex, @event, true, 33);
 		}
 		
 		[TestMethod]
 		public void CpuProfiler_EventBatch_Deserialize()
 		{
-			GenericEvent Event = DeserializeAndAssertGenericEvent(CpuProfilerEventBatchHex1, CpuProfilerEventBatchEvent.EventType, false, 39);
-			byte[] EventData = Event.GetFields()[0].GetArray()!;
-			Assert.AreEqual(241, EventData.Length);
+			GenericEvent @event = DeserializeAndAssertGenericEvent(_cpuProfilerEventBatchHex1, CpuProfilerEventBatchEvent.EventType, false, 39);
+			byte[] eventData = @event.GetFields()[0].GetArray()!;
+			Assert.AreEqual(241, eventData.Length);
 			
-			CpuProfilerSerializer Serializer = new CpuProfilerSerializer(10000000);
-			Serializer.Read(EventData);
+			CpuProfilerSerializer serializer = new CpuProfilerSerializer(10000000);
+			serializer.Read(eventData);
 
-			byte[] EventDataSerialized = Serializer.Write()[0];
+			byte[] eventDataSerialized = serializer.Write()[0];
 			
-			AssertHexString(EventData, EventDataSerialized);
+			AssertHexString(eventData, eventDataSerialized);
 		}
 		
 		[TestMethod]
 		public void CpuProfiler_EventBatch_Serialize()
 		{
-			byte[] Data = HexStringToBytes("87 A7 EB AE 8A D1 0B 01 A7 02 02 04 FB 23 03 A4 0A 61 04 0B 05 F2 80 42 AA 1C EB 5C 06 C0 BC 03 17 07 02 47 08 C2 02 AB 15 09 AA 03 07 0A AE 01 05 0B D8 0D 05 0C B7 01 0D D8 03 00 03 0E A4 03 6F 0F 89 FF 03 10 F3 C4 04 11 60 06 A7 02 10 AB 19 12 B3 04 10 9F 11 13 0A 04 02 04 ED 03 10 91 11 14 80 01 02 97 02 10 A9 D6 0B 15 10 06 DD 03 10 CB 85 04 16 12 06 D3 01 10 89 16 17 0C 02 D6 02 07 18 DA 84 28 99 01 19 18 05 1A 0A A5 10 1B CD 01 1C 06 88 6A C5 49 1D B8 EB 24 A9 17 20 AF 01 10 CD F9 01 21 71 22 F2 AE 12 85 04 23 D1 0D 24 85 74 25 EA 58 91 0D 25 80 E6 02 A9 1C 25 88 CD 06 A7 27 25 D0 DA 09 8B 20 25 DE AD 0B AF 33 25 C0 A9 09 ED 23 25 C4 81 0A F5 2D 25 9A EC 0A F3 35 25 D4 CF 03 FF 19 25 BE 92 01 B1 0D 25 98 22");
-			CpuProfilerEventBatchEvent Event = new CpuProfilerEventBatchEvent(Data);
-			SerializeAndAssertEvent(CpuProfilerEventBatchHex1, Event, false, 39);
+			byte[] data = HexStringToBytes("87 A7 EB AE 8A D1 0B 01 A7 02 02 04 FB 23 03 A4 0A 61 04 0B 05 F2 80 42 AA 1C EB 5C 06 C0 BC 03 17 07 02 47 08 C2 02 AB 15 09 AA 03 07 0A AE 01 05 0B D8 0D 05 0C B7 01 0D D8 03 00 03 0E A4 03 6F 0F 89 FF 03 10 F3 C4 04 11 60 06 A7 02 10 AB 19 12 B3 04 10 9F 11 13 0A 04 02 04 ED 03 10 91 11 14 80 01 02 97 02 10 A9 D6 0B 15 10 06 DD 03 10 CB 85 04 16 12 06 D3 01 10 89 16 17 0C 02 D6 02 07 18 DA 84 28 99 01 19 18 05 1A 0A A5 10 1B CD 01 1C 06 88 6A C5 49 1D B8 EB 24 A9 17 20 AF 01 10 CD F9 01 21 71 22 F2 AE 12 85 04 23 D1 0D 24 85 74 25 EA 58 91 0D 25 80 E6 02 A9 1C 25 88 CD 06 A7 27 25 D0 DA 09 8B 20 25 DE AD 0B AF 33 25 C0 A9 09 ED 23 25 C4 81 0A F5 2D 25 9A EC 0A F3 35 25 D4 CF 03 FF 19 25 BE 92 01 B1 0D 25 98 22");
+			CpuProfilerEventBatchEvent @event = new CpuProfilerEventBatchEvent(data);
+			SerializeAndAssertEvent(_cpuProfilerEventBatchHex1, @event, false, 39);
 		}
 		
 		[TestMethod]
 		public void CpuProfiler_EventBatch_EncodeDecodeTimestamp()
 		{
-			static void AssertTimestamp(ulong ExpectedTimestamp, bool ExpectedIsScopeEnter)
+			static void AssertTimestamp(ulong expectedTimestamp, bool expectedIsScopeEnter)
 			{
-				ulong Temp = CpuProfilerSerializer.EncodeTimestamp(ExpectedTimestamp, ExpectedIsScopeEnter);
-				(ulong ActualTimestamp, bool ActualIsScopeEnter) = CpuProfilerSerializer.DecodeTimestamp(Temp);
-				Assert.AreEqual(ExpectedTimestamp, ActualTimestamp);
-				Assert.AreEqual(ExpectedIsScopeEnter, ActualIsScopeEnter);
+				ulong temp = CpuProfilerSerializer.EncodeTimestamp(expectedTimestamp, expectedIsScopeEnter);
+				(ulong actualTimestamp, bool actualIsScopeEnter) = CpuProfilerSerializer.DecodeTimestamp(temp);
+				Assert.AreEqual(expectedTimestamp, actualTimestamp);
+				Assert.AreEqual(expectedIsScopeEnter, actualIsScopeEnter);
 			}
 
 			AssertTimestamp(0, true);
@@ -251,84 +251,87 @@ namespace EpicGames.Tracing.Tests.UnrealInsights
 		[TestMethod]
 		public void Deserialize_Logging_LogMessage()
 		{
-			GenericEvent Event = DeserializeAndAssertGenericEvent(LoggingLogMessageHex, LoggingLogMessage.EventType, false, 26);
-			Field[] Fields = Event.GetFields();
-			Assert.AreEqual(140719801695668, Fields[0].Long!.Value);
-			Assert.AreEqual(25582215343012, Fields[1].Long!.Value);
-			byte[] Bytes = Fields[2].GetArray()!;
-			Assert.AreEqual(29, Bytes.Length);
-			Assert.AreEqual(0x02, Bytes[0]);
-			Assert.AreEqual(0xC2, Bytes[1]);
-			Assert.AreEqual(0x44, Bytes[2]);
-			Assert.AreEqual(0x7E, Bytes[25]);
+			GenericEvent @event = DeserializeAndAssertGenericEvent(_loggingLogMessageHex, LoggingLogMessage.EventType, false, 26);
+			Field[] fields = @event.GetFields();
+			Assert.AreEqual(140719801695668, fields[0].Long!.Value);
+			Assert.AreEqual(25582215343012, fields[1].Long!.Value);
+			byte[] bytes = fields[2].GetArray()!;
+			Assert.AreEqual(29, bytes.Length);
+			Assert.AreEqual(0x02, bytes[0]);
+			Assert.AreEqual(0xC2, bytes[1]);
+			Assert.AreEqual(0x44, bytes[2]);
+			Assert.AreEqual(0x7E, bytes[25]);
 		}
 		
 		[TestMethod]
 		public void SerializeEventWithAuxAndImportant()
 		{
-			ushort Uid = 33;
-			using MemoryStream Ms = new MemoryStream();
-			using BinaryWriter Writer = new BinaryWriter(Ms);
+			ushort uid = 33;
+			using MemoryStream ms = new MemoryStream();
+			using BinaryWriter writer = new BinaryWriter(ms);
 
-			Field[] Fields =
+			Field[] fields =
 			{
 				Field.FromInt(9524),
 				Field.FromString("K2Node_VariableGet /Engine/EngineSky/BP_Sky_Sphere.BP_Sky_Sphere:UserConstructionScript.K2Node_VariableGet_983"),
 			};
-			GenericEvent Event = new GenericEvent(98989, Fields, CpuProfilerEventSpecEvent.EventType);
-			Assert.AreEqual(0x77, Event.Size);
+			GenericEvent @event = new GenericEvent(98989, fields, CpuProfilerEventSpecEvent.EventType);
+			Assert.AreEqual(0x77, @event.Size);
 			
-			TraceImportantEventHeader EventHeader = new TraceImportantEventHeader(Uid, Event.Size);
-			EventHeader.Serialize(Writer);
-			Event.Serialize(20, Writer);
+			TraceImportantEventHeader eventHeader = new TraceImportantEventHeader(uid, @event.Size);
+			eventHeader.Serialize(writer);
+			@event.Serialize(20, writer);
 
-			AssertHexString(CpuProfilerEventSpecHex, Ms.ToArray());
+			AssertHexString(_cpuProfilerEventSpecHex, ms.ToArray());
 		}
 
 		[TestMethod]
 		public void AuxHeaderDeserialize()
 		{
 			// Aux header from field "Name" in CpuProfiler.EventSpec
-			using MemoryStream Ms = new MemoryStream(new byte[] {0x01, 0xC1, 0x0D, 0x00});
-			using BinaryReader Reader = new BinaryReader(Ms);
+			using MemoryStream ms = new MemoryStream(new byte[] { 0x01, 0xC1, 0x0D, 0x00 });
+			using BinaryReader reader = new BinaryReader(ms);
 
-			(ushort Uid, int FieldIndex, int Size) = GenericEvent.DeserializeAuxHeader(Reader.ReadUInt32());
-			Assert.AreEqual(PredefinedEventUid.AuxData, Uid);
-			Assert.AreEqual(1, FieldIndex);
-			Assert.AreEqual(110, Size);
+			(ushort uid, int fieldIndex, int size) = GenericEvent.DeserializeAuxHeader(reader.ReadUInt32());
+			Assert.AreEqual(PredefinedEventUid.AuxData, uid);
+			Assert.AreEqual(1, fieldIndex);
+			Assert.AreEqual(110, size);
 		}
 
 		[TestMethod]
 		public void AuxHeaderSerialize()
 		{
-			uint Header = GenericEvent.SerializeAuxHeader(1, 1, 110);
-			using MemoryStream Ms = new MemoryStream();
-			using BinaryWriter Writer = new BinaryWriter(Ms);
-			Writer.Write(Header);
-			byte[] Data = Ms.ToArray();
-			Assert.AreEqual(0x01, Data[0]);
-			Assert.AreEqual(0xC1, Data[1]);
-			Assert.AreEqual(0x0D, Data[2]);
-			Assert.AreEqual(0x00, Data[3]);
+			uint header = GenericEvent.SerializeAuxHeader(1, 1, 110);
+			using MemoryStream ms = new MemoryStream();
+			using BinaryWriter writer = new BinaryWriter(ms);
+			writer.Write(header);
+			byte[] data = ms.ToArray();
+			Assert.AreEqual(0x01, data[0]);
+			Assert.AreEqual(0xC1, data[1]);
+			Assert.AreEqual(0x0D, data[2]);
+			Assert.AreEqual(0x00, data[3]);
 			
-			(ushort Uid, int FieldIndex, int Size) = GenericEvent.DeserializeAuxHeader(Header);
-			Assert.AreEqual(PredefinedEventUid.AuxData, Uid);
-			Assert.AreEqual(1, FieldIndex);
-			Assert.AreEqual(110, Size);
+			(ushort uid, int fieldIndex, int size) = GenericEvent.DeserializeAuxHeader(header);
+			Assert.AreEqual(PredefinedEventUid.AuxData, uid);
+			Assert.AreEqual(1, fieldIndex);
+			Assert.AreEqual(110, size);
 		}
 
-		internal static byte[] StringToByteArray(string Hex)
+		internal static byte[] StringToByteArray(string hex)
 		{
-			int NumChars = Hex.Length;
-			byte[] Bytes = new byte[NumChars / 2];
-			for (int i = 0; i < NumChars; i += 2)
-				Bytes[i / 2] = Convert.ToByte(Hex.Substring(i, 2), 16);
-			return Bytes;
+			int numChars = hex.Length;
+			byte[] bytes = new byte[numChars / 2];
+			for (int i = 0; i < numChars; i += 2)
+			{
+				bytes[i / 2] = Convert.ToByte(hex.Substring(i, 2), 16);
+			}
+
+			return bytes;
 		}
 		
-		internal static byte[] HexStringToBytes(string HexData)
+		internal static byte[] HexStringToBytes(string hexData)
 		{
-			return StringToByteArray(HexData.Replace(" ", "", StringComparison.Ordinal));
+			return StringToByteArray(hexData.Replace(" ", "", StringComparison.Ordinal));
 		}
 		
 		public static string ToHexString(byte[] ba)
@@ -336,87 +339,87 @@ namespace EpicGames.Tracing.Tests.UnrealInsights
 			return BitConverter.ToString(ba).Replace("-", " ", StringComparison.Ordinal);
 		}
 
-		internal static void AssertHexString(string Expected, string Actual)
+		internal static void AssertHexString(string expected, string actual)
 		{
-			if (Expected != Actual)
+			if (expected != actual)
 			{
-				Console.WriteLine("Expected:\n" + Expected);
-				Console.WriteLine("Actual:\n" + Actual);
+				Console.WriteLine("Expected:\n" + expected);
+				Console.WriteLine("Actual:\n" + actual);
 			}
 
-			Assert.AreEqual(Expected, Actual);
+			Assert.AreEqual(expected, actual);
 		}
 		
-		internal static void AssertHexString(string Expected, byte[] Actual)
+		internal static void AssertHexString(string expected, byte[] actual)
 		{
-			AssertHexString(Expected, ToHexString(Actual));
+			AssertHexString(expected, ToHexString(actual));
 		}
 		
-		internal static void AssertHexString(byte[] Expected, byte[] Actual)
+		internal static void AssertHexString(byte[] expected, byte[] actual)
 		{
-			AssertHexString(ToHexString(Expected), ToHexString(Actual));
+			AssertHexString(ToHexString(expected), ToHexString(actual));
 		}
 		
-		private static GenericEvent DeserializeAndAssertGenericEvent(string HexData, EventType EventType, bool IsImportant, ushort ExpectedUid, uint? ExpectedSerial = 0)
+		private static GenericEvent DeserializeAndAssertGenericEvent(string hexData, EventType eventType, bool isImportant, ushort expectedUid, uint? expectedSerial = 0)
 		{
-			byte[] Block = HexStringToBytes(HexData);
-			using MemoryStream Ms = new MemoryStream(Block);
-			using BinaryReader Reader = new BinaryReader(Ms);
+			byte[] block = HexStringToBytes(hexData);
+			using MemoryStream ms = new MemoryStream(block);
+			using BinaryReader reader = new BinaryReader(ms);
 
-			ushort Uid;
-			if (IsImportant)
+			ushort uid;
+			if (isImportant)
 			{
-				TraceImportantEventHeader EventHeader = TraceImportantEventHeader.Deserialize(Reader);
-				Assert.AreEqual(ExpectedUid, EventHeader.Uid);
-				Assert.AreEqual(Block.Length - TraceImportantEventHeader.HeaderSize, EventHeader.EventSize);
-				Uid = EventHeader.Uid;
+				TraceImportantEventHeader eventHeader = TraceImportantEventHeader.Deserialize(reader);
+				Assert.AreEqual(expectedUid, eventHeader.Uid);
+				Assert.AreEqual(block.Length - TraceImportantEventHeader.HeaderSize, eventHeader.EventSize);
+				uid = eventHeader.Uid;
 			}
 			else
 			{
-				Uid = Reader.ReadPackedUid(out _);
-				Assert.AreEqual(ExpectedUid, Uid);	
+				uid = reader.ReadPackedUid(out _);
+				Assert.AreEqual(expectedUid, uid);	
 			}
 			
-			GenericEvent Event = GenericEvent.Deserialize(Uid, Reader, EventType);
-			Reader.EnsureEntireStreamIsConsumed();
+			GenericEvent @event = GenericEvent.Deserialize(uid, reader, eventType);
+			reader.EnsureEntireStreamIsConsumed();
 			
-			Assert.AreEqual(ExpectedSerial, Event.Serial);
-			Assert.AreEqual(EventType.Fields.Count, Event.GetFields().Length);
-			return Event;
+			Assert.AreEqual(expectedSerial, @event.Serial);
+			Assert.AreEqual(eventType.Fields.Count, @event.GetFields().Length);
+			return @event;
 		}
 
-		private static void SerializeAndAssertGenericEvent(string ExpectedHexData, EventType EventType, bool IsImportant, ushort Uid, Field[] Fields, uint Serial = 0)
+		private static void SerializeAndAssertGenericEvent(string expectedHexData, EventType eventType, bool isImportant, ushort uid, Field[] fields, uint serial = 0)
 		{
-			using MemoryStream Ms = new MemoryStream();
-			using BinaryWriter Writer = new BinaryWriter(Ms);
+			using MemoryStream ms = new MemoryStream();
+			using BinaryWriter writer = new BinaryWriter(ms);
 
-			byte[] ExpectedData = HexStringToBytes(ExpectedHexData);
-			GenericEvent Event = new GenericEvent(Serial, Fields, EventType);
+			byte[] expectedData = HexStringToBytes(expectedHexData);
+			GenericEvent @event = new GenericEvent(serial, fields, eventType);
 
-			if (IsImportant)
+			if (isImportant)
 			{
-				int ExpectedEventSize = ExpectedData.Length - TraceImportantEventHeader.HeaderSize;
-				Assert.AreEqual(ExpectedEventSize, Event.Size);
-				TraceImportantEventHeader EventHeader = new TraceImportantEventHeader(Uid, Event.Size);
-				EventHeader.Serialize(Writer);
+				int expectedEventSize = expectedData.Length - TraceImportantEventHeader.HeaderSize;
+				Assert.AreEqual(expectedEventSize, @event.Size);
+				TraceImportantEventHeader eventHeader = new TraceImportantEventHeader(uid, @event.Size);
+				eventHeader.Serialize(writer);
 			}
 			else
 			{
-				Writer.WritePackedUid(Uid);
+				writer.WritePackedUid(uid);
 			}
 			
-			Event.Serialize(20, Writer);
+			@event.Serialize(20, writer);
 			
-			AssertHexString(ExpectedHexData, Ms.ToArray());
+			AssertHexString(expectedHexData, ms.ToArray());
 		}
 		
-		private static void SerializeAndAssertEvent(string ExpectedHexData, ITraceEvent Event, bool IsImportant, ushort Uid)
+		private static void SerializeAndAssertEvent(string expectedHexData, ITraceEvent @event, bool isImportant, ushort uid)
 		{
-			using MemoryStream Ms = new MemoryStream();
-			using BinaryWriter Writer = new BinaryWriter(Ms);
-			Event.Serialize(Uid, Writer);
+			using MemoryStream ms = new MemoryStream();
+			using BinaryWriter writer = new BinaryWriter(ms);
+			@event.Serialize(uid, writer);
 			
-			AssertHexString(ExpectedHexData, Ms.ToArray());
+			AssertHexString(expectedHexData, ms.ToArray());
 		}
 	}
 }
\ No newline at end of file
diff --git a/Engine/Source/Programs/Shared/EpicGames.Tracing.Tests/UnrealInsights/TransportPacketTest.cs b/Engine/Source/Programs/Shared/EpicGames.Tracing.Tests/UnrealInsights/TransportPacketTest.cs
--- a/Engine/Source/Programs/Shared/EpicGames.Tracing.Tests/UnrealInsights/TransportPacketTest.cs
+++ b/Engine/Source/Programs/Shared/EpicGames.Tracing.Tests/UnrealInsights/TransportPacketTest.cs
@@ -14,25 +14,25 @@ namespace EpicGames.Tracing.Tests.UnrealInsights
 		[TestMethod]
 		public void ShortByteHandling()
 		{
-			ushort Val1 = 0x1234;
-			Assert.AreEqual(0x34, BinaryReaderExtensions.GetLowByte(Val1));
-			Assert.AreEqual(0x12, BinaryReaderExtensions.GetHighByte(Val1));
+			ushort val1 = 0x1234;
+			Assert.AreEqual(0x34, BinaryReaderExtensions.GetLowByte(val1));
+			Assert.AreEqual(0x12, BinaryReaderExtensions.GetHighByte(val1));
 			
-			ushort Val2 = BinaryReaderExtensions.MakeShort(0x56, 0x78);
-			Assert.AreEqual(0x56, BinaryReaderExtensions.GetLowByte(Val2));
-			Assert.AreEqual(0x78, BinaryReaderExtensions.GetHighByte(Val2));
+			ushort val2 = BinaryReaderExtensions.MakeShort(0x56, 0x78);
+			Assert.AreEqual(0x56, BinaryReaderExtensions.GetLowByte(val2));
+			Assert.AreEqual(0x78, BinaryReaderExtensions.GetHighByte(val2));
 			
-			ushort Val3 = BinaryReaderExtensions.MakeShort(0x25, 0x00);
-			Assert.AreEqual(0x25, Val3);
+			ushort val3 = BinaryReaderExtensions.MakeShort(0x25, 0x00);
+			Assert.AreEqual(0x25, val3);
 		}
 
 		[TestMethod]
 		public void GetPackedUid()
 		{
-			static void AssertGetPackedUid(ushort ExpectedUid, bool ExpectedIsTwoByteUid, byte UidLow, byte UidHigh)
+			static void AssertGetPackedUid(ushort expectedUid, bool expectedIsTwoByteUid, byte uidLow, byte uidHigh)
 			{
-				Assert.AreEqual(ExpectedUid, BinaryReaderExtensions.GetPackedUid(UidLow, UidHigh, out bool IsTwoByteUid));
-				Assert.AreEqual(ExpectedIsTwoByteUid, IsTwoByteUid);
+				Assert.AreEqual(expectedUid, BinaryReaderExtensions.GetPackedUid(uidLow, uidHigh, out bool isTwoByteUid));
+				Assert.AreEqual(expectedIsTwoByteUid, isTwoByteUid);
 			}
 
 			AssertGetPackedUid(0x05, false, 0x0A, 0x08);
@@ -57,17 +57,17 @@ namespace EpicGames.Tracing.Tests.UnrealInsights
 		[TestMethod]
 		public void WritePackedUid()
 		{
-			static void AssertWriteUid(ushort Uid, bool ExpectTwoByteUid)
+			static void AssertWriteUid(ushort uid, bool expectTwoByteUid)
 			{
-				using MemoryStream Ms = new MemoryStream();
-				using BinaryWriter Writer = new BinaryWriter(Ms);
-				Writer.WritePackedUid(Uid);
-				byte[] Data = Ms.ToArray();
-				Assert.AreEqual(ExpectTwoByteUid ? 2 : 1, Data.Length);
-				byte UidHigh = Data.Length == 2 ? Data[1] : (byte)0xFF;
-				ushort DeserializedUid = BinaryReaderExtensions.GetPackedUid(Data[0], UidHigh, out bool IsTwoByteUid);
-				Assert.AreEqual(ExpectTwoByteUid, IsTwoByteUid);
-				Assert.AreEqual(Uid, DeserializedUid);
+				using MemoryStream ms = new MemoryStream();
+				using BinaryWriter writer = new BinaryWriter(ms);
+				writer.WritePackedUid(uid);
+				byte[] data = ms.ToArray();
+				Assert.AreEqual(expectTwoByteUid ? 2 : 1, data.Length);
+				byte uidHigh = data.Length == 2 ? data[1] : (byte)0xFF;
+				ushort deserializedUid = BinaryReaderExtensions.GetPackedUid(data[0], uidHigh, out bool isTwoByteUid);
+				Assert.AreEqual(expectTwoByteUid, isTwoByteUid);
+				Assert.AreEqual(uid, deserializedUid);
 			}
 			
 			AssertWriteUid(0x00, false);
@@ -86,12 +86,12 @@ namespace EpicGames.Tracing.Tests.UnrealInsights
 		[TestMethod]
 		public void Deserialize7Bit()
 		{
-			static void Assert7Bit(ulong ExpectedValue, string ActualBytesHex)
+			static void Assert7Bit(ulong expectedValue, string actualBytesHex)
 			{
-				byte[] ActualBytes = GenericEventTest.StringToByteArray(ActualBytesHex.Replace(" ", "", StringComparison.Ordinal));
-				using MemoryStream Ms = new MemoryStream(ActualBytes);
-				using BinaryReader Reader = new BinaryReader(Ms);
-				Assert.AreEqual(ExpectedValue, TraceUtils.Read7BitUint(Reader));
+				byte[] actualBytes = GenericEventTest.StringToByteArray(actualBytesHex.Replace(" ", "", StringComparison.Ordinal));
+				using MemoryStream ms = new MemoryStream(actualBytes);
+				using BinaryReader reader = new BinaryReader(ms);
+				Assert.AreEqual(expectedValue, TraceUtils.Read7BitUint(reader));
 			}
 
 			Assert7Bit(0, "00");
@@ -124,18 +124,18 @@ namespace EpicGames.Tracing.Tests.UnrealInsights
 		[TestMethod]
 		public void Serialize7Bit()
 		{
-			static void Assert7Bit(ulong Value, string ExpectedBytesHex)
+			static void Assert7Bit(ulong value, string expectedBytesHex)
 			{
-				byte[] ExpectedBytes = GenericEventTest.StringToByteArray(ExpectedBytesHex.Replace(" ", "", StringComparison.Ordinal));
-				using MemoryStream Ms = new MemoryStream();
-				using BinaryWriter Writer = new BinaryWriter(Ms);
+				byte[] expectedBytes = GenericEventTest.StringToByteArray(expectedBytesHex.Replace(" ", "", StringComparison.Ordinal));
+				using MemoryStream ms = new MemoryStream();
+				using BinaryWriter writer = new BinaryWriter(ms);
 
-				int ActualNumBytesWritten = TraceUtils.Write7BitUint(Writer, Value);
-				byte[] ActualBytes = Ms.ToArray();
+				int actualNumBytesWritten = TraceUtils.Write7BitUint(writer, value);
+				byte[] actualBytes = ms.ToArray();
 
-				Assert.AreEqual(ExpectedBytes.Length, ActualNumBytesWritten);
-				Assert.AreEqual(ExpectedBytes.Length, ActualBytes.Length);
-				GenericEventTest.AssertHexString(ExpectedBytesHex, ActualBytes);
+				Assert.AreEqual(expectedBytes.Length, actualNumBytesWritten);
+				Assert.AreEqual(expectedBytes.Length, actualBytes.Length);
+				GenericEventTest.AssertHexString(expectedBytesHex, actualBytes);
 			}
 
 			Assert7Bit(0, "00");
@@ -168,98 +168,98 @@ namespace EpicGames.Tracing.Tests.UnrealInsights
 		[TestMethod]
 		public void SerializeDeserializeWithStubs()
 		{
-			StubTraceEvent Event1 = new StubTraceEvent(1000, 20);
-			StubTraceEvent Event2 = new StubTraceEvent(2000, 30);
-			ushort Uid1 = 5555;
-			ushort Uid2 = 6666;
+			StubTraceEvent event1 = new StubTraceEvent(1000, 20);
+			StubTraceEvent event2 = new StubTraceEvent(2000, 30);
+			ushort uid1 = 5555;
+			ushort uid2 = 6666;
 
 			{
-				using MemoryStream Ms = new MemoryStream();
-				using BinaryWriter Writer = new BinaryWriter(Ms);
-				TransportPacket Packet = TransportPacket.Create(0, 0);
-				Packet.Serialize(Writer, new [] { (Uid1, Event1 as ITraceEvent), (Uid2, Event2 as ITraceEvent) });
+				using MemoryStream ms = new MemoryStream();
+				using BinaryWriter writer = new BinaryWriter(ms);
+				TransportPacket packet = TransportPacket.Create(0, 0);
+				packet.Serialize(writer, new [] { (uid1, event1 as ITraceEvent), (uid2, event2 as ITraceEvent) });
 				
-				Ms.Position = 0;
-				using BinaryReader Reader = new BinaryReader(Ms);
+				ms.Position = 0;
+				using BinaryReader reader = new BinaryReader(ms);
 
-				TransportPacket DeserializedPacket = TransportPacket.Deserialize(Reader);
-				Assert.AreEqual(0, DeserializedPacket.ThreadIdAndMarkers);
-				Assert.AreEqual(Event1.Size + Event2.Size, DeserializedPacket.GetData().Length);
+				TransportPacket deserializedPacket = TransportPacket.Deserialize(reader);
+				Assert.AreEqual(0, deserializedPacket.ThreadIdAndMarkers);
+				Assert.AreEqual(event1.Size + event2.Size, deserializedPacket.GetData().Length);
 			}
 			
 			{
-				using MemoryStream Ms = new MemoryStream();
-				using BinaryWriter Writer = new BinaryWriter(Ms);
-				TransportPacket Packet = TransportPacket.Create(0, 1);
-				Packet.Serialize(Writer, Array.Empty<(ushort, ITraceEvent)>());
+				using MemoryStream ms = new MemoryStream();
+				using BinaryWriter writer = new BinaryWriter(ms);
+				TransportPacket packet = TransportPacket.Create(0, 1);
+				packet.Serialize(writer, Array.Empty<(ushort, ITraceEvent)>());
 
-				Ms.Position = 0;
-				using BinaryReader Reader = new BinaryReader(Ms);
+				ms.Position = 0;
+				using BinaryReader reader = new BinaryReader(ms);
 
-				TransportPacket DeserializedPacket = TransportPacket.Deserialize(Reader);
-				Assert.AreEqual(1, DeserializedPacket.ThreadIdAndMarkers);
-				Assert.AreEqual(0, DeserializedPacket.GetData().Length);
+				TransportPacket deserializedPacket = TransportPacket.Deserialize(reader);
+				Assert.AreEqual(1, deserializedPacket.ThreadIdAndMarkers);
+				Assert.AreEqual(0, deserializedPacket.GetData().Length);
 			}
 			
 			{
-				using MemoryStream Ms = new MemoryStream();
-				using BinaryWriter Writer = new BinaryWriter(Ms);
-				TransportPacket Packet = TransportPacket.Create(0, 12345);
-				Packet.Serialize(Writer, Array.Empty<(ushort, ITraceEvent)>());
+				using MemoryStream ms = new MemoryStream();
+				using BinaryWriter writer = new BinaryWriter(ms);
+				TransportPacket packet = TransportPacket.Create(0, 12345);
+				packet.Serialize(writer, Array.Empty<(ushort, ITraceEvent)>());
 
-				Ms.Position = 0;
-				using BinaryReader Reader = new BinaryReader(Ms);
+				ms.Position = 0;
+				using BinaryReader reader = new BinaryReader(ms);
 
-				TransportPacket DeserializedPacket = TransportPacket.Deserialize(Reader);
-				Assert.AreEqual(12345, DeserializedPacket.ThreadIdAndMarkers);
-				Assert.AreEqual(0, DeserializedPacket.GetData().Length);
+				TransportPacket deserializedPacket = TransportPacket.Deserialize(reader);
+				Assert.AreEqual(12345, deserializedPacket.ThreadIdAndMarkers);
+				Assert.AreEqual(0, deserializedPacket.GetData().Length);
 			}
 		}
 		
 		[TestMethod]
 		public void SerializeDeserialize()
 		{
-			TraceNewTraceEvent Event1 = new TraceNewTraceEvent(25582215261913, 10000000, 21069, 8);
-			DiagnosticsSession2Event Event2 = new DiagnosticsSession2Event("Win64", "UnrealEditor", @" D:\depot\starship-main\QAGame\QAGame.uproject -trace=cpu -game", "++UE5+Main", "++UE5+Main-CL-17442524", 17442524, 3, 4);
-			CpuProfilerEventSpecEvent Event3 = new CpuProfilerEventSpecEvent(9524, "K2Node_VariableGet /Engine/EngineSky/BP_Sky_Sphere.BP_Sky_Sphere:UserConstructionScript.K2Node_VariableGet_983");
+			TraceNewTraceEvent event1 = new TraceNewTraceEvent(25582215261913, 10000000, 21069, 8);
+			DiagnosticsSession2Event event2 = new DiagnosticsSession2Event("Win64", "UnrealEditor", @" D:\depot\starship-main\QAGame\QAGame.uproject -trace=cpu -game", "++UE5+Main", "++UE5+Main-CL-17442524", 17442524, 3, 4);
+			CpuProfilerEventSpecEvent event3 = new CpuProfilerEventSpecEvent(9524, "K2Node_VariableGet /Engine/EngineSky/BP_Sky_Sphere.BP_Sky_Sphere:UserConstructionScript.K2Node_VariableGet_983");
 
-			ushort Uid1 = 10;
-			ushort Uid2 = 20;
-			ushort Uid3 = 30;
+			ushort uid1 = 10;
+			ushort uid2 = 20;
+			ushort uid3 = 30;
 
 			{
-				using MemoryStream Ms = new MemoryStream();
-				using BinaryWriter Writer = new BinaryWriter(Ms);
-				TransportPacket Packet = TransportPacket.Create(0, 0);
-				Packet.Serialize(Writer, new [] { (Uid1, Event1 as ITraceEvent), (Uid2, Event2 as ITraceEvent) });
+				using MemoryStream ms = new MemoryStream();
+				using BinaryWriter writer = new BinaryWriter(ms);
+				TransportPacket packet = TransportPacket.Create(0, 0);
+				packet.Serialize(writer, new [] { (uid1, event1 as ITraceEvent), (uid2, event2 as ITraceEvent) });
 				
-				Ms.Position = 0;
-				using BinaryReader Reader = new BinaryReader(Ms);
+				ms.Position = 0;
+				using BinaryReader reader = new BinaryReader(ms);
 
-				TransportPacket DeserializedPacket = TransportPacket.Deserialize(Reader);
-				Assert.AreEqual(0, DeserializedPacket.ThreadIdAndMarkers);
-				Assert.AreEqual(Event1.Size + Event2.Size, DeserializedPacket.GetData().Length);
-				Reader.EnsureEntireStreamIsConsumed();
+				TransportPacket deserializedPacket = TransportPacket.Deserialize(reader);
+				Assert.AreEqual(0, deserializedPacket.ThreadIdAndMarkers);
+				Assert.AreEqual(event1.Size + event2.Size, deserializedPacket.GetData().Length);
+				reader.EnsureEntireStreamIsConsumed();
 			}
 			
 			{
-				using MemoryStream Ms = new MemoryStream();
-				using BinaryWriter Writer = new BinaryWriter(Ms);
-				TransportPacket Packet = TransportPacket.Create(0, 0);
-				Packet.Serialize(Writer, new []
+				using MemoryStream ms = new MemoryStream();
+				using BinaryWriter writer = new BinaryWriter(ms);
+				TransportPacket packet = TransportPacket.Create(0, 0);
+				packet.Serialize(writer, new []
 				{
-					(Uid1, Event1.Type as ITraceEvent),
-					(Uid2, Event2.Type as ITraceEvent),
-					(Uid3, Event3.Type as ITraceEvent)
+					(uid1, event1.Type as ITraceEvent),
+					(uid2, event2.Type as ITraceEvent),
+					(uid3, event3.Type as ITraceEvent)
 				});
 				
-				Ms.Position = 0;
-				using BinaryReader Reader = new BinaryReader(Ms);
+				ms.Position = 0;
+				using BinaryReader reader = new BinaryReader(ms);
 
-				TransportPacket DeserializedPacket = TransportPacket.Deserialize(Reader);
-				Assert.AreEqual(0, DeserializedPacket.ThreadIdAndMarkers);
-				Assert.AreEqual(Event1.Type.Size + Event2.Type.Size + Event3.Type.Size, DeserializedPacket.GetData().Length);
-				Reader.EnsureEntireStreamIsConsumed();
+				TransportPacket deserializedPacket = TransportPacket.Deserialize(reader);
+				Assert.AreEqual(0, deserializedPacket.ThreadIdAndMarkers);
+				Assert.AreEqual(event1.Type.Size + event2.Type.Size + event3.Type.Size, deserializedPacket.GetData().Length);
+				reader.EnsureEntireStreamIsConsumed();
 			}
 		}
 	}
diff --git a/Engine/Source/Programs/Shared/EpicGames.Tracing/UnrealInsights/Events.cs b/Engine/Source/Programs/Shared/EpicGames.Tracing/UnrealInsights/Events.cs
--- a/Engine/Source/Programs/Shared/EpicGames.Tracing/UnrealInsights/Events.cs
+++ b/Engine/Source/Programs/Shared/EpicGames.Tracing/UnrealInsights/Events.cs
@@ -12,29 +12,31 @@ namespace EpicGames.Tracing.UnrealInsights
 {
 	public static class BinaryReaderExtensions
 	{
-		public static byte[] ReadBytesStrict(this BinaryReader Reader, int Count)
+		public static byte[] ReadBytesStrict(this BinaryReader reader, int count)
 		{
-			byte[] Result = Reader.ReadBytes(Count);
-			if (Result.Length != Count)
+			byte[] result = reader.ReadBytes(count);
+			if (result.Length != count)
+			{
 				throw new EndOfStreamException();
+			}
 
-			return Result;
+			return result;
 		}
 
 		/// <summary>
 		/// Check if bit is set for a given byte
 		/// </summary>
-		/// <param name="Value">Value to check</param>
-		/// <param name="Pos">Pos 0 is least significant bit, pos 7 is most</param>
+		/// <param name="value">Value to check</param>
+		/// <param name="pos">Pos 0 is least significant bit, pos 7 is most</param>
 		/// <returns>True if set</returns>
-		internal static bool IsBitSet(byte Value, int Pos)
+		internal static bool IsBitSet(byte value, int pos)
 		{
-			return (Value & (1 << Pos)) != 0;
+			return (value & (1 << pos)) != 0;
 		}
 
-		internal static bool IsTwoByteUid(byte UidLow)
+		internal static bool IsTwoByteUid(byte uidLow)
 		{
-			return IsBitSet(UidLow, 0);
+			return IsBitSet(uidLow, 0);
 		}
 
 		internal static byte GetHighByte(ushort a)
@@ -57,12 +59,12 @@ namespace EpicGames.Tracing.UnrealInsights
 			return (ushort)(a & 0xffff);
 		}
 
-		internal static ushort MakeShort(byte LowByte, byte HighByte)
+		internal static ushort MakeShort(byte lowByte, byte highByte)
 		{
-			return (ushort)((byte)(LowByte & 0xff) | (ushort)(HighByte & 0xff) << 8);
+			return (ushort)((byte)(lowByte & 0xff) | (ushort)(highByte & 0xff) << 8);
 		}
 
-		internal static ushort GetPackedUid(byte UidLow, byte UidHigh, out bool IsTwoByteUidArg)
+		internal static ushort GetPackedUid(byte uidLow, byte uidHigh, out bool isTwoByteUidArg)
 		{
 			// struct packed_uid
 			// {
@@ -71,61 +73,61 @@ namespace EpicGames.Tracing.UnrealInsights
 			// 	(uint8  uid_high) // if is_two_byte_uid == 1
 			// }
 
-			IsTwoByteUidArg = IsTwoByteUid(UidLow);
-			byte UidLowNoBit = (byte) (UidLow >> 1); // Strip the is_two_byte_uid bit
+			isTwoByteUidArg = IsTwoByteUid(uidLow);
+			byte uidLowNoBit = (byte)(uidLow >> 1); // Strip the is_two_byte_uid bit
 
-			return MakeShort(UidLowNoBit, IsTwoByteUidArg ? UidHigh : (byte) 0x00);
+			return MakeShort(uidLowNoBit, isTwoByteUidArg ? uidHigh : (byte) 0x00);
 		}
 
-		public static ushort ReadPackedUid(this BinaryReader Reader, out bool IsTwoByteUid)
+		public static ushort ReadPackedUid(this BinaryReader reader, out bool isTwoByteUid)
 		{
-			byte UidLow = Reader.ReadByte();
-			byte UidHigh = Reader.ReadByte();
-			Reader.BaseStream.Position -= 1; // UidLow is always consumed, but UidHigh maybe not.
+			byte uidLow = reader.ReadByte();
+			byte uidHigh = reader.ReadByte();
+			reader.BaseStream.Position -= 1; // UidLow is always consumed, but UidHigh maybe not.
 			
-			ushort Uid = GetPackedUid(UidLow, UidHigh, out IsTwoByteUid);
-			if (IsTwoByteUid)
+			ushort uid = GetPackedUid(uidLow, uidHigh, out isTwoByteUid);
+			if (isTwoByteUid)
 			{
-				Reader.ReadByte();  // Consume UidHigh
+				reader.ReadByte();  // Consume UidHigh
 			}
 
-			return Uid;
+			return uid;
 		}
 		
-		public static void EnsureEntireStreamIsConsumed(this BinaryReader Reader)
+		public static void EnsureEntireStreamIsConsumed(this BinaryReader reader)
 		{
-			bool IsEntireStreamConsumed = Reader.BaseStream.Position == Reader.BaseStream.Length;
-			if (!IsEntireStreamConsumed)
+			bool isEntireStreamConsumed = reader.BaseStream.Position == reader.BaseStream.Length;
+			if (!isEntireStreamConsumed)
 			{
-				throw new Exception($"Entire stream/buffer was not consumed. Pos={Reader.BaseStream.Position} Len={Reader.BaseStream.Length}");
+				throw new Exception($"Entire stream/buffer was not consumed. Pos={reader.BaseStream.Position} Len={reader.BaseStream.Length}");
 			}
 		}
 		
-		public static bool IsEntireStreamConsumed(this BinaryReader Reader)
+		public static bool IsEntireStreamConsumed(this BinaryReader reader)
 		{
-			return Reader.BaseStream.Position == Reader.BaseStream.Length;
+			return reader.BaseStream.Position == reader.BaseStream.Length;
 		}
 	}
 
 	public static class BinaryWriterExtensions
 	{
-		public static void WritePackedUid(this BinaryWriter Writer, ushort Uid)
+		public static void WritePackedUid(this BinaryWriter writer, ushort uid)
 		{
-			if (Uid < PredefinedEventUid._WellKnownNum)
+			if (uid < PredefinedEventUid.WellKnownNum)
 			{
-				byte UidLow = (byte) Uid;
-				UidLow = (byte) (UidLow << 1);
+				byte uidLow = (byte)uid;
+				uidLow = (byte) (uidLow << 1);
 				// LSB is 0 after shifting, indicating a one-byte UID
-				Writer.Write(UidLow);
+				writer.Write(uidLow);
 			}
-			else if (Uid < 127)
+			else if (uid < 127)
 			{
-				byte UidLow = BinaryReaderExtensions.GetLowByte(Uid);
-				byte UidHigh = BinaryReaderExtensions.GetHighByte(Uid);
-				UidLow = (byte) (UidLow << 1);
-				UidLow = (byte) (UidLow | 1);
-				Writer.Write(UidLow);
-				Writer.Write(UidHigh);
+				byte uidLow = BinaryReaderExtensions.GetLowByte(uid);
+				byte uidHigh = BinaryReaderExtensions.GetHighByte(uid);
+				uidLow = (byte) (uidLow << 1);
+				uidLow = (byte) (uidLow | 1);
+				writer.Write(uidLow);
+				writer.Write(uidHigh);
 			}
 			else
 			{
@@ -138,21 +140,21 @@ namespace EpicGames.Tracing.UnrealInsights
 	{
 		ushort Size { get; }
 		EventType Type { get; }
-		public void Serialize(ushort Uid, BinaryWriter Writer);
+		public void Serialize(ushort uid, BinaryWriter writer);
 	}
 
 	public class EnterScopeEvent : ITraceEvent
 	{
 		public ushort Size => 0;
 		public EventType Type => EventType.WellKnown(PredefinedEventUid.EnterScope, "EnterScope");
-		public void Serialize(ushort Uid, BinaryWriter Writer) { throw new NotImplementedException(); }
+		public void Serialize(ushort uid, BinaryWriter writer) { throw new NotImplementedException(); }
 	}
 	
 	public class LeaveScopeEvent : ITraceEvent
 	{
 		public ushort Size => 0;
 		public EventType Type => EventType.WellKnown(PredefinedEventUid.LeaveScope, "LeaveScope");
-		public void Serialize(ushort Uid, BinaryWriter Writer) { throw new NotImplementedException(); }
+		public void Serialize(ushort uid, BinaryWriter writer) { throw new NotImplementedException(); }
 	}
 	
 	public class EnterScopeEventTimestamp : ITraceEvent
@@ -161,25 +163,25 @@ namespace EpicGames.Tracing.UnrealInsights
 		public EventType Type => EventType.WellKnown(PredefinedEventUid.EnterScope_T, "EnterScopeTimestamp");
 		public ulong Timestamp { get; }
 
-		public EnterScopeEventTimestamp(ulong Timestamp)
+		public EnterScopeEventTimestamp(ulong timestamp)
 		{
-			this.Timestamp = Timestamp;
+			Timestamp = timestamp;
 		}
 
-		public void Serialize(ushort Uid, BinaryWriter Writer) { throw new NotImplementedException(); }
+		public void Serialize(ushort uid, BinaryWriter writer) { throw new NotImplementedException(); }
 
-		public static EnterScopeEventTimestamp Deserialize(BinaryReader Reader)
+		public static EnterScopeEventTimestamp Deserialize(BinaryReader reader)
 		{
-			ulong Value = Reader.ReadUInt64();
-			ushort UidFound = BinaryReaderExtensions.GetPackedUid((byte) (Value & 0xFF), 0x00, out bool _);
+			ulong value = reader.ReadUInt64();
+			ushort uidFound = BinaryReaderExtensions.GetPackedUid((byte)(value & 0xFF), 0x00, out bool _);
 
-			if (UidFound != PredefinedEventUid.EnterScope_T)
+			if (uidFound != PredefinedEventUid.EnterScope_T)
 			{
-				throw new ArgumentException($"Bad UID found when deserializing 0x{UidFound:X4}/{UidFound}");
+				throw new ArgumentException($"Bad UID found when deserializing 0x{uidFound:X4}/{uidFound}");
 			}
 
-			ulong Timestamp = Value >> 8;
-			return new EnterScopeEventTimestamp(Timestamp);
+			ulong timestamp = value >> 8;
+			return new EnterScopeEventTimestamp(timestamp);
 		}
 	}
 	
@@ -187,27 +189,27 @@ namespace EpicGames.Tracing.UnrealInsights
 	{
 		public ushort Size => 7;
 		public EventType Type => EventType.WellKnown(PredefinedEventUid.EnterScope_T, "LeaveScopeTimestamp");
-		readonly ulong Timestamp;
+		readonly ulong _timestamp;
 
-		public LeaveScopeEventTimestamp(ulong Timestamp)
+		public LeaveScopeEventTimestamp(ulong timestamp)
 		{
-			this.Timestamp = Timestamp;
+			_timestamp = timestamp;
 		}
 
-		public void Serialize(ushort Uid, BinaryWriter Writer) { throw new NotImplementedException(); }
+		public void Serialize(ushort uid, BinaryWriter writer) { throw new NotImplementedException(); }
 		
-		public static LeaveScopeEventTimestamp Deserialize(BinaryReader Reader)
+		public static LeaveScopeEventTimestamp Deserialize(BinaryReader reader)
 		{
-			ulong Value = Reader.ReadUInt64();
-			ushort UidFound = BinaryReaderExtensions.GetPackedUid((byte) (Value & 0xFF), 0x00, out bool _);
+			ulong value = reader.ReadUInt64();
+			ushort uidFound = BinaryReaderExtensions.GetPackedUid((byte)(value & 0xFF), 0x00, out bool _);
 
-			if (UidFound != PredefinedEventUid.LeaveScope_T)
+			if (uidFound != PredefinedEventUid.LeaveScope_T)
 			{
-				throw new ArgumentException($"Bad UID found when deserializing 0x{UidFound:X4}/{UidFound}");
+				throw new ArgumentException($"Bad UID found when deserializing 0x{uidFound:X4}/{uidFound}");
 			}
 
-			ulong Timestamp = Value >> 8;
-			return new LeaveScopeEventTimestamp(Timestamp);
+			ulong timestamp = value >> 8;
+			return new LeaveScopeEventTimestamp(timestamp);
 		}
 	}
 
@@ -218,25 +220,25 @@ namespace EpicGames.Tracing.UnrealInsights
 
 		public const ushort HeaderSize = sizeof(ushort) + sizeof(ushort); 
 
-		public TraceImportantEventHeader(ushort Uid, ushort EventSize)
+		public TraceImportantEventHeader(ushort uid, ushort eventSize)
 		{
 #pragma warning disable CA1508 // Avoid dead conditional code
 			Debug.Assert(HeaderSize == 4);
 #pragma warning restore CA1508 // Avoid dead conditional code
-			this.Uid = Uid;
-			this.EventSize = EventSize;
+			Uid = uid;
+			EventSize = eventSize;
 		}
 
-		public void Serialize(BinaryWriter Writer)
+		public void Serialize(BinaryWriter writer)
 		{
-			Writer.Write(Uid);
-			Writer.Write(EventSize);
+			writer.Write(Uid);
+			writer.Write(EventSize);
 		}
 
-		public static TraceImportantEventHeader Deserialize(BinaryReader Reader)
+		public static TraceImportantEventHeader Deserialize(BinaryReader reader)
 		{
-			ushort Uid = Reader.ReadUInt16();
-			return new TraceImportantEventHeader(Uid, Reader.ReadUInt16());
+			ushort uid = reader.ReadUInt16();
+			return new TraceImportantEventHeader(uid, reader.ReadUInt16());
 		}
 	}
 }
\ No newline at end of file
diff --git a/Engine/Source/Programs/Shared/EpicGames.Tracing/UnrealInsights/Events/CpuProfiler.cs b/Engine/Source/Programs/Shared/EpicGames.Tracing/UnrealInsights/Events/CpuProfiler.cs
--- a/Engine/Source/Programs/Shared/EpicGames.Tracing/UnrealInsights/Events/CpuProfiler.cs
+++ b/Engine/Source/Programs/Shared/EpicGames.Tracing/UnrealInsights/Events/CpuProfiler.cs
@@ -14,32 +14,32 @@ namespace EpicGames.Tracing.UnrealInsights.Events
 				new EventTypeField(4, 0, EventTypeField.TypeAnsiString, "Name")
 			});
 		
-		public ushort Size => (ushort) (GenericEvent.Size + TraceImportantEventHeader.HeaderSize);
+		public ushort Size => (ushort) (_genericEvent.Size + TraceImportantEventHeader.HeaderSize);
 		public EventType Type => EventType;
 
 		public uint Id { get; }
-		readonly string Name;
+		readonly string _name;
 
-		private readonly GenericEvent GenericEvent;
+		private readonly GenericEvent _genericEvent;
 
-		public CpuProfilerEventSpecEvent(uint Id, string Name)
+		public CpuProfilerEventSpecEvent(uint id, string name)
 		{
-			this.Id = Id;
-			this.Name = Name;
+			Id = id;
+			_name = name;
 			
-			Field[] Fields =
+			Field[] fields =
 			{
-				Field.FromInt((int) Id),
-				Field.FromString(Name),
+				Field.FromInt((int) id),
+				Field.FromString(name),
 			};
 			
-			GenericEvent = new GenericEvent(0, Fields, EventType);
+			_genericEvent = new GenericEvent(0, fields, EventType);
 		}
 
-		public void Serialize(ushort Uid, BinaryWriter Writer)
+		public void Serialize(ushort uid, BinaryWriter writer)
 		{
-			new TraceImportantEventHeader(Uid, GenericEvent.Size).Serialize(Writer);
-			GenericEvent.Serialize(Uid, Writer);
+			new TraceImportantEventHeader(uid, _genericEvent.Size).Serialize(writer);
+			_genericEvent.Serialize(uid, writer);
 		}
 	}
 	
@@ -48,21 +48,21 @@ namespace EpicGames.Tracing.UnrealInsights.Events
 		public static readonly EventType EventType = new EventType(0, "CpuProfiler", "EventBatch", EventType.FlagMaybeHasAux | EventType.FlagNoSync,
 			new List<EventTypeField>() { new EventTypeField(0, 0, EventTypeField.TypeArray, "Data") });
 		
-		public ushort Size => (ushort) (GenericEvent.Size + sizeof(ushort));
+		public ushort Size => (ushort) (_genericEvent.Size + sizeof(ushort));
 		public EventType Type => EventType;
 
-		readonly GenericEvent GenericEvent;
+		readonly GenericEvent _genericEvent;
 
-		public CpuProfilerEventBatchEvent(byte[] Data)
+		public CpuProfilerEventBatchEvent(byte[] data)
 		{
-			Field[] Fields = {Field.FromArray(Data)};
-			GenericEvent = new GenericEvent(0, Fields, EventType);
+			Field[] fields = { Field.FromArray(data) };
+			_genericEvent = new GenericEvent(0, fields, EventType);
 		}
 		
-		public void Serialize(ushort Uid, BinaryWriter Writer)
+		public void Serialize(ushort uid, BinaryWriter writer)
 		{
-			Writer.WritePackedUid(Uid);
-			GenericEvent.Serialize(Uid, Writer);
+			writer.WritePackedUid(uid);
+			_genericEvent.Serialize(uid, writer);
 		}
 	}
 	
@@ -72,17 +72,17 @@ namespace EpicGames.Tracing.UnrealInsights.Events
 		public bool IsEnterScope { get; }
 		public uint? SpecId { get; }
 
-		public CpuProfilerScopeEvent(ulong Timestamp, bool IsEnterScope, uint? SpecId)
+		public CpuProfilerScopeEvent(ulong timestamp, bool isEnterScope, uint? specId)
 		{
-			this.Timestamp = Timestamp;
-			this.IsEnterScope = IsEnterScope;
-			this.SpecId = SpecId;
+			Timestamp = timestamp;
+			IsEnterScope = isEnterScope;
+			SpecId = specId;
 		}
 
 		public override string ToString()
 		{
-			string Name = IsEnterScope ? "EnterScope" : "ExitScope";
-			return $"{nameof(CpuProfilerScopeEvent)}({Name} {nameof(Timestamp)}={Timestamp} {nameof(SpecId)}={SpecId})";
+			string name = IsEnterScope ? "EnterScope" : "ExitScope";
+			return $"{nameof(CpuProfilerScopeEvent)}({name} {nameof(Timestamp)}={Timestamp} {nameof(SpecId)}={SpecId})";
 		}
 	}
 	
@@ -91,104 +91,104 @@ namespace EpicGames.Tracing.UnrealInsights.Events
 	/// </summary>
 	public class CpuProfilerSerializer
 	{
-		private readonly ulong CycleFrequency;
+		private readonly ulong _cycleFrequency;
 		public List<CpuProfilerScopeEvent> ScopeEvents { get; } = new List<CpuProfilerScopeEvent>();
 
-		public CpuProfilerSerializer(ulong CycleFrequency)
+		public CpuProfilerSerializer(ulong cycleFrequency)
 		{
-			this.CycleFrequency = CycleFrequency;
+			_cycleFrequency = cycleFrequency;
 		}
 		
-		public void Read(byte[] EventBatchData)
+		public void Read(byte[] eventBatchData)
 		{
-			using MemoryStream Ms = new MemoryStream(EventBatchData);
-			using BinaryReader Reader = new BinaryReader(Ms);
+			using MemoryStream ms = new MemoryStream(eventBatchData);
+			using BinaryReader reader = new BinaryReader(ms);
 			
-			ulong AbsTimestamp = 0;
-			ulong LastTimestamp = 0;
+			ulong absTimestamp = 0;
+			ulong lastTimestamp = 0;
 
-			while (Reader.BaseStream.Position != Reader.BaseStream.Length)
+			while (reader.BaseStream.Position != reader.BaseStream.Length)
 			{
-				ulong Value = TraceUtils.Read7BitUint(Reader);
-				(ulong Timestamp, bool IsEnterScope) = DecodeTimestamp(Value);
+				ulong value = TraceUtils.Read7BitUint(reader);
+				(ulong timestamp, bool isEnterScope) = DecodeTimestamp(value);
 
-				if (AbsTimestamp == 0)
+				if (absTimestamp == 0)
 				{
-					AbsTimestamp = Timestamp; // First timestamp is always absolute
+					absTimestamp = timestamp; // First timestamp is always absolute
 				}
 				else
 				{
-					Timestamp = LastTimestamp + Timestamp * CycleFrequency; // Make timestamp absolute
+					timestamp = lastTimestamp + timestamp * _cycleFrequency; // Make timestamp absolute
 				}
-				LastTimestamp = Timestamp;
+				lastTimestamp = timestamp;
 
-				if (IsEnterScope)
+				if (isEnterScope)
 				{
-					ulong SpecId = TraceUtils.Read7BitUint(Reader);
-					CpuProfilerScopeEvent ScopeEvent = new CpuProfilerScopeEvent(Timestamp, true, (uint?) SpecId);
-					ScopeEvents.Add(ScopeEvent);
+					ulong specId = TraceUtils.Read7BitUint(reader);
+					CpuProfilerScopeEvent scopeEvent = new CpuProfilerScopeEvent(timestamp, true, (uint?)specId);
+					ScopeEvents.Add(scopeEvent);
 				}
 				else
 				{
-					CpuProfilerScopeEvent ScopeEvent = new CpuProfilerScopeEvent(Timestamp, false, null);
-					ScopeEvents.Add(ScopeEvent);
+					CpuProfilerScopeEvent scopeEvent = new CpuProfilerScopeEvent(timestamp, false, null);
+					ScopeEvents.Add(scopeEvent);
 				}
 			}
 		}
 		
 		public List<byte[]> Write()
 		{
-			List<MemoryStream> Streams = new List<MemoryStream>();
-			MemoryStream Stream = new MemoryStream();
-			BinaryWriter Writer = new BinaryWriter(Stream);
-			Streams.Add(Stream);
+			List<MemoryStream> streams = new List<MemoryStream>();
+			MemoryStream stream = new MemoryStream();
+			BinaryWriter writer = new BinaryWriter(stream);
+			streams.Add(stream);
 			
-			ulong LastAbsTimestamp = 0;
-			foreach (CpuProfilerScopeEvent Scope in ScopeEvents)
+			ulong lastAbsTimestamp = 0;
+			foreach (CpuProfilerScopeEvent scope in ScopeEvents)
 			{
-				ulong Timestamp;
-				if (LastAbsTimestamp == 0)
+				ulong timestamp;
+				if (lastAbsTimestamp == 0)
 				{
-					Timestamp = Scope.Timestamp;
+					timestamp = scope.Timestamp;
 				}
 				else
 				{
-					Timestamp = (Scope.Timestamp - LastAbsTimestamp) / CycleFrequency;
+					timestamp = (scope.Timestamp - lastAbsTimestamp) / _cycleFrequency;
 				}
-				LastAbsTimestamp = Scope.Timestamp;
+				lastAbsTimestamp = scope.Timestamp;
 				
-				TraceUtils.Write7BitUint(Writer, EncodeTimestamp(Timestamp, Scope.IsEnterScope));
-				if (Scope.IsEnterScope)
+				TraceUtils.Write7BitUint(writer, EncodeTimestamp(timestamp, scope.IsEnterScope));
+				if (scope.IsEnterScope)
 				{
-					TraceUtils.Write7BitUint(Writer, (ulong) Scope.SpecId!);
+					TraceUtils.Write7BitUint(writer, (ulong) scope.SpecId!);
 				}
 
 				// Split buffers after 250 bytes
-				if (Stream.Position > 250)
+				if (stream.Position > 250)
 				{
-					Writer.Close();
-					Stream = new MemoryStream();
-					Writer = new BinaryWriter(Stream);
-					LastAbsTimestamp = 0;
+					writer.Close();
+					stream = new MemoryStream();
+					writer = new BinaryWriter(stream);
+					lastAbsTimestamp = 0;
 				}
 			}
 
-			Writer.Close();
+			writer.Close();
 
-			return Streams.Select(x => x.ToArray()).ToList();
+			return streams.Select(x => x.ToArray()).ToList();
 		}
 		
-		internal static ulong EncodeTimestamp(ulong Timestamp, bool IsScopeEnter)
+		internal static ulong EncodeTimestamp(ulong timestamp, bool isScopeEnter)
 		{
-			ulong Value = (Timestamp << 1) | (uint) (IsScopeEnter ? 1 : 0);
-			return Value;
+			ulong value = (timestamp << 1) | (uint)(isScopeEnter ? 1 : 0);
+			return value;
 		}
 		
-		internal static (ulong Timestamp, bool IsScopeEnter) DecodeTimestamp(ulong Value)
+		internal static (ulong Timestamp, bool IsScopeEnter) DecodeTimestamp(ulong value)
 		{
-			bool IsScopeEnter = (Value & 1) != 0;
-			ulong Timestamp = Value >> 1; // Strip the IsScopeEnter bit
-			return (Timestamp, IsScopeEnter);
+			bool isScopeEnter = (value & 1) != 0;
+			ulong timestamp = value >> 1; // Strip the IsScopeEnter bit
+			return (timestamp, isScopeEnter);
 		}
 	}
 }
\ No newline at end of file
diff --git a/Engine/Source/Programs/Shared/EpicGames.Tracing/UnrealInsights/Events/Diagnostics.cs b/Engine/Source/Programs/Shared/EpicGames.Tracing/UnrealInsights/Events/Diagnostics.cs
--- a/Engine/Source/Programs/Shared/EpicGames.Tracing/UnrealInsights/Events/Diagnostics.cs
+++ b/Engine/Source/Programs/Shared/EpicGames.Tracing/UnrealInsights/Events/Diagnostics.cs
@@ -19,31 +19,31 @@ namespace EpicGames.Tracing.UnrealInsights.Events
 				new EventTypeField(5, 1, EventTypeField.TypeInt8, "TargetType"),
 			});
 		
-		public ushort Size => (ushort) (GenericEvent.Size + TraceImportantEventHeader.HeaderSize);
+		public ushort Size => (ushort) (_genericEvent.Size + TraceImportantEventHeader.HeaderSize);
 		public EventType Type => EventType;
-		private readonly GenericEvent GenericEvent;
+		private readonly GenericEvent _genericEvent;
 
-		public DiagnosticsSession2Event(string Platform, string AppName, string CommandLine, string Branch, string BuildVersion, int ChangeList, int ConfigurationType, int TargetType)
+		public DiagnosticsSession2Event(string platform, string appName, string commandLine, string branch, string buildVersion, int changeList, int configurationType, int targetType)
 		{
-			Field[] Fields =
+			Field[] fields =
 			{
-				Field.FromString(Platform),
-				Field.FromString(AppName),
-				Field.FromString(CommandLine),
-				Field.FromString(Branch),
-				Field.FromString(BuildVersion),
-				Field.FromInt(ChangeList),
-				Field.FromInt(ConfigurationType),
-				Field.FromInt(TargetType),
+				Field.FromString(platform),
+				Field.FromString(appName),
+				Field.FromString(commandLine),
+				Field.FromString(branch),
+				Field.FromString(buildVersion),
+				Field.FromInt(changeList),
+				Field.FromInt(configurationType),
+				Field.FromInt(targetType),
 			};
 
-			GenericEvent = new GenericEvent(0, Fields, EventType);
+			_genericEvent = new GenericEvent(0, fields, EventType);
 		}
 		
-		public void Serialize(ushort Uid, BinaryWriter Writer)
+		public void Serialize(ushort uid, BinaryWriter writer)
 		{
-			new TraceImportantEventHeader(Uid, GenericEvent.Size).Serialize(Writer);
-			GenericEvent.Serialize(Uid, Writer);
+			new TraceImportantEventHeader(uid, _genericEvent.Size).Serialize(writer);
+			_genericEvent.Serialize(uid, writer);
 		}
 	}
 }
\ No newline at end of file
diff --git a/Engine/Source/Programs/Shared/EpicGames.Tracing/UnrealInsights/Events/GenericEvent.cs b/Engine/Source/Programs/Shared/EpicGames.Tracing/UnrealInsights/Events/GenericEvent.cs
--- a/Engine/Source/Programs/Shared/EpicGames.Tracing/UnrealInsights/Events/GenericEvent.cs
+++ b/Engine/Source/Programs/Shared/EpicGames.Tracing/UnrealInsights/Events/GenericEvent.cs
@@ -13,26 +13,26 @@ namespace EpicGames.Tracing.UnrealInsights.Events
 		public long? Long { get; }
 		public double? Float { get; }
 		public string? String { get; }
-		readonly byte[]? Array;
+		readonly byte[]? _array;
 		public byte[]? GetArray()
 		{
-			return Array;
+			return _array;
 		}
 
-		private Field(bool Value) { Bool = Value; }
-		private Field(int Value) { Int = Value; }
-		private Field(long Value) { Long = Value; }
-		private Field(float Value) { Float = Value; }
-		private Field(double Value) { Float = Value; }
-		private Field(string Value) { String = Value; }
-		private Field(byte[] Value) { Array = Value; }
-		public static Field FromBool(bool Value) => new Field(Value);
-		public static Field FromInt(int Value) => new Field(Value);
-		public static Field FromLong(long Value) => new Field(Value);
-		public static Field FromFloat(float Value) => new Field(Value);
-		public static Field FromDouble(double Value) => new Field(Value);
-		public static Field FromString(string Value) => new Field(Value);
-		public static Field FromArray(byte[] Value) => new Field(Value);
+		private Field(bool value) { Bool = value; }
+		private Field(int value) { Int = value; }
+		private Field(long value) { Long = value; }
+		private Field(float value) { Float = value; }
+		private Field(double value) { Float = value; }
+		private Field(string value) { String = value; }
+		private Field(byte[] value) { _array = value; }
+		public static Field FromBool(bool value) => new Field(value);
+		public static Field FromInt(int value) => new Field(value);
+		public static Field FromLong(long value) => new Field(value);
+		public static Field FromFloat(float value) => new Field(value);
+		public static Field FromDouble(double value) => new Field(value);
+		public static Field FromString(string value) => new Field(value);
+		public static Field FromArray(byte[] value) => new Field(value);
 	}
 
 	/// <summary>
@@ -41,309 +41,316 @@ namespace EpicGames.Tracing.UnrealInsights.Events
 	public class GenericEvent : ITraceEvent
 	{
 		public uint Serial { get; }
-		readonly Field[] Fields;
-		public Field[] GetFields() => Fields;
-		public EventType Type => EventType;
-		private readonly EventType EventType;
+		readonly Field[] _fields;
+		public Field[] GetFields() => _fields;
+		public EventType Type => _eventType;
+		private readonly EventType _eventType;
 
-		public GenericEvent(uint Serial, Field[] Fields, EventType EventType)
+		public GenericEvent(uint serial, Field[] fields, EventType eventType)
 		{
-			this.Serial = Serial;
-			this.Fields = Fields;
-			this.EventType = EventType;
+			Serial = serial;
+			_fields = fields;
+			_eventType = eventType;
 		}
 
 		public ushort Size
 		{
 			get
 			{
-				ushort TotalSize = 0;
-				if (EventType.HasSerial)
-					TotalSize += 3; // 24-bit serial
-				TotalSize += EventType.GetEventSize();
-				if (EventType.MaybeHasAux())
+				ushort totalSize = 0;
+				if (_eventType.HasSerial)
 				{
-					for (int i = 0; i < EventType.Fields.Count; i++)
+					totalSize += 3; // 24-bit serial
+				}
+
+				totalSize += _eventType.GetEventSize();
+				if (_eventType.MaybeHasAux())
+				{
+					for (int i = 0; i < _eventType.Fields.Count; i++)
 					{
-						EventTypeField FieldType = EventType.Fields[i];
-						if (!FieldType.IsAuxData())
+						EventTypeField fieldType = _eventType.Fields[i];
+						if (!fieldType.IsAuxData())
+						{
 							continue;
+						}
 
-						ushort AuxDataSize;
-						if (FieldType.TypeInfo == EventTypeField.TypeAnsiString)
+						ushort auxDataSize;
+						if (fieldType.TypeInfo == EventTypeField.TypeAnsiString)
 						{
-							AuxDataSize = (ushort) Fields[i].String!.Length;
+							auxDataSize = (ushort) _fields[i].String!.Length;
 						}
-						else if (FieldType.TypeInfo == EventTypeField.TypeWideString)
+						else if (fieldType.TypeInfo == EventTypeField.TypeWideString)
 						{
-							AuxDataSize = (ushort) (Fields[i].String!.Length * 2);
+							auxDataSize = (ushort) (_fields[i].String!.Length * 2);
 						}
 						else
 						{
-							AuxDataSize = (ushort) Fields[i].GetArray()!.Length;
+							auxDataSize = (ushort) _fields[i].GetArray()!.Length;
 						}
 
-						TotalSize += sizeof(uint); // AuxHeader
-						TotalSize += AuxDataSize;
+						totalSize += sizeof(uint); // AuxHeader
+						totalSize += auxDataSize;
 					}
-					TotalSize += sizeof(byte); // AuxTerminal UID
+					totalSize += sizeof(byte); // AuxTerminal UID
 				}
 				
-				return TotalSize;
+				return totalSize;
 			}
 		}
 
-		public void Serialize(ushort Uid, BinaryWriter Writer)
+		public void Serialize(ushort uid, BinaryWriter writer)
 		{
-			if (EventType.HasSerial)
+			if (_eventType.HasSerial)
 			{
-				byte SerialLow = (byte) (Serial & 0xFF);
-				Writer.Write(SerialLow);
-				Writer.Write(BinaryReaderExtensions.GetHighWord(Serial));
+				byte serialLow = (byte)(Serial & 0xFF);
+				writer.Write(serialLow);
+				writer.Write(BinaryReaderExtensions.GetHighWord(Serial));
 			}
 			
-			for (int i = 0; i < EventType.Fields.Count; i++)
+			for (int i = 0; i < _eventType.Fields.Count; i++)
 			{
-				EventTypeField FieldType = EventType.Fields[i];
+				EventTypeField fieldType = _eventType.Fields[i];
 
-				if (FieldType.TypeInfo == EventTypeField.TypeInt8)
+				if (fieldType.TypeInfo == EventTypeField.TypeInt8)
 				{
-					Writer.Write((byte)Fields[i].Int!);
+					writer.Write((byte)_fields[i].Int!);
 				}
-				else if (FieldType.TypeInfo == EventTypeField.TypeInt16)
+				else if (fieldType.TypeInfo == EventTypeField.TypeInt16)
 				{
-					Writer.Write((ushort)Fields[i].Int!);
+					writer.Write((ushort)_fields[i].Int!);
 				}
-				else if (FieldType.TypeInfo == EventTypeField.TypeInt32)
+				else if (fieldType.TypeInfo == EventTypeField.TypeInt32)
 				{
-					Writer.Write((uint)Fields[i].Int!);
+					writer.Write((uint)_fields[i].Int!);
 				}
-				else if (FieldType.TypeInfo == EventTypeField.TypeInt64)
+				else if (fieldType.TypeInfo == EventTypeField.TypeInt64)
 				{
-					Writer.Write((ulong)Fields[i].Long!);
+					writer.Write((ulong)_fields[i].Long!);
 				}
-				else if (FieldType.TypeInfo == EventTypeField.TypePointer)
+				else if (fieldType.TypeInfo == EventTypeField.TypePointer)
 				{
-					Writer.Write((ulong)Fields[i].Long!);
+					writer.Write((ulong)_fields[i].Long!);
 				}
-				else if (FieldType.TypeInfo == EventTypeField.TypeFloat32)
+				else if (fieldType.TypeInfo == EventTypeField.TypeFloat32)
 				{
-					Writer.Write((float)Fields[i].Float!);
+					writer.Write((float)_fields[i].Float!);
 				}
-				else if (FieldType.TypeInfo == EventTypeField.TypeFloat64)
+				else if (fieldType.TypeInfo == EventTypeField.TypeFloat64)
 				{
-					Writer.Write((double)Fields[i].Float!);
+					writer.Write((double)_fields[i].Float!);
 				}
-				else if (FieldType.TypeInfo == EventTypeField.TypeAnsiString)
+				else if (fieldType.TypeInfo == EventTypeField.TypeAnsiString)
 				{
 				} // Write later as aux-data
-				else if (FieldType.TypeInfo == EventTypeField.TypeWideString)
+				else if (fieldType.TypeInfo == EventTypeField.TypeWideString)
 				{
 				} // Write later as aux-data
-				else if (FieldType.TypeInfo == EventTypeField.TypeArray)
+				else if (fieldType.TypeInfo == EventTypeField.TypeArray)
 				{
 				} // Write later as aux-data
-				else if (FieldType.TypeInfo == EventTypeField.TypeBool)
+				else if (fieldType.TypeInfo == EventTypeField.TypeBool)
 				{
-					Writer.Write(Fields[i].Bool!.Value ? (byte)1 : (byte)0);
+					writer.Write(_fields[i].Bool!.Value ? (byte)1 : (byte)0);
 				}
 				else
 				{
-					throw new Exception($"Found unknown TypeInfo {FieldType.TypeInfo}");
+					throw new Exception($"Found unknown TypeInfo {fieldType.TypeInfo}");
 				}
 			}
 
-			if (EventType.MaybeHasAux())
+			if (_eventType.MaybeHasAux())
 			{
-				for (int i = 0; i < EventType.Fields.Count; i++)
+				for (int i = 0; i < _eventType.Fields.Count; i++)
 				{
-					EventTypeField FieldType = EventType.Fields[i];
+					EventTypeField fieldType = _eventType.Fields[i];
 
-					if (FieldType.Size > 0)
+					if (fieldType.Size > 0)
+					{
 						continue; // Skip any non-aux fields
+					}
 
-					ushort AuxUid = PredefinedEventUid.AuxData;
-					if (!EventType.IsImportant())
+					ushort auxUid = PredefinedEventUid.AuxData;
+					if (!_eventType.IsImportant())
 					{
 						// Well-known UIDs are only shifted in non-important events
-						AuxUid = (ushort) (AuxUid << 1);
+						auxUid = (ushort) (auxUid << 1);
 					}
 
-					byte[] Data;
-					if (FieldType.TypeInfo == EventTypeField.TypeAnsiString)
+					byte[] data;
+					if (fieldType.TypeInfo == EventTypeField.TypeAnsiString)
 					{
-						Data = Encoding.ASCII.GetBytes(Fields[i].String!);
+						data = Encoding.ASCII.GetBytes(_fields[i].String!);
 					}
-					else if (FieldType.TypeInfo == EventTypeField.TypeWideString)
+					else if (fieldType.TypeInfo == EventTypeField.TypeWideString)
 					{
-						Data = Encoding.Unicode.GetBytes(Fields[i].String!);
+						data = Encoding.Unicode.GetBytes(_fields[i].String!);
 					}
 					else
 					{
-						Data = Fields[i].GetArray()!;
+						data = _fields[i].GetArray()!;
 					}
 
-					uint AuxHeader = 0;
-					AuxHeader = SetBits(AuxHeader, AuxUid, 0, 8);
-					AuxHeader = SetBits(AuxHeader, (uint) i, 8, 5);
-					AuxHeader = SetBits(AuxHeader, (uint) Data.Length, 13, 19);
+					uint auxHeader = 0;
+					auxHeader = SetBits(auxHeader, auxUid, 0, 8);
+					auxHeader = SetBits(auxHeader, (uint) i, 8, 5);
+					auxHeader = SetBits(auxHeader, (uint) data.Length, 13, 19);
 					
-					Writer.Write(AuxHeader);
-					Writer.Write(Data);
+					writer.Write(auxHeader);
+					writer.Write(data);
 				}
 				
-				ushort AuxDataTerminalUid = PredefinedEventUid.AuxDataTerminal;
+				ushort auxDataTerminalUid = PredefinedEventUid.AuxDataTerminal;
 				
-				if (!EventType.IsImportant())
+				if (!_eventType.IsImportant())
 				{
 					// Well-known UIDs are only shifted in non-important events
-					AuxDataTerminalUid = (ushort) (AuxDataTerminalUid << 1);
+					auxDataTerminalUid = (ushort) (auxDataTerminalUid << 1);
 				}
-				Writer.Write((byte)AuxDataTerminalUid);
+				writer.Write((byte)auxDataTerminalUid);
 			}
 		}
 
-		public static GenericEvent Deserialize(ushort Uid, BinaryReader Reader, EventType EventType)
+		public static GenericEvent Deserialize(ushort uid, BinaryReader reader, EventType eventType)
 		{
-			Field[] Fields = new Field[EventType.Fields.Count];
-			uint Serial = 0;
+			Field[] fields = new Field[eventType.Fields.Count];
+			uint serial = 0;
 			
-			if (!EventType.IsNoSync() && !EventType.IsImportant())
+			if (!eventType.IsNoSync() && !eventType.IsImportant())
 			{
 				// Read uint24 as serial
-				byte SerialLow = Reader.ReadByte();
-				ushort SerialHigh = Reader.ReadUInt16();
-				Serial = SerialLow | ((uint)SerialHigh << 16);
+				byte serialLow = reader.ReadByte();
+				ushort serialHigh = reader.ReadUInt16();
+				serial = serialLow | ((uint)serialHigh << 16);
 			}
 
-			for (int i = 0; i < EventType.Fields.Count; i++)
+			for (int i = 0; i < eventType.Fields.Count; i++)
 			{
-				EventTypeField FieldType = EventType.Fields[i];
+				EventTypeField fieldType = eventType.Fields[i];
 
-				if (FieldType.TypeInfo == EventTypeField.TypeInt8)
+				if (fieldType.TypeInfo == EventTypeField.TypeInt8)
 				{
-					Fields[i] = Field.FromInt(Reader.ReadByte());
+					fields[i] = Field.FromInt(reader.ReadByte());
 				}
-				else if (FieldType.TypeInfo == EventTypeField.TypeInt16)
+				else if (fieldType.TypeInfo == EventTypeField.TypeInt16)
 				{
-					Fields[i] = Field.FromInt(Reader.ReadUInt16());
+					fields[i] = Field.FromInt(reader.ReadUInt16());
 				}
-				else if (FieldType.TypeInfo == EventTypeField.TypeInt32)
+				else if (fieldType.TypeInfo == EventTypeField.TypeInt32)
 				{
-					Fields[i] = Field.FromInt((int)Reader.ReadUInt32());
+					fields[i] = Field.FromInt((int)reader.ReadUInt32());
 				}
-				else if (FieldType.TypeInfo == EventTypeField.TypeInt64)
+				else if (fieldType.TypeInfo == EventTypeField.TypeInt64)
 				{
-					Fields[i] = Field.FromLong((long) Reader.ReadUInt64());
+					fields[i] = Field.FromLong((long) reader.ReadUInt64());
 				}
-				else if (FieldType.TypeInfo == EventTypeField.TypePointer)
+				else if (fieldType.TypeInfo == EventTypeField.TypePointer)
 				{
-					Fields[i] = Field.FromLong((long) Reader.ReadUInt64());
+					fields[i] = Field.FromLong((long) reader.ReadUInt64());
 				}
-				else if (FieldType.TypeInfo == EventTypeField.TypeFloat32)
+				else if (fieldType.TypeInfo == EventTypeField.TypeFloat32)
 				{
-					Fields[i] = Field.FromFloat(Reader.ReadSingle());
+					fields[i] = Field.FromFloat(reader.ReadSingle());
 				}
-				else if (FieldType.TypeInfo == EventTypeField.TypeFloat64)
+				else if (fieldType.TypeInfo == EventTypeField.TypeFloat64)
 				{
-					Fields[i] = Field.FromDouble(Reader.ReadDouble());
+					fields[i] = Field.FromDouble(reader.ReadDouble());
 				}
-				else if (FieldType.TypeInfo == EventTypeField.TypeAnsiString)
+				else if (fieldType.TypeInfo == EventTypeField.TypeAnsiString)
 				{
 				} // Read later as aux-data
-				else if (FieldType.TypeInfo == EventTypeField.TypeWideString)
+				else if (fieldType.TypeInfo == EventTypeField.TypeWideString)
 				{
 				} // Read later as aux-data
-				else if (FieldType.TypeInfo == EventTypeField.TypeArray)
+				else if (fieldType.TypeInfo == EventTypeField.TypeArray)
 				{
 				} // Read later as aux-data
-				else if (FieldType.TypeInfo == EventTypeField.TypeBool)
+				else if (fieldType.TypeInfo == EventTypeField.TypeBool)
 				{
-					Fields[i] = Field.FromBool(Reader.ReadByte() == 1);
+					fields[i] = Field.FromBool(reader.ReadByte() == 1);
 				}
 				else
 				{
-					throw new Exception($"Found unknown TypeInfo {FieldType.TypeInfo}");
+					throw new Exception($"Found unknown TypeInfo {fieldType.TypeInfo}");
 				}
 			}
 
-			if (EventType.MaybeHasAux())
+			if (eventType.MaybeHasAux())
 			{
 				for (;;)
 				{
-					ushort AuxUid = Reader.ReadByte();
-					if (!EventType.IsImportant())
+					ushort auxUid = reader.ReadByte();
+					if (!eventType.IsImportant())
 					{
 						// Well-known UIDs are only shifted in non-important events
-						AuxUid = (ushort) (AuxUid >> 1);
+						auxUid = (ushort) (auxUid >> 1);
 					}
 
-					if (AuxUid == PredefinedEventUid.AuxData)
+					if (auxUid == PredefinedEventUid.AuxData)
 					{
-						Reader.BaseStream.Position -= 1; // Include the UID for parsing rest of header as a single uint32
-						uint AuxHeader = Reader.ReadUInt32();
+						reader.BaseStream.Position -= 1; // Include the UID for parsing rest of header as a single uint32
+						uint auxHeader = reader.ReadUInt32();
 
-						int FieldIndex = (int) ReadBits(AuxHeader, 8, 5);
-						uint Size = ReadBits(AuxHeader, 13, 19);
+						int fieldIndex = (int)ReadBits(auxHeader, 8, 5);
+						uint size = ReadBits(auxHeader, 13, 19);
 
-						byte[] AuxData = Reader.ReadBytesStrict((int) Size);
+						byte[] auxData = reader.ReadBytesStrict((int)size);
 						
-						if (EventType.Fields[FieldIndex].TypeInfo == EventTypeField.TypeAnsiString)
+						if (eventType.Fields[fieldIndex].TypeInfo == EventTypeField.TypeAnsiString)
 						{
-							Fields[FieldIndex] = Field.FromString(Encoding.ASCII.GetString(AuxData));
+							fields[fieldIndex] = Field.FromString(Encoding.ASCII.GetString(auxData));
 						}
-						else if (EventType.Fields[FieldIndex].TypeInfo == EventTypeField.TypeWideString)
+						else if (eventType.Fields[fieldIndex].TypeInfo == EventTypeField.TypeWideString)
 						{
-							Fields[FieldIndex] = Field.FromString(Encoding.Unicode.GetString(AuxData));
+							fields[fieldIndex] = Field.FromString(Encoding.Unicode.GetString(auxData));
 						}
 						else
 						{
-							Fields[FieldIndex] = Field.FromArray(AuxData);
+							fields[fieldIndex] = Field.FromArray(auxData);
 						}
 					}
-					else if (AuxUid == PredefinedEventUid.AuxDataTerminal)
+					else if (auxUid == PredefinedEventUid.AuxDataTerminal)
 					{
 						break;
 					}
 					else
 					{
-						throw new Exception($"Invalid AuxUid found: {AuxUid} / 0x{AuxUid:X4}");
+						throw new Exception($"Invalid AuxUid found: {auxUid} / 0x{auxUid:X4}");
 					}
 				}
 			}
 
-			return new GenericEvent(Serial, Fields, EventType);
+			return new GenericEvent(serial, fields, eventType);
 		}
 
-		private static uint SetBits(uint Word, uint Value, int Pos, int Size)
+		private static uint SetBits(uint word, uint value, int pos, int size)
 		{
-			uint mask = (((uint)1 << Size) - 1) << Pos;
-			Word &= ~mask;
-			Word |= (Value << Pos) & mask;
-			return Word;
+			uint mask = (((uint)1 << size) - 1) << pos;
+			word &= ~mask;
+			word |= (value << pos) & mask;
+			return word;
 		}
 
-		private static uint ReadBits(uint Word, int Pos, int Size)
+		private static uint ReadBits(uint word, int pos, int size)
 		{
-			uint Mask = (((uint)1 << Size) - 1) << Pos;
-			return (Word & Mask) >> Pos;
+			uint mask = (((uint)1 << size) - 1) << pos;
+			return (word & mask) >> pos;
 		}
 
-		internal static (ushort Uid, int FieldIndex, int Size) DeserializeAuxHeader(uint Header)
+		internal static (ushort Uid, int FieldIndex, int Size) DeserializeAuxHeader(uint header)
 		{
-			ushort Uid = (ushort) ReadBits(Header, 0, 8);
-			int FieldIndex = (int) ReadBits(Header, 8, 5);
-			int Size = (int) ReadBits(Header, 13, 19);
-			return (Uid, FieldIndex, Size);
+			ushort uid = (ushort)ReadBits(header, 0, 8);
+			int fieldIndex = (int)ReadBits(header, 8, 5);
+			int size = (int)ReadBits(header, 13, 19);
+			return (uid, fieldIndex, size);
 		}
 		
-		internal static uint SerializeAuxHeader(ushort Uid, int FieldIndex, int Size)
+		internal static uint SerializeAuxHeader(ushort uid, int fieldIndex, int size)
 		{
-			uint Header = 0;
-			Header = SetBits(Header, Uid, 0, 8);
-			Header = SetBits(Header, (uint) FieldIndex, 8, 5);
-			Header = SetBits(Header, (uint) Size, 13, 19);
-			return Header;
+			uint header = 0;
+			header = SetBits(header, uid, 0, 8);
+			header = SetBits(header, (uint) fieldIndex, 8, 5);
+			header = SetBits(header, (uint) size, 13, 19);
+			return header;
 		}
 	}
 }
\ No newline at end of file
diff --git a/Engine/Source/Programs/Shared/EpicGames.Tracing/UnrealInsights/Events/Logging.cs b/Engine/Source/Programs/Shared/EpicGames.Tracing/UnrealInsights/Events/Logging.cs
--- a/Engine/Source/Programs/Shared/EpicGames.Tracing/UnrealInsights/Events/Logging.cs
+++ b/Engine/Source/Programs/Shared/EpicGames.Tracing/UnrealInsights/Events/Logging.cs
@@ -18,7 +18,7 @@ namespace EpicGames.Tracing.UnrealInsights.Events
 		public ushort Size => throw new NotImplementedException();
 		public EventType Type => EventType;
 
-		public void Serialize(ushort Uid, BinaryWriter Writer)
+		public void Serialize(ushort uid, BinaryWriter writer)
 		{
 			throw new NotImplementedException();
 		}
diff --git a/Engine/Source/Programs/Shared/EpicGames.Tracing/UnrealInsights/Events/Memory.cs b/Engine/Source/Programs/Shared/EpicGames.Tracing/UnrealInsights/Events/Memory.cs
--- a/Engine/Source/Programs/Shared/EpicGames.Tracing/UnrealInsights/Events/Memory.cs
+++ b/Engine/Source/Programs/Shared/EpicGames.Tracing/UnrealInsights/Events/Memory.cs
@@ -14,7 +14,7 @@ namespace EpicGames.Tracing.UnrealInsights.Events
 		public ushort Size => throw new NotImplementedException();
 		public EventType Type => EventType;
 		
-		public void Serialize(ushort Uid, BinaryWriter Writer)
+		public void Serialize(ushort uid, BinaryWriter writer)
 		{
 			throw new NotImplementedException();
 		}
diff --git a/Engine/Source/Programs/Shared/EpicGames.Tracing/UnrealInsights/Events/NewEvent.cs b/Engine/Source/Programs/Shared/EpicGames.Tracing/UnrealInsights/Events/NewEvent.cs
--- a/Engine/Source/Programs/Shared/EpicGames.Tracing/UnrealInsights/Events/NewEvent.cs
+++ b/Engine/Source/Programs/Shared/EpicGames.Tracing/UnrealInsights/Events/NewEvent.cs
@@ -54,27 +54,60 @@ namespace EpicGames.Tracing.UnrealInsights.Events
 		public static string TypeInfoToString(byte ByteInfo)
 		{
 			if (ByteInfo == TypeInt8)
+			{
 				return "int8";
+			}
+
 			if (ByteInfo == TypeInt16)
+			{
 				return "int16";
+			}
+
 			if (ByteInfo == TypeInt32)
+			{
 				return "int32";
+			}
+
 			if (ByteInfo == TypeInt64)
+			{
 				return "int64";
+			}
+
 			if (ByteInfo == TypePointer)
+			{
 				return "ptr";
+			}
+
 			if (ByteInfo == TypeFloat32)
+			{
 				return "float32";
+			}
+
 			if (ByteInfo == TypeFloat64)
+			{
 				return "float64";
+			}
+
 			if (ByteInfo == TypeAnsiString)
+			{
 				return "ansi_str";
+			}
+
 			if (ByteInfo == TypeWideString)
+			{
 				return "wide_str";
+			}
+
 			if (ByteInfo == TypeArray)
+			{
 				return "array";
+			}
+
 			if (ByteInfo == TypeBool)
+			{
 				return "bool";
+			}
+
 			throw new Exception($"Unable to convert type info {ByteInfo} to string");
 		}
 
@@ -140,35 +173,35 @@ namespace EpicGames.Tracing.UnrealInsights.Events
 		
 		public string Name => $"{LoggerName}.{EventName}";
 
-		private int EventSize;
+		private int _eventSize;
 
 		public const byte FlagNone = 0;
 		public const byte FlagImportant = 1 << 0;
 		public const byte FlagMaybeHasAux = 1 << 1;
 		public const byte FlagNoSync = 1 << 2;
 
-		public EventType(ushort NewEventUid, string LoggerName, string EventName, byte Flags, List<EventTypeField> Fields)
+		public EventType(ushort newEventUid, string loggerName, string eventName, byte flags, List<EventTypeField> fields)
 		{
-			this.NewEventUid = NewEventUid;
-			this.LoggerName = LoggerName;
-			this.EventName = EventName;
-			this.Flags = Flags;
-			this.Fields = Fields;
-			EventSize += this.Fields.Sum(f => f.Size);
+			NewEventUid = newEventUid;
+			LoggerName = loggerName;
+			EventName = eventName;
+			Flags = flags;
+			Fields = fields;
+			_eventSize += Fields.Sum(f => f.Size);
 		}
 
-		public EventType(string LoggerName, string EventName, byte Flags)
+		public EventType(string loggerName, string eventName, byte flags)
 		{
 			NewEventUid = 0;
-			this.LoggerName = LoggerName;
-			this.EventName = EventName;
-			this.Flags = Flags;
+			LoggerName = loggerName;
+			EventName = eventName;
+			Flags = flags;
 			Fields = new List<EventTypeField>();
 		}
 
-		public static EventType WellKnown(ushort Uid, string Name)
+		public static EventType WellKnown(ushort uid, string name)
 		{
-			return new EventType(Uid, "WellKnown", Name, 0, new List<EventTypeField>());
+			return new EventType(uid, "WellKnown", name, 0, new List<EventTypeField>());
 		}
 		
 		private static EventType Self()
@@ -176,16 +209,16 @@ namespace EpicGames.Tracing.UnrealInsights.Events
 			return new EventType(0, "EventType", "Self", 0, new List<EventTypeField>());
 		}
 
-		public void AddEventType(ushort Offset, ushort Size, byte TypeInfo, string Name)
+		public void AddEventType(ushort offset, ushort size, byte typeInfo, string name)
 		{
-			EventTypeField Field = new EventTypeField(Offset, Size, TypeInfo, Name);
-			Fields.Add(Field);
-			EventSize += Field.Size;
+			EventTypeField field = new EventTypeField(offset, size, typeInfo, name);
+			Fields.Add(field);
+			_eventSize += field.Size;
 		}
 		
 		public bool IsWellKnown()
 		{
-			return NewEventUid < PredefinedEventUid._WellKnownNum;
+			return NewEventUid < PredefinedEventUid.WellKnownNum;
 		}
 
 		public bool IsImportant()
@@ -214,7 +247,7 @@ namespace EpicGames.Tracing.UnrealInsights.Events
 		/// <returns></returns>
 		public ushort GetEventSize()
 		{
-			return (ushort) EventSize;
+			return (ushort) _eventSize;
 		}
 
 		public override string ToString()
@@ -224,31 +257,40 @@ namespace EpicGames.Tracing.UnrealInsights.Events
 		
 		public string ToStringDetailed()
 		{
-			string FieldText = String.Join(',', Fields.Select(x => x.ToString()));
-			string FlagText = "";
+			string fieldText = String.Join(',', Fields.Select(x => x.ToString()));
+			string flagText = "";
 			if (IsImportant())
-				FlagText += ",Important";
+			{
+				flagText += ",Important";
+			}
+
 			if (MaybeHasAux())
-				FlagText += ",MaybeHasAux";
+			{
+				flagText += ",MaybeHasAux";
+			}
+
 			if (IsNoSync())
-				FlagText += ",NoSync";
-			FlagText = FlagText.Trim(',');
-			return $"EventType(Uid={NewEventUid} {LoggerName} {EventName} Flags={FlagText} FlagRaw={Flags} Fields={FieldText})";
+			{
+				flagText += ",NoSync";
+			}
+
+			flagText = flagText.Trim(',');
+			return $"EventType(Uid={NewEventUid} {LoggerName} {EventName} Flags={flagText} FlagRaw={Flags} Fields={fieldText})";
 		}
 
 		public ushort Size
 		{
 			get
 			{
-				byte[] LoggerNameBytes = Encoding.UTF8.GetBytes(LoggerName);
-				byte[] EventNameBytes = Encoding.UTF8.GetBytes(EventName);
+				byte[] loggerNameBytes = Encoding.UTF8.GetBytes(LoggerName);
+				byte[] eventNameBytes = Encoding.UTF8.GetBytes(EventName);
 
-				ushort FieldArraySize = (ushort) (Fields.Count * EventTypeField.StructSize);
-				ushort NamesSize = (ushort) (LoggerNameBytes.Length + EventNameBytes.Length + Fields.Sum(x => x.NameSize));
-				ushort NewEventSize = (ushort) (2 + 1 + 1 + 1 + 1 + FieldArraySize + NamesSize);
+				ushort fieldArraySize = (ushort)(Fields.Count * EventTypeField.StructSize);
+				ushort namesSize = (ushort)(loggerNameBytes.Length + eventNameBytes.Length + Fields.Sum(x => x.NameSize));
+				ushort newEventSize = (ushort)(2 + 1 + 1 + 1 + 1 + fieldArraySize + namesSize);
 				
-				NewEventSize += TraceImportantEventHeader.HeaderSize;
-				return NewEventSize;
+				newEventSize += TraceImportantEventHeader.HeaderSize;
+				return newEventSize;
 			}
 		}
 
@@ -267,72 +309,81 @@ namespace EpicGames.Tracing.UnrealInsights.Events
 		public override bool Equals(object? obj)
 		{
 			if (obj is null)
+			{
 				return false;
+			}
+
 			if (ReferenceEquals(this, obj))
+			{
 				return true;
+			}
+
 			if (obj.GetType() != GetType())
+			{
 				return false;
+			}
+
 			return Equals((EventType) obj);
 		}
 
-		public void Serialize(ushort Uid, BinaryWriter Writer)
+		public void Serialize(ushort uid, BinaryWriter writer)
 		{
-			new TraceImportantEventHeader(PredefinedEventUid.NewEvent, (ushort) (Size - TraceImportantEventHeader.HeaderSize)).Serialize(Writer);
+			new TraceImportantEventHeader(PredefinedEventUid.NewEvent, (ushort) (Size - TraceImportantEventHeader.HeaderSize)).Serialize(writer);
 			
-			byte[] LoggerNameBytes = Encoding.UTF8.GetBytes(LoggerName);
-			byte[] EventNameBytes = Encoding.UTF8.GetBytes(EventName);
+			byte[] loggerNameBytes = Encoding.UTF8.GetBytes(LoggerName);
+			byte[] eventNameBytes = Encoding.UTF8.GetBytes(EventName);
 
-			Writer.Write(Uid); // UID of the new event to declare
-			Writer.Write((byte) Fields.Count);
-			Writer.Write(Flags);
-			Writer.Write((byte) LoggerNameBytes.Length);
-			Writer.Write((byte) EventNameBytes.Length);
+			writer.Write(uid); // UID of the new event to declare
+			writer.Write((byte) Fields.Count);
+			writer.Write(Flags);
+			writer.Write((byte) loggerNameBytes.Length);
+			writer.Write((byte) eventNameBytes.Length);
 
-			foreach (EventTypeField Field in Fields)
+			foreach (EventTypeField field in Fields)
 			{
-				Writer.Write(Field.Offset);
-				Writer.Write(Field.Size);
-				Writer.Write(Field.TypeInfo);
-				Writer.Write(Field.NameSize);
+				writer.Write(field.Offset);
+				writer.Write(field.Size);
+				writer.Write(field.TypeInfo);
+				writer.Write(field.NameSize);
 			}
 
-			Writer.Write(LoggerNameBytes);
-			Writer.Write(EventNameBytes);
+			writer.Write(loggerNameBytes);
+			writer.Write(eventNameBytes);
 
-			foreach (EventTypeField Field in Fields)
+			foreach (EventTypeField field in Fields)
 			{
-				Writer.Write(Encoding.UTF8.GetBytes(Field.Name));
+				writer.Write(Encoding.UTF8.GetBytes(field.Name));
 			}
 		}
 
-		public static (ushort, EventType) Deserialize(BinaryReader Reader)
+		public static (ushort, EventType) Deserialize(BinaryReader reader)
 		{
-			ushort NewEventUid = Reader.ReadUInt16();
-			byte FieldCount = Reader.ReadByte();
-			byte Flags = Reader.ReadByte();
-			byte LoggerNameSize = Reader.ReadByte();
-			byte EventNameSize = Reader.ReadByte();
+			ushort newEventUid = reader.ReadUInt16();
+			byte fieldCount = reader.ReadByte();
+			byte flags = reader.ReadByte();
+			byte loggerNameSize = reader.ReadByte();
+			byte eventNameSize = reader.ReadByte();
 			
-			EventTypeField[] Fields = new EventTypeField[FieldCount];
+			EventTypeField[] fields = new EventTypeField[fieldCount];
 
-			for (int i = 0; i < FieldCount; i++)
+			for (int i = 0; i < fieldCount; i++)
 			{
-				ushort Offset = Reader.ReadUInt16();
-				ushort Size = Reader.ReadUInt16();
-				byte TypeInfo = Reader.ReadByte();
-				byte NameSize = Reader.ReadByte();
-				Fields[i] = new EventTypeField(Offset, Size, TypeInfo, NameSize);
+				ushort offset = reader.ReadUInt16();
+				ushort size = reader.ReadUInt16();
+				byte typeInfo = reader.ReadByte();
+				byte nameSize = reader.ReadByte();
+				fields[i] = new EventTypeField(offset, size, typeInfo, nameSize);
 			}
 
-			string LoggerName = Encoding.UTF8.GetString(Reader.ReadBytesStrict(LoggerNameSize));
-			string EventName = Encoding.UTF8.GetString(Reader.ReadBytesStrict(EventNameSize));
+			string loggerName = Encoding.UTF8.GetString(reader.ReadBytesStrict(loggerNameSize));
+			string eventName = Encoding.UTF8.GetString(reader.ReadBytesStrict(eventNameSize));
 
-			for (int i = 0; i < FieldCount; i++)
+			for (int i = 0; i < fieldCount; i++)
 			{
-				Fields[i].SetName(Encoding.UTF8.GetString(Reader.ReadBytesStrict(Fields[i].NameSize)));
+				fields[i].SetName(Encoding.UTF8.GetString(reader.ReadBytesStrict(fields[i].NameSize)));
 			}
 
-			return (NewEventUid, new EventType(NewEventUid, LoggerName, EventName, Flags, Fields.ToList()));
+			return (newEventUid, new EventType(newEventUid, loggerName, eventName, flags, fields.ToList()));
 		}
 	}
 }
\ No newline at end of file
diff --git a/Engine/Source/Programs/Shared/EpicGames.Tracing/UnrealInsights/Events/Trace.cs b/Engine/Source/Programs/Shared/EpicGames.Tracing/UnrealInsights/Events/Trace.cs
--- a/Engine/Source/Programs/Shared/EpicGames.Tracing/UnrealInsights/Events/Trace.cs
+++ b/Engine/Source/Programs/Shared/EpicGames.Tracing/UnrealInsights/Events/Trace.cs
@@ -23,30 +23,30 @@ namespace EpicGames.Tracing.UnrealInsights.Events
 		public ushort Size => (ushort) (EventType.GetEventSize() + TraceImportantEventHeader.HeaderSize);
 		public EventType Type => EventType;
 		
-		public TraceNewTraceEvent(ulong StartCycle, ulong CycleFrequency, ushort Endian, byte PointerSize)
+		public TraceNewTraceEvent(ulong startCycle, ulong cycleFrequency, ushort endian, byte pointerSize)
 		{
-			this.StartCycle = StartCycle;
-			this.CycleFrequency = CycleFrequency;
-			this.Endian = Endian;
-			this.PointerSize = PointerSize;
+			StartCycle = startCycle;
+			CycleFrequency = cycleFrequency;
+			Endian = endian;
+			PointerSize = pointerSize;
 		}
 		
-		public void Serialize(ushort Uid, BinaryWriter Writer)
+		public void Serialize(ushort uid, BinaryWriter writer)
 		{
-			new TraceImportantEventHeader(Uid, EventType.GetEventSize()).Serialize(Writer);
-			Writer.Write(StartCycle);
-			Writer.Write(CycleFrequency);
-			Writer.Write(Endian);
-			Writer.Write(PointerSize);
+			new TraceImportantEventHeader(uid, EventType.GetEventSize()).Serialize(writer);
+			writer.Write(StartCycle);
+			writer.Write(CycleFrequency);
+			writer.Write(Endian);
+			writer.Write(PointerSize);
 		}
 		
-		public static TraceNewTraceEvent Deserialize(BinaryReader Reader)
+		public static TraceNewTraceEvent Deserialize(BinaryReader reader)
 		{
-			ulong StartCycle = Reader.ReadUInt64();
-			ulong CycleFrequency = Reader.ReadUInt64();
-			ushort Endian = Reader.ReadUInt16();
-			byte PointerSize = Reader.ReadByte();
-			return new TraceNewTraceEvent(StartCycle, CycleFrequency, Endian, PointerSize);
+			ulong startCycle = reader.ReadUInt64();
+			ulong cycleFrequency = reader.ReadUInt64();
+			ushort endian = reader.ReadUInt16();
+			byte pointerSize = reader.ReadByte();
+			return new TraceNewTraceEvent(startCycle, cycleFrequency, endian, pointerSize);
 		}
 	}
 
@@ -61,27 +61,27 @@ namespace EpicGames.Tracing.UnrealInsights.Events
 			});
 			
 
-		public ushort Size => (ushort) (GenericEvent.Size + TraceImportantEventHeader.HeaderSize);
+		public ushort Size => (ushort) (_genericEvent.Size + TraceImportantEventHeader.HeaderSize);
 		public EventType Type => EventType;
-		private readonly GenericEvent GenericEvent;
+		private readonly GenericEvent _genericEvent;
 		
-		public TraceThreadInfoEvent(int ThreadId, int SystemId, int SortHint, string Name)
+		public TraceThreadInfoEvent(int threadId, int systemId, int sortHint, string name)
 		{
-			Field[] Fields =
+			Field[] fields =
 			{
-				Field.FromInt((int) ThreadId),
-				Field.FromInt((int) SystemId),
-				Field.FromInt((int) SortHint),
-				Field.FromString(Name),
+				Field.FromInt((int) threadId),
+				Field.FromInt((int) systemId),
+				Field.FromInt((int) sortHint),
+				Field.FromString(name),
 			};
 
-			GenericEvent = new GenericEvent(0, Fields, EventType);
+			_genericEvent = new GenericEvent(0, fields, EventType);
 		}
 
-		public void Serialize(ushort Uid, BinaryWriter Writer)
+		public void Serialize(ushort uid, BinaryWriter writer)
 		{
-			new TraceImportantEventHeader(Uid, GenericEvent.Size).Serialize(Writer);
-			GenericEvent.Serialize(Uid, Writer);
+			new TraceImportantEventHeader(uid, _genericEvent.Size).Serialize(writer);
+			_genericEvent.Serialize(uid, writer);
 		}
 	}
 }
\ No newline at end of file
diff --git a/Engine/Source/Programs/Shared/EpicGames.Tracing/UnrealInsights/Reader.cs b/Engine/Source/Programs/Shared/EpicGames.Tracing/UnrealInsights/Reader.cs
--- a/Engine/Source/Programs/Shared/EpicGames.Tracing/UnrealInsights/Reader.cs
+++ b/Engine/Source/Programs/Shared/EpicGames.Tracing/UnrealInsights/Reader.cs
@@ -27,7 +27,7 @@ namespace EpicGames.Tracing.UnrealInsights
 		// public const ushort _LeaveScope_T_Unused0 = 13;
 		// public const ushort _LeaveScope_T_Unused1 = 14;
 		// public const ushort _LeaveScope_T_Unused2 = 15;
-		public const ushort _WellKnownNum = 16;
+		public const ushort WellKnownNum = 16;
 	}
 	
 	public class UnrealInsightsReader
@@ -43,22 +43,22 @@ namespace EpicGames.Tracing.UnrealInsights
 			EventTypes[PredefinedEventUid.LeaveScope_T] = new LeaveScopeEventTimestamp(0).Type;
 		}
 
-		public void Read(Stream Stream)
+		public void Read(Stream stream)
 		{
-			using BinaryReader Reader = new BinaryReader(Stream);
-			StreamHeader.Deserialize(Reader);
+			using BinaryReader reader = new BinaryReader(stream);
+			StreamHeader.Deserialize(reader);
 
-			Dictionary<ushort, MemoryStream> ThreadStreams = Demux(Reader);
-			Reader.EnsureEntireStreamIsConsumed(); // There should be no more trailing bytes once all packets are consumed.
+			Dictionary<ushort, MemoryStream> threadStreams = Demux(reader);
+			reader.EnsureEntireStreamIsConsumed(); // There should be no more trailing bytes once all packets are consumed.
 			
-			void ReadThreadStream(ushort ThreadId)
+			void ReadThreadStream(ushort threadId)
 			{
-				MemoryStream ThreadStream = ThreadStreams[ThreadId];
-				using BinaryReader ThreadReader = new BinaryReader(ThreadStream);
-				EventsPerThread[ThreadId] = ReadEvents(ThreadId, ThreadReader);
-				ThreadReader.EnsureEntireStreamIsConsumed();
-				ThreadStream.Dispose();
-				ThreadStreams.Remove(ThreadId);
+				MemoryStream threadStream = threadStreams[threadId];
+				using BinaryReader threadReader = new BinaryReader(threadStream);
+				EventsPerThread[threadId] = ReadEvents(threadId, threadReader);
+				threadReader.EnsureEntireStreamIsConsumed();
+				threadStream.Dispose();
+				threadStreams.Remove(threadId);
 			}
 
 			// Read the NewEvents first to ensure event types are registered
@@ -67,9 +67,9 @@ namespace EpicGames.Tracing.UnrealInsights
 			// Read the important events next
 			ReadThreadStream(TransportPacket.ThreadIdImportants);
 			
-			foreach (ushort ThreadId in ThreadStreams.Keys)
+			foreach (ushort threadId in threadStreams.Keys)
 			{
-				ReadThreadStream(ThreadId);
+				ReadThreadStream(threadId);
 			}
 		}
 
@@ -79,69 +79,71 @@ namespace EpicGames.Tracing.UnrealInsights
 		///
 		/// It's a naive and simple enough for our parsing purposes but it definitely won't read larger trace files.
 		/// </summary>
-		/// <param name="Reader">Reader to consume</param>
+		/// <param name="reader">Reader to consume</param>
 		/// <returns>A stream of data per thread</returns>
-		public Dictionary<ushort, MemoryStream> Demux(BinaryReader Reader)
+		public Dictionary<ushort, MemoryStream> Demux(BinaryReader reader)
 		{
-			Dictionary<ushort, MemoryStream> ThreadStreams = new Dictionary<ushort, MemoryStream>();
-			while (Reader.BaseStream.Position < Reader.BaseStream.Length)
+			Dictionary<ushort, MemoryStream> threadStreams = new Dictionary<ushort, MemoryStream>();
+			while (reader.BaseStream.Position < reader.BaseStream.Length)
 			{
-				TransportPacket Packet = TransportPacket.Deserialize(Reader);
+				TransportPacket packet = TransportPacket.Deserialize(reader);
 
-				if (!ThreadStreams.TryGetValue(Packet.GetThreadId(), out MemoryStream? ThreadStream))
+				if (!threadStreams.TryGetValue(packet.GetThreadId(), out MemoryStream? threadStream))
 				{
-					ThreadStream = new MemoryStream(5 * 1024 * 1024);
-					ThreadStreams[Packet.GetThreadId()] = ThreadStream;
+					threadStream = new MemoryStream(5 * 1024 * 1024);
+					threadStreams[packet.GetThreadId()] = threadStream;
 				}
 
-				ThreadStream.Write(Packet.GetData());
+				threadStream.Write(packet.GetData());
 				NumTransportPacketsRead++;
 			}
-			Reader.EnsureEntireStreamIsConsumed();
+			reader.EnsureEntireStreamIsConsumed();
 			
 			// Reset stream positions so it can be read from the beginning once returned
-			ThreadStreams.Values.ToList().ForEach(Ms => Ms.Position = 0);
+			threadStreams.Values.ToList().ForEach(ms => ms.Position = 0);
 
-			return ThreadStreams;
+			return threadStreams;
 		}
 
-		public List<ITraceEvent> ReadEvents(int ThreadId, BinaryReader Reader)
+		public List<ITraceEvent> ReadEvents(int threadId, BinaryReader reader)
 		{
-			List<ITraceEvent> EventsRead = new List<ITraceEvent>();
-			while (Reader.BaseStream.Position != Reader.BaseStream.Length)
+			List<ITraceEvent> eventsRead = new List<ITraceEvent>();
+			while (reader.BaseStream.Position != reader.BaseStream.Length)
 			{
-				ushort Uid;
-				ushort Size = 0;
-				if (ThreadId == TransportPacket.ThreadIdEvents || ThreadId == TransportPacket.ThreadIdImportants)
+				ushort uid;
+				ushort size = 0;
+				if (threadId == TransportPacket.ThreadIdEvents || threadId == TransportPacket.ThreadIdImportants)
 				{
-					Uid = Reader.ReadUInt16();
-					Size = Reader.ReadUInt16(); // TODO: validate and respect the size
+					uid = reader.ReadUInt16();
+					size = reader.ReadUInt16(); // TODO: validate and respect the size
 				}
 				else
 				{
-					Uid = Reader.ReadPackedUid(out _);
+					uid = reader.ReadPackedUid(out _);
 				}
 				
-				if (Uid >= PredefinedEventUid._WellKnownNum)
+				if (uid >= PredefinedEventUid.WellKnownNum)
 				{
-					if (!EventTypes.TryGetValue(Uid, out EventType? EventType))
+					if (!EventTypes.TryGetValue(uid, out EventType? eventType))
 					{
-						throw new Exception($"No event type registered for UID {Uid} / 0x{Uid:X4}");
+						throw new Exception($"No event type registered for UID {uid} / 0x{uid:X4}");
 					}
 				
-					GenericEvent Event = GenericEvent.Deserialize(Uid, Reader, EventType);
-					EventsRead.Add(Event);
+					GenericEvent @event = GenericEvent.Deserialize(uid, reader, eventType);
+					eventsRead.Add(@event);
 				}
 				else
 				{
-					switch (Uid)
+					switch (uid)
 					{
 						case PredefinedEventUid.NewEvent:
-							if (ThreadId != TransportPacket.ThreadIdEvents)
+							if (threadId != TransportPacket.ThreadIdEvents)
+							{
 								throw new Exception("NewEvents are only allowed on thread ID " + TransportPacket.ThreadIdEvents);
+							}
 
-							(ushort NewEventUid, EventType NewEvent) = EventType.Deserialize(Reader);
-							EventTypes[NewEventUid] = NewEvent;
+							(ushort newEventUid, EventType newEvent) = EventType.Deserialize(reader);
+							EventTypes[newEventUid] = newEvent;
 							break;
 						
 						case PredefinedEventUid.EnterScope:
@@ -151,50 +153,50 @@ namespace EpicGames.Tracing.UnrealInsights
 							break;
 						
 						case PredefinedEventUid.EnterScope_T:
-							Reader.BaseStream.Position -= 1; // Reset so UID can be read inside event's deserialize method
-							EventsRead.Add(EnterScopeEventTimestamp.Deserialize(Reader));
+							reader.BaseStream.Position -= 1; // Reset so UID can be read inside event's deserialize method
+							eventsRead.Add(EnterScopeEventTimestamp.Deserialize(reader));
 							break;
 
 						case PredefinedEventUid.LeaveScope_T:
-							Reader.BaseStream.Position -= 1; // Reset so UID can be read inside event's deserialize method
-							EventsRead.Add(LeaveScopeEventTimestamp.Deserialize(Reader));
+							reader.BaseStream.Position -= 1; // Reset so UID can be read inside event's deserialize method
+							eventsRead.Add(LeaveScopeEventTimestamp.Deserialize(reader));
 							break;
 
 						default:
-							throw new Exception($"Cannot handle unknown UID {Uid}/0x{Uid:X4}");
+							throw new Exception($"Cannot handle unknown UID {uid}/0x{uid:X4}");
 					}
 				}
 			}
 
-			return EventsRead;
+			return eventsRead;
 		}
 		
 		public Dictionary<ushort, List<ITraceEvent>> GetEventsPerUid()
 		{
-			Dictionary<ushort, List<ITraceEvent>> EventsPerUid = new Dictionary<ushort, List<ITraceEvent>>();
+			Dictionary<ushort, List<ITraceEvent>> eventsPerUid = new Dictionary<ushort, List<ITraceEvent>>();
 			
-			foreach ((ushort _, List<ITraceEvent> ThreadEvents) in EventsPerThread)
+			foreach ((ushort _, List<ITraceEvent> threadEvents) in EventsPerThread)
 			{
-				foreach (ITraceEvent Event in ThreadEvents)
+				foreach (ITraceEvent @event in threadEvents)
 				{
-					ushort Uid = EventTypes.First(x => Event.Type.Name == x.Value.Name).Key;
-					if (!EventsPerUid.TryGetValue(Uid, out List<ITraceEvent>? Events))
+					ushort uid = EventTypes.First(x => @event.Type.Name == x.Value.Name).Key;
+					if (!eventsPerUid.TryGetValue(uid, out List<ITraceEvent>? events))
 					{
-						Events = new List<ITraceEvent>();
-						EventsPerUid[Uid] = Events;
+						events = new List<ITraceEvent>();
+						eventsPerUid[uid] = events;
 					}
-					Events.Add(Event);
+					events.Add(@event);
 				}
 			}
 		
-			return EventsPerUid;
+			return eventsPerUid;
 		}
 
 		public void PrintEventSummary()
 		{
-			foreach ((ushort Uid, List<ITraceEvent> Events) in GetEventsPerUid())
+			foreach ((ushort uid, List<ITraceEvent> events) in GetEventsPerUid())
 			{
-				Console.WriteLine($"{Uid,4} {EventTypes[Uid].Name,-45} {Events.Count}");
+				Console.WriteLine($"{uid,4} {EventTypes[uid].Name,-45} {events.Count}");
 			}
 		}
 	}
diff --git a/Engine/Source/Programs/Shared/EpicGames.Tracing/UnrealInsights/StreamHeader.cs b/Engine/Source/Programs/Shared/EpicGames.Tracing/UnrealInsights/StreamHeader.cs
--- a/Engine/Source/Programs/Shared/EpicGames.Tracing/UnrealInsights/StreamHeader.cs
+++ b/Engine/Source/Programs/Shared/EpicGames.Tracing/UnrealInsights/StreamHeader.cs
@@ -10,13 +10,13 @@ namespace EpicGames.Tracing.UnrealInsights
 		public const uint MagicTrc = 1414677317; // TRCE
 		public const uint MagicTrc2 = 1414677298; // TRC2
 
-		uint MagicFourCc;
-		ushort MetadataSize;
-		ushort MetadataField0;
-		ushort ControlPort;
+		uint _magicFourCc;
+		ushort _metadataSize;
+		ushort _metadataField0;
+		ushort _controlPort;
 
-		byte TransportVersion;
-		byte ProtocolVersion;
+		byte _transportVersion;
+		byte _protocolVersion;
 
 		private StreamHeader()
 		{
@@ -24,39 +24,44 @@ namespace EpicGames.Tracing.UnrealInsights
 
 		public static StreamHeader Default()
 		{
-			StreamHeader Handshake = new StreamHeader();
-			Handshake.MagicFourCc = MagicTrc;
-			Handshake.TransportVersion = 3;
-			Handshake.ProtocolVersion = 5;
-			return Handshake;
+			StreamHeader handshake = new StreamHeader();
+			handshake._magicFourCc = MagicTrc;
+			handshake._transportVersion = 3;
+			handshake._protocolVersion = 5;
+			return handshake;
 		}
 
-		public void Serialize(BinaryWriter Writer)
+		public void Serialize(BinaryWriter writer)
 		{
-			Writer.Write(MagicTrc);
-			Writer.Write(TransportVersion);
-			Writer.Write(ProtocolVersion);
+			writer.Write(MagicTrc);
+			writer.Write(_transportVersion);
+			writer.Write(_protocolVersion);
 		}
 
-		public static StreamHeader Deserialize(BinaryReader Reader)
+		public static StreamHeader Deserialize(BinaryReader reader)
 		{
-			StreamHeader Header = new StreamHeader();
-			Header.MagicFourCc = Reader.ReadUInt32();
-			if (Header.MagicFourCc != MagicTrc && Header.MagicFourCc != MagicTrc2)
+			StreamHeader header = new StreamHeader();
+			header._magicFourCc = reader.ReadUInt32();
+			if (header._magicFourCc != MagicTrc && header._magicFourCc != MagicTrc2)
+			{
 				throw new ArgumentException("Only support magic number TRCE and TRC2");
+			}
 
-			if (Header.MagicFourCc == MagicTrc2)
+			if (header._magicFourCc == MagicTrc2)
 			{
-				Header.MetadataSize = Reader.ReadUInt16();
-				if (Header.MetadataSize != 4)
-					throw new ArgumentException("Only support metadata size of 4 bytes (got " + Header.MetadataSize + ")");
-				Header.MetadataField0 = Reader.ReadUInt16();
-				Header.ControlPort = Reader.ReadUInt16();
+				header._metadataSize = reader.ReadUInt16();
+				if (header._metadataSize != 4)
+				{
+					throw new ArgumentException("Only support metadata size of 4 bytes (got " + header._metadataSize + ")");
+				}
+
+				header._metadataField0 = reader.ReadUInt16();
+				header._controlPort = reader.ReadUInt16();
 			}
 
-			Header.TransportVersion = Reader.ReadByte();
-			Header.ProtocolVersion = Reader.ReadByte();
-			return Header;
+			header._transportVersion = reader.ReadByte();
+			header._protocolVersion = reader.ReadByte();
+			return header;
 		}
 	}
 }
\ No newline at end of file
diff --git a/Engine/Source/Programs/Shared/EpicGames.Tracing/UnrealInsights/TraceUtils.cs b/Engine/Source/Programs/Shared/EpicGames.Tracing/UnrealInsights/TraceUtils.cs
--- a/Engine/Source/Programs/Shared/EpicGames.Tracing/UnrealInsights/TraceUtils.cs
+++ b/Engine/Source/Programs/Shared/EpicGames.Tracing/UnrealInsights/TraceUtils.cs
@@ -6,39 +6,39 @@ namespace EpicGames.Tracing.UnrealInsights
 {
 	public static class TraceUtils
 	{
-		public static ulong Read7BitUint(BinaryReader Reader)
+		public static ulong Read7BitUint(BinaryReader reader)
 		{
-			ulong Value = 0;
-			ulong ByteIndex = 0;
-			bool HasMoreBytes;
+			ulong value = 0;
+			ulong byteIndex = 0;
+			bool hasMoreBytes;
 			
 			do
 			{
-				byte ByteValue = Reader.ReadByte();
-				HasMoreBytes = (ByteValue & 0x80) != 0;
-				Value |= (ulong)(ByteValue & 0x7f) << (int)(ByteIndex * 7);
-				++ByteIndex;
-			} while (HasMoreBytes);
+				byte byteValue = reader.ReadByte();
+				hasMoreBytes = (byteValue & 0x80) != 0;
+				value |= (ulong)(byteValue & 0x7f) << (int)(byteIndex * 7);
+				++byteIndex;
+			} while (hasMoreBytes);
 			
-			return Value;
+			return value;
 		}
 
-		public static int Write7BitUint(BinaryWriter Writer, ulong Value)
+		public static int Write7BitUint(BinaryWriter writer, ulong value)
 		{
-			int NumBytesWritten = 0;
+			int numBytesWritten = 0;
 			do
 			{
-				byte HasMoreBytesBit = (byte) (Value > 0x7F ? 1 : 0);
-				byte HasMoreBytes = (byte)(HasMoreBytesBit << 7);
-				byte ByteToWrite = (byte)((Value & 0x7F) | HasMoreBytes);
+				byte hasMoreBytesBit = (byte)(value > 0x7F ? 1 : 0);
+				byte hasMoreBytes = (byte)(hasMoreBytesBit << 7);
+				byte byteToWrite = (byte)((value & 0x7F) | hasMoreBytes);
 				
-				Writer.Write(ByteToWrite);
-				Value >>= 7;
-				NumBytesWritten++;
+				writer.Write(byteToWrite);
+				value >>= 7;
+				numBytesWritten++;
 				
-			} while (Value > 0);
+			} while (value > 0);
 
-			return NumBytesWritten;
+			return numBytesWritten;
 		}
 	}
 }
\ No newline at end of file
diff --git a/Engine/Source/Programs/Shared/EpicGames.Tracing/UnrealInsights/TransportPacket.cs b/Engine/Source/Programs/Shared/EpicGames.Tracing/UnrealInsights/TransportPacket.cs
--- a/Engine/Source/Programs/Shared/EpicGames.Tracing/UnrealInsights/TransportPacket.cs
+++ b/Engine/Source/Programs/Shared/EpicGames.Tracing/UnrealInsights/TransportPacket.cs
@@ -21,10 +21,10 @@ namespace EpicGames.Tracing.UnrealInsights
 		public ushort PacketSize { get; private set; }
 		public ushort ThreadIdAndMarkers { get; private set; }
 		public ushort DecodedSize { get; private set; }
-		byte[] Data = Array.Empty<byte>();
+		byte[] _data = Array.Empty<byte>();
 		public byte[] GetData()
 		{
-			return Data;
+			return _data;
 		}
 	
 		private TransportPacket()
@@ -46,32 +46,41 @@ namespace EpicGames.Tracing.UnrealInsights
 			return (ushort) (ThreadIdAndMarkers & ThreadIdMask);
 		}
 		
-		public static bool IsNormalThread(ushort ThreadId)
+		public static bool IsNormalThread(ushort threadId)
 		{
-			if (ThreadId == ThreadIdEvents)
+			if (threadId == ThreadIdEvents)
+			{
 				return false;
-			if (ThreadId == ThreadIdInternal)
+			}
+
+			if (threadId == ThreadIdInternal)
+			{
 				return false;
-			if (ThreadId == ThreadIdImportants)
+			}
+
+			if (threadId == ThreadIdImportants)
+			{
 				return false;
+			}
+
 			return true;
 		}
 
-		public void Serialize(BinaryWriter Writer, IEnumerable<(ushort, ITraceEvent)> Events)
+		public void Serialize(BinaryWriter writer, IEnumerable<(ushort, ITraceEvent)> events)
 		{
-			ushort TotalSize = (ushort) Events.Sum(x =>
+			ushort totalSize = (ushort)events.Sum(x =>
 			{
-				ITraceEvent Event = x.Item2;
-				return Event.Size;
+				ITraceEvent @event = x.Item2;
+				return @event.Size;
 			});
-			TotalSize += 4; // The two uint16 writes below are included
+			totalSize += 4; // The two uint16 writes below are included
 			
-			Writer.Write(TotalSize);
-			Writer.Write(ThreadIdAndMarkers);
+			writer.Write(totalSize);
+			writer.Write(ThreadIdAndMarkers);
 
-			foreach ((ushort Uid, ITraceEvent Event) in Events)
+			foreach ((ushort uid, ITraceEvent @event) in events)
 			{
-				Event.Serialize(Uid, Writer);
+				@event.Serialize(uid, writer);
 			}
 		}
 
@@ -80,30 +89,30 @@ namespace EpicGames.Tracing.UnrealInsights
 			return $"TransportPacket(ThreadId={GetThreadId()} PacketSize={PacketSize} IsEncoded={IsEncoded()} IsPartial={IsPartial()})";
 		}
 
-		public static TransportPacket Create(ushort PacketSize, ushort ThreadId)
+		public static TransportPacket Create(ushort packetSize, ushort threadId)
 		{
-			TransportPacket Packet = new TransportPacket();
-			Packet.PacketSize = PacketSize;
-			Packet.ThreadIdAndMarkers = ThreadId;
-			return Packet;
+			TransportPacket packet = new TransportPacket();
+			packet.PacketSize = packetSize;
+			packet.ThreadIdAndMarkers = threadId;
+			return packet;
 		}
 
-		public static TransportPacket Deserialize(BinaryReader Reader)
+		public static TransportPacket Deserialize(BinaryReader reader)
 		{
-			TransportPacket Packet = new TransportPacket();
-			Packet.PacketSize = Reader.ReadUInt16();
-			Packet.ThreadIdAndMarkers = Reader.ReadUInt16();
+			TransportPacket packet = new TransportPacket();
+			packet.PacketSize = reader.ReadUInt16();
+			packet.ThreadIdAndMarkers = reader.ReadUInt16();
 
-			int HeaderSize = sizeof(ushort) + sizeof(ushort); // PacketSize + ThreadId
-			if (Packet.IsEncoded())
+			int headerSize = sizeof(ushort) + sizeof(ushort); // PacketSize + ThreadId
+			if (packet.IsEncoded())
 			{
-				Packet.DecodedSize = Reader.ReadUInt16();
-				HeaderSize += sizeof(ushort); // sizeof(DecodedSize) 
+				packet.DecodedSize = reader.ReadUInt16();
+				headerSize += sizeof(ushort); // sizeof(DecodedSize) 
 			}
 
-			int BytesToRead = Packet.PacketSize - HeaderSize;
-			Packet.Data = Reader.ReadBytesStrict(BytesToRead);
-			return Packet;
+			int bytesToRead = packet.PacketSize - headerSize;
+			packet._data = reader.ReadBytesStrict(bytesToRead);
+			return packet;
 		}
 	}
 }
\ No newline at end of file
diff --git a/Engine/Source/Programs/Shared/EpicGames.Tracing/UnrealInsights/Writer.cs b/Engine/Source/Programs/Shared/EpicGames.Tracing/UnrealInsights/Writer.cs
--- a/Engine/Source/Programs/Shared/EpicGames.Tracing/UnrealInsights/Writer.cs
+++ b/Engine/Source/Programs/Shared/EpicGames.Tracing/UnrealInsights/Writer.cs
@@ -10,99 +10,101 @@ namespace EpicGames.Tracing.UnrealInsights
 {
 	public class UnrealInsightsWriter
 	{
-		readonly Dictionary<EventType, ushort> EventTypeToUids = new Dictionary<EventType, ushort>();
-		readonly Dictionary<ushort, EventType> UidToEventTypes = new Dictionary<ushort, EventType>();
-		readonly Dictionary<ushort, List<ITraceEvent>> ThreadToEvents = new Dictionary<ushort, List<ITraceEvent>>();
+		readonly Dictionary<EventType, ushort> _eventTypeToUids = new Dictionary<EventType, ushort>();
+		readonly Dictionary<ushort, EventType> _uidToEventTypes = new Dictionary<ushort, EventType>();
+		readonly Dictionary<ushort, List<ITraceEvent>> _threadToEvents = new Dictionary<ushort, List<ITraceEvent>>();
 		
-		private ushort UidCounter = PredefinedEventUid._WellKnownNum; // IDs below 16 are reserved for well-known events
-		private readonly object Lock = new object();
+		private ushort _uidCounter = PredefinedEventUid.WellKnownNum; // IDs below 16 are reserved for well-known events
+		private readonly object _lock = new object();
 
 		public UnrealInsightsWriter()
 		{
-			ThreadToEvents[TransportPacket.ThreadIdEvents] = new List<ITraceEvent>();
+			_threadToEvents[TransportPacket.ThreadIdEvents] = new List<ITraceEvent>();
 
-			EnterScopeEvent EnterScopeEvent = new EnterScopeEvent();
-			RegisterEventType(PredefinedEventUid.EnterScope, EnterScopeEvent.Type);
+			EnterScopeEvent enterScopeEvent = new EnterScopeEvent();
+			RegisterEventType(PredefinedEventUid.EnterScope, enterScopeEvent.Type);
 			
-			LeaveScopeEvent LeaveScopeEvent = new LeaveScopeEvent();
-			RegisterEventType(PredefinedEventUid.LeaveScope, LeaveScopeEvent.Type);
+			LeaveScopeEvent leaveScopeEvent = new LeaveScopeEvent();
+			RegisterEventType(PredefinedEventUid.LeaveScope, leaveScopeEvent.Type);
 			
-			EnterScopeEvent EnterScopeTimestampEvent = new EnterScopeEvent();
-			RegisterEventType(PredefinedEventUid.EnterScope_T, EnterScopeTimestampEvent.Type);
+			EnterScopeEvent enterScopeTimestampEvent = new EnterScopeEvent();
+			RegisterEventType(PredefinedEventUid.EnterScope_T, enterScopeTimestampEvent.Type);
 			
-			LeaveScopeEventTimestamp LeaveScopeTimestampEvent = new LeaveScopeEventTimestamp(0);
-			RegisterEventType(PredefinedEventUid.LeaveScope_T, LeaveScopeTimestampEvent.Type);
+			LeaveScopeEventTimestamp leaveScopeTimestampEvent = new LeaveScopeEventTimestamp(0);
+			RegisterEventType(PredefinedEventUid.LeaveScope_T, leaveScopeTimestampEvent.Type);
 		}
 
-		public void AddEvent(ushort ThreadId, ITraceEvent Event)
+		public void AddEvent(ushort threadId, ITraceEvent @event)
 		{
-			if (ThreadId == TransportPacket.ThreadIdEvents)
+			if (threadId == TransportPacket.ThreadIdEvents)
 			{
 				throw new ArgumentException("Cannot add events directly to new-event thread");
 			}
 			
-			lock (Lock)
+			lock (_lock)
 			{
-				if (!ThreadToEvents.TryGetValue(ThreadId, out List<ITraceEvent>? Events))
+				if (!_threadToEvents.TryGetValue(threadId, out List<ITraceEvent>? events))
 				{
-					Events = new List<ITraceEvent>();
-					ThreadToEvents[ThreadId] = Events;
+					events = new List<ITraceEvent>();
+					_threadToEvents[threadId] = events;
 				}
 
-				if (!EventTypeToUids.ContainsKey(Event.Type))
+				if (!_eventTypeToUids.ContainsKey(@event.Type))
 				{
-					ushort NewUid = UidCounter++;
-					RegisterEventType(NewUid, Event.Type);
-					ThreadToEvents[TransportPacket.ThreadIdEvents].Add(Event);
+					ushort newUid = _uidCounter++;
+					RegisterEventType(newUid, @event.Type);
+					_threadToEvents[TransportPacket.ThreadIdEvents].Add(@event);
 				}
 				
-				Events.Add(Event);
+				events.Add(@event);
 			}
 		}
 
-		private void RegisterEventType(ushort Uid, EventType EventType)
+		private void RegisterEventType(ushort uid, EventType eventType)
 		{
-			UidToEventTypes[Uid] = EventType;
-			EventTypeToUids[EventType] = Uid;
+			_uidToEventTypes[uid] = eventType;
+			_eventTypeToUids[eventType] = uid;
 		}
 
-		public void Write(BinaryWriter Writer)
+		public void Write(BinaryWriter writer)
 		{
-			lock (Lock)
+			lock (_lock)
 			{
-				WriteHeader(Writer);
-				WriteThread(TransportPacket.ThreadIdImportants, Writer);
+				WriteHeader(writer);
+				WriteThread(TransportPacket.ThreadIdImportants, writer);
 
-				foreach (ushort ThreadId in ThreadToEvents.Keys)
+				foreach (ushort threadId in _threadToEvents.Keys)
 				{
-					if (ThreadId == TransportPacket.ThreadIdEvents || ThreadId == TransportPacket.ThreadIdImportants)
+					if (threadId == TransportPacket.ThreadIdEvents || threadId == TransportPacket.ThreadIdImportants)
+					{
 						continue;
-					
-					WriteThread(ThreadId, Writer);
+					}
+
+					WriteThread(threadId, writer);
 				}
 			}
 		}
 
-		private void WriteHeader(BinaryWriter Writer)
+		private void WriteHeader(BinaryWriter writer)
 		{
-			StreamHeader.Default().Serialize(Writer);
-			TransportPacket TransportPacket = TransportPacket.Create(0, 0);
-			IEnumerable<(ushort Uid, ITraceEvent Event)> NewEvents = ThreadToEvents[TransportPacket.ThreadIdEvents].Select(e =>
+			StreamHeader.Default().Serialize(writer);
+			TransportPacket transportPacket = TransportPacket.Create(0, 0);
+			IEnumerable<(ushort Uid, ITraceEvent Event)> newEvents = _threadToEvents[TransportPacket.ThreadIdEvents].Select(e =>
 			{
-				ushort NewEventUid = EventTypeToUids[e.Type];
-				ITraceEvent MyEventType = e.Type;
-				return (NewEventUid, MyEventType);
+				ushort newEventUid = _eventTypeToUids[e.Type];
+				ITraceEvent myEventType = e.Type;
+				return (newEventUid, myEventType);
 			});
-			TransportPacket.Serialize(Writer, NewEvents);
+			transportPacket.Serialize(writer, newEvents);
 		}
 
-		private void WriteThread(ushort ThreadId, BinaryWriter Writer)
+		private void WriteThread(ushort threadId, BinaryWriter writer)
 		{
-			foreach (ITraceEvent Event in ThreadToEvents[ThreadId])
+			foreach (ITraceEvent @event in _threadToEvents[threadId])
 			{
-				ushort Uid = EventTypeToUids[Event.Type];
-				TransportPacket TransportPacket = TransportPacket.Create(0, ThreadId);
-				TransportPacket.Serialize(Writer, new [] { (Uid, Event) });
+				ushort uid = _eventTypeToUids[@event.Type];
+				TransportPacket transportPacket = TransportPacket.Create(0, threadId);
+				transportPacket.Serialize(writer, new [] { (uid, @event) });
 			}
 		}
 	}
