EpicGames.Tracing.Tests: More fixed variable naming conventions

#rnx

diff --git a/Engine/Source/Programs/Shared/EpicGames.Tracing.Tests/UnrealInsights/UnrealInsightsTest.cs b/Engine/Source/Programs/Shared/EpicGames.Tracing.Tests/UnrealInsights/UnrealInsightsTest.cs
--- a/Engine/Source/Programs/Shared/EpicGames.Tracing.Tests/UnrealInsights/UnrealInsightsTest.cs
+++ b/Engine/Source/Programs/Shared/EpicGames.Tracing.Tests/UnrealInsights/UnrealInsightsTest.cs
@@ -13,27 +13,27 @@ namespace EpicGames.Tracing.Tests.UnrealInsights
 {
 	public class StubTraceEvent : ITraceEvent
 	{
-		readonly uint Val1;
-		readonly byte Val2;
+		readonly uint _val1;
+		readonly byte _val2;
 
-		public StubTraceEvent(uint Val1, byte Val2)
+		public StubTraceEvent(uint val1, byte val2)
 		{
-			this.Val1 = Val1;
-			this.Val2 = Val2;
+			_val1 = val1;
+			_val2 = val2;
 		}
 
 		public ushort Size => sizeof(uint) + sizeof(byte);
 		public EventType Type => EventType.WellKnown(15, "StubEvent");
 
-		public void Serialize(ushort Uid, BinaryWriter Writer)
+		public void Serialize(ushort uid, BinaryWriter writer)
 		{
-			Writer.Write(Val1);
-			Writer.Write(Val2);
+			writer.Write(_val1);
+			writer.Write(_val2);
 		}
 		
-		public static StubTraceEvent Deserialize(BinaryReader Reader)
+		public static StubTraceEvent Deserialize(BinaryReader reader)
 		{
-			return new StubTraceEvent(Reader.ReadUInt32(), Reader.ReadByte());
+			return new StubTraceEvent(reader.ReadUInt32(), reader.ReadByte());
 		}
 	}
 	
@@ -43,26 +43,26 @@ namespace EpicGames.Tracing.Tests.UnrealInsights
 		[TestMethod]
 		public void SerializeDeserialize()
 		{
-			uint Test1 = 222;
-			uint Test2 = 333;
-			ulong Test3 = 444;
-			ushort EventSize = 4 + 4 + 8;
+			uint test1 = 222;
+			uint test2 = 333;
+			ulong test3 = 444;
+			ushort eventSize = 4 + 4 + 8;
 			
-			using MemoryStream Ms = new MemoryStream();
-			using BinaryWriter Writer = new BinaryWriter(Ms);
-			new TraceImportantEventHeader(1000, EventSize).Serialize(Writer);
-			Writer.Write(Test1);
-			Writer.Write(Test2);
-			Writer.Write(Test3);
+			using MemoryStream ms = new MemoryStream();
+			using BinaryWriter writer = new BinaryWriter(ms);
+			new TraceImportantEventHeader(1000, eventSize).Serialize(writer);
+			writer.Write(test1);
+			writer.Write(test2);
+			writer.Write(test3);
 			
-			Ms.Position = 0;
-			using BinaryReader Reader = new BinaryReader(Ms);
-			TraceImportantEventHeader EventHeader = TraceImportantEventHeader.Deserialize(Reader);
-			Assert.AreEqual(1000, EventHeader.Uid);
-			Assert.AreEqual(EventSize, EventHeader.EventSize);
-			Assert.AreEqual(Test1, Reader.ReadUInt32());
-			Assert.AreEqual(Test2, Reader.ReadUInt32());
-			Assert.AreEqual(Test3, Reader.ReadUInt64());
+			ms.Position = 0;
+			using BinaryReader reader = new BinaryReader(ms);
+			TraceImportantEventHeader eventHeader = TraceImportantEventHeader.Deserialize(reader);
+			Assert.AreEqual(1000, eventHeader.Uid);
+			Assert.AreEqual(eventSize, eventHeader.EventSize);
+			Assert.AreEqual(test1, reader.ReadUInt32());
+			Assert.AreEqual(test2, reader.ReadUInt32());
+			Assert.AreEqual(test3, reader.ReadUInt64());
 		}
 	}
 
@@ -72,61 +72,61 @@ namespace EpicGames.Tracing.Tests.UnrealInsights
 		[TestMethod]
 		public void WriteUtraceFile()
 		{
-			using MemoryStream Ms = new MemoryStream();
-			using BinaryWriter BinaryWriter = new BinaryWriter(Ms);
+			using MemoryStream ms = new MemoryStream();
+			using BinaryWriter binaryWriter = new BinaryWriter(ms);
 
-			UnrealInsightsWriter Writer = new UnrealInsightsWriter();
+			UnrealInsightsWriter writer = new UnrealInsightsWriter();
 			
-			TraceNewTraceEvent NewTraceEvent1 = new TraceNewTraceEvent(1001, 2001, 31, 8);
-			TraceNewTraceEvent NewTraceEvent2 = new TraceNewTraceEvent(1002, 2002, 32, 8);
-			TraceNewTraceEvent NewTraceEvent3 = new TraceNewTraceEvent(1003, 2003, 33, 8);
-			CpuProfilerEventSpecEvent Cpu1 = new CpuProfilerEventSpecEvent(400, "MyCpuEventSpecEvent1");
-
-			Writer.AddEvent(TransportPacket.ThreadIdImportants, NewTraceEvent1);
-			Writer.AddEvent(TransportPacket.ThreadIdImportants, NewTraceEvent2);
-			Writer.AddEvent(TransportPacket.ThreadIdImportants, NewTraceEvent3);
-			Writer.AddEvent(TransportPacket.ThreadIdImportants, Cpu1);
-			Writer.Write(BinaryWriter);
-
-			Ms.Position = 0;
-
-			UnrealInsightsReader Reader = new UnrealInsightsReader();
-			Reader.Read(Ms);
-			Reader.PrintEventSummary();
-			Assert.AreEqual(4, Reader.EventTypes.Count);
+			TraceNewTraceEvent newTraceEvent1 = new TraceNewTraceEvent(1001, 2001, 31, 8);
+			TraceNewTraceEvent newTraceEvent2 = new TraceNewTraceEvent(1002, 2002, 32, 8);
+			TraceNewTraceEvent newTraceEvent3 = new TraceNewTraceEvent(1003, 2003, 33, 8);
+			CpuProfilerEventSpecEvent cpu1 = new CpuProfilerEventSpecEvent(400, "MyCpuEventSpecEvent1");
+
+			writer.AddEvent(TransportPacket.ThreadIdImportants, newTraceEvent1);
+			writer.AddEvent(TransportPacket.ThreadIdImportants, newTraceEvent2);
+			writer.AddEvent(TransportPacket.ThreadIdImportants, newTraceEvent3);
+			writer.AddEvent(TransportPacket.ThreadIdImportants, cpu1);
+			writer.Write(binaryWriter);
+
+			ms.Position = 0;
+
+			UnrealInsightsReader reader = new UnrealInsightsReader();
+			reader.Read(ms);
+			reader.PrintEventSummary();
+			Assert.AreEqual(4, reader.EventTypes.Count);
 			
-			Assert.AreEqual(0, Reader.EventsPerThread[TransportPacket.ThreadIdEvents].Count);
-			Assert.AreEqual(4, Reader.EventsPerThread[TransportPacket.ThreadIdImportants].Count);
+			Assert.AreEqual(0, reader.EventsPerThread[TransportPacket.ThreadIdEvents].Count);
+			Assert.AreEqual(4, reader.EventsPerThread[TransportPacket.ThreadIdImportants].Count);
 			
 			{
-				GenericEvent Event = (GenericEvent) Reader.EventsPerThread[1][0];
-				Field[] Fields = Event.GetFields();
-				Assert.AreEqual(1001, Fields[0].Long!.Value);
-				Assert.AreEqual(2001, Fields[1].Long!.Value);
-				Assert.AreEqual(31, Fields[2].Int!.Value);
-				Assert.AreEqual(8, Fields[3].Int!.Value);
+				GenericEvent @event = (GenericEvent)reader.EventsPerThread[1][0];
+				Field[] fields = @event.GetFields();
+				Assert.AreEqual(1001, fields[0].Long!.Value);
+				Assert.AreEqual(2001, fields[1].Long!.Value);
+				Assert.AreEqual(31, fields[2].Int!.Value);
+				Assert.AreEqual(8, fields[3].Int!.Value);
 			}
 			{
-				GenericEvent Event = (GenericEvent) Reader.EventsPerThread[1][1];
-				Field[] Fields = Event.GetFields();
-				Assert.AreEqual(1002, Fields[0].Long!.Value);
-				Assert.AreEqual(2002, Fields[1].Long!.Value);
-				Assert.AreEqual(32, Fields[2].Int!.Value);
-				Assert.AreEqual(8, Fields[3].Int!.Value);
+				GenericEvent @event = (GenericEvent)reader.EventsPerThread[1][1];
+				Field[] fields = @event.GetFields();
+				Assert.AreEqual(1002, fields[0].Long!.Value);
+				Assert.AreEqual(2002, fields[1].Long!.Value);
+				Assert.AreEqual(32, fields[2].Int!.Value);
+				Assert.AreEqual(8, fields[3].Int!.Value);
 			}
 			{
-				GenericEvent Event = (GenericEvent) Reader.EventsPerThread[1][2];
-				Field[] Fields = Event.GetFields();
-				Assert.AreEqual(1003, Fields[0].Long!.Value);
-				Assert.AreEqual(2003, Fields[1].Long!.Value);
-				Assert.AreEqual(33, Fields[2].Int!.Value);
-				Assert.AreEqual(8, Fields[3].Int!.Value);
+				GenericEvent @event = (GenericEvent)reader.EventsPerThread[1][2];
+				Field[] fields = @event.GetFields();
+				Assert.AreEqual(1003, fields[0].Long!.Value);
+				Assert.AreEqual(2003, fields[1].Long!.Value);
+				Assert.AreEqual(33, fields[2].Int!.Value);
+				Assert.AreEqual(8, fields[3].Int!.Value);
 			}
 			{
-				GenericEvent Event = (GenericEvent) Reader.EventsPerThread[1][3];
-				Field[] Fields = Event.GetFields();
-				Assert.AreEqual(400, Fields[0].Int!.Value);
-				Assert.AreEqual("MyCpuEventSpecEvent1", Fields[1].String!);
+				GenericEvent @event = (GenericEvent)reader.EventsPerThread[1][3];
+				Field[] fields = @event.GetFields();
+				Assert.AreEqual(400, fields[0].Int!.Value);
+				Assert.AreEqual("MyCpuEventSpecEvent1", fields[1].String!);
 			}
 		}
 		
@@ -134,87 +134,87 @@ namespace EpicGames.Tracing.Tests.UnrealInsights
 		[Ignore]
 		public void WriteUtraceExample()
 		{
-			string FileName = "d:\\Temp\\mytrace.utrace";
+			string fileName = "d:\\Temp\\mytrace.utrace";
 			
-			FileStream Fs = File.Open(FileName, FileMode.Create);
-			using BinaryWriter BinaryWriter = new BinaryWriter(Fs);
-			UnrealInsightsWriter Writer = new UnrealInsightsWriter();
+			FileStream fs = File.Open(fileName, FileMode.Create);
+			using BinaryWriter binaryWriter = new BinaryWriter(fs);
+			UnrealInsightsWriter writer = new UnrealInsightsWriter();
 			
-			byte[] CpuBatchData1 = GenericEventTest.HexStringToBytes("87 A7 EB AE 8A D1 0B 01 A7 02 02 04 FB 23 03 A4 0A 61 04 0B 05 F2 80 42 AA 1C EB 5C 06 C0 BC 03 17 07 02 47 08 C2 02 AB 15 09 AA 03 07 0A AE 01 05 0B D8 0D 05 0C B7 01 0D D8 03 00 03 0E A4 03 6F 0F 89 FF 03 10 F3 C4 04 11 60 06 A7 02 10 AB 19 12 B3 04 10 9F 11 13 0A 04 02 04 ED 03 10 91 11 14 80 01 02 97 02 10 A9 D6 0B 15 10 06 DD 03 10 CB 85 04 16 12 06 D3 01 10 89 16 17 0C 02 D6 02 07 18 DA 84 28 99 01 19 18 05 1A 0A A5 10 1B CD 01 1C 06 88 6A C5 49 1D B8 EB 24 A9 17 20 AF 01 10 CD F9 01 21 71 22 F2 AE 12 85 04 23 D1 0D 24 85 74 25 EA 58 91 0D 25 80 E6 02 A9 1C 25 88 CD 06 A7 27 25 D0 DA 09 8B 20 25 DE AD 0B AF 33 25 C0 A9 09 ED 23 25 C4 81 0A F5 2D 25 9A EC 0A F3 35 25 D4 CF 03 FF 19 25 BE 92 01 B1 0D 25 98 22");
-
-			ulong CycleFrequency = 10000000;
-			TraceNewTraceEvent NewTrace1 = new TraceNewTraceEvent(25582215261913, CycleFrequency, 21069, 8);
-			TraceThreadInfoEvent ThreadInfo1 = new TraceThreadInfoEvent(2, 171380, -1, "GameThread");
-			TraceThreadInfoEvent ThreadInfo2 = new TraceThreadInfoEvent(3, 0, 2147483647, "Trace");
-			DiagnosticsSession2Event DiagnosticsSession = new DiagnosticsSession2Event("Finally it works", "HELLO WORLD", "QAGame.uproject -trace=cpu -game", "++UE5+Main", "++UE5+Main-CL-17442524", 1000, 3, 4);
-			CpuProfilerEventSpecEvent CpuEventSpec1 = new CpuProfilerEventSpecEvent(1, "MyCpuEventSpecEvent1");
-			CpuProfilerEventSpecEvent CpuEventSpec2 = new CpuProfilerEventSpecEvent(2, "MyCpuEventSpecEvent2");
-			CpuProfilerEventSpecEvent CpuEventSpec3 = new CpuProfilerEventSpecEvent(3, "MyCpuEventSpecEvent3");
-			CpuProfilerEventSpecEvent CpuEventSpec4 = new CpuProfilerEventSpecEvent(4, "MyCpuEventSpecEvent4");
-			CpuProfilerEventSpecEvent CpuEventSpec5 = new CpuProfilerEventSpecEvent(5, "MyCpuEventSpecEvent5");
-			CpuProfilerEventSpecEvent CpuEventSpec6 = new CpuProfilerEventSpecEvent(6, "MyCpuEventSpecEvent6");
-			CpuProfilerEventSpecEvent CpuEventSpec7 = new CpuProfilerEventSpecEvent(7, "MyCpuEventSpecEvent7");
-			CpuProfilerEventSpecEvent CpuEventSpec8 = new CpuProfilerEventSpecEvent(8, "MyCpuEventSpecEvent8");
-			CpuProfilerEventSpecEvent CpuEventSpecTesting = new CpuProfilerEventSpecEvent(9, "TestingSome");
-			CpuProfilerEventBatchEvent CpuBatch1 = new CpuProfilerEventBatchEvent(CpuBatchData1);
-
-			Writer.AddEvent(TransportPacket.ThreadIdImportants, NewTrace1);
-			Writer.AddEvent(TransportPacket.ThreadIdImportants, ThreadInfo1);
-			Writer.AddEvent(TransportPacket.ThreadIdImportants, ThreadInfo2);
-			Writer.AddEvent(TransportPacket.ThreadIdImportants, DiagnosticsSession);
+			byte[] cpuBatchData1 = GenericEventTest.HexStringToBytes("87 A7 EB AE 8A D1 0B 01 A7 02 02 04 FB 23 03 A4 0A 61 04 0B 05 F2 80 42 AA 1C EB 5C 06 C0 BC 03 17 07 02 47 08 C2 02 AB 15 09 AA 03 07 0A AE 01 05 0B D8 0D 05 0C B7 01 0D D8 03 00 03 0E A4 03 6F 0F 89 FF 03 10 F3 C4 04 11 60 06 A7 02 10 AB 19 12 B3 04 10 9F 11 13 0A 04 02 04 ED 03 10 91 11 14 80 01 02 97 02 10 A9 D6 0B 15 10 06 DD 03 10 CB 85 04 16 12 06 D3 01 10 89 16 17 0C 02 D6 02 07 18 DA 84 28 99 01 19 18 05 1A 0A A5 10 1B CD 01 1C 06 88 6A C5 49 1D B8 EB 24 A9 17 20 AF 01 10 CD F9 01 21 71 22 F2 AE 12 85 04 23 D1 0D 24 85 74 25 EA 58 91 0D 25 80 E6 02 A9 1C 25 88 CD 06 A7 27 25 D0 DA 09 8B 20 25 DE AD 0B AF 33 25 C0 A9 09 ED 23 25 C4 81 0A F5 2D 25 9A EC 0A F3 35 25 D4 CF 03 FF 19 25 BE 92 01 B1 0D 25 98 22");
+
+			ulong cycleFrequency = 10000000;
+			TraceNewTraceEvent newTrace1 = new TraceNewTraceEvent(25582215261913, cycleFrequency, 21069, 8);
+			TraceThreadInfoEvent threadInfo1 = new TraceThreadInfoEvent(2, 171380, -1, "GameThread");
+			TraceThreadInfoEvent threadInfo2 = new TraceThreadInfoEvent(3, 0, 2147483647, "Trace");
+			DiagnosticsSession2Event diagnosticsSession = new DiagnosticsSession2Event("Finally it works", "HELLO WORLD", "QAGame.uproject -trace=cpu -game", "++UE5+Main", "++UE5+Main-CL-17442524", 1000, 3, 4);
+			CpuProfilerEventSpecEvent cpuEventSpec1 = new CpuProfilerEventSpecEvent(1, "MyCpuEventSpecEvent1");
+			CpuProfilerEventSpecEvent cpuEventSpec2 = new CpuProfilerEventSpecEvent(2, "MyCpuEventSpecEvent2");
+			CpuProfilerEventSpecEvent cpuEventSpec3 = new CpuProfilerEventSpecEvent(3, "MyCpuEventSpecEvent3");
+			CpuProfilerEventSpecEvent cpuEventSpec4 = new CpuProfilerEventSpecEvent(4, "MyCpuEventSpecEvent4");
+			CpuProfilerEventSpecEvent cpuEventSpec5 = new CpuProfilerEventSpecEvent(5, "MyCpuEventSpecEvent5");
+			CpuProfilerEventSpecEvent cpuEventSpec6 = new CpuProfilerEventSpecEvent(6, "MyCpuEventSpecEvent6");
+			CpuProfilerEventSpecEvent cpuEventSpec7 = new CpuProfilerEventSpecEvent(7, "MyCpuEventSpecEvent7");
+			CpuProfilerEventSpecEvent cpuEventSpec8 = new CpuProfilerEventSpecEvent(8, "MyCpuEventSpecEvent8");
+			CpuProfilerEventSpecEvent cpuEventSpecTesting = new CpuProfilerEventSpecEvent(9, "TestingSome");
+			CpuProfilerEventBatchEvent cpuBatch1 = new CpuProfilerEventBatchEvent(cpuBatchData1);
+
+			writer.AddEvent(TransportPacket.ThreadIdImportants, newTrace1);
+			writer.AddEvent(TransportPacket.ThreadIdImportants, threadInfo1);
+			writer.AddEvent(TransportPacket.ThreadIdImportants, threadInfo2);
+			writer.AddEvent(TransportPacket.ThreadIdImportants, diagnosticsSession);
 			
-			Writer.AddEvent(TransportPacket.ThreadIdImportants, CpuEventSpec1);
-			Writer.AddEvent(TransportPacket.ThreadIdImportants, CpuEventSpec2);
-			Writer.AddEvent(TransportPacket.ThreadIdImportants, CpuEventSpec3);
-			Writer.AddEvent(TransportPacket.ThreadIdImportants, CpuEventSpec4);
-			Writer.AddEvent(TransportPacket.ThreadIdImportants, CpuEventSpec5);
-			Writer.AddEvent(TransportPacket.ThreadIdImportants, CpuEventSpec6);
-			Writer.AddEvent(TransportPacket.ThreadIdImportants, CpuEventSpec7);
-			Writer.AddEvent(TransportPacket.ThreadIdImportants, CpuEventSpec8);
-			Writer.AddEvent(TransportPacket.ThreadIdImportants, CpuEventSpecTesting);
+			writer.AddEvent(TransportPacket.ThreadIdImportants, cpuEventSpec1);
+			writer.AddEvent(TransportPacket.ThreadIdImportants, cpuEventSpec2);
+			writer.AddEvent(TransportPacket.ThreadIdImportants, cpuEventSpec3);
+			writer.AddEvent(TransportPacket.ThreadIdImportants, cpuEventSpec4);
+			writer.AddEvent(TransportPacket.ThreadIdImportants, cpuEventSpec5);
+			writer.AddEvent(TransportPacket.ThreadIdImportants, cpuEventSpec6);
+			writer.AddEvent(TransportPacket.ThreadIdImportants, cpuEventSpec7);
+			writer.AddEvent(TransportPacket.ThreadIdImportants, cpuEventSpec8);
+			writer.AddEvent(TransportPacket.ThreadIdImportants, cpuEventSpecTesting);
 			
-			Writer.AddEvent(2, CpuBatch1);
+			writer.AddEvent(2, cpuBatch1);
 			
-			CpuProfilerSerializer CpuSerializerThread3 = new CpuProfilerSerializer(CycleFrequency);
-			CpuSerializerThread3.ScopeEvents.Add(new CpuProfilerScopeEvent(25582215261913 + 400 * 1000, true, CpuEventSpecTesting.Id));
-			CpuSerializerThread3.ScopeEvents.Add(new CpuProfilerScopeEvent(25582215261913 + (CycleFrequency * CycleFrequency), false, null));
+			CpuProfilerSerializer cpuSerializerThread3 = new CpuProfilerSerializer(cycleFrequency);
+			cpuSerializerThread3.ScopeEvents.Add(new CpuProfilerScopeEvent(25582215261913 + 400 * 1000, true, cpuEventSpecTesting.Id));
+			cpuSerializerThread3.ScopeEvents.Add(new CpuProfilerScopeEvent(25582215261913 + (cycleFrequency * cycleFrequency), false, null));
 			
-			List<byte[]> EventBatchData = CpuSerializerThread3.Write();
-			foreach (byte[] Data in EventBatchData)
+			List<byte[]> eventBatchData = cpuSerializerThread3.Write();
+			foreach (byte[] data in eventBatchData)
 			{
-				Writer.AddEvent(3, new CpuProfilerEventBatchEvent(Data));
+				writer.AddEvent(3, new CpuProfilerEventBatchEvent(data));
 			}
-			Writer.Write(BinaryWriter);
+			writer.Write(binaryWriter);
 
-			Fs.Close();
+			fs.Close();
 
-			UnrealInsightsReader Reader = new UnrealInsightsReader();
-			using FileStream Stream = File.Open(FileName, FileMode.Open);
-			Reader.Read(Stream);
-			Reader.PrintEventSummary();
+			UnrealInsightsReader reader = new UnrealInsightsReader();
+			using FileStream stream = File.Open(fileName, FileMode.Open);
+			reader.Read(stream);
+			reader.PrintEventSummary();
 
-			Assert.AreEqual(0, Reader.EventsPerThread[TransportPacket.ThreadIdEvents].Count);
+			Assert.AreEqual(0, reader.EventsPerThread[TransportPacket.ThreadIdEvents].Count);
 		}
 		
 		[TestMethod]
 		public void ReadUtraceFile()
 		{
-			string ExampleDecompTrace = "UnrealInsights/example_trace.decomp.utrace";
+			string exampleDecompTrace = "UnrealInsights/example_trace.decomp.utrace";
 
-			UnrealInsightsReader Reader = new UnrealInsightsReader();
-			using FileStream Stream = File.Open(ExampleDecompTrace, FileMode.Open);
-			Reader.Read(Stream);
-			Reader.PrintEventSummary();
+			UnrealInsightsReader reader = new UnrealInsightsReader();
+			using FileStream stream = File.Open(exampleDecompTrace, FileMode.Open);
+			reader.Read(stream);
+			reader.PrintEventSummary();
 			
-			Assert.AreEqual(6281, Reader.NumTransportPacketsRead);
-			Dictionary<ushort,List<ITraceEvent>> EventsPerUid = Reader.GetEventsPerUid();
+			Assert.AreEqual(6281, reader.NumTransportPacketsRead);
+			Dictionary<ushort, List<ITraceEvent>> eventsPerUid = reader.GetEventsPerUid();
 
-			void AssertEvent(ushort Uid, string ExpectedName, int ExpectedCount)
+			void AssertEvent(ushort uid, string expectedName, int expectedCount)
 			{
-				Assert.IsTrue(Reader.EventTypes.ContainsKey(Uid));
-				Assert.IsTrue(EventsPerUid.ContainsKey(Uid));
-				Assert.AreEqual(ExpectedName, Reader.EventTypes[Uid].Name);
-				Assert.AreEqual(ExpectedCount, EventsPerUid[Uid].Count);
+				Assert.IsTrue(reader.EventTypes.ContainsKey(uid));
+				Assert.IsTrue(eventsPerUid.ContainsKey(uid));
+				Assert.AreEqual(expectedName, reader.EventTypes[uid].Name);
+				Assert.AreEqual(expectedCount, eventsPerUid[uid].Count);
 			}
 
 			AssertEvent(16, "$Trace.NewTrace", 1);
