More C# analysis fixes

#rnx

diff --git a/Engine/Source/Programs/Shared/EpicGames.Horde/Compute/AgentMessageHandler.cs b/Engine/Source/Programs/Shared/EpicGames.Horde/Compute/AgentMessageHandler.cs
--- a/Engine/Source/Programs/Shared/EpicGames.Horde/Compute/AgentMessageHandler.cs
+++ b/Engine/Source/Programs/Shared/EpicGames.Horde/Compute/AgentMessageHandler.cs
@@ -410,7 +410,7 @@ namespace EpicGames.Horde.Compute
 					}
 
 					ReadOnlySpan<byte> line = span.Slice(0, newlineIdx);
-					if (line.Length > 0 && line[line.Length - 1] == (byte)'\r')
+					if (line.Length > 0 && line[^1] == (byte)'\r')
 					{
 						line = line.Slice(0, line.Length - 1);
 					}
diff --git a/Engine/Source/Programs/Shared/EpicGames.Horde/HordeClient.cs b/Engine/Source/Programs/Shared/EpicGames.Horde/HordeClient.cs
--- a/Engine/Source/Programs/Shared/EpicGames.Horde/HordeClient.cs
+++ b/Engine/Source/Programs/Shared/EpicGames.Horde/HordeClient.cs
@@ -84,9 +84,9 @@ namespace EpicGames.Horde
 		/// <inheritdoc/>
 		public IStorageClient CreateStorageClient(string basePath)
 		{
-			Func<HttpClient> createClient = () => _httpClientFactory.CreateClient(HordeHttpClient.HttpClientName);
-			Func<HttpClient> createUploadRedirectClient = () => _httpClientFactory.CreateClient(HordeHttpClient.UploadRedirectHttpClientName);
-			HttpStorageBackend httpStorageBackend = new HttpStorageBackend(basePath, createClient, createUploadRedirectClient, _loggerFactory.CreateLogger<HttpStorageBackend>());
+			HttpClient CreateClient() => _httpClientFactory.CreateClient(HordeHttpClient.HttpClientName);
+			HttpClient CreateUploadRedirectClient() => _httpClientFactory.CreateClient(HordeHttpClient.UploadRedirectHttpClientName);
+			HttpStorageBackend httpStorageBackend = new HttpStorageBackend(basePath, CreateClient, CreateUploadRedirectClient, _loggerFactory.CreateLogger<HttpStorageBackend>());
 			return new BundleStorageClient(httpStorageBackend, _bundleCache, _hordeOptions.Bundle, _loggerFactory.CreateLogger<BundleStorageClient>());
 		}
 	}
diff --git a/Engine/Source/Programs/Shared/EpicGames.Horde/Storage/StorageBackendCache.cs b/Engine/Source/Programs/Shared/EpicGames.Horde/Storage/StorageBackendCache.cs
--- a/Engine/Source/Programs/Shared/EpicGames.Horde/Storage/StorageBackendCache.cs
+++ b/Engine/Source/Programs/Shared/EpicGames.Horde/Storage/StorageBackendCache.cs
@@ -223,7 +223,7 @@ namespace EpicGames.Horde.Storage
 		/// <param name="logger">Logger for error/warning messages</param>
 		public StorageBackendCache(DirectoryReference? cacheDir, long? maxSize, ILogger logger)
 		{
-			cacheDir ??= new DirectoryReference(Path.Combine(Path.GetTempPath(), $"horde-{Guid.NewGuid().ToString("n")}"));
+			cacheDir ??= new DirectoryReference(Path.Combine(Path.GetTempPath(), $"horde-{Guid.NewGuid():n}"));
 			FileUtils.ForceDeleteDirectoryContents(cacheDir);
 
 			_memoryMappedFileCache = new MemoryMappedFileCache();
diff --git a/Engine/Source/Programs/Shared/EpicGames.Redis/Converters/RedisClassConverter.cs b/Engine/Source/Programs/Shared/EpicGames.Redis/Converters/RedisClassConverter.cs
--- a/Engine/Source/Programs/Shared/EpicGames.Redis/Converters/RedisClassConverter.cs
+++ b/Engine/Source/Programs/Shared/EpicGames.Redis/Converters/RedisClassConverter.cs
@@ -158,7 +158,7 @@ namespace EpicGames.Redis.Converters
 		{
 			if (data == null)
 			{
-				return default(T)!;
+				return default!;
 			}
 
 			object?[] arguments = new object?[_properties.Length];
diff --git a/Engine/Source/Programs/Shared/EpicGames.Redis/RedisString.cs b/Engine/Source/Programs/Shared/EpicGames.Redis/RedisString.cs
--- a/Engine/Source/Programs/Shared/EpicGames.Redis/RedisString.cs
+++ b/Engine/Source/Programs/Shared/EpicGames.Redis/RedisString.cs
@@ -76,7 +76,7 @@ namespace EpicGames.Redis
 			RedisValue value = await target.StringGetAsync(key.Inner, flags);
 			if (value.IsNullOrEmpty)
 			{
-				return default(TValue);
+				return default;
 			}
 			return RedisSerializer.Deserialize<TValue>(value)!;
 		}
diff --git a/Engine/Source/Programs/Shared/EpicGames.Redis/RedisSubscription.cs b/Engine/Source/Programs/Shared/EpicGames.Redis/RedisSubscription.cs
--- a/Engine/Source/Programs/Shared/EpicGames.Redis/RedisSubscription.cs
+++ b/Engine/Source/Programs/Shared/EpicGames.Redis/RedisSubscription.cs
@@ -97,15 +97,15 @@ namespace EpicGames.Redis
 		/// <inheritdoc cref="ISubscriber.SubscribeAsync(RedisChannel, Action{RedisChannel, RedisValue}, CommandFlags)"/>
 		public static Task<RedisSubscription> SubscribeAsync<T>(this IConnectionMultiplexer connection, RedisChannel<T> channel, Action<T> handler, CommandFlags flags = CommandFlags.None)
 		{
-			Action<RedisChannel, RedisValue> action = (_, v) => handler(RedisSerializer.Deserialize<T>(v));
-			return SubscribeAsync(connection, channel.Channel, action, flags);
+			void Action(RedisChannel _, RedisValue v) => handler(RedisSerializer.Deserialize<T>(v));
+			return SubscribeAsync(connection, channel.Channel, Action, flags);
 		}
 
 		/// <inheritdoc cref="ISubscriber.SubscribeAsync(RedisChannel, Action{RedisChannel, RedisValue}, CommandFlags)"/>
 		public static Task<RedisSubscription> SubscribeAsync<T>(this IConnectionMultiplexer connection, RedisChannel<T> channel, Action<RedisChannel<T>, T> handler, CommandFlags flags = CommandFlags.None)
 		{
-			Action<RedisChannel, RedisValue> action = (_, v) => handler(channel, RedisSerializer.Deserialize<T>(v));
-			return SubscribeAsync(connection, channel.Channel, action, flags);
+			void Action(RedisChannel _, RedisValue v) => handler(channel, RedisSerializer.Deserialize<T>(v));
+			return SubscribeAsync(connection, channel.Channel, Action, flags);
 		}
 
 		#endregion
