[UBA]
* Added experimental code for writing all .obj files compressed on disk.

diff --git a/Engine/Source/Programs/Shared/EpicGames.UBA/ISessionServer.cs b/Engine/Source/Programs/Shared/EpicGames.UBA/ISessionServer.cs
--- a/Engine/Source/Programs/Shared/EpicGames.UBA/ISessionServer.cs
+++ b/Engine/Source/Programs/Shared/EpicGames.UBA/ISessionServer.cs
@@ -54,6 +54,8 @@ namespace EpicGames.UBA
 		/// </summary>
 		public bool AllowKillOnMem { get; init; }
 
+		public bool StoreObjFilesCompressed { get; init; }
+
 		/// <summary>
 		/// Constructor
 		/// </summary>
@@ -66,7 +68,7 @@ namespace EpicGames.UBA
 		/// <param name="detailedTrace">More detailed trace information</param>
 		/// <param name="allowWaitOnMem">Wait for memory before starting new processes</param>
 		/// <param name="allowKillOnMem">Kill processes when close to run out of memory</param>
-		public SessionServerCreateInfo(string rootDirectory, string traceOutputFile, bool disableCustomAllocator, bool launchVisualizer, bool resetCas, bool writeToDisk, bool detailedTrace, bool allowWaitOnMem, bool allowKillOnMem)
+		public SessionServerCreateInfo(string rootDirectory, string traceOutputFile, bool disableCustomAllocator, bool launchVisualizer, bool resetCas, bool writeToDisk, bool detailedTrace, bool allowWaitOnMem, bool allowKillOnMem, bool storeObjFilesCompressed)
 		{
 			RootDirectory = rootDirectory;
 			TraceOutputFile = traceOutputFile;
@@ -77,6 +79,7 @@ namespace EpicGames.UBA
 			DetailedTrace = detailedTrace;
 			AllowWaitOnMem = allowWaitOnMem;
 			AllowKillOnMem = allowKillOnMem;
+			StoreObjFilesCompressed = storeObjFilesCompressed;
 		}
 	}
 
diff --git a/Engine/Source/Programs/Shared/EpicGames.UBA/Impl/SessionServerImpl.cs b/Engine/Source/Programs/Shared/EpicGames.UBA/Impl/SessionServerImpl.cs
--- a/Engine/Source/Programs/Shared/EpicGames.UBA/Impl/SessionServerImpl.cs
+++ b/Engine/Source/Programs/Shared/EpicGames.UBA/Impl/SessionServerImpl.cs
@@ -17,7 +17,7 @@ namespace EpicGames.UBA
 		#region DllImport
 		[DllImport("UbaHost", CharSet = CharSet.Auto)]
 		static extern IntPtr CreateSessionServerCreateInfo(IntPtr storage, IntPtr client, IntPtr logger, string rootDir, string traceOutputFile,
-			bool disableCustomAllocator, bool launchVisualizer, bool resetCas, bool writeToDisk, bool detailedTrace, bool allowWaitOnMem, bool allowKillOnMem);
+			bool disableCustomAllocator, bool launchVisualizer, bool resetCas, bool writeToDisk, bool detailedTrace, bool allowWaitOnMem, bool allowKillOnMem, bool storeObjFilesCompressed);
 
 		[DllImport("UbaHost", CharSet = CharSet.Auto)]
 		static extern void DestroySessionServerCreateInfo(IntPtr server);
@@ -28,7 +28,7 @@ namespace EpicGames.UBA
 			_storage = storage;
 			_client = client;
 			_logger = logger;
-			_handle = CreateSessionServerCreateInfo(_storage.GetHandle(), _client.GetHandle(), _logger.GetHandle(), info.RootDirectory, info.TraceOutputFile, info.DisableCustomAllocator, info.LaunchVisualizer, info.ResetCas, info.WriteToDisk, info.DetailedTrace, info.AllowWaitOnMem, info.AllowKillOnMem);
+			_handle = CreateSessionServerCreateInfo(_storage.GetHandle(), _client.GetHandle(), _logger.GetHandle(), info.RootDirectory, info.TraceOutputFile, info.DisableCustomAllocator, info.LaunchVisualizer, info.ResetCas, info.WriteToDisk, info.DetailedTrace, info.AllowWaitOnMem, info.AllowKillOnMem, info.StoreObjFilesCompressed);
 		}
 
 		#region IDisposable
diff --git a/Engine/Source/Programs/UnrealBuildTool/Executors/UnrealBuildAccelerator/UBAExecutor.cs b/Engine/Source/Programs/UnrealBuildTool/Executors/UnrealBuildAccelerator/UBAExecutor.cs
--- a/Engine/Source/Programs/UnrealBuildTool/Executors/UnrealBuildAccelerator/UBAExecutor.cs
+++ b/Engine/Source/Programs/UnrealBuildTool/Executors/UnrealBuildAccelerator/UBAExecutor.cs
@@ -243,7 +243,7 @@ namespace UnrealBuildTool
 				using (Server = IServer.CreateServer(UBAConfig.MaxWorkers, UBAConfig.SendSize, ubaLogger, UBAConfig.bUseQuic))
 				{
 					using IStorageServer ubaStorageServer = IStorageServer.CreateStorageServer(Server, ubaLogger, new StorageServerCreateInfo(_rootDirRef.FullName, ((ulong)UBAConfig.StoreCapacityGb) * 1000 * 1000 * 1000, !UBAConfig.bStoreRaw, UBAConfig.Zone));
-					using ISessionServerCreateInfo serverCreateInfo = ISessionServerCreateInfo.CreateSessionServerCreateInfo(ubaStorageServer, Server, ubaLogger, new SessionServerCreateInfo(_rootDirRef.FullName, ubaTraceFile.FullName, UBAConfig.bDisableCustomAlloc, false, UBAConfig.bResetCas, UBAConfig.bWriteToDisk, UBAConfig.bDetailedTrace, !UBAConfig.bDisableWaitOnMem, UBAConfig.bAllowKillOnMem));
+					using ISessionServerCreateInfo serverCreateInfo = ISessionServerCreateInfo.CreateSessionServerCreateInfo(ubaStorageServer, Server, ubaLogger, new SessionServerCreateInfo(_rootDirRef.FullName, ubaTraceFile.FullName, UBAConfig.bDisableCustomAlloc, false, UBAConfig.bResetCas, UBAConfig.bWriteToDisk, UBAConfig.bDetailedTrace, !UBAConfig.bDisableWaitOnMem, UBAConfig.bAllowKillOnMem, UBAConfig.bStoreObjFilesCompressed));
 					using (_session = ISessionServer.CreateSessionServer(serverCreateInfo))
 					{
 
diff --git a/Engine/Source/Programs/UnrealBuildTool/Executors/UnrealBuildAccelerator/UnrealBuildAcceleratorConfig.cs b/Engine/Source/Programs/UnrealBuildTool/Executors/UnrealBuildAccelerator/UnrealBuildAcceleratorConfig.cs
--- a/Engine/Source/Programs/UnrealBuildTool/Executors/UnrealBuildAccelerator/UnrealBuildAcceleratorConfig.cs
+++ b/Engine/Source/Programs/UnrealBuildTool/Executors/UnrealBuildAccelerator/UnrealBuildAcceleratorConfig.cs
@@ -158,6 +158,13 @@ namespace UnrealBuildTool
 		[CommandLine("-UBAAllowKillOnMem", Value = "true")]
 		public bool bAllowKillOnMem { get; set; }
 
+		/// <summary>
+		/// Store object (.obj) compressed on disk. Requires uba to do link step where it will decompress obj files again
+		/// </summary>
+		[XmlConfigFile(Category = "UnrealBuildAccelerator")]
+		[CommandLine("-UBAStoreObjFilesCompressed", Value = "true")]
+		public bool bStoreObjFilesCompressed { get; set; }
+		
 		/// <summary>
 		/// Threshold for when executor should output logging for the process. Defaults to never
 		/// </summary>
