[UBA]
* Added documentation to c# code to fix warnings

diff --git a/Engine/Source/Programs/Shared/EpicGames.UBA/ISessionServer.cs b/Engine/Source/Programs/Shared/EpicGames.UBA/ISessionServer.cs
--- a/Engine/Source/Programs/Shared/EpicGames.UBA/ISessionServer.cs
+++ b/Engine/Source/Programs/Shared/EpicGames.UBA/ISessionServer.cs
@@ -54,6 +54,9 @@ namespace EpicGames.UBA
 		/// </summary>
 		public bool AllowKillOnMem { get; init; }
 
+		/// <summary>
+		/// Store .obj files compressed on disk
+		/// </summary>
 		public bool StoreObjFilesCompressed { get; init; }
 
 		/// <summary>
@@ -68,6 +71,7 @@ namespace EpicGames.UBA
 		/// <param name="detailedTrace">More detailed trace information</param>
 		/// <param name="allowWaitOnMem">Wait for memory before starting new processes</param>
 		/// <param name="allowKillOnMem">Kill processes when close to run out of memory</param>
+		/// <param name="storeObjFilesCompressed">Store .obj files compressed on disk</param>
 		public SessionServerCreateInfo(string rootDirectory, string traceOutputFile, bool disableCustomAllocator, bool launchVisualizer, bool resetCas, bool writeToDisk, bool detailedTrace, bool allowWaitOnMem, bool allowKillOnMem, bool storeObjFilesCompressed)
 		{
 			RootDirectory = rootDirectory;
