[EpicGames.Horde]
* Renamed information message from "Error" to "Exception" when background send is interrupted to see if this solves the failed submittool builds

diff --git a/Engine/Source/Programs/Shared/EpicGames.Horde/Compute/ComputeSocket.cs b/Engine/Source/Programs/Shared/EpicGames.Horde/Compute/ComputeSocket.cs
--- a/Engine/Source/Programs/Shared/EpicGames.Horde/Compute/ComputeSocket.cs
+++ b/Engine/Source/Programs/Shared/EpicGames.Horde/Compute/ComputeSocket.cs
@@ -853,7 +853,7 @@ namespace EpicGames.Horde.Compute
 			}
 			catch (Exception ex)
 			{
-				_logger.LogInformation(ex, "Error in background send: {Message}", ex.Message);
+				_logger.LogInformation(ex, "Exception in background send: {Message}", ex.Message);
 			}
 		}
 	}
