Horde: Fix decoding of UTF8 characters as part of unescaping Json strings.

diff --git a/Engine/Source/Programs/Shared/EpicGames.Core/JsonExtensions.cs b/Engine/Source/Programs/Shared/EpicGames.Core/JsonExtensions.cs
--- a/Engine/Source/Programs/Shared/EpicGames.Core/JsonExtensions.cs
+++ b/Engine/Source/Programs/Shared/EpicGames.Core/JsonExtensions.cs
@@ -167,7 +167,7 @@ namespace EpicGames.Core
 				output[1] = (byte)(0b10000000 | (c & 0x1f));
 				return 2;
 			}
-			else if (c <= 0x7ff)
+			else if (c <= 0xffff)
 			{
 				output[0] = (byte)(0b11100000 | (c >> 12));
 				output[1] = (byte)(0b10000000 | ((c >> 6) & 0x1f));
