Horde: Suppress static analysis warnings.

diff --git a/Engine/Source/Programs/Shared/EpicGames.Horde/Accounts/AccountMessages.cs b/Engine/Source/Programs/Shared/EpicGames.Horde/Accounts/AccountMessages.cs
--- a/Engine/Source/Programs/Shared/EpicGames.Horde/Accounts/AccountMessages.cs
+++ b/Engine/Source/Programs/Shared/EpicGames.Horde/Accounts/AccountMessages.cs
@@ -4,6 +4,7 @@ using System.Collections.Generic;
 using EpicGames.Horde.Acls;
 
 #pragma warning disable CA1054 // URI-like parameters should not be strings
+#pragma warning disable CA1056 // Change string to URI
 
 namespace EpicGames.Horde.Accounts
 {
diff --git a/Engine/Source/Programs/Shared/EpicGames.Horde/Compute/AgentMessage.cs b/Engine/Source/Programs/Shared/EpicGames.Horde/Compute/AgentMessage.cs
--- a/Engine/Source/Programs/Shared/EpicGames.Horde/Compute/AgentMessage.cs
+++ b/Engine/Source/Programs/Shared/EpicGames.Horde/Compute/AgentMessage.cs
@@ -11,6 +11,7 @@ using EpicGames.Core;
 using EpicGames.Horde.Storage;
 
 #pragma warning disable CA1054 // URI-like parameters should not be strings
+#pragma warning disable CA1056 // Change string to URI
 
 namespace EpicGames.Horde.Compute
 {
diff --git a/Engine/Source/Programs/Shared/EpicGames.Horde/Dashboard/DashboardMessages.cs b/Engine/Source/Programs/Shared/EpicGames.Horde/Dashboard/DashboardMessages.cs
--- a/Engine/Source/Programs/Shared/EpicGames.Horde/Dashboard/DashboardMessages.cs
+++ b/Engine/Source/Programs/Shared/EpicGames.Horde/Dashboard/DashboardMessages.cs
@@ -6,6 +6,7 @@ using EpicGames.Horde.Common;
 using EpicGames.Horde.Server;
 
 #pragma warning disable CA2227 // Collection properties should be read only
+#pragma warning disable CA1056 // Change string to URI
 
 namespace EpicGames.Horde.Dashboard
 {
diff --git a/Engine/Source/Programs/Shared/EpicGames.Horde/Server/ServerMessages.cs b/Engine/Source/Programs/Shared/EpicGames.Horde/Server/ServerMessages.cs
--- a/Engine/Source/Programs/Shared/EpicGames.Horde/Server/ServerMessages.cs
+++ b/Engine/Source/Programs/Shared/EpicGames.Horde/Server/ServerMessages.cs
@@ -3,6 +3,8 @@
 using System;
 using System.Text.Json.Serialization;
 
+#pragma warning disable CA1056 // Change string to URI
+
 namespace EpicGames.Horde.Server
 {
 	/// <summary>
