Set c# structs to be readonly if possible

#rnx

diff --git a/Engine/Source/Programs/Shared/EpicGames.Horde/Accounts/AccountId.cs b/Engine/Source/Programs/Shared/EpicGames.Horde/Accounts/AccountId.cs
--- a/Engine/Source/Programs/Shared/EpicGames.Horde/Accounts/AccountId.cs
+++ b/Engine/Source/Programs/Shared/EpicGames.Horde/Accounts/AccountId.cs
@@ -14,7 +14,7 @@ namespace EpicGames.Horde.Accounts
 	[JsonSchemaString]
 	[TypeConverter(typeof(BinaryIdTypeConverter<AccountId, AccountIdConverter>))]
 	[BinaryIdConverter(typeof(AccountIdConverter))]
-	public record struct AccountId(BinaryId Id)
+	public readonly record struct AccountId(BinaryId Id)
 	{
 		/// <inheritdoc cref="BinaryId.Parse(System.String)"/>
 		public static AccountId Parse(string text) => new AccountId(BinaryId.Parse(text));
diff --git a/Engine/Source/Programs/Shared/EpicGames.Horde/Acls/AclAction.cs b/Engine/Source/Programs/Shared/EpicGames.Horde/Acls/AclAction.cs
--- a/Engine/Source/Programs/Shared/EpicGames.Horde/Acls/AclAction.cs
+++ b/Engine/Source/Programs/Shared/EpicGames.Horde/Acls/AclAction.cs
@@ -16,7 +16,7 @@ namespace EpicGames.Horde.Acls
 	[JsonSchemaString]
 	[JsonConverter(typeof(AclActionJsonConverter))]
 	[TypeConverter(typeof(AclActionTypeConverter))]
-	public record struct AclAction(string Name)
+	public readonly record struct AclAction(string Name)
 	{
 		/// <inheritdoc/>
 		public override string ToString() => Name;
diff --git a/Engine/Source/Programs/Shared/EpicGames.Horde/Agents/Leases/LeaseId.cs b/Engine/Source/Programs/Shared/EpicGames.Horde/Agents/Leases/LeaseId.cs
--- a/Engine/Source/Programs/Shared/EpicGames.Horde/Agents/Leases/LeaseId.cs
+++ b/Engine/Source/Programs/Shared/EpicGames.Horde/Agents/Leases/LeaseId.cs
@@ -13,7 +13,7 @@ namespace EpicGames.Horde.Agents.Leases
 	[JsonSchemaString]
 	[TypeConverter(typeof(BinaryIdTypeConverter<LeaseId, LeaseIdConverter>))]
 	[BinaryIdConverter(typeof(LeaseIdConverter))]
-	public record struct LeaseId(BinaryId Id)
+	public readonly record struct LeaseId(BinaryId Id)
 	{
 		/// <inheritdoc cref="BinaryId.TryParse(System.String, out BinaryId)"/>
 		public static bool TryParse(string text, out LeaseId leaseId)
diff --git a/Engine/Source/Programs/Shared/EpicGames.Horde/Agents/Pools/PoolId.cs b/Engine/Source/Programs/Shared/EpicGames.Horde/Agents/Pools/PoolId.cs
--- a/Engine/Source/Programs/Shared/EpicGames.Horde/Agents/Pools/PoolId.cs
+++ b/Engine/Source/Programs/Shared/EpicGames.Horde/Agents/Pools/PoolId.cs
@@ -15,7 +15,7 @@ namespace EpicGames.Horde.Agents.Pools
 	[TypeConverter(typeof(StringIdTypeConverter<PoolId, PoolIdConverter>))]
 	[StringIdConverter(typeof(PoolIdConverter))]
 	[CbConverter(typeof(StringIdCbConverter<PoolId, PoolIdConverter>))]
-	public record struct PoolId(StringId Id)
+	public readonly record struct PoolId(StringId Id)
 	{
 		/// <summary>
 		/// Constructor
diff --git a/Engine/Source/Programs/Shared/EpicGames.Horde/Agents/Sessions/SessionId.cs b/Engine/Source/Programs/Shared/EpicGames.Horde/Agents/Sessions/SessionId.cs
--- a/Engine/Source/Programs/Shared/EpicGames.Horde/Agents/Sessions/SessionId.cs
+++ b/Engine/Source/Programs/Shared/EpicGames.Horde/Agents/Sessions/SessionId.cs
@@ -13,7 +13,7 @@ namespace EpicGames.Horde.Agents.Sessions
 	[JsonSchemaString]
 	[TypeConverter(typeof(BinaryIdTypeConverter<SessionId, SessionIdConverter>))]
 	[BinaryIdConverter(typeof(SessionIdConverter))]
-	public record struct SessionId(BinaryId Id)
+	public readonly record struct SessionId(BinaryId Id)
 	{
 		/// <summary>
 		/// Default empty value for session id
diff --git a/Engine/Source/Programs/Shared/EpicGames.Horde/Artifacts/ArtifactId.cs b/Engine/Source/Programs/Shared/EpicGames.Horde/Artifacts/ArtifactId.cs
--- a/Engine/Source/Programs/Shared/EpicGames.Horde/Artifacts/ArtifactId.cs
+++ b/Engine/Source/Programs/Shared/EpicGames.Horde/Artifacts/ArtifactId.cs
@@ -13,7 +13,7 @@ namespace EpicGames.Horde.Artifacts
 	[JsonSchemaString]
 	[TypeConverter(typeof(BinaryIdTypeConverter<ArtifactId, ArtifactIdConverter>))]
 	[BinaryIdConverter(typeof(ArtifactIdConverter))]
-	public record struct ArtifactId(BinaryId Id)
+	public readonly record struct ArtifactId(BinaryId Id)
 	{
 		/// <inheritdoc cref="BinaryId.Parse(System.String)"/>
 		public static ArtifactId Parse(string text) => new ArtifactId(BinaryId.Parse(text));
diff --git a/Engine/Source/Programs/Shared/EpicGames.Horde/Artifacts/ArtifactName.cs b/Engine/Source/Programs/Shared/EpicGames.Horde/Artifacts/ArtifactName.cs
--- a/Engine/Source/Programs/Shared/EpicGames.Horde/Artifacts/ArtifactName.cs
+++ b/Engine/Source/Programs/Shared/EpicGames.Horde/Artifacts/ArtifactName.cs
@@ -12,7 +12,7 @@ namespace EpicGames.Horde.Artifacts
 	[JsonSchemaString]
 	[StringIdConverter(typeof(ArtifactNameConverter))]
 	[TypeConverter(typeof(StringIdTypeConverter<ArtifactName, ArtifactNameConverter>))]
-	public record struct ArtifactName(StringId Id)
+	public readonly record struct ArtifactName(StringId Id)
 	{
 		/// <summary>
 		/// Constructor
diff --git a/Engine/Source/Programs/Shared/EpicGames.Horde/Artifacts/ArtifactType.cs b/Engine/Source/Programs/Shared/EpicGames.Horde/Artifacts/ArtifactType.cs
--- a/Engine/Source/Programs/Shared/EpicGames.Horde/Artifacts/ArtifactType.cs
+++ b/Engine/Source/Programs/Shared/EpicGames.Horde/Artifacts/ArtifactType.cs
@@ -12,7 +12,7 @@ namespace EpicGames.Horde.Artifacts
 	[JsonSchemaString]
 	[StringIdConverter(typeof(ArtifactTypeConverter))]
 	[TypeConverter(typeof(StringIdTypeConverter<ArtifactType, ArtifactTypeConverter>))]
-	public record struct ArtifactType(StringId Id)
+	public readonly record struct ArtifactType(StringId Id)
 	{
 		/// <summary>
 		/// Default artifact type 
diff --git a/Engine/Source/Programs/Shared/EpicGames.Horde/Compute/AgentMessage.cs b/Engine/Source/Programs/Shared/EpicGames.Horde/Compute/AgentMessage.cs
--- a/Engine/Source/Programs/Shared/EpicGames.Horde/Compute/AgentMessage.cs
+++ b/Engine/Source/Programs/Shared/EpicGames.Horde/Compute/AgentMessage.cs
@@ -226,27 +226,27 @@ namespace EpicGames.Horde.Compute
 	/// <summary>
 	/// Message for reporting an error
 	/// </summary>
-	public record struct ExceptionMessage(string Message, string Description);
+	public readonly record struct ExceptionMessage(string Message, string Description);
 
 	/// <summary>
 	/// Message requesting that the message loop be forked
 	/// </summary>
 	/// <param name="ChannelId">New channel to communicate on</param>
 	/// <param name="BufferSize">Size of the buffer</param>
-	public record struct ForkMessage(int ChannelId, int BufferSize);
+	public readonly record struct ForkMessage(int ChannelId, int BufferSize);
 
 	/// <summary>
 	/// Extract files from a bundle to a path in the remote sandbox
 	/// </summary>
 	/// <param name="Name">Path to extract the files to</param>
 	/// <param name="Locator">Locator for the tree to extract</param>
-	public record struct UploadFilesMessage(string Name, BlobLocator Locator);
+	public readonly record struct UploadFilesMessage(string Name, BlobLocator Locator);
 
 	/// <summary>
 	/// Deletes files or directories in the remote
 	/// </summary>
 	/// <param name="Filter">Filter for files to delete</param>
-	public record struct DeleteFilesMessage(IReadOnlyList<string> Filter);
+	public readonly record struct DeleteFilesMessage(IReadOnlyList<string> Filter);
 
 	/// <summary>
 	/// Message to execute a new child process
@@ -257,25 +257,25 @@ namespace EpicGames.Horde.Compute
 	/// <param name="EnvVars">Environment variables for the child process. Null values unset variables.</param>
 	/// <param name="Flags">Additional execution flags</param>
 	/// <param name="ContainerImageUrl">URL to container image. If specified, process will be executed inside this container</param>
-	public record struct ExecuteProcessMessage(string Executable, IReadOnlyList<string> Arguments, string? WorkingDir, IReadOnlyDictionary<string, string?> EnvVars, ExecuteProcessFlags Flags, string? ContainerImageUrl);
+	public readonly record struct ExecuteProcessMessage(string Executable, IReadOnlyList<string> Arguments, string? WorkingDir, IReadOnlyDictionary<string, string?> EnvVars, ExecuteProcessFlags Flags, string? ContainerImageUrl);
 
 	/// <summary>
 	/// Response from executing a child process
 	/// </summary>
 	/// <param name="ExitCode">Exit code for the process</param>
-	public record struct ExecuteProcessResponseMessage(int ExitCode);
+	public readonly record struct ExecuteProcessResponseMessage(int ExitCode);
 
 	/// <summary>
 	/// Creates a blob read request
 	/// </summary>
-	public record struct ReadBlobMessage(BlobLocator Locator, int Offset, int Length);
+	public readonly record struct ReadBlobMessage(BlobLocator Locator, int Offset, int Length);
 
 	/// <summary>
 	/// Message for running an XOR command
 	/// </summary>
 	/// <param name="Data">Data to xor</param>
 	/// <param name="Value">Value to XOR with</param>
-	public record struct XorRequestMessage(ReadOnlyMemory<byte> Data, byte Value);
+	public readonly record struct XorRequestMessage(ReadOnlyMemory<byte> Data, byte Value);
 
 	/// <summary>
 	/// Wraps various requests across compute channels
diff --git a/Engine/Source/Programs/Shared/EpicGames.Horde/Compute/ClusterId.cs b/Engine/Source/Programs/Shared/EpicGames.Horde/Compute/ClusterId.cs
--- a/Engine/Source/Programs/Shared/EpicGames.Horde/Compute/ClusterId.cs
+++ b/Engine/Source/Programs/Shared/EpicGames.Horde/Compute/ClusterId.cs
@@ -18,7 +18,7 @@ namespace EpicGames.Horde.Compute
 	[CbConverter(typeof(ClusterIdCbConverter))]
 	[JsonConverter(typeof(ClusterIdJsonConverter))]
 	[TypeConverter(typeof(ClusterIdTypeConverter))]
-	public struct ClusterId : IEquatable<ClusterId>
+	public readonly struct ClusterId : IEquatable<ClusterId>
 	{
 		/// <summary>
 		/// The text representing this id
diff --git a/Engine/Source/Programs/Shared/EpicGames.Horde/Devices/DeviceId.cs b/Engine/Source/Programs/Shared/EpicGames.Horde/Devices/DeviceId.cs
--- a/Engine/Source/Programs/Shared/EpicGames.Horde/Devices/DeviceId.cs
+++ b/Engine/Source/Programs/Shared/EpicGames.Horde/Devices/DeviceId.cs
@@ -15,7 +15,7 @@ namespace EpicGames.Horde.Devices
 	[TypeConverter(typeof(StringIdTypeConverter<DeviceId, DeviceIdConverter>))]
 	[StringIdConverter(typeof(DeviceIdConverter))]
 	[CbConverter(typeof(StringIdCbConverter<DeviceId, DeviceIdConverter>))]
-	public record struct DeviceId(StringId Id)
+	public readonly record struct DeviceId(StringId Id)
 	{
 		/// <summary>
 		/// Constructor
diff --git a/Engine/Source/Programs/Shared/EpicGames.Horde/Devices/DevicePlatformId.cs b/Engine/Source/Programs/Shared/EpicGames.Horde/Devices/DevicePlatformId.cs
--- a/Engine/Source/Programs/Shared/EpicGames.Horde/Devices/DevicePlatformId.cs
+++ b/Engine/Source/Programs/Shared/EpicGames.Horde/Devices/DevicePlatformId.cs
@@ -15,7 +15,7 @@ namespace EpicGames.Horde.Devices
 	[TypeConverter(typeof(StringIdTypeConverter<DevicePlatformId, DevicePlatformIdConverter>))]
 	[StringIdConverter(typeof(DevicePlatformIdConverter))]
 	[CbConverter(typeof(StringIdCbConverter<DevicePlatformId, DevicePlatformIdConverter>))]
-	public record struct DevicePlatformId(StringId Id)
+	public readonly record struct DevicePlatformId(StringId Id)
 	{
 		/// <summary>
 		/// Constructor
diff --git a/Engine/Source/Programs/Shared/EpicGames.Horde/Devices/DevicePoolId.cs b/Engine/Source/Programs/Shared/EpicGames.Horde/Devices/DevicePoolId.cs
--- a/Engine/Source/Programs/Shared/EpicGames.Horde/Devices/DevicePoolId.cs
+++ b/Engine/Source/Programs/Shared/EpicGames.Horde/Devices/DevicePoolId.cs
@@ -15,7 +15,7 @@ namespace EpicGames.Horde.Devices
 	[TypeConverter(typeof(StringIdTypeConverter<DevicePoolId, DevicePoolIdConverter>))]
 	[StringIdConverter(typeof(DevicePoolIdConverter))]
 	[CbConverter(typeof(StringIdCbConverter<DevicePoolId, DevicePoolIdConverter>))]
-	public record struct DevicePoolId(StringId Id)
+	public readonly record struct DevicePoolId(StringId Id)
 	{
 		/// <summary>
 		/// Constructor
diff --git a/Engine/Source/Programs/Shared/EpicGames.Horde/Issues/WorkflowId.cs b/Engine/Source/Programs/Shared/EpicGames.Horde/Issues/WorkflowId.cs
--- a/Engine/Source/Programs/Shared/EpicGames.Horde/Issues/WorkflowId.cs
+++ b/Engine/Source/Programs/Shared/EpicGames.Horde/Issues/WorkflowId.cs
@@ -14,7 +14,7 @@ namespace EpicGames.Horde.Issues
 	[TypeConverter(typeof(StringIdTypeConverter<WorkflowId, WorkflowIdConverter>))]
 	[StringIdConverter(typeof(WorkflowIdConverter))]
 	[CbConverter(typeof(StringIdCbConverter<WorkflowId, WorkflowIdConverter>))]
-	public record struct WorkflowId(StringId Id)
+	public readonly record struct WorkflowId(StringId Id)
 	{
 		/// <summary>
 		/// Empty workflow id constant
diff --git a/Engine/Source/Programs/Shared/EpicGames.Horde/Jobs/Bisect/BisectTaskId.cs b/Engine/Source/Programs/Shared/EpicGames.Horde/Jobs/Bisect/BisectTaskId.cs
--- a/Engine/Source/Programs/Shared/EpicGames.Horde/Jobs/Bisect/BisectTaskId.cs
+++ b/Engine/Source/Programs/Shared/EpicGames.Horde/Jobs/Bisect/BisectTaskId.cs
@@ -13,7 +13,7 @@ namespace EpicGames.Horde.Jobs.Bisect
 	[JsonSchemaString]
 	[TypeConverter(typeof(BinaryIdTypeConverter<BisectTaskId, BisectTaskIdConverter>))]
 	[BinaryIdConverter(typeof(BisectTaskIdConverter))]
-	public record struct BisectTaskId(BinaryId Id)
+	public readonly record struct BisectTaskId(BinaryId Id)
 	{
 		/// <inheritdoc cref="BinaryId.Parse(System.String)"/>
 		public static BisectTaskId Parse(string text) => new BisectTaskId(BinaryId.Parse(text));
diff --git a/Engine/Source/Programs/Shared/EpicGames.Horde/Jobs/JobId.cs b/Engine/Source/Programs/Shared/EpicGames.Horde/Jobs/JobId.cs
--- a/Engine/Source/Programs/Shared/EpicGames.Horde/Jobs/JobId.cs
+++ b/Engine/Source/Programs/Shared/EpicGames.Horde/Jobs/JobId.cs
@@ -13,7 +13,7 @@ namespace EpicGames.Horde.Jobs
 	[JsonSchemaString]
 	[TypeConverter(typeof(BinaryIdTypeConverter<JobId, JobIdConverter>))]
 	[BinaryIdConverter(typeof(JobIdConverter))]
-	public record struct JobId(BinaryId Id)
+	public readonly record struct JobId(BinaryId Id)
 	{
 		/// <summary>
 		/// Constant value for an empty job id
diff --git a/Engine/Source/Programs/Shared/EpicGames.Horde/Jobs/JobStepBatchId.cs b/Engine/Source/Programs/Shared/EpicGames.Horde/Jobs/JobStepBatchId.cs
--- a/Engine/Source/Programs/Shared/EpicGames.Horde/Jobs/JobStepBatchId.cs
+++ b/Engine/Source/Programs/Shared/EpicGames.Horde/Jobs/JobStepBatchId.cs
@@ -13,7 +13,7 @@ namespace EpicGames.Horde.Jobs
 	[JsonSchemaString]
 	[TypeConverter(typeof(SubResourceIdTypeConverter<JobStepBatchId, JobStepBatchIdConverter>))]
 	[SubResourceIdConverter(typeof(JobStepBatchIdConverter))]
-	public record struct JobStepBatchId(SubResourceId SubResourceId)
+	public readonly record struct JobStepBatchId(SubResourceId SubResourceId)
 	{
 		/// <summary>
 		/// Constructor
diff --git a/Engine/Source/Programs/Shared/EpicGames.Horde/Jobs/JobStepId.cs b/Engine/Source/Programs/Shared/EpicGames.Horde/Jobs/JobStepId.cs
--- a/Engine/Source/Programs/Shared/EpicGames.Horde/Jobs/JobStepId.cs
+++ b/Engine/Source/Programs/Shared/EpicGames.Horde/Jobs/JobStepId.cs
@@ -13,7 +13,7 @@ namespace EpicGames.Horde.Jobs
 	[JsonSchemaString]
 	[TypeConverter(typeof(SubResourceIdTypeConverter<JobStepId, JobStepIdConverter>))]
 	[SubResourceIdConverter(typeof(JobStepIdConverter))]
-	public record struct JobStepId(SubResourceId Id)
+	public readonly record struct JobStepId(SubResourceId Id)
 	{
 		/// <summary>
 		/// Constructor
diff --git a/Engine/Source/Programs/Shared/EpicGames.Horde/Jobs/SubResourceId.cs b/Engine/Source/Programs/Shared/EpicGames.Horde/Jobs/SubResourceId.cs
--- a/Engine/Source/Programs/Shared/EpicGames.Horde/Jobs/SubResourceId.cs
+++ b/Engine/Source/Programs/Shared/EpicGames.Horde/Jobs/SubResourceId.cs
@@ -14,7 +14,7 @@ namespace EpicGames.Horde.Jobs
 	/// identifier with more entropy than just incrementing the value but an identical period before repeating, in order to make URL fragments more distinctive.
 	/// </summary>
 	[TypeConverter(typeof(SubResourceIdTypeConverter))]
-	public struct SubResourceId : IEquatable<SubResourceId>
+	public readonly struct SubResourceId : IEquatable<SubResourceId>
 	{
 		/// <summary>
 		/// The unique identifier value
diff --git a/Engine/Source/Programs/Shared/EpicGames.Horde/Jobs/Templates/TemplateId.cs b/Engine/Source/Programs/Shared/EpicGames.Horde/Jobs/Templates/TemplateId.cs
--- a/Engine/Source/Programs/Shared/EpicGames.Horde/Jobs/Templates/TemplateId.cs
+++ b/Engine/Source/Programs/Shared/EpicGames.Horde/Jobs/Templates/TemplateId.cs
@@ -15,7 +15,7 @@ namespace EpicGames.Horde.Jobs.Templates
 	[TypeConverter(typeof(StringIdTypeConverter<TemplateId, TemplateIdConverter>))]
 	[StringIdConverter(typeof(TemplateIdConverter))]
 	[CbConverter(typeof(StringIdCbConverter<TemplateId, TemplateIdConverter>))]
-	public record struct TemplateId(StringId Id)
+	public readonly record struct TemplateId(StringId Id)
 	{
 		/// <summary>
 		/// Constructor
diff --git a/Engine/Source/Programs/Shared/EpicGames.Horde/Logs/LogId.cs b/Engine/Source/Programs/Shared/EpicGames.Horde/Logs/LogId.cs
--- a/Engine/Source/Programs/Shared/EpicGames.Horde/Logs/LogId.cs
+++ b/Engine/Source/Programs/Shared/EpicGames.Horde/Logs/LogId.cs
@@ -13,7 +13,7 @@ namespace EpicGames.Horde.Logs
 	[JsonSchemaString]
 	[TypeConverter(typeof(BinaryIdTypeConverter<LogId, LogIdConverter>))]
 	[BinaryIdConverter(typeof(LogIdConverter))]
-	public record struct LogId(BinaryId Id)
+	public readonly record struct LogId(BinaryId Id)
 	{
 		/// <summary>
 		/// Constant value for empty user id
diff --git a/Engine/Source/Programs/Shared/EpicGames.Horde/Projects/ProjectId.cs b/Engine/Source/Programs/Shared/EpicGames.Horde/Projects/ProjectId.cs
--- a/Engine/Source/Programs/Shared/EpicGames.Horde/Projects/ProjectId.cs
+++ b/Engine/Source/Programs/Shared/EpicGames.Horde/Projects/ProjectId.cs
@@ -15,7 +15,7 @@ namespace EpicGames.Horde.Projects
 	[TypeConverter(typeof(StringIdTypeConverter<ProjectId, ProjectIdConverter>))]
 	[StringIdConverter(typeof(ProjectIdConverter))]
 	[CbConverter(typeof(StringIdCbConverter<ProjectId, ProjectIdConverter>))]
-	public record struct ProjectId(StringId Id)
+	public readonly record struct ProjectId(StringId Id)
 	{
 		/// <summary>
 		/// Constructor
diff --git a/Engine/Source/Programs/Shared/EpicGames.Horde/Replicators/ReplicatorId.cs b/Engine/Source/Programs/Shared/EpicGames.Horde/Replicators/ReplicatorId.cs
--- a/Engine/Source/Programs/Shared/EpicGames.Horde/Replicators/ReplicatorId.cs
+++ b/Engine/Source/Programs/Shared/EpicGames.Horde/Replicators/ReplicatorId.cs
@@ -18,7 +18,7 @@ namespace EpicGames.Horde.Replicators
 	[JsonSchemaString]
 	[TypeConverter(typeof(ReplicatorIdTypeConverter))]
 	[JsonConverter(typeof(ReplicatorIdJsonConverter))]
-	public record struct ReplicatorId(StreamId StreamId, StreamReplicatorId StreamReplicatorId)
+	public readonly record struct ReplicatorId(StreamId StreamId, StreamReplicatorId StreamReplicatorId)
 	{
 		/// <summary>
 		/// Parse a replicator id
@@ -104,7 +104,7 @@ namespace EpicGames.Horde.Replicators
 	[TypeConverter(typeof(StringIdTypeConverter<StreamReplicatorId, StreamReplicatorIdConverter>))]
 	[StringIdConverter(typeof(StreamReplicatorIdConverter))]
 	[CbConverter(typeof(StringIdCbConverter<StreamReplicatorId, StreamReplicatorIdConverter>))]
-	public record struct StreamReplicatorId(StringId Id)
+	public readonly record struct StreamReplicatorId(StringId Id)
 	{
 		/// <summary>
 		/// Constructor
diff --git a/Engine/Source/Programs/Shared/EpicGames.Horde/Secrets/SecretId.cs b/Engine/Source/Programs/Shared/EpicGames.Horde/Secrets/SecretId.cs
--- a/Engine/Source/Programs/Shared/EpicGames.Horde/Secrets/SecretId.cs
+++ b/Engine/Source/Programs/Shared/EpicGames.Horde/Secrets/SecretId.cs
@@ -15,7 +15,7 @@ namespace EpicGames.Horde.Secrets
 	[TypeConverter(typeof(StringIdTypeConverter<SecretId, SecretIdConverter>))]
 	[StringIdConverter(typeof(SecretIdConverter))]
 	[CbConverter(typeof(StringIdCbConverter<SecretId, SecretIdConverter>))]
-	public record struct SecretId(StringId Id)
+	public readonly record struct SecretId(StringId Id)
 	{
 		/// <summary>
 		/// Constructor
diff --git a/Engine/Source/Programs/Shared/EpicGames.Horde/ServiceAccounts/ServiceAccountId.cs b/Engine/Source/Programs/Shared/EpicGames.Horde/ServiceAccounts/ServiceAccountId.cs
--- a/Engine/Source/Programs/Shared/EpicGames.Horde/ServiceAccounts/ServiceAccountId.cs
+++ b/Engine/Source/Programs/Shared/EpicGames.Horde/ServiceAccounts/ServiceAccountId.cs
@@ -13,7 +13,7 @@ namespace EpicGames.Horde.ServiceAccounts
 	[JsonSchemaString]
 	[TypeConverter(typeof(BinaryIdTypeConverter<ServiceAccountId, ServiceAccountIdConverter>))]
 	[BinaryIdConverter(typeof(ServiceAccountIdConverter))]
-	public record struct ServiceAccountId(BinaryId Id)
+	public readonly record struct ServiceAccountId(BinaryId Id)
 	{
 		/// <inheritdoc cref="BinaryId.Parse(System.String)"/>
 		public static ServiceAccountId Parse(string text) => new ServiceAccountId(BinaryId.Parse(text));
diff --git a/Engine/Source/Programs/Shared/EpicGames.Horde/Storage/BlobType.cs b/Engine/Source/Programs/Shared/EpicGames.Horde/Storage/BlobType.cs
--- a/Engine/Source/Programs/Shared/EpicGames.Horde/Storage/BlobType.cs
+++ b/Engine/Source/Programs/Shared/EpicGames.Horde/Storage/BlobType.cs
@@ -11,7 +11,7 @@ namespace EpicGames.Horde.Storage
 	/// </summary>
 	/// <param name="Guid">Nominal identifier for the type</param>
 	/// <param name="Version">Version number for the serializer</param>
-	public record struct BlobType(Guid Guid, int Version)
+	public readonly record struct BlobType(Guid Guid, int Version)
 	{
 		/// <summary>
 		/// Number of bytes in a serialized blob type instance
diff --git a/Engine/Source/Programs/Shared/EpicGames.Horde/Storage/Bundles/Bundle.cs b/Engine/Source/Programs/Shared/EpicGames.Horde/Storage/Bundles/Bundle.cs
--- a/Engine/Source/Programs/Shared/EpicGames.Horde/Storage/Bundles/Bundle.cs
+++ b/Engine/Source/Programs/Shared/EpicGames.Horde/Storage/Bundles/Bundle.cs
@@ -70,7 +70,7 @@ namespace EpicGames.Horde.Storage.Bundles
 	/// </summary>
 	/// <param name="Version">Version number for the following file data</param>
 	/// <param name="HeaderLength">Length of the initial header</param>
-	public record struct BundleSignature(BundleVersion Version, int HeaderLength)
+	public readonly record struct BundleSignature(BundleVersion Version, int HeaderLength)
 	{
 		/// <summary>
 		/// Number of bytes in a signature when serialized
diff --git a/Engine/Source/Programs/Shared/EpicGames.Horde/Storage/Bundles/BundleCache.cs b/Engine/Source/Programs/Shared/EpicGames.Horde/Storage/Bundles/BundleCache.cs
--- a/Engine/Source/Programs/Shared/EpicGames.Horde/Storage/Bundles/BundleCache.cs
+++ b/Engine/Source/Programs/Shared/EpicGames.Horde/Storage/Bundles/BundleCache.cs
@@ -43,7 +43,7 @@ namespace EpicGames.Horde.Storage.Bundles
 		public long PacketCacheSize { get; set; } = 192 * 1024 * 1024;
 	}
 
-	record struct BundlePageCacheKey(BundleHandle Bundle, int Index)
+	readonly record struct BundlePageCacheKey(BundleHandle Bundle, int Index)
 	{
 		public override string ToString()
 			=> $"bundle-page:{Bundle}:{Index}";
diff --git a/Engine/Source/Programs/Shared/EpicGames.Horde/Storage/Bundles/V1/BundleHeader.cs b/Engine/Source/Programs/Shared/EpicGames.Horde/Storage/Bundles/V1/BundleHeader.cs
--- a/Engine/Source/Programs/Shared/EpicGames.Horde/Storage/Bundles/V1/BundleHeader.cs
+++ b/Engine/Source/Programs/Shared/EpicGames.Horde/Storage/Bundles/V1/BundleHeader.cs
@@ -417,7 +417,7 @@ namespace EpicGames.Horde.Storage.Bundles.V1
 	/// <summary>
 	/// Collection of node types in a bundle
 	/// </summary>
-	public struct BundleTypeCollection : IReadOnlyList<BlobType>
+	public readonly struct BundleTypeCollection : IReadOnlyList<BlobType>
 	{
 		readonly BlobType[] _types;
 
@@ -486,7 +486,7 @@ namespace EpicGames.Horde.Storage.Bundles.V1
 	/// <summary>
 	/// Collection of imported node references
 	/// </summary>
-	public struct BundleImportCollection : IReadOnlyList<BlobLocator>
+	public readonly struct BundleImportCollection : IReadOnlyList<BlobLocator>
 	{
 		readonly BlobLocator[] _imports;
 
@@ -564,7 +564,7 @@ namespace EpicGames.Horde.Storage.Bundles.V1
 	/// <summary>
 	/// Descriptor for a compression packet
 	/// </summary>
-	public struct BundlePacket
+	public readonly struct BundlePacket
 	{
 		/// <summary>
 		/// Size of this structure when serialized
@@ -633,7 +633,7 @@ namespace EpicGames.Horde.Storage.Bundles.V1
 	/// <summary>
 	/// Collection of information about packets in a bundle
 	/// </summary>
-	public struct BundlePacketCollection : IReadOnlyCollection<BundlePacket>
+	public readonly struct BundlePacketCollection : IReadOnlyCollection<BundlePacket>
 	{
 		readonly BundlePacket[] _packets;
 
@@ -704,7 +704,7 @@ namespace EpicGames.Horde.Storage.Bundles.V1
 	/// <summary>
 	/// Entry for a node exported from an object
 	/// </summary>
-	public struct BundleExport
+	public readonly struct BundleExport
 	{
 		/// <summary>
 		/// Number of bytes in a serialized export object
@@ -791,7 +791,7 @@ namespace EpicGames.Horde.Storage.Bundles.V1
 	/// <summary>
 	/// Entry for a node exported from an object
 	/// </summary>
-	public struct BundleExportCollection : IReadOnlyList<BundleExport>
+	public readonly struct BundleExportCollection : IReadOnlyList<BundleExport>
 	{
 		readonly BundleExport[] _exports;
 
@@ -994,7 +994,7 @@ namespace EpicGames.Horde.Storage.Bundles.V1
 	/// </summary>
 	/// <param name="ImportIdx">Index into the import table of the blob containing the referenced node. Can be -1 for references within the same bundle.</param>
 	/// <param name="NodeIdx">Node imported from the bundle</param>
-	public record struct BundleExportRef(int ImportIdx, int NodeIdx)
+	public readonly record struct BundleExportRef(int ImportIdx, int NodeIdx)
 	{
 		/// <summary>
 		/// Number of bytes in the serialized object
@@ -1033,7 +1033,7 @@ namespace EpicGames.Horde.Storage.Bundles.V1
 	/// <summary>
 	/// Collection of information about exported nodes
 	/// </summary>
-	public struct BundleExportRefCollection : IReadOnlyList<BundleExportRef>
+	public readonly struct BundleExportRefCollection : IReadOnlyList<BundleExportRef>
 	{
 		/// <summary>
 		/// Data used to store this collection
diff --git a/Engine/Source/Programs/Shared/EpicGames.Horde/Storage/Bundles/V2/Packet.cs b/Engine/Source/Programs/Shared/EpicGames.Horde/Storage/Bundles/V2/Packet.cs
--- a/Engine/Source/Programs/Shared/EpicGames.Horde/Storage/Bundles/V2/Packet.cs
+++ b/Engine/Source/Programs/Shared/EpicGames.Horde/Storage/Bundles/V2/Packet.cs
@@ -236,7 +236,7 @@ namespace EpicGames.Horde.Storage.Bundles.V2
 	/// </summary>
 	/// <param name="BaseIdx">Base index for this locator</param>
 	/// <param name="Fragment">The utf8 fragment appended to the base index</param>
-	public record struct PacketImport(int BaseIdx, Utf8String Fragment)
+	public readonly record struct PacketImport(int BaseIdx, Utf8String Fragment)
 	{
 		/// <summary>
 		/// Bias for indexes into the import table
@@ -274,7 +274,7 @@ namespace EpicGames.Horde.Storage.Bundles.V2
 	/// Data for an exported node in a packet
 	/// </summary>
 	[DebuggerTypeProxy(typeof(PacketExport.DebugProxy))]
-	public struct PacketExport
+	public readonly struct PacketExport
 	{
 		class DebugProxy
 		{
@@ -325,7 +325,7 @@ namespace EpicGames.Horde.Storage.Bundles.V2
 	/// <summary>
 	/// Data for an exported node in a bundle
 	/// </summary>
-	public struct PacketExportHeader
+	public readonly struct PacketExportHeader
 	{
 		/// <summary>
 		/// Index of the type for this export
diff --git a/Engine/Source/Programs/Shared/EpicGames.Horde/Storage/Bundles/V2/PacketReader.cs b/Engine/Source/Programs/Shared/EpicGames.Horde/Storage/Bundles/V2/PacketReader.cs
--- a/Engine/Source/Programs/Shared/EpicGames.Horde/Storage/Bundles/V2/PacketReader.cs
+++ b/Engine/Source/Programs/Shared/EpicGames.Horde/Storage/Bundles/V2/PacketReader.cs
@@ -23,7 +23,7 @@ namespace EpicGames.Horde.Storage.Bundles.V2
 		}
 	}
 
-	record struct PacketReaderCacheKey(BundleHandle Bundle, int Offset)
+	readonly record struct PacketReaderCacheKey(BundleHandle Bundle, int Offset)
 	{
 		public override string ToString()
 			=> $"packet-reader:{Bundle}@{Offset}";
diff --git a/Engine/Source/Programs/Shared/EpicGames.Horde/Storage/EncodedBlobData.cs b/Engine/Source/Programs/Shared/EpicGames.Horde/Storage/EncodedBlobData.cs
--- a/Engine/Source/Programs/Shared/EpicGames.Horde/Storage/EncodedBlobData.cs
+++ b/Engine/Source/Programs/Shared/EpicGames.Horde/Storage/EncodedBlobData.cs
@@ -22,7 +22,7 @@ namespace EpicGames.Horde.Storage
 
 		static readonly Version s_currentVersion = Enum.GetValues(typeof(Version)).Cast<Version>().Max();
 
-		struct LocatorCollection : IReadOnlyList<BlobLocator>
+		readonly struct LocatorCollection : IReadOnlyList<BlobLocator>
 		{
 			public readonly JaggedReadOnlyMemoryArray Array;
 
diff --git a/Engine/Source/Programs/Shared/EpicGames.Horde/Storage/IStorageClient.cs b/Engine/Source/Programs/Shared/EpicGames.Horde/Storage/IStorageClient.cs
--- a/Engine/Source/Programs/Shared/EpicGames.Horde/Storage/IStorageClient.cs
+++ b/Engine/Source/Programs/Shared/EpicGames.Horde/Storage/IStorageClient.cs
@@ -161,7 +161,7 @@ namespace EpicGames.Horde.Storage
 	/// Indicates the maximum age of a entry returned from a cache in the hierarchy
 	/// </summary>
 	/// <param name="Utc">Oldest allowed timestamp for a returned result</param>
-	public record struct RefCacheTime(DateTime Utc)
+	public readonly record struct RefCacheTime(DateTime Utc)
 	{
 		/// <summary>
 		/// Maximum age for a cached value to be returned
diff --git a/Engine/Source/Programs/Shared/EpicGames.Horde/Storage/NamespaceId.cs b/Engine/Source/Programs/Shared/EpicGames.Horde/Storage/NamespaceId.cs
--- a/Engine/Source/Programs/Shared/EpicGames.Horde/Storage/NamespaceId.cs
+++ b/Engine/Source/Programs/Shared/EpicGames.Horde/Storage/NamespaceId.cs
@@ -16,7 +16,7 @@ namespace EpicGames.Horde.Storage
 	[JsonSchemaString]
 	[JsonConverter(typeof(NamespaceIdJsonConverter))]
 	[TypeConverter(typeof(NamespaceIdTypeConverter))]
-	public struct NamespaceId : IEquatable<NamespaceId>
+	public readonly struct NamespaceId : IEquatable<NamespaceId>
 	{
 		/// <summary>
 		/// The text representing this id
diff --git a/Engine/Source/Programs/Shared/EpicGames.Horde/Storage/Nodes/DdcRefNode.cs b/Engine/Source/Programs/Shared/EpicGames.Horde/Storage/Nodes/DdcRefNode.cs
--- a/Engine/Source/Programs/Shared/EpicGames.Horde/Storage/Nodes/DdcRefNode.cs
+++ b/Engine/Source/Programs/Shared/EpicGames.Horde/Storage/Nodes/DdcRefNode.cs
@@ -49,7 +49,7 @@ namespace Horde.Server.Ddc
 
 	class DdcRefNodeConverter : BlobConverter<DdcRefNode>
 	{
-		static BlobType s_blobType = new BlobType(DdcRefNode.BlobTypeGuid, 1);
+		static readonly BlobType s_blobType = new BlobType(DdcRefNode.BlobTypeGuid, 1);
 
 		public override DdcRefNode Read(IBlobReader reader, BlobSerializerOptions options)
 		{
diff --git a/Engine/Source/Programs/Shared/EpicGames.Horde/Storage/RefName.cs b/Engine/Source/Programs/Shared/EpicGames.Horde/Storage/RefName.cs
--- a/Engine/Source/Programs/Shared/EpicGames.Horde/Storage/RefName.cs
+++ b/Engine/Source/Programs/Shared/EpicGames.Horde/Storage/RefName.cs
@@ -18,7 +18,7 @@ namespace EpicGames.Horde.Storage
 	[JsonConverter(typeof(RefNameJsonConverter))]
 	[TypeConverter(typeof(RefNameTypeConverter))]
 	[CbConverter(typeof(RefNameCbConverter))]
-	public struct RefName : IEquatable<RefName>, IComparable<RefName>
+	public readonly struct RefName : IEquatable<RefName>, IComparable<RefName>
 	{
 		/// <summary>
 		/// Empty ref name
diff --git a/Engine/Source/Programs/Shared/EpicGames.Horde/Storage/WorkspaceLayerId.cs b/Engine/Source/Programs/Shared/EpicGames.Horde/Storage/WorkspaceLayerId.cs
--- a/Engine/Source/Programs/Shared/EpicGames.Horde/Storage/WorkspaceLayerId.cs
+++ b/Engine/Source/Programs/Shared/EpicGames.Horde/Storage/WorkspaceLayerId.cs
@@ -12,7 +12,7 @@ namespace EpicGames.Horde.Storage
 	[LogValueType]
 	[TypeConverter(typeof(StringIdTypeConverter<WorkspaceLayerId, WorkspaceLayerIdConverter>))]
 	[BinaryIdConverter(typeof(WorkspaceLayerIdConverter))]
-	public record struct WorkspaceLayerId(StringId Id)
+	public readonly record struct WorkspaceLayerId(StringId Id)
 	{
 		/// <summary>
 		/// Name of the default layer
diff --git a/Engine/Source/Programs/Shared/EpicGames.Horde/Streams/StreamId.cs b/Engine/Source/Programs/Shared/EpicGames.Horde/Streams/StreamId.cs
--- a/Engine/Source/Programs/Shared/EpicGames.Horde/Streams/StreamId.cs
+++ b/Engine/Source/Programs/Shared/EpicGames.Horde/Streams/StreamId.cs
@@ -15,7 +15,7 @@ namespace EpicGames.Horde.Streams
 	[TypeConverter(typeof(StringIdTypeConverter<StreamId, StreamIdConverter>))]
 	[StringIdConverter(typeof(StreamIdConverter))]
 	[CbConverter(typeof(StringIdCbConverter<StreamId, StreamIdConverter>))]
-	public record struct StreamId(StringId Id)
+	public readonly record struct StreamId(StringId Id)
 	{
 		/// <summary>
 		/// Constructor
diff --git a/Engine/Source/Programs/Shared/EpicGames.Horde/StringId.cs b/Engine/Source/Programs/Shared/EpicGames.Horde/StringId.cs
--- a/Engine/Source/Programs/Shared/EpicGames.Horde/StringId.cs
+++ b/Engine/Source/Programs/Shared/EpicGames.Horde/StringId.cs
@@ -17,7 +17,7 @@ namespace EpicGames.Horde
 	[JsonSchemaString]
 	[JsonConverter(typeof(StringIdJsonConverter))]
 	[TypeConverter(typeof(StringIdTypeConverter))]
-	public struct StringId : IEquatable<StringId>, IEquatable<string>, IEquatable<ReadOnlyMemory<char>>
+	public readonly struct StringId : IEquatable<StringId>, IEquatable<string>, IEquatable<ReadOnlyMemory<char>>
 	{
 		/// <summary>
 		/// Enum used to disable validation on string arguments
diff --git a/Engine/Source/Programs/Shared/EpicGames.Horde/Telemetry/Metrics/MetricId.cs b/Engine/Source/Programs/Shared/EpicGames.Horde/Telemetry/Metrics/MetricId.cs
--- a/Engine/Source/Programs/Shared/EpicGames.Horde/Telemetry/Metrics/MetricId.cs
+++ b/Engine/Source/Programs/Shared/EpicGames.Horde/Telemetry/Metrics/MetricId.cs
@@ -15,7 +15,7 @@ namespace EpicGames.Horde.Telemetry.Metrics
 	[TypeConverter(typeof(StringIdTypeConverter<MetricId, MetricIdConverter>))]
 	[StringIdConverter(typeof(MetricIdConverter))]
 	[CbConverter(typeof(StringIdCbConverter<MetricId, MetricIdConverter>))]
-	public record struct MetricId(StringId Id)
+	public readonly record struct MetricId(StringId Id)
 	{
 		/// <summary>
 		/// Constructor
diff --git a/Engine/Source/Programs/Shared/EpicGames.Horde/Telemetry/TelemetryStoreId.cs b/Engine/Source/Programs/Shared/EpicGames.Horde/Telemetry/TelemetryStoreId.cs
--- a/Engine/Source/Programs/Shared/EpicGames.Horde/Telemetry/TelemetryStoreId.cs
+++ b/Engine/Source/Programs/Shared/EpicGames.Horde/Telemetry/TelemetryStoreId.cs
@@ -15,7 +15,7 @@ namespace EpicGames.Horde.Telemetry
 	[TypeConverter(typeof(StringIdTypeConverter<TelemetryStoreId, TelemetryStoreIdConverter>))]
 	[StringIdConverter(typeof(TelemetryStoreIdConverter))]
 	[CbConverter(typeof(StringIdCbConverter<TelemetryStoreId, TelemetryStoreIdConverter>))]
-	public record struct TelemetryStoreId(StringId Id)
+	public readonly record struct TelemetryStoreId(StringId Id)
 	{
 		/// <summary>
 		/// Default telemetry store for Horde internal metrics
diff --git a/Engine/Source/Programs/Shared/EpicGames.Horde/Tools/ToolDeploymentId.cs b/Engine/Source/Programs/Shared/EpicGames.Horde/Tools/ToolDeploymentId.cs
--- a/Engine/Source/Programs/Shared/EpicGames.Horde/Tools/ToolDeploymentId.cs
+++ b/Engine/Source/Programs/Shared/EpicGames.Horde/Tools/ToolDeploymentId.cs
@@ -13,7 +13,7 @@ namespace EpicGames.Horde.Tools
 	[JsonSchemaString]
 	[TypeConverter(typeof(BinaryIdTypeConverter<ToolDeploymentId, ToolDeploymentIdConverter>))]
 	[BinaryIdConverter(typeof(ToolDeploymentIdConverter))]
-	public record struct ToolDeploymentId(BinaryId Id)
+	public readonly record struct ToolDeploymentId(BinaryId Id)
 	{
 		/// <inheritdoc cref="BinaryId.Parse(System.String)"/>
 		public static ToolDeploymentId Parse(string text) => new ToolDeploymentId(BinaryId.Parse(text));
diff --git a/Engine/Source/Programs/Shared/EpicGames.Horde/Tools/ToolId.cs b/Engine/Source/Programs/Shared/EpicGames.Horde/Tools/ToolId.cs
--- a/Engine/Source/Programs/Shared/EpicGames.Horde/Tools/ToolId.cs
+++ b/Engine/Source/Programs/Shared/EpicGames.Horde/Tools/ToolId.cs
@@ -15,7 +15,7 @@ namespace EpicGames.Horde.Tools
 	[TypeConverter(typeof(StringIdTypeConverter<ToolId, ToolIdConverter>))]
 	[StringIdConverter(typeof(ToolIdConverter))]
 	[CbConverter(typeof(StringIdCbConverter<ToolId, ToolIdConverter>))]
-	public record struct ToolId(StringId Id)
+	public readonly record struct ToolId(StringId Id)
 	{
 		/// <summary>
 		/// Constructor
diff --git a/Engine/Source/Programs/Shared/EpicGames.Horde/Users/UserId.cs b/Engine/Source/Programs/Shared/EpicGames.Horde/Users/UserId.cs
--- a/Engine/Source/Programs/Shared/EpicGames.Horde/Users/UserId.cs
+++ b/Engine/Source/Programs/Shared/EpicGames.Horde/Users/UserId.cs
@@ -13,7 +13,7 @@ namespace EpicGames.Horde.Users
 	[JsonSchemaString]
 	[TypeConverter(typeof(BinaryIdTypeConverter<UserId, UserIdConverter>))]
 	[BinaryIdConverter(typeof(UserIdConverter))]
-	public record struct UserId(BinaryId Id)
+	public readonly record struct UserId(BinaryId Id)
 	{
 		/// <summary>
 		/// Constant value for empty user id
diff --git a/Engine/Source/Programs/Shared/EpicGames.Perforce/Records/FileType.cs b/Engine/Source/Programs/Shared/EpicGames.Perforce/Records/FileType.cs
--- a/Engine/Source/Programs/Shared/EpicGames.Perforce/Records/FileType.cs
+++ b/Engine/Source/Programs/Shared/EpicGames.Perforce/Records/FileType.cs
@@ -112,7 +112,7 @@ namespace EpicGames.Perforce.Managed
 	/// <summary>
 	/// Indicates the type of a file
 	/// </summary>
-	public struct FileType : IEquatable<FileType>
+	public readonly struct FileType : IEquatable<FileType>
 	{
 		/// <summary>
 		/// Size of the object when serialized to memory
