Horde: Fix auth state not being invalidated when interactive login settings change.

diff --git a/Engine/Source/Programs/Shared/EpicGames.Horde/HordeHttpAuthHandler.cs b/Engine/Source/Programs/Shared/EpicGames.Horde/HordeHttpAuthHandler.cs
--- a/Engine/Source/Programs/Shared/EpicGames.Horde/HordeHttpAuthHandler.cs
+++ b/Engine/Source/Programs/Shared/EpicGames.Horde/HordeHttpAuthHandler.cs
@@ -61,10 +61,11 @@ namespace EpicGames.Horde
 						{
 							return response;
 						}
-
-						_authState.Invalidate(_authHeader);
 					}
 
+					// Invalidate the current auth header
+					_authState.Invalidate(_authHeader);
+
 					// Otherwise update the auth header and try again
 					_authHeader = await _authState.TryGetAuthHeaderAsync(cancellationToken);
 					if (_authHeader != null)
@@ -171,12 +172,12 @@ namespace EpicGames.Horde
 		/// Invalidate a cached header value
 		/// </summary>
 		/// <param name="authHeader">The auth header to invalidate</param>
-		public void Invalidate(AuthenticationHeaderValue authHeader)
+		public void Invalidate(AuthenticationHeaderValue? authHeader)
 		{
 			lock (_lockObject)
 			{
 #pragma warning disable VSTHRD002
-				if (_authStateTask != null && _authStateTask.IsCompleted && Object.Equals(_authStateTask.Result?.TokenInfo?.AccessToken, authHeader.Parameter))
+				if (_authStateTask != null && _authStateTask.IsCompleted && Object.Equals(_authStateTask.Result?.TokenInfo?.AccessToken, authHeader?.Parameter))
 				{
 					_authStateTask = null;
 				}
