UGS: Disable hang monitor for Perforce connections.

diff --git a/Engine/Source/Programs/Shared/EpicGames.Perforce/PerforceSettings.cs b/Engine/Source/Programs/Shared/EpicGames.Perforce/PerforceSettings.cs
--- a/Engine/Source/Programs/Shared/EpicGames.Perforce/PerforceSettings.cs
+++ b/Engine/Source/Programs/Shared/EpicGames.Perforce/PerforceSettings.cs
@@ -51,6 +51,11 @@ namespace EpicGames.Perforce
 		/// Whether to create a native client rather than running the p4 child process, if possible
 		/// </summary>
 		public bool PreferNativeClient { get; }
+
+		/// <summary>
+		/// Whether to monitor the connection for commands that are not producing any output
+		/// </summary>
+		public bool EnableHangMonitor { get; }
 	}
 
 	/// <summary>
@@ -110,6 +115,9 @@ namespace EpicGames.Perforce
 		/// <inheritdoc/>
 		public bool PreferNativeClient { get; set; }
 
+		/// <inheritdoc/>
+		public bool EnableHangMonitor { get; set; }
+
 		/// <summary>
 		/// Default constructor
 		/// </summary>
@@ -120,6 +128,7 @@ namespace EpicGames.Perforce
 			Password = environment.GetValue("P4PASSWD");
 			HostName = environment.GetValue("P4HOST");
 			ClientName = environment.GetValue("P4CLIENT");
+			EnableHangMonitor = true;
 		}
 
 		/// <summary>
@@ -131,6 +140,7 @@ namespace EpicGames.Perforce
 		{
 			ServerAndPort = serverAndPort;
 			UserName = userName;
+			EnableHangMonitor = true;
 		}
 
 		/// <summary>
@@ -146,6 +156,7 @@ namespace EpicGames.Perforce
 			_appName = other.AppName;
 			_appVersion = other.AppVersion;
 			PreferNativeClient = other.PreferNativeClient;
+			EnableHangMonitor = other.EnableHangMonitor;
 		}
 
 		/// <summary>
diff --git a/Engine/Source/Programs/UnrealGameSync/UnrealGameSyncShared/WorkspaceUpdate.cs b/Engine/Source/Programs/UnrealGameSync/UnrealGameSyncShared/WorkspaceUpdate.cs
--- a/Engine/Source/Programs/UnrealGameSync/UnrealGameSyncShared/WorkspaceUpdate.cs
+++ b/Engine/Source/Programs/UnrealGameSync/UnrealGameSyncShared/WorkspaceUpdate.cs
@@ -583,7 +583,7 @@ namespace UnrealGameSync
 
 		public async Task<(WorkspaceUpdateResult, string)> ExecuteAsync(IPerforceSettings perforceSettings, ProjectInfo project, WorkspaceStateWrapper stateMgr, ILogger logger, CancellationToken cancellationToken)
 		{
-			using IPerforceConnection perforce = await PerforceConnection.CreateAsync(perforceSettings, logger);
+			using IPerforceConnection perforce = await PerforceConnection.CreateAsync(new PerforceSettings(perforceSettings) { EnableHangMonitor = false }, logger);
 
 			ReadOnlyWorkspaceState state = stateMgr.Current;
 
