OidcToken: Make OidcTokenManager more testable

Wraps it in an interface and adds two helper classes to make testing easier. Since they're used by other libraries, putting them here makes most sense even if they're only used during testing.

diff --git a/Engine/Source/Programs/Shared/EpicGames.OIDC/FakeOidcTokenManager.cs b/Engine/Source/Programs/Shared/EpicGames.OIDC/FakeOidcTokenManager.cs
new file mode 100644
--- /dev/null
+++ b/Engine/Source/Programs/Shared/EpicGames.OIDC/FakeOidcTokenManager.cs
@@ -0,0 +1,92 @@
+// Copyright Epic Games, Inc. All Rights Reserved.
+
+using System;
+using System.Threading;
+using System.Threading.Tasks;
+
+namespace EpicGames.OIDC;
+
+#pragma warning disable CA1721 // Property names should not match get methods
+
+/// <summary>
+/// A fake implementation of OidcTokenManager for testing
+/// Assumes just a single provider.
+/// </summary>
+public class FakeOidcTokenManager : IOidcTokenManager
+{
+	/// <summary>
+	/// Refresh token in use
+	/// </summary>
+	public string? RefreshToken { get; set; }
+	
+	/// <summary>
+	/// Latest access token
+	/// </summary>
+	public string? AccessToken { get; set; }
+	
+	/// <summary>
+	/// Expiry time for access token
+	/// </summary>
+	public DateTimeOffset AccessTokenExpiry { get; set; } = DateTimeOffset.UnixEpoch;
+	
+	private int _refreshCounter = 1;
+	private int _accessCounter = 1;
+	
+	/// <inheritdoc/>
+	public Task<OidcTokenInfo> LoginAsync(string providerIdentifier, CancellationToken cancellationToken = default)
+	{
+		Console.WriteLine("FakeOidcTokenManager.LoginAsync");
+		RefreshToken = "fakeRefreshToken-" + _refreshCounter++;
+		RefreshAccessToken();
+		return Task.FromResult(GetOidcTokenInfo());
+	}
+
+	/// <inheritdoc/>
+	public Task<OidcTokenInfo> GetAccessToken(string providerIdentifier, CancellationToken cancellationToken = default)
+	{
+		throw new NotImplementedException("Method not in use");
+	}
+
+	/// <inheritdoc/>
+	public Task<OidcTokenInfo?> TryGetAccessToken(string providerIdentifier, CancellationToken cancellationToken = default)
+	{
+		Console.WriteLine("FakeOidcTokenManager.TryGetAccessToken");
+		if (String.IsNullOrEmpty(RefreshToken))
+		{
+			throw new NotLoggedInException();
+		}
+		
+		bool tokenValid = AccessTokenExpiry.AddMinutes(2) > DateTime.Now || AccessTokenExpiry == DateTimeOffset.MinValue;
+		if (!String.IsNullOrEmpty(AccessToken) && tokenValid)
+		{
+			return Task.FromResult(GetOidcTokenInfo())!;
+		}
+		
+		// Refresh access token with provider
+		RefreshAccessToken();
+		
+		// Assume the updated access token is valid
+		return Task.FromResult(GetOidcTokenInfo())!;
+	}
+
+	/// <inheritdoc/>
+	public OidcStatus GetStatusForProvider(string providerIdentifier)
+	{
+		Console.WriteLine("FakeOidcTokenManager.GetStatusForProvider");
+		return OidcTokenClient.GetStatus(RefreshToken, AccessToken, AccessTokenExpiry);
+	}
+
+	private OidcTokenInfo GetOidcTokenInfo()
+	{
+		return new OidcTokenInfo { RefreshToken = RefreshToken, AccessToken = AccessToken, TokenExpiry = AccessTokenExpiry };
+	}
+
+	/// <summary>
+	/// Perform fake refresh of the access token
+	/// </summary>
+	private void RefreshAccessToken()
+	{
+		AccessToken = "fakeAccessToken-" + _accessCounter++;;
+		AccessTokenExpiry = DateTimeOffset.UtcNow + TimeSpan.FromHours(1);
+	}
+}
diff --git a/Engine/Source/Programs/Shared/EpicGames.OIDC/InMemoryTokenStore.cs b/Engine/Source/Programs/Shared/EpicGames.OIDC/InMemoryTokenStore.cs
new file mode 100644
--- /dev/null
+++ b/Engine/Source/Programs/Shared/EpicGames.OIDC/InMemoryTokenStore.cs
@@ -0,0 +1,45 @@
+// Copyright Epic Games, Inc. All Rights Reserved.
+
+using System;
+using System.Collections.Generic;
+
+namespace EpicGames.OIDC;
+
+/// <summary>
+/// An in-memory token store without encryption, intended for testing
+/// </summary>
+public sealed class InMemoryTokenStore : ITokenStore
+{
+	private readonly Dictionary<string, string> _providerToRefreshToken = new ();
+
+	/// <inheritdoc/>
+	public bool TryGetRefreshToken(string oidcProvider, out string refreshToken)
+	{
+		if (!_providerToRefreshToken.TryGetValue(oidcProvider, out string? result))
+		{
+			refreshToken = "";
+			return false;
+		}
+
+		refreshToken = result;
+		return true;
+	}
+
+	/// <inheritdoc/>
+	public void AddRefreshToken(string providerIdentifier, string refreshToken)
+	{
+		_providerToRefreshToken[providerIdentifier] = refreshToken;
+	}
+
+	/// <inheritdoc/>
+	public void Save()
+	{
+		// No need - data exists only in-memory
+	}
+
+	/// <inheritdoc/>
+	public void Dispose()
+	{
+		GC.SuppressFinalize(this);
+	}
+}
diff --git a/Engine/Source/Programs/Shared/EpicGames.OIDC/OidcTokenManager.cs b/Engine/Source/Programs/Shared/EpicGames.OIDC/OidcTokenManager.cs
--- a/Engine/Source/Programs/Shared/EpicGames.OIDC/OidcTokenManager.cs
+++ b/Engine/Source/Programs/Shared/EpicGames.OIDC/OidcTokenManager.cs
@@ -24,7 +24,18 @@ using Microsoft.Extensions.Options;
 
 namespace EpicGames.OIDC
 {
-	public class OidcTokenManager
+	public interface IOidcTokenManager
+	{
+		public Task<OidcTokenInfo> LoginAsync(string providerIdentifier, CancellationToken cancellationToken = default);
+
+		public Task<OidcTokenInfo> GetAccessToken(string providerIdentifier, CancellationToken cancellationToken = default);
+
+		public Task<OidcTokenInfo?> TryGetAccessToken(string providerIdentifier, CancellationToken cancellationToken = default);
+
+		public OidcStatus GetStatusForProvider(string providerIdentifier);
+	}
+	
+	public class OidcTokenManager : IOidcTokenManager
 	{
 		private readonly object _lockObject = new object();
 		private readonly Dictionary<string, OidcTokenClient> _tokenClients = new Dictionary<string, OidcTokenClient>(StringComparer.OrdinalIgnoreCase);
@@ -126,6 +137,7 @@ namespace EpicGames.OIDC
 			return _tokenClients.Any(pair => pair.Value.GetStatus() == OidcStatus.NotLoggedIn);
 		}
 
+		/// <inheritdoc/>
 		public async Task<OidcTokenInfo> LoginAsync(string providerIdentifier, CancellationToken cancellationToken = default)
 		{
 			OidcTokenClient tokenClient = _tokenClients[providerIdentifier];
@@ -140,16 +152,19 @@ namespace EpicGames.OIDC
 			return tokenInfo;
 		}
 
+		/// <inheritdoc/>
 		public Task<OidcTokenInfo> GetAccessToken(string providerIdentifier, CancellationToken cancellationToken = default)
 		{
 			return _tokenClients[providerIdentifier].GetAccessTokenAsync(cancellationToken);
 		}
 
+		/// <inheritdoc/>
 		public Task<OidcTokenInfo?> TryGetAccessToken(string providerIdentifier, CancellationToken cancellationToken = default)
 		{
 			return _tokenClients[providerIdentifier].TryGetAccessTokenAsync(cancellationToken);
 		}
 
+		/// <inheritdoc/>
 		public OidcStatus GetStatusForProvider(string providerIdentifier)
 		{
 			return _tokenClients[providerIdentifier].GetStatus();
@@ -158,8 +173,19 @@ namespace EpicGames.OIDC
 
 	public enum OidcStatus
 	{
+		/// <summary>
+		/// Both access and refresh token are valid
+		/// </summary>
 		Connected,
+		
+		/// <summary>
+		/// No refresh token is set and requires login
+		/// </summary>
 		NotLoggedIn,
+		
+		/// <summary>
+		/// Access token has not been generated or has expired, and needs refreshing 
+		/// </summary>
 		TokenRefreshRequired
 	}
 
@@ -167,6 +193,10 @@ namespace EpicGames.OIDC
 	{
 		public string? RefreshToken { get; set; }
 		public string? AccessToken { get; set; }
+		
+		/// <summary>
+		/// When access token expires
+		/// </summary>
 		public DateTimeOffset TokenExpiry { get; set; }
 
 		public bool IsValid => RefreshToken != null && AccessToken != null;
@@ -562,17 +592,22 @@ namespace EpicGames.OIDC
 
 		public OidcStatus GetStatus()
 		{
-			if (String.IsNullOrEmpty(_refreshToken))
+			return GetStatus(_refreshToken, _accessToken, _tokenExpiry);
+		}
+		
+		public static OidcStatus GetStatus(string? refreshToken, string? accessToken, DateTimeOffset accessTokenExpiry)
+		{
+			if (String.IsNullOrEmpty(refreshToken))
 			{
 				return OidcStatus.NotLoggedIn;
 			}
 
-			if (String.IsNullOrEmpty(_accessToken))
+			if (String.IsNullOrEmpty(accessToken))
 			{
 				return OidcStatus.TokenRefreshRequired;
 			}
 
-			if (_tokenExpiry < DateTime.Now)
+			if (accessTokenExpiry < DateTime.Now)
 			{
 				return OidcStatus.TokenRefreshRequired;
 			}
