Horde: Fix some auth handling issues in UGS.

* OIDC profile was using a hard-coded name of "Horde".
* Interactive logins were always being triggered from background threads.
* Non-interactive logins were not triggering token refresh.

diff --git a/Engine/Source/Programs/Shared/EpicGames.Horde/HordeClient.cs b/Engine/Source/Programs/Shared/EpicGames.Horde/HordeClient.cs
--- a/Engine/Source/Programs/Shared/EpicGames.Horde/HordeClient.cs
+++ b/Engine/Source/Programs/Shared/EpicGames.Horde/HordeClient.cs
@@ -53,15 +53,7 @@ namespace EpicGames.Horde
 		/// <inheritdoc/>
 		public async Task<bool> ConnectAsync(bool allowLogin, CancellationToken cancellationToken)
 		{
-			if (_authHandlerState.TryGetConfiguredAuthHeader() != null)
-			{
-				return true;
-			}
-
-			_authHandlerState.Invalidate();
-			await _authHandlerState.RefreshAsync(allowLogin, cancellationToken);
-
-			return IsConnected();
+			return await _authHandlerState.LoginAsync(allowLogin, cancellationToken);
 		}
 
 		/// <inheritdoc/>
@@ -69,7 +61,7 @@ namespace EpicGames.Horde
 		{
 			try
 			{
-				return _authHandlerState.IsAuthenticated();
+				return _authHandlerState.IsLoggedIn();
 			}
 			catch
 			{
diff --git a/Engine/Source/Programs/Shared/EpicGames.Horde/HordeHttpAuthHandler.cs b/Engine/Source/Programs/Shared/EpicGames.Horde/HordeHttpAuthHandler.cs
--- a/Engine/Source/Programs/Shared/EpicGames.Horde/HordeHttpAuthHandler.cs
+++ b/Engine/Source/Programs/Shared/EpicGames.Horde/HordeHttpAuthHandler.cs
@@ -27,8 +27,6 @@ namespace EpicGames.Horde
 		readonly HordeHttpAuthHandlerState _authState;
 		readonly IOptions<HordeOptions> _options;
 
-		AuthenticationHeaderValue? _authHeader;
-
 		/// <summary>
 		/// Constructor
 		/// </summary>
@@ -41,45 +39,42 @@ namespace EpicGames.Horde
 		/// <inheritdoc/>
 		protected override async Task<HttpResponseMessage> SendAsync(HttpRequestMessage request, CancellationToken cancellationToken)
 		{
-			if (request.Headers.Authorization == null)
+			// If the request already has a custom auth header, send the request as it is
+			if (request.Headers.Authorization != null)
 			{
-				AuthenticationHeaderValue? configuredAuthHeader = _authState.TryGetConfiguredAuthHeader();
-				if (configuredAuthHeader != null)
-				{
-					// Use the configured header
-					request.Headers.Authorization = configuredAuthHeader;
-				}
-				else if (_options.Value.AllowAuthPrompt)
-				{
-					// Try to use the cached auth header
-					if (_authHeader != null)
-					{
-						request.Headers.Authorization = _authHeader;
+				return await base.SendAsync(request, cancellationToken);
+			}
 
-						HttpResponseMessage response = await base.SendAsync(request, cancellationToken);
-						if (response.StatusCode != HttpStatusCode.Unauthorized)
-						{
-							return response;
-						}
-					}
+			// Get the current access token and send the request with that
+			string? accessToken = await _authState.GetAccessTokenAsync(_options.Value.AllowAuthPrompt, cancellationToken);
+			for (int attempt = 0; ; attempt++)
+			{
+				// Attempt to perform the request with this access token
+				request.Headers.Authorization = (accessToken == null)? null : new AuthenticationHeaderValue("Bearer", accessToken);
+				HttpResponseMessage response = await base.SendAsync(request, cancellationToken);
 
-					// Invalidate the current auth header
-					_authState.Invalidate(_authHeader);
+				const int MaxAttempts = 3;
+				if (response.StatusCode != HttpStatusCode.Unauthorized || attempt >= MaxAttempts)
+				{
+					return response;
+				}
 
-					// Otherwise update the auth header and try again
-					_authHeader = await _authState.TryGetAuthHeaderAsync(cancellationToken);
-					if (_authHeader != null)
-					{
-						request.Headers.Authorization = _authHeader;
-					}
+				// Mark this access token as invalid
+				if (accessToken != null)
+				{
+					_authState.Invalidate(accessToken);
 				}
-				else
+
+				// Get the next token, and quit out if it's the same
+				string? nextAccessToken = await _authState.GetAccessTokenAsync(_options.Value.AllowAuthPrompt, cancellationToken);
+				if (String.Equals(accessToken, nextAccessToken, StringComparison.Ordinal))
 				{
-					// Use whatever cached auth header we currently have
-					request.Headers.Authorization = _authHeader;
+					return response;
 				}
+
+				// Otherwise update the token and try again
+				accessToken = nextAccessToken;
 			}
-			return await base.SendAsync(request, cancellationToken);
 		}
 	}
 
@@ -93,10 +88,10 @@ namespace EpicGames.Horde
 		/// </summary>
 		public const string HttpClientName = "HordeHttpAuthState";
 
-		record class AuthState(AuthMethod Method, OidcTokenInfo? TokenInfo)
+		record class AuthState(AuthMethod Method, OidcTokenInfo? TokenInfo, bool Interactive)
 		{
 			public bool IsAuthorized()
-				=> (Method == AuthMethod.Anonymous) || (TokenInfo != null && TokenInfo.IsValid && TokenInfo.TokenExpiry > DateTimeOffset.Now);
+				=> (Method == AuthMethod.Anonymous) || (TokenInfo != null && TokenInfo.IsValid);
 		}
 
 		readonly object _lockObject = new object();
@@ -139,13 +134,13 @@ namespace EpicGames.Horde
 		/// <summary>
 		/// Checks if we have a valid auth header at the moment
 		/// </summary>
-		public bool IsAuthenticated()
+		public bool IsLoggedIn()
 		{
-			if (TryGetConfiguredAuthHeader() != null)
+			if (GetAccessTokenFromConfig() != null)
 			{
 				return true;
 			}
-			if (_authStateTask != null && _authStateTask.TryGetResult(out AuthState? authState) && authState.IsAuthorized())
+			if (_authStateTask != null && _authStateTask.IsCompletedSuccessfully && _authStateTask.TryGetResult(out AuthState? authState) && authState.IsAuthorized())
 			{
 				return true;
 			}
@@ -153,31 +148,15 @@ namespace EpicGames.Horde
 		}
 
 		/// <summary>
-		/// Invalidate the cached header value
-		/// </summary>
-		public void Invalidate()
-		{
-			lock (_lockObject)
-			{
-#pragma warning disable VSTHRD002
-				if (_authStateTask != null && _authStateTask.IsCompleted)
-				{
-					_authStateTask = null;
-				}
-#pragma warning restore VSTHRD002
-			}
-		}
-
-		/// <summary>
-		/// Invalidate a cached header value
+		/// Marks the given access token as invalid, having attempted to use it and got an unauthorized response
 		/// </summary>
-		/// <param name="authHeader">The auth header to invalidate</param>
-		public void Invalidate(AuthenticationHeaderValue? authHeader)
+		/// <param name="accessToken">The access  header to invalidate</param>
+		public void Invalidate(string? accessToken)
 		{
 			lock (_lockObject)
 			{
 #pragma warning disable VSTHRD002
-				if (_authStateTask != null && _authStateTask.IsCompleted && Object.Equals(_authStateTask.Result?.TokenInfo?.AccessToken, authHeader?.Parameter))
+				if (_authStateTask != null && _authStateTask.IsCompleted && Object.Equals(_authStateTask.Result?.TokenInfo?.AccessToken, accessToken))
 				{
 					_authStateTask = null;
 				}
@@ -188,28 +167,15 @@ namespace EpicGames.Horde
 		/// <summary>
 		/// Try to get a configured auth header
 		/// </summary>
-		public AuthenticationHeaderValue? TryGetConfiguredAuthHeader()
+		public string? GetAccessTokenFromConfig()
 		{
+			// If an explicit access token is specified, just use that
 			if (_options.Value.AccessToken != null)
 			{
-				// If an explicit access token is specified, just use that
-				return new AuthenticationHeaderValue("Bearer", _options.Value.AccessToken);
-			}
-			else if (TryGetAccessTokenFromEnvironment(out string? accessToken))
-			{
-				// Use the access token specified in the environment
-				return new AuthenticationHeaderValue("Bearer", accessToken);
-			}
-			else
-			{
-				// Will need to login asynchronously
-				return null;
+				return _options.Value.AccessToken;
 			}
-		}
 
-		bool TryGetAccessTokenFromEnvironment(out string? accessToken)
-		{
-			// Only use the token from the environment if the configured base address is missing or matches the one configured in the environment
+			// Check environment variables for an access token matching the current server
 			string? hordeUrlEnvVar = Environment.GetEnvironmentVariable(HordeHttpClient.HordeUrlEnvVarName);
 			if (!String.IsNullOrEmpty(hordeUrlEnvVar))
 			{
@@ -219,59 +185,75 @@ namespace EpicGames.Horde
 					string? hordeToken = Environment.GetEnvironmentVariable(HordeHttpClient.HordeTokenEnvVarName);
 					if (!String.IsNullOrEmpty(hordeToken))
 					{
-						accessToken = hordeToken;
-						return true;
+						return hordeToken;
 					}
 				}
 			}
 
-			accessToken = null;
-			return false;
+			// Otherwise we need to find the access token dynamically
+			return null;
 		}
 
 		/// <summary>
 		/// Refresh the auth state
 		/// </summary>
-		/// <param name="allowLogin">Whether to allow logging in</param>
-		/// <param name="cancellationToken"></param>
+		/// <param name="interactive">Whether to allow logging in interactively</param>
+		/// <param name="cancellationToken">Cancellation token for the operation</param>
 		/// <returns></returns>
-		public async Task RefreshAsync(bool allowLogin, CancellationToken cancellationToken)
+		public async Task<bool> LoginAsync(bool interactive, CancellationToken cancellationToken)
 		{
-			Invalidate();
-			await GetAuthStateAsync(allowLogin, cancellationToken);
+			if (GetAccessTokenFromConfig() != null)
+			{
+				return true;
+			}
+
+			AuthState? state = await GetAuthStateAsync(interactive, cancellationToken);
+			return state?.IsAuthorized() ?? false;
 		}
 
 		/// <summary>
-		/// Gets a new auth header
+		/// Gets the current access token
 		/// </summary>
-		/// <param name="cancellationToken">Cancellation token for the operation</param>
-		public async ValueTask<AuthenticationHeaderValue?> TryGetAuthHeaderAsync(CancellationToken cancellationToken)
+		public async Task<string?> GetAccessTokenAsync(bool allowAuthPrompt, CancellationToken cancellationToken)
 		{
-			AuthState? authState = await GetAuthStateAsync(true, cancellationToken);
-			if (authState?.TokenInfo == null)
+			string? accessToken = GetAccessTokenFromConfig();
+			if (accessToken != null)
 			{
-				return null;
+				return accessToken;
 			}
-			return new AuthenticationHeaderValue("Bearer", authState.TokenInfo.AccessToken);
+
+			AuthState? authState = await GetAuthStateAsync(allowAuthPrompt, cancellationToken);
+			return authState?.TokenInfo?.AccessToken;
 		}
 
-		async Task<AuthState?> GetAuthStateAsync(bool allowLogin, CancellationToken cancellationToken)
+		async Task<AuthState?> GetAuthStateAsync(bool interactive, CancellationToken cancellationToken)
 		{
-			if (TryGetConfiguredAuthHeader() != null)
+			if (GetAccessTokenFromConfig() != null)
 			{
 				return null;
 			}
 
-			Task<AuthState> authStateTask;
-			lock (_lockObject)
+			Task<AuthState>? authStateTask = null;
+			for (; ; )
 			{
-				_authStateTask ??= Task.Run(() => GetAuthStateInternalAsync(allowLogin, _cancellationTokenSource.Token), _cancellationTokenSource.Token);
-				authStateTask = _authStateTask;
+				lock (_lockObject)
+				{
+					if (_authStateTask == null || _authStateTask == authStateTask)
+					{
+						_authStateTask = Task.Run(() => GetAuthStateInternalAsync(interactive, _cancellationTokenSource.Token), _cancellationTokenSource.Token);
+					}
+					authStateTask = _authStateTask;
+				}
+
+				AuthState authState = await authStateTask.WaitAsync(cancellationToken);
+				if (authState.IsAuthorized() || !interactive || authState.Interactive)
+				{
+					return authState;
+				}
 			}
-			return await authStateTask.WaitAsync(cancellationToken);
 		}
 
-		async Task<AuthState> GetAuthStateInternalAsync(bool allowLogin, CancellationToken cancellationToken)
+		async Task<AuthState> GetAuthStateInternalAsync(bool interactive, CancellationToken cancellationToken)
 		{
 			Uri serverUrl;
 
@@ -298,7 +280,7 @@ namespace EpicGames.Horde
 
 			if (authConfig.Method == AuthMethod.Anonymous)
 			{
-				return new AuthState(authConfig.Method, null);
+				return new AuthState(authConfig.Method, null, true);
 			}
 
 			string? localRedirectUrl = authConfig.LocalRedirectUrls?.FirstOrDefault();
@@ -310,7 +292,7 @@ namespace EpicGames.Horde
 			string oidcProvider = authConfig.ProfileName ?? "Horde";
 
 			Dictionary<string, string?> values = new Dictionary<string, string?>();
-			values[$"Providers:{oidcProvider}:DisplayName"] = "Horde";
+			values[$"Providers:{oidcProvider}:DisplayName"] = oidcProvider;
 			values[$"Providers:{oidcProvider}:ServerUri"] = authConfig.ServerUrl;
 			values[$"Providers:{oidcProvider}:ClientId"] = authConfig.ClientId;
 			values[$"Providers:{oidcProvider}:RedirectUri"] = localRedirectUrl;
@@ -335,13 +317,13 @@ namespace EpicGames.Horde
 					_logger.LogTrace(ex, "Unable to get access token; attempting login: {Message}", ex.Message);
 				}
 			}
-			if (result == null && allowLogin)
+			if (result == null && interactive)
 			{
 				_logger.LogInformation("Logging in to {Server}...", serverUrl);
 				result = await oidcTokenManager.LoginAsync(oidcProvider, cancellationToken);
 			}
 
-			return new AuthState(authConfig.Method, result);
+			return new AuthState(authConfig.Method, result, interactive);
 		}
 	}
 }
diff --git a/Engine/Source/Programs/UnrealGameSync/UnrealGameSync/Program.cs b/Engine/Source/Programs/UnrealGameSync/UnrealGameSync/Program.cs
--- a/Engine/Source/Programs/UnrealGameSync/UnrealGameSync/Program.cs
+++ b/Engine/Source/Programs/UnrealGameSync/UnrealGameSync/Program.cs
@@ -209,7 +209,11 @@ namespace UnrealGameSync
 
 				if (launcherSettings.HordeServer != null)
 				{
-					services.AddHorde(options => options.ServerUrl = new Uri(launcherSettings.HordeServer));
+					services.AddHorde(options =>
+					{
+						options.ServerUrl = new Uri(launcherSettings.HordeServer);
+						options.AllowAuthPrompt = false;
+					});
 				}
 
 				ServiceProvider serviceProvider = services.BuildServiceProvider();
diff --git a/Engine/Source/Programs/UnrealGameSync/UnrealGameSync/ProgramApplicationContext.cs b/Engine/Source/Programs/UnrealGameSync/UnrealGameSync/ProgramApplicationContext.cs
--- a/Engine/Source/Programs/UnrealGameSync/UnrealGameSync/ProgramApplicationContext.cs
+++ b/Engine/Source/Programs/UnrealGameSync/UnrealGameSync/ProgramApplicationContext.cs
@@ -260,7 +260,7 @@ namespace UnrealGameSync
 
 			if (hordeAuthState != null)
 			{
-				await hordeAuthState.RefreshAsync(false, CancellationToken.None);
+				await hordeAuthState.LoginAsync(false, CancellationToken.None);
 			}
 		}
 
