Horde: Add tests for HordeHttpAuthHandler

Primarily to debug why access tokens aren't refreshed (UE-209637). But it turns out moving auth handling from UBA executor over to this handler did fix the issue.

diff --git a/Engine/Source/Programs/Shared/EpicGames.Horde.Tests/HordeHttpTests.cs b/Engine/Source/Programs/Shared/EpicGames.Horde.Tests/HordeHttpTests.cs
new file mode 100644
--- /dev/null
+++ b/Engine/Source/Programs/Shared/EpicGames.Horde.Tests/HordeHttpTests.cs
@@ -0,0 +1,137 @@
+// Copyright Epic Games, Inc. All Rights Reserved.
+
+using System;
+using System.Collections.Generic;
+using System.Diagnostics.CodeAnalysis;
+using System.Linq;
+using System.Net;
+using System.Net.Http;
+using System.Text.Json;
+using System.Threading;
+using System.Threading.Tasks;
+using EpicGames.Horde.Server;
+using EpicGames.OIDC;
+using Microsoft.Extensions.Logging;
+using Microsoft.Extensions.Options;
+using Microsoft.VisualStudio.TestTools.UnitTesting;
+
+namespace EpicGames.Horde.Tests;
+
+public class StubHttpClientFactory : IHttpClientFactory
+{
+	private readonly HttpClient _httpClient;
+
+	public StubHttpClientFactory(HttpClient httpClient)
+	{
+		_httpClient = httpClient;
+	}
+
+	public HttpClient CreateClient(string name)
+	{
+		return _httpClient;
+	}
+}
+
+public class StubMessageHandler : HttpMessageHandler
+{
+	public HttpStatusCode StatusCode { get; }
+	public string Content { get; }
+	public List<HttpRequestMessage> HttpRequests { get; } = new();
+
+	public StubMessageHandler(HttpStatusCode statusCode, string content)
+	{
+		StatusCode = statusCode;
+		Content = content;
+	}
+	
+	protected override Task<HttpResponseMessage> SendAsync(HttpRequestMessage request, CancellationToken cancellationToken)
+	{
+		Console.WriteLine($"Saving request {request.RequestUri} {request.Headers.Authorization}");
+		HttpRequests.Add(request);
+		return Task.FromResult(new HttpResponseMessage(StatusCode) { Content = new StringContent(Content) });
+	}
+}
+
+[TestClass]
+public class HordeHttpAuthHandlerTests
+{
+	const string HordeServerUrl = "http://horde-server-test";
+	const string FakeServerResponse = "fakeServerResponse";
+	
+	[TestMethod]
+	public async Task AccessToken_Valid_IsReusedAsync()
+	{
+		FakeOidcTokenManager oidc = new ();
+		(HttpClient client, StubMessageHandler server) = CreateClientServer(oidc);
+
+		await SendHttpRequestAsync(client);
+		string firstAccessToken = GetLastUsedAccessToken(server);
+		Assert.AreEqual(oidc.AccessToken, firstAccessToken);
+		
+		await SendHttpRequestAsync(client);
+		Assert.AreEqual(oidc.AccessToken, GetLastUsedAccessToken(server));
+		Assert.AreEqual(firstAccessToken, GetLastUsedAccessToken(server));
+	}
+	
+	[TestMethod]
+	public async Task AccessToken_Expired_IsRefreshedAsync()
+	{
+		FakeOidcTokenManager oidc = new ();
+		(HttpClient client, StubMessageHandler server) = CreateClientServer(oidc);
+
+		oidc.RefreshToken = "someRefreshToken";
+		oidc.AccessToken = "expiredAccessToken";
+		
+		// Set it as expired two hours ago, causing a refresh of access token
+		oidc.AccessTokenExpiry = DateTimeOffset.Now.AddHours(-2);
+		
+		await SendHttpRequestAsync(client);
+		Assert.AreEqual(oidc.AccessToken, GetLastUsedAccessToken(server));
+	}
+
+	[SuppressMessage("Reliability", "CA2000:Dispose objects before losing scope")]
+	private static (HttpClient client, StubMessageHandler server) CreateClientServer(IOidcTokenManager oidcTokenManager, HordeOptions? hordeOptions = null)
+	{
+		GetAuthConfigResponse authConfig = new () { Method = AuthMethod.OpenIdConnect, ServerUrl = HordeServerUrl, LocalRedirectUrls = new [] { HordeServerUrl } };
+		IHttpClientFactory httpFactory = GetHordeHttpClientFactory(HordeServerUrl, authConfig);
+		using ILoggerFactory loggerFactory = LoggerFactory.Create(builder =>
+		{
+			builder.SetMinimumLevel(LogLevel.Debug);
+			builder.AddSimpleConsole(options => { options.SingleLine = true; });
+		});
+
+		ILogger<HordeHttpAuthHandler> logger = loggerFactory.CreateLogger<HordeHttpAuthHandler>();
+
+		OptionsWrapper<HordeOptions> options = new (hordeOptions ?? new HordeOptions());
+		InMemoryTokenStore inMemoryTokenStore = new ();
+		HordeHttpAuthHandlerState state = new (httpFactory, options, logger, inMemoryTokenStore, oidcTokenManager);
+		HordeHttpAuthHandler authHandler = new (state, options);
+		
+		StubMessageHandler server = new (HttpStatusCode.OK, FakeServerResponse);
+		authHandler.InnerHandler = server;
+		HttpClient client = new (authHandler);
+		return (client, server);
+	}
+
+	private static string GetLastUsedAccessToken(StubMessageHandler server)
+	{
+		return server.HttpRequests.Last().Headers.Authorization!.Parameter!;
+	}
+	
+	private static async Task SendHttpRequestAsync(HttpClient client)
+	{
+		using HttpRequestMessage req = new (HttpMethod.Get, HordeServerUrl);
+		HttpResponseMessage res = await client.SendAsync(req);
+		Assert.AreEqual(HttpStatusCode.OK, res.StatusCode);
+		Assert.AreEqual(FakeServerResponse, await res.Content.ReadAsStringAsync());
+	}
+	
+	private static IHttpClientFactory GetHordeHttpClientFactory(string? baseAddress, GetAuthConfigResponse response)
+	{
+		using StubMessageHandler stubMessageHandler = new (HttpStatusCode.OK, JsonSerializer.Serialize(response));
+		return new StubHttpClientFactory(new HttpClient(stubMessageHandler)
+		{
+			BaseAddress = baseAddress != null ? new Uri(baseAddress) : null,
+		});
+	}
+}
diff --git a/Engine/Source/Programs/Shared/EpicGames.Horde/HordeHttpAuthHandler.cs b/Engine/Source/Programs/Shared/EpicGames.Horde/HordeHttpAuthHandler.cs
--- a/Engine/Source/Programs/Shared/EpicGames.Horde/HordeHttpAuthHandler.cs
+++ b/Engine/Source/Programs/Shared/EpicGames.Horde/HordeHttpAuthHandler.cs
@@ -100,15 +100,26 @@ namespace EpicGames.Horde
 		readonly IHttpClientFactory _httpClientFactory;
 		readonly IOptions<HordeOptions> _options;
 		readonly ILogger _logger;
+		
+		// Allow these to be overridden in tests
+		readonly ITokenStore? _tokenStore;
+		readonly IOidcTokenManager? _oidcTokenManager;
 
 		/// <summary>
 		/// Constructor
 		/// </summary>
-		public HordeHttpAuthHandlerState(IHttpClientFactory httpClientFactory, IOptions<HordeOptions> options, ILogger<HordeHttpAuthHandler> logger)
+		public HordeHttpAuthHandlerState(
+			IHttpClientFactory httpClientFactory,
+			IOptions<HordeOptions> options,
+			ILogger<HordeHttpAuthHandler> logger,
+			ITokenStore? tokenStore = null,
+			IOidcTokenManager? oidcTokenManager = null)
 		{
 			_httpClientFactory = httpClientFactory;
 			_options = options;
 			_logger = logger;
+			_tokenStore = tokenStore;
+			_oidcTokenManager = oidcTokenManager;
 		}
 
 		/// <inheritdoc/>
@@ -291,19 +302,8 @@ namespace EpicGames.Horde
 
 			string oidcProvider = authConfig.ProfileName ?? "Horde";
 
-			Dictionary<string, string?> values = new Dictionary<string, string?>();
-			values[$"Providers:{oidcProvider}:DisplayName"] = oidcProvider;
-			values[$"Providers:{oidcProvider}:ServerUri"] = authConfig.ServerUrl;
-			values[$"Providers:{oidcProvider}:ClientId"] = authConfig.ClientId;
-			values[$"Providers:{oidcProvider}:RedirectUri"] = localRedirectUrl;
-
-			ConfigurationBuilder builder = new ConfigurationBuilder();
-			builder.AddInMemoryCollection(values);
-
-			IConfiguration configuration = builder.Build();
-
-			using ITokenStore tokenStore = TokenStoreFactory.CreateTokenStore();
-			OidcTokenManager oidcTokenManager = OidcTokenManager.CreateTokenManager(configuration, tokenStore, new List<string>() { oidcProvider });
+			using ITokenStore tokenStore = CreateTokenStore();
+			IOidcTokenManager oidcTokenManager = CreateOidcTokenManager(tokenStore, authConfig, localRedirectUrl);
 
 			OidcTokenInfo? result = null;
 			if (oidcTokenManager.GetStatusForProvider(oidcProvider) != OidcStatus.NotLoggedIn)
@@ -325,5 +325,32 @@ namespace EpicGames.Horde
 
 			return new AuthState(authConfig.Method, result, interactive);
 		}
+		
+		private ITokenStore CreateTokenStore()
+		{
+			return _tokenStore ?? TokenStoreFactory.CreateTokenStore();
+		}
+
+		private IOidcTokenManager CreateOidcTokenManager(ITokenStore tokenStore, GetAuthConfigResponse authConfig, string localRedirectUrl)
+		{
+			if (_oidcTokenManager != null)
+			{
+				return _oidcTokenManager;
+			}
+			
+			string oidcProvider = authConfig.ProfileName ?? "Horde";
+
+			Dictionary<string, string?> values = new ();
+			values[$"Providers:{oidcProvider}:DisplayName"] = oidcProvider;
+			values[$"Providers:{oidcProvider}:ServerUri"] = authConfig.ServerUrl;
+			values[$"Providers:{oidcProvider}:ClientId"] = authConfig.ClientId;
+			values[$"Providers:{oidcProvider}:RedirectUri"] = localRedirectUrl;
+			
+			ConfigurationBuilder builder = new ();
+			builder.AddInMemoryCollection(values);
+			IConfiguration configuration = builder.Build();
+			
+			return OidcTokenManager.CreateTokenManager(configuration, tokenStore, new List<string>() { oidcProvider });
+		}
 	}
 }
