Incr gc: turn on GC barriers for argument-mutating BlueprintCallable C++ functions

diff --git a/Engine/Source/Programs/Shared/EpicGames.UHT/Exporters/CodeGen/UhtHeaderCodeGeneratorCppFile.cs b/Engine/Source/Programs/Shared/EpicGames.UHT/Exporters/CodeGen/UhtHeaderCodeGeneratorCppFile.cs
--- a/Engine/Source/Programs/Shared/EpicGames.UHT/Exporters/CodeGen/UhtHeaderCodeGeneratorCppFile.cs
+++ b/Engine/Source/Programs/Shared/EpicGames.UHT/Exporters/CodeGen/UhtHeaderCodeGeneratorCppFile.cs
@@ -123,6 +123,20 @@ namespace EpicGames.UHT.Exporters.CodeGen
 							{
 								AddIncludeForProperty(property, requireIncludeForClasses, addedIncludes, includesToAdd);
 							}
+
+							foreach (UhtType parameter in function.ParameterProperties.Span)
+							{
+								if (parameter is UhtProperty property
+										&& property.NeedsGCBarrierWhenPassedToFunction(function)
+										&& property is UhtObjectProperty objectProperty)
+								{
+									UhtClass uhtClass = objectProperty.Class;
+									if (!uhtClass.HeaderFile.IsNoExportTypes && addedIncludes.Add(uhtClass.HeaderFile))
+									{
+										includesToAdd.Add(HeaderInfos[uhtClass.HeaderFile.HeaderFileTypeIndex].IncludePath);
+									}
+								}
+							}
 						}
 
 						// Properties
diff --git a/Engine/Source/Runtime/CoreUObject/Public/UObject/ScriptMacros.h b/Engine/Source/Runtime/CoreUObject/Public/UObject/ScriptMacros.h
--- a/Engine/Source/Runtime/CoreUObject/Public/UObject/ScriptMacros.h
+++ b/Engine/Source/Runtime/CoreUObject/Public/UObject/ScriptMacros.h
@@ -115,7 +115,7 @@ enum {MAX_VARIABLE_SIZE = 0x0FFF };
 #define P_NATIVE_END   }
 
 #ifndef UE_SCRIPT_ARGS_GC_BARRIER
-#define UE_SCRIPT_ARGS_GC_BARRIER 0
+#define UE_SCRIPT_ARGS_GC_BARRIER 1
 #endif
 
 #if UE_SCRIPT_ARGS_GC_BARRIER
