EpicGames.Core: Move IClock interface to Core

IClock is useful for testing any code, not just Horde. The specific Redis-based implementation is still left with Horde.

diff --git a/Engine/Source/Programs/Horde/Horde.Server.Tests/Logs/LogTailServiceTests.cs b/Engine/Source/Programs/Horde/Horde.Server.Tests/Logs/LogTailServiceTests.cs
--- a/Engine/Source/Programs/Horde/Horde.Server.Tests/Logs/LogTailServiceTests.cs
+++ b/Engine/Source/Programs/Horde/Horde.Server.Tests/Logs/LogTailServiceTests.cs
@@ -10,7 +10,6 @@ using EpicGames.Horde.Logs;
 using EpicGames.Redis;
 using Horde.Server.Logs;
 using Horde.Server.Server;
-using HordeCommon;
 using Microsoft.Extensions.DependencyInjection;
 using Microsoft.Extensions.Logging;
 using Microsoft.Extensions.Options;
diff --git a/Engine/Source/Programs/Horde/Horde.Server/Agents/AgentCollection.cs b/Engine/Source/Programs/Horde/Horde.Server/Agents/AgentCollection.cs
--- a/Engine/Source/Programs/Horde/Horde.Server/Agents/AgentCollection.cs
+++ b/Engine/Source/Programs/Horde/Horde.Server/Agents/AgentCollection.cs
@@ -6,6 +6,7 @@ using System.Linq;
 using System.Linq.Expressions;
 using System.Threading;
 using System.Threading.Tasks;
+using EpicGames.Core;
 using EpicGames.Horde.Agents;
 using EpicGames.Horde.Agents.Leases;
 using EpicGames.Horde.Agents.Pools;
@@ -14,7 +15,6 @@ using EpicGames.Redis;
 using Google.Protobuf.WellKnownTypes;
 using Horde.Server.Auditing;
 using Horde.Server.Server;
-using HordeCommon;
 using HordeCommon.Rpc.Tasks;
 using MongoDB.Bson;
 using MongoDB.Bson.Serialization.Attributes;
diff --git a/Engine/Source/Programs/Horde/Horde.Server/Agents/AgentReportService.cs b/Engine/Source/Programs/Horde/Horde.Server/Agents/AgentReportService.cs
--- a/Engine/Source/Programs/Horde/Horde.Server/Agents/AgentReportService.cs
+++ b/Engine/Source/Programs/Horde/Horde.Server/Agents/AgentReportService.cs
@@ -4,6 +4,7 @@ using System;
 using System.Collections.Generic;
 using System.Threading;
 using System.Threading.Tasks;
+using EpicGames.Core;
 using EpicGames.Horde.Agents;
 using Horde.Server.Jobs;
 using Horde.Server.Notifications;
diff --git a/Engine/Source/Programs/Horde/Horde.Server/Agents/Enrollment/EnrollmentService.cs b/Engine/Source/Programs/Horde/Horde.Server/Agents/Enrollment/EnrollmentService.cs
--- a/Engine/Source/Programs/Horde/Horde.Server/Agents/Enrollment/EnrollmentService.cs
+++ b/Engine/Source/Programs/Horde/Horde.Server/Agents/Enrollment/EnrollmentService.cs
@@ -8,7 +8,6 @@ using EpicGames.Core;
 using EpicGames.Horde.Agents;
 using EpicGames.Redis;
 using Horde.Server.Server;
-using HordeCommon;
 using Microsoft.Extensions.Hosting;
 using StackExchange.Redis;
 
diff --git a/Engine/Source/Programs/Horde/Horde.Server/Agents/Fleet/AwsAutoScalingLifecycle.cs b/Engine/Source/Programs/Horde/Horde.Server/Agents/Fleet/AwsAutoScalingLifecycle.cs
--- a/Engine/Source/Programs/Horde/Horde.Server/Agents/Fleet/AwsAutoScalingLifecycle.cs
+++ b/Engine/Source/Programs/Horde/Horde.Server/Agents/Fleet/AwsAutoScalingLifecycle.cs
@@ -18,7 +18,6 @@ using EpicGames.Core;
 using EpicGames.Horde.Agents;
 using Horde.Server.Server;
 using Horde.Server.Utilities;
-using HordeCommon;
 using Microsoft.AspNetCore.Authorization;
 using Microsoft.AspNetCore.Http;
 using Microsoft.AspNetCore.Mvc;
diff --git a/Engine/Source/Programs/Horde/Horde.Server/Agents/Fleet/FleetManagerFactory.cs b/Engine/Source/Programs/Horde/Horde.Server/Agents/Fleet/FleetManagerFactory.cs
--- a/Engine/Source/Programs/Horde/Horde.Server/Agents/Fleet/FleetManagerFactory.cs
+++ b/Engine/Source/Programs/Horde/Horde.Server/Agents/Fleet/FleetManagerFactory.cs
@@ -6,8 +6,8 @@ using System.IO;
 using System.Text.Json;
 using Amazon.AutoScaling;
 using Amazon.EC2;
+using EpicGames.Core;
 using Horde.Server.Agents.Fleet.Providers;
-using HordeCommon;
 using Microsoft.Extensions.DependencyInjection;
 using Microsoft.Extensions.Logging;
 using Microsoft.Extensions.Options;
diff --git a/Engine/Source/Programs/Horde/Horde.Server/Agents/Fleet/FleetService.cs b/Engine/Source/Programs/Horde/Horde.Server/Agents/Fleet/FleetService.cs
--- a/Engine/Source/Programs/Horde/Horde.Server/Agents/Fleet/FleetService.cs
+++ b/Engine/Source/Programs/Horde/Horde.Server/Agents/Fleet/FleetService.cs
@@ -8,6 +8,7 @@ using System.Linq;
 using System.Text.Json;
 using System.Threading;
 using System.Threading.Tasks;
+using EpicGames.Core;
 using EpicGames.Horde.Agents;
 using EpicGames.Horde.Agents.Pools;
 using Horde.Server.Agents.Leases;
diff --git a/Engine/Source/Programs/Horde/Horde.Server/Agents/Fleet/JobQueueStrategy.cs b/Engine/Source/Programs/Horde/Horde.Server/Agents/Fleet/JobQueueStrategy.cs
--- a/Engine/Source/Programs/Horde/Horde.Server/Agents/Fleet/JobQueueStrategy.cs
+++ b/Engine/Source/Programs/Horde/Horde.Server/Agents/Fleet/JobQueueStrategy.cs
@@ -7,6 +7,7 @@ using System.Text;
 using System.Text.Json.Serialization;
 using System.Threading;
 using System.Threading.Tasks;
+using EpicGames.Core;
 using EpicGames.Horde.Agents.Pools;
 using EpicGames.Horde.Streams;
 using Horde.Server.Agents.Pools;
@@ -15,7 +16,6 @@ using Horde.Server.Jobs.Graphs;
 using Horde.Server.Server;
 using Horde.Server.Streams;
 using Horde.Server.Utilities;
-using HordeCommon;
 using Microsoft.Extensions.Caching.Memory;
 using Microsoft.Extensions.Options;
 using OpenTelemetry.Trace;
diff --git a/Engine/Source/Programs/Horde/Horde.Server/Agents/Fleet/LeaseUtilizationAwsMetricStrategy.cs b/Engine/Source/Programs/Horde/Horde.Server/Agents/Fleet/LeaseUtilizationAwsMetricStrategy.cs
--- a/Engine/Source/Programs/Horde/Horde.Server/Agents/Fleet/LeaseUtilizationAwsMetricStrategy.cs
+++ b/Engine/Source/Programs/Horde/Horde.Server/Agents/Fleet/LeaseUtilizationAwsMetricStrategy.cs
@@ -9,10 +9,10 @@ using System.Threading;
 using System.Threading.Tasks;
 using Amazon.CloudWatch;
 using Amazon.CloudWatch.Model;
+using EpicGames.Core;
 using Horde.Server.Agents.Leases;
 using Horde.Server.Agents.Pools;
 using Horde.Server.Utilities;
-using HordeCommon;
 using Microsoft.Extensions.Logging;
 using OpenTelemetry.Trace;
 
diff --git a/Engine/Source/Programs/Horde/Horde.Server/Agents/Fleet/LeaseUtilizationStrategy.cs b/Engine/Source/Programs/Horde/Horde.Server/Agents/Fleet/LeaseUtilizationStrategy.cs
--- a/Engine/Source/Programs/Horde/Horde.Server/Agents/Fleet/LeaseUtilizationStrategy.cs
+++ b/Engine/Source/Programs/Horde/Horde.Server/Agents/Fleet/LeaseUtilizationStrategy.cs
@@ -6,13 +6,13 @@ using System.Linq;
 using System.Text;
 using System.Threading;
 using System.Threading.Tasks;
+using EpicGames.Core;
 using EpicGames.Horde.Agents;
 using EpicGames.Horde.Agents.Pools;
 using Google.Protobuf.WellKnownTypes;
 using Horde.Server.Agents.Leases;
 using Horde.Server.Agents.Pools;
 using Horde.Server.Utilities;
-using HordeCommon;
 using HordeCommon.Rpc.Tasks;
 using Microsoft.Extensions.Caching.Memory;
 using OpenTelemetry.Trace;
diff --git a/Engine/Source/Programs/Horde/Horde.Server/Agents/Fleet/Providers/AwsRecyclingFleetManager.cs b/Engine/Source/Programs/Horde/Horde.Server/Agents/Fleet/Providers/AwsRecyclingFleetManager.cs
--- a/Engine/Source/Programs/Horde/Horde.Server/Agents/Fleet/Providers/AwsRecyclingFleetManager.cs
+++ b/Engine/Source/Programs/Horde/Horde.Server/Agents/Fleet/Providers/AwsRecyclingFleetManager.cs
@@ -12,7 +12,6 @@ using Amazon.EC2.Model;
 using Amazon.Runtime;
 using EpicGames.Core;
 using Horde.Server.Agents.Pools;
-using HordeCommon;
 using Microsoft.Extensions.Logging;
 using OpenTelemetry.Trace;
 
diff --git a/Engine/Source/Programs/Horde/Horde.Server/Agents/Pools/PoolService.cs b/Engine/Source/Programs/Horde/Horde.Server/Agents/Pools/PoolService.cs
--- a/Engine/Source/Programs/Horde/Horde.Server/Agents/Pools/PoolService.cs
+++ b/Engine/Source/Programs/Horde/Horde.Server/Agents/Pools/PoolService.cs
@@ -5,10 +5,10 @@ using System.Collections.Generic;
 using System.Linq;
 using System.Threading;
 using System.Threading.Tasks;
+using EpicGames.Core;
 using EpicGames.Horde;
 using EpicGames.Horde.Agents.Pools;
 using Horde.Server.Server;
-using HordeCommon;
 
 namespace Horde.Server.Agents.Pools
 {
diff --git a/Engine/Source/Programs/Horde/Horde.Server/Agents/Pools/PoolUpdateService.cs b/Engine/Source/Programs/Horde/Horde.Server/Agents/Pools/PoolUpdateService.cs
--- a/Engine/Source/Programs/Horde/Horde.Server/Agents/Pools/PoolUpdateService.cs
+++ b/Engine/Source/Programs/Horde/Horde.Server/Agents/Pools/PoolUpdateService.cs
@@ -5,6 +5,7 @@ using System.Collections.Generic;
 using System.Linq;
 using System.Threading;
 using System.Threading.Tasks;
+using EpicGames.Core;
 using EpicGames.Horde.Agents.Pools;
 using Horde.Server.Server;
 using Horde.Server.Streams;
diff --git a/Engine/Source/Programs/Horde/Horde.Server/Agents/Pools/PoolsController.cs b/Engine/Source/Programs/Horde/Horde.Server/Agents/Pools/PoolsController.cs
--- a/Engine/Source/Programs/Horde/Horde.Server/Agents/Pools/PoolsController.cs
+++ b/Engine/Source/Programs/Horde/Horde.Server/Agents/Pools/PoolsController.cs
@@ -5,6 +5,7 @@ using System.Collections.Generic;
 using System.Linq;
 using System.Threading;
 using System.Threading.Tasks;
+using EpicGames.Core;
 using EpicGames.Horde;
 using EpicGames.Horde.Agents.Pools;
 using EpicGames.Horde.Common;
@@ -12,7 +13,6 @@ using Horde.Server.Agents.Fleet;
 using Horde.Server.Agents.Utilization;
 using Horde.Server.Server;
 using Horde.Server.Utilities;
-using HordeCommon;
 using Microsoft.AspNetCore.Authorization;
 using Microsoft.AspNetCore.Mvc;
 using Microsoft.Extensions.Options;
diff --git a/Engine/Source/Programs/Horde/Horde.Server/Agents/Relay/AgentRelayService.cs b/Engine/Source/Programs/Horde/Horde.Server/Agents/Relay/AgentRelayService.cs
--- a/Engine/Source/Programs/Horde/Horde.Server/Agents/Relay/AgentRelayService.cs
+++ b/Engine/Source/Programs/Horde/Horde.Server/Agents/Relay/AgentRelayService.cs
@@ -8,6 +8,7 @@ using System.Linq;
 using System.Net;
 using System.Threading;
 using System.Threading.Tasks;
+using EpicGames.Core;
 using EpicGames.Horde.Agents.Leases;
 using EpicGames.Horde.Compute;
 using EpicGames.Redis;
@@ -17,7 +18,6 @@ using Grpc.Core;
 using Horde.Common.Rpc;
 using Horde.Server.Server;
 using Horde.Server.Utilities;
-using HordeCommon;
 using Microsoft.Extensions.Hosting;
 using Microsoft.Extensions.Logging;
 using StackExchange.Redis;
diff --git a/Engine/Source/Programs/Horde/Horde.Server/Agents/Software/AgentSoftwareController.cs b/Engine/Source/Programs/Horde/Horde.Server/Agents/Software/AgentSoftwareController.cs
--- a/Engine/Source/Programs/Horde/Horde.Server/Agents/Software/AgentSoftwareController.cs
+++ b/Engine/Source/Programs/Horde/Horde.Server/Agents/Software/AgentSoftwareController.cs
@@ -4,9 +4,9 @@ using System;
 using System.IO;
 using System.Threading;
 using System.Threading.Tasks;
+using EpicGames.Core;
 using Horde.Server.Server;
 using Horde.Server.Tools;
-using HordeCommon;
 using Microsoft.AspNetCore.Authorization;
 using Microsoft.AspNetCore.Mvc;
 using Microsoft.Extensions.Options;
diff --git a/Engine/Source/Programs/Horde/Horde.Server/Agents/Utilization/UtilizationDataService.cs b/Engine/Source/Programs/Horde/Horde.Server/Agents/Utilization/UtilizationDataService.cs
--- a/Engine/Source/Programs/Horde/Horde.Server/Agents/Utilization/UtilizationDataService.cs
+++ b/Engine/Source/Programs/Horde/Horde.Server/Agents/Utilization/UtilizationDataService.cs
@@ -5,6 +5,7 @@ using System.Collections.Generic;
 using System.Linq;
 using System.Threading;
 using System.Threading.Tasks;
+using EpicGames.Core;
 using EpicGames.Horde.Agents;
 using EpicGames.Horde.Agents.Pools;
 using Horde.Server.Agents.Fleet;
diff --git a/Engine/Source/Programs/Horde/Horde.Server/Artifacts/ArtifactCollection.cs b/Engine/Source/Programs/Horde/Horde.Server/Artifacts/ArtifactCollection.cs
--- a/Engine/Source/Programs/Horde/Horde.Server/Artifacts/ArtifactCollection.cs
+++ b/Engine/Source/Programs/Horde/Horde.Server/Artifacts/ArtifactCollection.cs
@@ -6,6 +6,7 @@ using System.Linq;
 using System.Runtime.CompilerServices;
 using System.Threading;
 using System.Threading.Tasks;
+using EpicGames.Core;
 using EpicGames.Horde.Artifacts;
 using EpicGames.Horde.Storage;
 using EpicGames.Horde.Streams;
@@ -13,7 +14,6 @@ using Horde.Server.Acls;
 using Horde.Server.Server;
 using Horde.Server.Storage;
 using Horde.Server.Utilities;
-using HordeCommon;
 using MongoDB.Bson;
 using MongoDB.Bson.Serialization.Attributes;
 using MongoDB.Driver;
diff --git a/Engine/Source/Programs/Horde/Horde.Server/Artifacts/ArtifactExpirationService.cs b/Engine/Source/Programs/Horde/Horde.Server/Artifacts/ArtifactExpirationService.cs
--- a/Engine/Source/Programs/Horde/Horde.Server/Artifacts/ArtifactExpirationService.cs
+++ b/Engine/Source/Programs/Horde/Horde.Server/Artifacts/ArtifactExpirationService.cs
@@ -6,6 +6,7 @@ using System.Diagnostics;
 using System.Linq;
 using System.Threading;
 using System.Threading.Tasks;
+using EpicGames.Core;
 using EpicGames.Horde.Storage;
 using EpicGames.Horde.Streams;
 using Horde.Server.Server;
diff --git a/Engine/Source/Programs/Horde/Horde.Server/Compute/ComputeService.cs b/Engine/Source/Programs/Horde/Horde.Server/Compute/ComputeService.cs
--- a/Engine/Source/Programs/Horde/Horde.Server/Compute/ComputeService.cs
+++ b/Engine/Source/Programs/Horde/Horde.Server/Compute/ComputeService.cs
@@ -10,6 +10,7 @@ using System.Text;
 using System.Text.Json;
 using System.Threading;
 using System.Threading.Tasks;
+using EpicGames.Core;
 using EpicGames.Horde.Agents.Leases;
 using EpicGames.Horde.Compute;
 using EpicGames.Horde.Compute.Clients;
diff --git a/Engine/Source/Programs/Horde/Horde.Server/Devices/DeviceReportService.cs b/Engine/Source/Programs/Horde/Horde.Server/Devices/DeviceReportService.cs
--- a/Engine/Source/Programs/Horde/Horde.Server/Devices/DeviceReportService.cs
+++ b/Engine/Source/Programs/Horde/Horde.Server/Devices/DeviceReportService.cs
@@ -5,6 +5,7 @@ using System.Collections.Generic;
 using System.Linq;
 using System.Threading;
 using System.Threading.Tasks;
+using EpicGames.Core;
 using Horde.Server.Jobs;
 using Horde.Server.Notifications;
 using Horde.Server.Server;
diff --git a/Engine/Source/Programs/Horde/Horde.Server/Jobs/Bisect/BisectService.cs b/Engine/Source/Programs/Horde/Horde.Server/Jobs/Bisect/BisectService.cs
--- a/Engine/Source/Programs/Horde/Horde.Server/Jobs/Bisect/BisectService.cs
+++ b/Engine/Source/Programs/Horde/Horde.Server/Jobs/Bisect/BisectService.cs
@@ -5,6 +5,7 @@ using System.Collections.Generic;
 using System.Linq;
 using System.Threading;
 using System.Threading.Tasks;
+using EpicGames.Core;
 using Horde.Server.Jobs.Graphs;
 using Horde.Server.Jobs.Templates;
 using Horde.Server.Perforce;
diff --git a/Engine/Source/Programs/Horde/Horde.Server/Perforce/PerforceLoadBalancer.cs b/Engine/Source/Programs/Horde/Horde.Server/Perforce/PerforceLoadBalancer.cs
--- a/Engine/Source/Programs/Horde/Horde.Server/Perforce/PerforceLoadBalancer.cs
+++ b/Engine/Source/Programs/Horde/Horde.Server/Perforce/PerforceLoadBalancer.cs
@@ -9,6 +9,7 @@ using System.Net.Http;
 using System.Text.Json;
 using System.Threading;
 using System.Threading.Tasks;
+using EpicGames.Core;
 using EpicGames.Perforce;
 using Google.Protobuf.WellKnownTypes;
 using Horde.Server.Agents;
diff --git a/Engine/Source/Programs/Horde/Horde.Server/Replicators/ReplicatorCollection.cs b/Engine/Source/Programs/Horde/Horde.Server/Replicators/ReplicatorCollection.cs
--- a/Engine/Source/Programs/Horde/Horde.Server/Replicators/ReplicatorCollection.cs
+++ b/Engine/Source/Programs/Horde/Horde.Server/Replicators/ReplicatorCollection.cs
@@ -2,14 +2,12 @@
 
 using System;
 using System.Collections.Generic;
-using System.Linq;
 using System.Threading;
 using System.Threading.Tasks;
+using EpicGames.Core;
 using EpicGames.Horde.Replicators;
 using Horde.Server.Server;
 using Horde.Server.Utilities;
-using HordeCommon;
-using MongoDB.Bson;
 using MongoDB.Bson.Serialization;
 using MongoDB.Bson.Serialization.Attributes;
 using MongoDB.Driver;
diff --git a/Engine/Source/Programs/Horde/Horde.Server/Server/ConsistencyService.cs b/Engine/Source/Programs/Horde/Horde.Server/Server/ConsistencyService.cs
--- a/Engine/Source/Programs/Horde/Horde.Server/Server/ConsistencyService.cs
+++ b/Engine/Source/Programs/Horde/Horde.Server/Server/ConsistencyService.cs
@@ -5,6 +5,7 @@ using System.Collections.Generic;
 using System.Linq;
 using System.Threading;
 using System.Threading.Tasks;
+using EpicGames.Core;
 using EpicGames.Horde.Agents;
 using EpicGames.Horde.Agents.Leases;
 using EpicGames.Horde.Agents.Sessions;
diff --git a/Engine/Source/Programs/Horde/Horde.Server/Server/DowntimeService.cs b/Engine/Source/Programs/Horde/Horde.Server/Server/DowntimeService.cs
--- a/Engine/Source/Programs/Horde/Horde.Server/Server/DowntimeService.cs
+++ b/Engine/Source/Programs/Horde/Horde.Server/Server/DowntimeService.cs
@@ -4,6 +4,7 @@ using System;
 using System.Linq;
 using System.Threading;
 using System.Threading.Tasks;
+using EpicGames.Core;
 using HordeCommon;
 using Microsoft.Extensions.Hosting;
 using Microsoft.Extensions.Logging;
diff --git a/Engine/Source/Programs/Horde/Horde.Server/Server/LifetimeService.cs b/Engine/Source/Programs/Horde/Horde.Server/Server/LifetimeService.cs
--- a/Engine/Source/Programs/Horde/Horde.Server/Server/LifetimeService.cs
+++ b/Engine/Source/Programs/Horde/Horde.Server/Server/LifetimeService.cs
@@ -3,6 +3,7 @@
 using System;
 using System.Threading;
 using System.Threading.Tasks;
+using EpicGames.Core;
 using HordeCommon;
 using Microsoft.Extensions.Hosting;
 using Microsoft.Extensions.Logging;
diff --git a/Engine/Source/Programs/Horde/Horde.Server/Server/Notices/NoticesController.cs b/Engine/Source/Programs/Horde/Horde.Server/Server/Notices/NoticesController.cs
--- a/Engine/Source/Programs/Horde/Horde.Server/Server/Notices/NoticesController.cs
+++ b/Engine/Source/Programs/Horde/Horde.Server/Server/Notices/NoticesController.cs
@@ -5,11 +5,11 @@ using System.Collections.Generic;
 using System.Linq;
 using System.Threading;
 using System.Threading.Tasks;
+using EpicGames.Core;
 using EpicGames.Horde.Users;
 using Horde.Server.Acls;
 using Horde.Server.Users;
 using Horde.Server.Utilities;
-using HordeCommon;
 using Microsoft.AspNetCore.Authorization;
 using Microsoft.AspNetCore.Mvc;
 using Microsoft.Extensions.Diagnostics.HealthChecks;
diff --git a/Engine/Source/Programs/Horde/Horde.Server/Server/ServerController.cs b/Engine/Source/Programs/Horde/Horde.Server/Server/ServerController.cs
--- a/Engine/Source/Programs/Horde/Horde.Server/Server/ServerController.cs
+++ b/Engine/Source/Programs/Horde/Horde.Server/Server/ServerController.cs
@@ -19,7 +19,6 @@ using Horde.Server.Configuration;
 using Horde.Server.Perforce;
 using Horde.Server.Tools;
 using Horde.Server.Utilities;
-using HordeCommon;
 using Microsoft.AspNetCore.Authorization;
 using Microsoft.AspNetCore.Mvc;
 using Microsoft.Extensions.DependencyInjection;
diff --git a/Engine/Source/Programs/Horde/Horde.Server/Server/ServerStatusService.cs b/Engine/Source/Programs/Horde/Horde.Server/Server/ServerStatusService.cs
--- a/Engine/Source/Programs/Horde/Horde.Server/Server/ServerStatusService.cs
+++ b/Engine/Source/Programs/Horde/Horde.Server/Server/ServerStatusService.cs
@@ -6,6 +6,7 @@ using System.Linq;
 using System.Text.Json;
 using System.Threading;
 using System.Threading.Tasks;
+using EpicGames.Core;
 using HordeCommon;
 using Microsoft.Extensions.Diagnostics.HealthChecks;
 using Microsoft.Extensions.Hosting;
diff --git a/Engine/Source/Programs/Horde/Horde.Server/Streams/StreamCollection.cs b/Engine/Source/Programs/Horde/Horde.Server/Streams/StreamCollection.cs
--- a/Engine/Source/Programs/Horde/Horde.Server/Streams/StreamCollection.cs
+++ b/Engine/Source/Programs/Horde/Horde.Server/Streams/StreamCollection.cs
@@ -6,13 +6,13 @@ using System.IO;
 using System.Linq;
 using System.Threading;
 using System.Threading.Tasks;
+using EpicGames.Core;
 using EpicGames.Horde.Jobs;
 using EpicGames.Horde.Jobs.Templates;
 using EpicGames.Horde.Streams;
 using EpicGames.Horde.Users;
 using Horde.Server.Server;
 using Horde.Server.Utilities;
-using HordeCommon;
 using MongoDB.Bson.Serialization.Attributes;
 using MongoDB.Driver;
 
diff --git a/Engine/Source/Programs/Horde/Horde.Server/Tasks/ConformTaskSource.cs b/Engine/Source/Programs/Horde/Horde.Server/Tasks/ConformTaskSource.cs
--- a/Engine/Source/Programs/Horde/Horde.Server/Tasks/ConformTaskSource.cs
+++ b/Engine/Source/Programs/Horde/Horde.Server/Tasks/ConformTaskSource.cs
@@ -5,6 +5,7 @@ using System.Collections.Generic;
 using System.Linq;
 using System.Threading;
 using System.Threading.Tasks;
+using EpicGames.Core;
 using EpicGames.Horde.Agents;
 using EpicGames.Horde.Agents.Leases;
 using EpicGames.Horde.Jobs;
diff --git a/Engine/Source/Programs/Horde/Horde.Server/Tasks/UpgradeTaskSource.cs b/Engine/Source/Programs/Horde/Horde.Server/Tasks/UpgradeTaskSource.cs
--- a/Engine/Source/Programs/Horde/Horde.Server/Tasks/UpgradeTaskSource.cs
+++ b/Engine/Source/Programs/Horde/Horde.Server/Tasks/UpgradeTaskSource.cs
@@ -16,7 +16,6 @@ using Horde.Server.Logs;
 using Horde.Server.Server;
 using Horde.Server.Tools;
 using Horde.Server.Utilities;
-using HordeCommon;
 using HordeCommon.Rpc.Tasks;
 using Microsoft.Extensions.Options;
 
diff --git a/Engine/Source/Programs/Horde/Horde.Server/Telemetry/Metrics/MetricCollection.cs b/Engine/Source/Programs/Horde/Horde.Server/Telemetry/Metrics/MetricCollection.cs
--- a/Engine/Source/Programs/Horde/Horde.Server/Telemetry/Metrics/MetricCollection.cs
+++ b/Engine/Source/Programs/Horde/Horde.Server/Telemetry/Metrics/MetricCollection.cs
@@ -11,7 +11,6 @@ using EpicGames.Horde.Telemetry;
 using EpicGames.Horde.Telemetry.Metrics;
 using Horde.Server.Server;
 using Horde.Server.Utilities;
-using HordeCommon;
 using Json.Path;
 using Microsoft.Extensions.Hosting;
 using Microsoft.Extensions.Logging;
diff --git a/Engine/Source/Programs/Horde/Horde.Server/Telemetry/TelemetryManager.cs b/Engine/Source/Programs/Horde/Horde.Server/Telemetry/TelemetryManager.cs
--- a/Engine/Source/Programs/Horde/Horde.Server/Telemetry/TelemetryManager.cs
+++ b/Engine/Source/Programs/Horde/Horde.Server/Telemetry/TelemetryManager.cs
@@ -6,6 +6,7 @@ using System.Net.Http;
 using System.Text.Json;
 using System.Threading;
 using System.Threading.Tasks;
+using EpicGames.Core;
 using EpicGames.Horde.Telemetry;
 using Horde.Server.Telemetry.Sinks;
 using HordeCommon;
diff --git a/Engine/Source/Programs/Horde/Horde.Server/Tools/ToolCollection.cs b/Engine/Source/Programs/Horde/Horde.Server/Tools/ToolCollection.cs
--- a/Engine/Source/Programs/Horde/Horde.Server/Tools/ToolCollection.cs
+++ b/Engine/Source/Programs/Horde/Horde.Server/Tools/ToolCollection.cs
@@ -20,10 +20,7 @@ using EpicGames.Serialization;
 using Horde.Server.Server;
 using Horde.Server.Storage;
 using Horde.Server.Utilities;
-using HordeCommon;
-using Microsoft.Extensions.Caching.Memory;
 using Microsoft.Extensions.Logging;
-using MongoDB.Bson;
 using MongoDB.Bson.Serialization.Attributes;
 using MongoDB.Driver;
 
diff --git a/Engine/Source/Programs/Horde/Horde.Server/Tools/ToolsController.cs b/Engine/Source/Programs/Horde/Horde.Server/Tools/ToolsController.cs
--- a/Engine/Source/Programs/Horde/Horde.Server/Tools/ToolsController.cs
+++ b/Engine/Source/Programs/Horde/Horde.Server/Tools/ToolsController.cs
@@ -13,7 +13,6 @@ using EpicGames.Horde.Tools;
 using Horde.Server.Server;
 using Horde.Server.Storage;
 using Horde.Server.Utilities;
-using HordeCommon;
 using Microsoft.AspNetCore.Authorization;
 using Microsoft.AspNetCore.Http;
 using Microsoft.AspNetCore.Mvc;
diff --git a/Engine/Source/Programs/Horde/Horde.Server/Utilities/Clock.cs b/Engine/Source/Programs/Horde/Horde.Server/Utilities/Clock.cs
--- a/Engine/Source/Programs/Horde/Horde.Server/Utilities/Clock.cs
+++ b/Engine/Source/Programs/Horde/Horde.Server/Utilities/Clock.cs
@@ -5,6 +5,7 @@ using System.Collections.Generic;
 using System.Diagnostics;
 using System.Threading;
 using System.Threading.Tasks;
+using EpicGames.Core;
 using EpicGames.Redis.Utility;
 using Horde.Server;
 using Horde.Server.Server;
@@ -16,22 +17,6 @@ using TimeZoneConverter;
 
 namespace HordeCommon
 {
-	/// <summary>
-	/// Base interface for a scheduled event
-	/// </summary>
-	public interface ITicker : IAsyncDisposable
-	{
-		/// <summary>
-		/// Start the ticker
-		/// </summary>
-		Task StartAsync();
-
-		/// <summary>
-		/// Stop the ticker
-		/// </summary>
-		Task StopAsync();
-	}
-
 	/// <summary>
 	/// Placeholder interface for ITicker
 	/// </summary>
@@ -47,43 +32,6 @@ namespace HordeCommon
 		public Task StopAsync() => Task.CompletedTask;
 	}
 
-	/// <summary>
-	/// Interface representing time and scheduling events which is pluggable during testing. In normal use, the Clock implementation below is used. 
-	/// </summary>
-	public interface IClock
-	{
-		/// <summary>
-		/// Return time expressed as the Coordinated Universal Time (UTC)
-		/// </summary>
-		DateTime UtcNow { get; }
-
-		/// <summary>
-		/// Time zone for schedules etc...
-		/// </summary>
-		TimeZoneInfo TimeZone { get; }
-
-		/// <summary>
-		/// Create an event that will trigger after the given time
-		/// </summary>
-		/// <param name="name">Name of the event</param>
-		/// <param name="interval">Time after which the event will trigger</param>
-		/// <param name="tickAsync">Callback for the tick. Returns the time interval until the next tick, or null to cancel the tick.</param>
-		/// <param name="logger">Logger for error messages</param>
-		/// <returns>Handle to the event</returns>
-		ITicker AddTicker(string name, TimeSpan interval, Func<CancellationToken, ValueTask<TimeSpan?>> tickAsync, ILogger logger);
-
-		/// <summary>
-		/// Create a ticker shared between all server processes.
-		/// Callback can be run inside any available process but will still only be called once per tick.
-		/// </summary>
-		/// <param name="name">Name of the event</param>
-		/// <param name="interval">Time after which the event will trigger</param>
-		/// <param name="tickAsync">Callback for the tick. Returns the time interval until the next tick, or null to cancel the tick.</param>
-		/// <param name="logger">Logger for error messages</param>
-		/// <returns>New ticker instance</returns>
-		ITicker AddSharedTicker(string name, TimeSpan interval, Func<CancellationToken, ValueTask> tickAsync, ILogger logger);
-	}
-
 	/// <summary>
 	/// Extension methods for <see cref="IClock"/>
 	/// </summary>
diff --git a/Engine/Source/Programs/Shared/EpicGames.Core/Clock.cs b/Engine/Source/Programs/Shared/EpicGames.Core/Clock.cs
new file mode 100644
--- /dev/null
+++ b/Engine/Source/Programs/Shared/EpicGames.Core/Clock.cs
@@ -0,0 +1,116 @@
+// Copyright Epic Games, Inc. All Rights Reserved.
+
+using System;
+using System.Threading;
+using System.Threading.Tasks;
+using Microsoft.Extensions.Logging;
+
+namespace EpicGames.Core;
+
+/// <summary>
+/// Base interface for a scheduled event
+/// </summary>
+public interface ITicker : IAsyncDisposable
+{
+	/// <summary>
+	/// Start the ticker
+	/// </summary>
+	Task StartAsync();
+
+	/// <summary>
+	/// Stop the ticker
+	/// </summary>
+	Task StopAsync();
+}
+
+/// <summary>
+/// Interface representing time and scheduling events which is pluggable during testing. In normal use, the Clock implementation below is used. 
+/// </summary>
+public interface IClock
+{
+	/// <summary>
+	/// Return time expressed as the Coordinated Universal Time (UTC)
+	/// </summary>
+	DateTime UtcNow { get; }
+
+	/// <summary>
+	/// Time zone for schedules etc...
+	/// </summary>
+	TimeZoneInfo TimeZone { get; }
+
+	/// <summary>
+	/// Create an event that will trigger after the given time
+	/// </summary>
+	/// <param name="name">Name of the event</param>
+	/// <param name="interval">Time after which the event will trigger</param>
+	/// <param name="tickAsync">Callback for the tick. Returns the time interval until the next tick, or null to cancel the tick.</param>
+	/// <param name="logger">Logger for error messages</param>
+	/// <returns>Handle to the event</returns>
+	ITicker AddTicker(string name, TimeSpan interval, Func<CancellationToken, ValueTask<TimeSpan?>> tickAsync, ILogger logger);
+
+	/// <summary>
+	/// Create a ticker shared between all server processes.
+	/// Callback can be run inside any available process but will still only be called once per tick.
+	/// </summary>
+	/// <param name="name">Name of the event</param>
+	/// <param name="interval">Time after which the event will trigger</param>
+	/// <param name="tickAsync">Callback for the tick. Returns the time interval until the next tick, or null to cancel the tick.</param>
+	/// <param name="logger">Logger for error messages</param>
+	/// <returns>New ticker instance</returns>
+	ITicker AddSharedTicker(string name, TimeSpan interval, Func<CancellationToken, ValueTask> tickAsync, ILogger logger);
+}
+
+/// <summary>
+/// A default implementation of IClock for normal production use
+/// </summary>
+public class DefaultClock : IClock
+{
+	/// <inheritdoc/>
+	public DateTime UtcNow => DateTime.UtcNow;
+	
+	/// <inheritdoc/>
+	public TimeZoneInfo TimeZone => TimeZoneInfo.Local;
+	
+	/// <inheritdoc/>
+	public ITicker AddTicker(string name, TimeSpan interval, Func<CancellationToken, ValueTask<TimeSpan?>> tickAsync, ILogger logger)
+	{
+		throw new NotImplementedException("Not available in default implementation");
+	}
+
+	/// <inheritdoc/>
+	public ITicker AddSharedTicker(string name, TimeSpan interval, Func<CancellationToken, ValueTask> tickAsync, ILogger logger)
+	{
+		throw new NotImplementedException("Not available in default implementation");
+	}
+}
+
+/// <summary>
+/// A stub implementation of IClock. Intended for testing to override time.
+/// </summary>
+public class StubCoreClock : IClock
+{
+	private DateTime _utcNow = DateTime.UtcNow;
+
+	/// <inheritdoc/>
+	public DateTime UtcNow
+	{
+		get => _utcNow;
+		set => _utcNow = value.ToUniversalTime();
+	}
+	
+	/// <inheritdoc/>
+	public TimeZoneInfo TimeZone { get; set; } = TimeZoneInfo.Local;
+	
+	/// <inheritdoc/>
+	public ITicker AddTicker(string name, TimeSpan interval, Func<CancellationToken, ValueTask<TimeSpan?>> tickAsync, ILogger logger)
+	{
+		throw new NotImplementedException("Not available in stub implementation");
+	}
+
+	/// <inheritdoc/>
+	public ITicker AddSharedTicker(string name, TimeSpan interval, Func<CancellationToken, ValueTask> tickAsync, ILogger logger)
+	{
+		throw new NotImplementedException("Not available in stub implementation");
+	}
+}
+
