Horde: Add log message whenever native P4 connection buffer increases in size above 16mb.

diff --git a/Engine/Source/Programs/Horde/Horde.Server/appsettings.Development.json b/Engine/Source/Programs/Horde/Horde.Server/appsettings.Development.json
--- a/Engine/Source/Programs/Horde/Horde.Server/appsettings.Development.json
+++ b/Engine/Source/Programs/Horde/Horde.Server/appsettings.Development.json
@@ -6,7 +6,10 @@
 	{
 		"MinimumLevel": 
 		{
-			"Default": "Debug"
+			"Default": "Debug",
+			"Override": {
+				"Horde.Server.Perforce.PerforceService": "Verbose"
+			}
 		}
 	}
 }
diff --git a/Engine/Source/Programs/Shared/EpicGames.Perforce/NativePerforceConnection.cs b/Engine/Source/Programs/Shared/EpicGames.Perforce/NativePerforceConnection.cs
--- a/Engine/Source/Programs/Shared/EpicGames.Perforce/NativePerforceConnection.cs
+++ b/Engine/Source/Programs/Shared/EpicGames.Perforce/NativePerforceConnection.cs
@@ -273,6 +273,7 @@ namespace EpicGames.Perforce
 		readonly ManualResetEvent _responseCompleteEvent;
 		readonly Stopwatch _stallTimer = new Stopwatch();
 		readonly HangMonitor _hangMonitor;
+		long _allocatedSize;
 		bool _disposed;
 
 		/// <inheritdoc/>
@@ -311,6 +312,7 @@ namespace EpicGames.Perforce
 			{
 				_buffers[idx] = new PinnedBuffer(bufferSize);
 				_writeBuffers.TryAdd(_buffers[idx]);
+				_allocatedSize += bufferSize;
 			}
 
 			_onBufferReadyInst = new OnBufferReadyFn(OnBufferReady);
@@ -373,6 +375,7 @@ namespace EpicGames.Perforce
 			if (buffer.MaxLength < minSize)
 			{
 				buffer.Resize(minSize);
+				_allocatedSize += minSize - buffer.MaxLength;
 			}
 
 			nativeWriteBuffer._data = buffer.BasePtr;
@@ -598,6 +601,7 @@ namespace EpicGames.Perforce
 			Logger.LogTrace("Conn {ConnectionId}: {Command} {Args}", _uniqueId, command, argList.ToString());
 
 			using IDisposable scope = _hangMonitor.Start($"{command} {argList}");
+			long initialAllocatedSize = _allocatedSize;
 
 			List<IntPtr> nativeArgs = new List<IntPtr>();
 			try
@@ -628,6 +632,12 @@ namespace EpicGames.Perforce
 				}
 			}
 
+			const long WarnSize = 16 * 1024 * 1024;
+			if (_allocatedSize > initialAllocatedSize && _allocatedSize > WarnSize)
+			{
+				Logger.LogTrace("Native P4 connection allocated {AllocatedSize} ({Command} {Args})", _allocatedSize, command, argList.ToString());
+			}
+
 			Logger.LogTrace("Conn {ConnectionId}: Request completed in {Time}ms", _uniqueId, timer.ElapsedMilliseconds);
 		}
 
