Horde: Handle socket shutdown errors gracefully rather than throwing exceptions.

diff --git a/Engine/Source/Programs/Shared/EpicGames.Horde/Compute/Transports/TcpTransport.cs b/Engine/Source/Programs/Shared/EpicGames.Horde/Compute/Transports/TcpTransport.cs
--- a/Engine/Source/Programs/Shared/EpicGames.Horde/Compute/Transports/TcpTransport.cs
+++ b/Engine/Source/Programs/Shared/EpicGames.Horde/Compute/Transports/TcpTransport.cs
@@ -34,7 +34,15 @@ namespace EpicGames.Horde.Compute.Transports
 		/// <inheritdoc/>
 		public override async ValueTask<int> RecvAsync(Memory<byte> buffer, CancellationToken cancellationToken)
 		{
-			int read = await _socket.ReceiveAsync(buffer, SocketFlags.None, cancellationToken);
+			int read;
+			try
+			{
+				read = await _socket.ReceiveAsync(buffer, SocketFlags.None, cancellationToken);
+			}
+			catch (SocketException ex) when (ex.SocketErrorCode == SocketError.Shutdown)
+			{
+				read = 0;
+			}
 			Position += read;
 			return read;
 		}
