EpicGames.Core: Improve test stub for IClock

diff --git a/Engine/Source/Programs/Shared/EpicGames.Core/Clock.cs b/Engine/Source/Programs/Shared/EpicGames.Core/Clock.cs
--- a/Engine/Source/Programs/Shared/EpicGames.Core/Clock.cs
+++ b/Engine/Source/Programs/Shared/EpicGames.Core/Clock.cs
@@ -87,7 +87,7 @@ public class DefaultClock : IClock
 /// <summary>
 /// A stub implementation of IClock. Intended for testing to override time.
 /// </summary>
-public class StubCoreClock : IClock
+public class StubClock : IClock
 {
 	private DateTime _utcNow = DateTime.UtcNow;
 
@@ -112,5 +112,14 @@ public class StubCoreClock : IClock
 	{
 		throw new NotImplementedException("Not available in stub implementation");
 	}
+
+	/// <summary>
+	/// Advance the time
+	/// </summary>
+	/// <param name="delta"></param>
+	public void Advance(TimeSpan delta)
+	{
+		_utcNow += delta;
+	}
 }
 
