EpicGames.OIDC: Fix issue using local time for token expiration check

- Having OidcTokenInfo.IsValid just check for empty tokens was not enough
- Adding a proper expiration check to it makes IsAuthorized() behave as intended

diff --git a/Engine/Source/Programs/Shared/EpicGames.Horde.Tests/HordeHttpTests.cs b/Engine/Source/Programs/Shared/EpicGames.Horde.Tests/HordeHttpTests.cs
--- a/Engine/Source/Programs/Shared/EpicGames.Horde.Tests/HordeHttpTests.cs
+++ b/Engine/Source/Programs/Shared/EpicGames.Horde.Tests/HordeHttpTests.cs
@@ -9,6 +9,7 @@ using System.Net.Http;
 using System.Text.Json;
 using System.Threading;
 using System.Threading.Tasks;
+using EpicGames.Core;
 using EpicGames.Horde.Server;
 using EpicGames.OIDC;
 using Microsoft.Extensions.Logging;
@@ -17,83 +18,111 @@ using Microsoft.VisualStudio.TestTools.UnitTesting;
 
 namespace EpicGames.Horde.Tests;
 
-public class StubHttpClientFactory : IHttpClientFactory
+public record StubHttpClientFactory(Func<HttpClient> CreateHttpClient) : IHttpClientFactory
 {
-	private readonly HttpClient _httpClient;
-
-	public StubHttpClientFactory(HttpClient httpClient)
-	{
-		_httpClient = httpClient;
-	}
-
 	public HttpClient CreateClient(string name)
 	{
-		return _httpClient;
+		return CreateHttpClient();
 	}
 }
 
-public class StubMessageHandler : HttpMessageHandler
+public class FakeServerWithAuth : HttpMessageHandler
 {
-	public HttpStatusCode StatusCode { get; }
-	public string Content { get; }
+	public const string ServerUrl = "http://fake-server-test";
+	public const string SuccessResponse = "fakeServerSuccess";
+	
 	public List<HttpRequestMessage> HttpRequests { get; } = new();
+	private readonly GetAuthConfigResponse? _authConfigResponse;
+	private readonly FakeOidcTokenManager _oidcTokenManager;
 
-	public StubMessageHandler(HttpStatusCode statusCode, string content)
+	public FakeServerWithAuth(FakeOidcTokenManager oidcTokenManager, GetAuthConfigResponse? authConfigResponse = null)
 	{
-		StatusCode = statusCode;
-		Content = content;
+		_authConfigResponse = authConfigResponse ?? new GetAuthConfigResponse
+		{
+			Method = AuthMethod.OpenIdConnect, ServerUrl = ServerUrl, LocalRedirectUrls = new[] { ServerUrl }
+		};
+		_oidcTokenManager = oidcTokenManager;
 	}
 	
+	/// <summary>
+	/// Get the access token used by a request received in the fake server
+	/// </summary>
+	/// <returns>Access token</returns>
+	public string GetAccessTokenUsed(int requestNum = -1)
+	{
+		HttpRequestMessage req = requestNum == -1 ? HttpRequests.Last() : HttpRequests[requestNum];
+		return req.Headers.Authorization!.Parameter!;
+	}
+
+	public IHttpClientFactory GetHttpClientFactory()
+	{
+		return new StubHttpClientFactory(() => new HttpClient(this) { BaseAddress = new Uri(ServerUrl) });
+	}
+	
+	/// <inheritdoc/>
 	protected override Task<HttpResponseMessage> SendAsync(HttpRequestMessage request, CancellationToken cancellationToken)
 	{
+		if (request.RequestUri?.AbsolutePath == "/api/v1/server/auth")
+		{
+			return Task.FromResult(new HttpResponseMessage(HttpStatusCode.OK) { Content = new StringContent(JsonSerializer.Serialize(_authConfigResponse)) });
+		}
+		
 		Console.WriteLine($"Saving request {request.RequestUri} {request.Headers.Authorization}");
 		HttpRequests.Add(request);
-		return Task.FromResult(new HttpResponseMessage(StatusCode) { Content = new StringContent(Content) });
+		try
+		{
+			_oidcTokenManager.ValidateAccessToken(request.Headers.Authorization?.Parameter);
+			return Task.FromResult(new HttpResponseMessage(HttpStatusCode.OK) { Content = new StringContent(SuccessResponse) });
+		}
+		catch (Exception e)
+		{
+			return Task.FromResult(new HttpResponseMessage(HttpStatusCode.Unauthorized) { Content = new StringContent(e.Message) });
+		}
 	}
 }
 
 [TestClass]
 public class HordeHttpAuthHandlerTests
 {
-	const string HordeServerUrl = "http://horde-server-test";
-	const string FakeServerResponse = "fakeServerResponse";
-	
+	private readonly FakeOidcTokenManager _oidc;
+	private readonly StubClock _clock = new();
+
+	public HordeHttpAuthHandlerTests()
+	{
+		_oidc = new FakeOidcTokenManager(() => _clock.UtcNow);
+	}
+
 	[TestMethod]
 	public async Task AccessToken_Valid_IsReusedAsync()
 	{
-		FakeOidcTokenManager oidc = new ();
-		(HttpClient client, StubMessageHandler server) = CreateClientServer(oidc);
+		(HttpClient client, FakeServerWithAuth server) = CreateClientServer(_oidc);
 
 		await SendHttpRequestAsync(client);
-		string firstAccessToken = GetLastUsedAccessToken(server);
-		Assert.AreEqual(oidc.AccessToken, firstAccessToken);
+		Assert.AreEqual(_oidc.AccessToken, server.GetAccessTokenUsed(0));
 		
 		await SendHttpRequestAsync(client);
-		Assert.AreEqual(oidc.AccessToken, GetLastUsedAccessToken(server));
-		Assert.AreEqual(firstAccessToken, GetLastUsedAccessToken(server));
+		Assert.AreEqual(_oidc.AccessToken, server.GetAccessTokenUsed(1));
+		Assert.AreEqual(server.GetAccessTokenUsed(0), server.GetAccessTokenUsed(1));
 	}
 	
 	[TestMethod]
 	public async Task AccessToken_Expired_IsRefreshedAsync()
 	{
-		FakeOidcTokenManager oidc = new ();
-		(HttpClient client, StubMessageHandler server) = CreateClientServer(oidc);
-
-		oidc.RefreshToken = "someRefreshToken";
-		oidc.AccessToken = "expiredAccessToken";
+		(HttpClient client, FakeServerWithAuth server) = CreateClientServer(_oidc);
 		
-		// Set it as expired two hours ago, causing a refresh of access token
-		oidc.AccessTokenExpiry = DateTimeOffset.Now.AddHours(-2);
+		// Send a first request to obtain an access token then advance time so it expires
+		await SendHttpRequestAsync(client);
+		_clock.Advance(TimeSpan.FromMinutes(30));
 		
 		await SendHttpRequestAsync(client);
-		Assert.AreEqual(oidc.AccessToken, GetLastUsedAccessToken(server));
+		Assert.AreEqual(_oidc.AccessToken, server.GetAccessTokenUsed());
+		Assert.AreNotEqual(server.GetAccessTokenUsed(0), server.GetAccessTokenUsed(1));
 	}
 
 	[SuppressMessage("Reliability", "CA2000:Dispose objects before losing scope")]
-	private static (HttpClient client, StubMessageHandler server) CreateClientServer(IOidcTokenManager oidcTokenManager, HordeOptions? hordeOptions = null)
+	private static (HttpClient client, FakeServerWithAuth server) CreateClientServer(FakeOidcTokenManager oidcTokenManager, HordeOptions? hordeOptions = null)
 	{
-		GetAuthConfigResponse authConfig = new () { Method = AuthMethod.OpenIdConnect, ServerUrl = HordeServerUrl, LocalRedirectUrls = new [] { HordeServerUrl } };
-		IHttpClientFactory httpFactory = GetHordeHttpClientFactory(HordeServerUrl, authConfig);
+		FakeServerWithAuth server = new (oidcTokenManager);
 		using ILoggerFactory loggerFactory = LoggerFactory.Create(builder =>
 		{
 			builder.SetMinimumLevel(LogLevel.Debug);
@@ -101,37 +130,21 @@ public class HordeHttpAuthHandlerTests
 		});
 
 		ILogger<HordeHttpAuthHandler> logger = loggerFactory.CreateLogger<HordeHttpAuthHandler>();
-
 		OptionsWrapper<HordeOptions> options = new (hordeOptions ?? new HordeOptions());
 		InMemoryTokenStore inMemoryTokenStore = new ();
-		HordeHttpAuthHandlerState state = new (httpFactory, options, logger, inMemoryTokenStore, oidcTokenManager);
+		HordeHttpAuthHandlerState state = new (server.GetHttpClientFactory(), options, logger, inMemoryTokenStore, oidcTokenManager);
 		HordeHttpAuthHandler authHandler = new (state, options);
-		
-		StubMessageHandler server = new (HttpStatusCode.OK, FakeServerResponse);
+
 		authHandler.InnerHandler = server;
 		HttpClient client = new (authHandler);
 		return (client, server);
 	}
-
-	private static string GetLastUsedAccessToken(StubMessageHandler server)
-	{
-		return server.HttpRequests.Last().Headers.Authorization!.Parameter!;
-	}
 	
 	private static async Task SendHttpRequestAsync(HttpClient client)
 	{
-		using HttpRequestMessage req = new (HttpMethod.Get, HordeServerUrl);
+		using HttpRequestMessage req = new (HttpMethod.Get, FakeServerWithAuth.ServerUrl);
 		HttpResponseMessage res = await client.SendAsync(req);
+		Assert.AreEqual(FakeServerWithAuth.SuccessResponse, await res.Content.ReadAsStringAsync());
 		Assert.AreEqual(HttpStatusCode.OK, res.StatusCode);
-		Assert.AreEqual(FakeServerResponse, await res.Content.ReadAsStringAsync());
-	}
-	
-	private static IHttpClientFactory GetHordeHttpClientFactory(string? baseAddress, GetAuthConfigResponse response)
-	{
-		using StubMessageHandler stubMessageHandler = new (HttpStatusCode.OK, JsonSerializer.Serialize(response));
-		return new StubHttpClientFactory(new HttpClient(stubMessageHandler)
-		{
-			BaseAddress = baseAddress != null ? new Uri(baseAddress) : null,
-		});
 	}
 }
diff --git a/Engine/Source/Programs/Shared/EpicGames.Horde/HordeHttpAuthHandler.cs b/Engine/Source/Programs/Shared/EpicGames.Horde/HordeHttpAuthHandler.cs
--- a/Engine/Source/Programs/Shared/EpicGames.Horde/HordeHttpAuthHandler.cs
+++ b/Engine/Source/Programs/Shared/EpicGames.Horde/HordeHttpAuthHandler.cs
@@ -88,10 +88,14 @@ namespace EpicGames.Horde
 		/// </summary>
 		public const string HttpClientName = "HordeHttpAuthState";
 
-		record class AuthState(AuthMethod Method, OidcTokenInfo? TokenInfo, bool Interactive)
+		record AuthState(IClock Clock, AuthMethod Method, OidcTokenInfo? TokenInfo, bool Interactive)
 		{
 			public bool IsAuthorized()
-				=> (Method == AuthMethod.Anonymous) || (TokenInfo != null && TokenInfo.IsValid);
+			{
+				bool isAnonymousAuth = Method == AuthMethod.Anonymous;
+				bool isTokenValid = TokenInfo != null && TokenInfo.IsValid(Clock.UtcNow);
+				return isAnonymousAuth || isTokenValid;
+			}
 		}
 
 		readonly object _lockObject = new object();
@@ -104,6 +108,7 @@ namespace EpicGames.Horde
 		// Allow these to be overridden in tests
 		readonly ITokenStore? _tokenStore;
 		readonly IOidcTokenManager? _oidcTokenManager;
+		readonly IClock _clock;
 
 		/// <summary>
 		/// Constructor
@@ -113,13 +118,15 @@ namespace EpicGames.Horde
 			IOptions<HordeOptions> options,
 			ILogger<HordeHttpAuthHandler> logger,
 			ITokenStore? tokenStore = null,
-			IOidcTokenManager? oidcTokenManager = null)
+			IOidcTokenManager? oidcTokenManager = null,
+			IClock? clock = null)
 		{
 			_httpClientFactory = httpClientFactory;
 			_options = options;
 			_logger = logger;
 			_tokenStore = tokenStore;
 			_oidcTokenManager = oidcTokenManager;
+			_clock = clock ?? new DefaultClock();
 		}
 
 		/// <inheritdoc/>
@@ -291,7 +298,7 @@ namespace EpicGames.Horde
 
 			if (authConfig.Method == AuthMethod.Anonymous)
 			{
-				return new AuthState(authConfig.Method, null, true);
+				return new AuthState(_clock, authConfig.Method, null, true);
 			}
 
 			string? localRedirectUrl = authConfig.LocalRedirectUrls?.FirstOrDefault();
@@ -323,7 +330,7 @@ namespace EpicGames.Horde
 				result = await oidcTokenManager.LoginAsync(oidcProvider, cancellationToken);
 			}
 
-			return new AuthState(authConfig.Method, result, interactive);
+			return new AuthState(_clock, authConfig.Method, result, interactive);
 		}
 		
 		private ITokenStore CreateTokenStore()
diff --git a/Engine/Source/Programs/Shared/EpicGames.OIDC/FakeOidcTokenManager.cs b/Engine/Source/Programs/Shared/EpicGames.OIDC/FakeOidcTokenManager.cs
--- a/Engine/Source/Programs/Shared/EpicGames.OIDC/FakeOidcTokenManager.cs
+++ b/Engine/Source/Programs/Shared/EpicGames.OIDC/FakeOidcTokenManager.cs
@@ -14,6 +14,8 @@ namespace EpicGames.OIDC;
 /// </summary>
 public class FakeOidcTokenManager : IOidcTokenManager
 {
+	private const string AccessTokenPrefix = "fakeAccessToken";
+	
 	/// <summary>
 	/// Refresh token in use
 	/// </summary>
@@ -23,6 +25,11 @@ public class FakeOidcTokenManager : IOidcTokenManager
 	/// Latest access token
 	/// </summary>
 	public string? AccessToken { get; set; }
+
+	/// <summary>
+	/// Time to live for newly minted access token
+	/// </summary>
+	public TimeSpan AccessTokenTtl { get; set; } = TimeSpan.FromMinutes(15);
 	
 	/// <summary>
 	/// Expiry time for access token
@@ -31,11 +38,19 @@ public class FakeOidcTokenManager : IOidcTokenManager
 	
 	private int _refreshCounter = 1;
 	private int _accessCounter = 1;
-	
+	private readonly Func<DateTimeOffset> _utcNow;
+
+	/// <summary>
+	/// Constructor
+	/// </summary>
+	public FakeOidcTokenManager(Func<DateTimeOffset> utcNow)
+	{
+		_utcNow = utcNow;
+	}
+
 	/// <inheritdoc/>
 	public Task<OidcTokenInfo> LoginAsync(string providerIdentifier, CancellationToken cancellationToken = default)
 	{
-		Console.WriteLine("FakeOidcTokenManager.LoginAsync");
 		RefreshToken = "fakeRefreshToken-" + _refreshCounter++;
 		RefreshAccessToken();
 		return Task.FromResult(GetOidcTokenInfo());
@@ -50,43 +65,76 @@ public class FakeOidcTokenManager : IOidcTokenManager
 	/// <inheritdoc/>
 	public Task<OidcTokenInfo?> TryGetAccessToken(string providerIdentifier, CancellationToken cancellationToken = default)
 	{
-		Console.WriteLine("FakeOidcTokenManager.TryGetAccessToken");
 		if (String.IsNullOrEmpty(RefreshToken))
 		{
 			throw new NotLoggedInException();
 		}
+
+		OidcTokenInfo tokenInfo = GetOidcTokenInfo();
 		
-		bool tokenValid = AccessTokenExpiry.AddMinutes(2) > DateTime.Now || AccessTokenExpiry == DateTimeOffset.MinValue;
-		if (!String.IsNullOrEmpty(AccessToken) && tokenValid)
+		// Ensure token is valid at least for another two minutes
+		if (tokenInfo.IsValid(_utcNow().AddMinutes(-2)))
 		{
-			return Task.FromResult(GetOidcTokenInfo())!;
+			return Task.FromResult<OidcTokenInfo?>(tokenInfo);
 		}
 		
-		// Refresh access token with provider
+		// Access token not valid, simulate a refresh with provider
 		RefreshAccessToken();
 		
 		// Assume the updated access token is valid
-		return Task.FromResult(GetOidcTokenInfo())!;
+		return Task.FromResult<OidcTokenInfo?>(GetOidcTokenInfo());
 	}
 
 	/// <inheritdoc/>
 	public OidcStatus GetStatusForProvider(string providerIdentifier)
 	{
-		Console.WriteLine("FakeOidcTokenManager.GetStatusForProvider");
 		return OidcTokenClient.GetStatus(RefreshToken, AccessToken, AccessTokenExpiry);
 	}
 
-	private OidcTokenInfo GetOidcTokenInfo()
+	/// <summary>
+	/// Validate a fake access token minted by this class
+	/// </summary>
+	/// <param name="token">Token to check</param>
+	/// <exception cref="Exception">If token is malformed or expired</exception>
+	public void ValidateAccessToken(string? token)
 	{
-		return new OidcTokenInfo { RefreshToken = RefreshToken, AccessToken = AccessToken, TokenExpiry = AccessTokenExpiry };
-	}
+		if (String.IsNullOrEmpty(token))
+		{
+			throw new Exception($"Empty or null token: {token}");
+		}
+		
+		string[] parts = token.Split('-');
+		string prefix = parts[0];
+		int id = Convert.ToInt32(parts[1]);
+		long expireUnixTime = Convert.ToInt64(parts[2]);
+		DateTimeOffset expireTime = DateTimeOffset.FromUnixTimeMilliseconds(expireUnixTime);
+		DateTimeOffset utcNow = _utcNow();
 
+		if (prefix != AccessTokenPrefix)
+		{
+			throw new Exception($"Bad prefix for token: {token}");
+		}
+
+		if (utcNow > expireTime)
+		{
+			Console.WriteLine($"_clock.UtcNow: {utcNow}");
+			Console.WriteLine($"   expireTime: {expireTime}");
+			int totalSeconds = (int)(_utcNow() - expireTime).TotalSeconds;
+			throw new Exception($"Access token expired ({totalSeconds} seconds ago). Token: {token}");
+		}
+	}
+	
 	/// <summary>
 	/// Perform fake refresh of the access token
 	/// </summary>
 	private void RefreshAccessToken()
 	{
-		AccessToken = "fakeAccessToken-" + _accessCounter++;;
-		AccessTokenExpiry = DateTimeOffset.UtcNow + TimeSpan.FromHours(1);
+		AccessTokenExpiry = _utcNow() + AccessTokenTtl;
+		AccessToken = $"{AccessTokenPrefix}-{_accessCounter++}-{AccessTokenExpiry.ToUnixTimeMilliseconds()}-{AccessTokenExpiry.ToString()}";
+	}
+	
+	private OidcTokenInfo GetOidcTokenInfo()
+	{
+		return new OidcTokenInfo { RefreshToken = RefreshToken, AccessToken = AccessToken, TokenExpiry = AccessTokenExpiry };
 	}
 }
diff --git a/Engine/Source/Programs/Shared/EpicGames.OIDC/OidcTokenManager.cs b/Engine/Source/Programs/Shared/EpicGames.OIDC/OidcTokenManager.cs
--- a/Engine/Source/Programs/Shared/EpicGames.OIDC/OidcTokenManager.cs
+++ b/Engine/Source/Programs/Shared/EpicGames.OIDC/OidcTokenManager.cs
@@ -199,8 +199,22 @@ namespace EpicGames.OIDC
 		/// </summary>
 		public DateTimeOffset TokenExpiry { get; set; }
 
-		public bool IsValid => RefreshToken != null && AccessToken != null;
-	};
+		public bool IsValid(DateTimeOffset currentTime)
+		{
+			if (String.IsNullOrEmpty(RefreshToken) || String.IsNullOrEmpty(AccessToken))
+			{
+				return false;
+			}
+
+			// An expiry of MinValue means it has no expiration time
+			if (TokenExpiry == DateTimeOffset.MinValue)
+			{
+				return true;
+			}
+
+			return currentTime <= TokenExpiry;
+		}
+	}
 
 	public class OidcTokenClient
 	{
@@ -573,18 +587,17 @@ namespace EpicGames.OIDC
 				throw new NotLoggedInException();
 			}
 
-			// if the token is valid for another few minutes we can use it
-			// we avoid using a token that is about to expire to make sure we can finish the call we expect to do with it before it expires
-			// if a token has infinite lifetime its expiry is set to 0 and is thus always valid
-			bool tokenValid = _tokenExpiry.AddMinutes(2) > DateTime.Now || _tokenExpiry == DateTimeOffset.MinValue;
-			if (!String.IsNullOrEmpty(_accessToken) && tokenValid)
+			OidcTokenInfo tokenInfo = new ()
 			{
-				return new OidcTokenInfo
-				{
-					RefreshToken = _refreshToken,
-					AccessToken = _accessToken,
-					TokenExpiry = _tokenExpiry
-				};
+				RefreshToken = _refreshToken,
+				AccessToken = _accessToken,
+				TokenExpiry = _tokenExpiry
+			};
+
+			// Ensure token is valid for at least another two minutes
+			if (tokenInfo.IsValid(DateTime.Now.AddMinutes(-2)))
+			{
+				return tokenInfo;
 			}
 
 			return await TryDoRefreshTokenAsync(_refreshToken, cancellationToken);
