Changed Gauntlet message to call out automation

diff --git a/Engine/Source/Programs/Shared/EpicGames.Horde/Issues/Handlers/GauntletIssueHandler.cs b/Engine/Source/Programs/Shared/EpicGames.Horde/Issues/Handlers/GauntletIssueHandler.cs
--- a/Engine/Source/Programs/Shared/EpicGames.Horde/Issues/Handlers/GauntletIssueHandler.cs
+++ b/Engine/Source/Programs/Shared/EpicGames.Horde/Issues/Handlers/GauntletIssueHandler.cs
@@ -195,7 +195,7 @@ namespace EpicGames.Horde.Issues.Handlers
 		{
 			if (issueEvent.EventId != null && IsMatchingEventId(issueEvent.EventId.Value))
 			{
-				IssueEventGroup issue = new IssueEventGroup("Gauntlet", "Gauntlet {Meta:Type} {Severity} {Meta:Context}", IssueChangeFilter.Code);
+				IssueEventGroup issue = new IssueEventGroup("Gauntlet", "Automation {Meta:Type} {Severity} {Meta:Context}", IssueChangeFilter.Code);
 				issue.Events.Add(issueEvent);
 
 				GetNames(issueEvent, issue.Keys, issue.Metadata);
