Horde: Move the CommitTag class into EpicGames.Horde.

#rnx

diff --git a/Engine/Source/Programs/Horde/Horde.Server.Tests/Jobs/JobServiceTests.cs b/Engine/Source/Programs/Horde/Horde.Server.Tests/Jobs/JobServiceTests.cs
--- a/Engine/Source/Programs/Horde/Horde.Server.Tests/Jobs/JobServiceTests.cs
+++ b/Engine/Source/Programs/Horde/Horde.Server.Tests/Jobs/JobServiceTests.cs
@@ -8,6 +8,7 @@ using System.Threading.Tasks;
 using EpicGames.Core;
 using EpicGames.Horde.Agents;
 using EpicGames.Horde.Agents.Pools;
+using EpicGames.Horde.Commits;
 using EpicGames.Horde.Jobs;
 using EpicGames.Horde.Jobs.Templates;
 using EpicGames.Horde.Streams;
diff --git a/Engine/Source/Programs/Horde/Horde.Server.Tests/Jobs/Schedules/SchedulerTests.cs b/Engine/Source/Programs/Horde/Horde.Server.Tests/Jobs/Schedules/SchedulerTests.cs
--- a/Engine/Source/Programs/Horde/Horde.Server.Tests/Jobs/Schedules/SchedulerTests.cs
+++ b/Engine/Source/Programs/Horde/Horde.Server.Tests/Jobs/Schedules/SchedulerTests.cs
@@ -4,6 +4,7 @@ using System;
 using System.Collections.Generic;
 using System.Linq;
 using System.Threading.Tasks;
+using EpicGames.Horde.Commits;
 using EpicGames.Horde.Jobs;
 using EpicGames.Horde.Jobs.Templates;
 using EpicGames.Horde.Projects;
@@ -13,7 +14,6 @@ using Horde.Server.Jobs.Graphs;
 using Horde.Server.Jobs.Schedules;
 using Horde.Server.Jobs.Templates;
 using Horde.Server.Logs;
-using Horde.Server.Perforce;
 using Horde.Server.Projects;
 using Horde.Server.Server;
 using Horde.Server.Streams;
diff --git a/Engine/Source/Programs/Horde/Horde.Server.Tests/Perforce/CommitTagTests.cs b/Engine/Source/Programs/Horde/Horde.Server.Tests/Perforce/CommitTagTests.cs
--- a/Engine/Source/Programs/Horde/Horde.Server.Tests/Perforce/CommitTagTests.cs
+++ b/Engine/Source/Programs/Horde/Horde.Server.Tests/Perforce/CommitTagTests.cs
@@ -2,7 +2,7 @@
 
 using System.Collections.Generic;
 using EpicGames.Core;
-using Horde.Server.Perforce;
+using EpicGames.Horde.Commits;
 using Horde.Server.Streams;
 using Microsoft.VisualStudio.TestTools.UnitTesting;
 
diff --git a/Engine/Source/Programs/Horde/Horde.Server.Tests/Stubs/Services/PerforceServiceStub.cs b/Engine/Source/Programs/Horde/Horde.Server.Tests/Stubs/Services/PerforceServiceStub.cs
--- a/Engine/Source/Programs/Horde/Horde.Server.Tests/Stubs/Services/PerforceServiceStub.cs
+++ b/Engine/Source/Programs/Horde/Horde.Server.Tests/Stubs/Services/PerforceServiceStub.cs
@@ -7,6 +7,7 @@ using System.Runtime.CompilerServices;
 using System.Threading;
 using System.Threading.Tasks;
 using EpicGames.Core;
+using EpicGames.Horde.Commits;
 using EpicGames.Horde.Streams;
 using EpicGames.Horde.Users;
 using Horde.Server.Perforce;
diff --git a/Engine/Source/Programs/Horde/Horde.Server/Jobs/Bisect/BisectController.cs b/Engine/Source/Programs/Horde/Horde.Server/Jobs/Bisect/BisectController.cs
--- a/Engine/Source/Programs/Horde/Horde.Server/Jobs/Bisect/BisectController.cs
+++ b/Engine/Source/Programs/Horde/Horde.Server/Jobs/Bisect/BisectController.cs
@@ -5,13 +5,13 @@ using System.Collections.Generic;
 using System.Linq;
 using System.Threading;
 using System.Threading.Tasks;
+using EpicGames.Horde.Commits;
 using EpicGames.Horde.Jobs;
 using EpicGames.Horde.Jobs.Bisect;
 using EpicGames.Horde.Jobs.Templates;
 using EpicGames.Horde.Streams;
 using EpicGames.Horde.Users;
 using Horde.Server.Jobs.Graphs;
-using Horde.Server.Perforce;
 using Horde.Server.Server;
 using Horde.Server.Streams;
 using Horde.Server.Users;
diff --git a/Engine/Source/Programs/Horde/Horde.Server/Jobs/Bisect/BisectTaskCollection.cs b/Engine/Source/Programs/Horde/Horde.Server/Jobs/Bisect/BisectTaskCollection.cs
--- a/Engine/Source/Programs/Horde/Horde.Server/Jobs/Bisect/BisectTaskCollection.cs
+++ b/Engine/Source/Programs/Horde/Horde.Server/Jobs/Bisect/BisectTaskCollection.cs
@@ -5,12 +5,12 @@ using System.Collections.Generic;
 using System.Runtime.CompilerServices;
 using System.Threading;
 using System.Threading.Tasks;
+using EpicGames.Horde.Commits;
 using EpicGames.Horde.Jobs;
 using EpicGames.Horde.Jobs.Bisect;
 using EpicGames.Horde.Jobs.Templates;
 using EpicGames.Horde.Streams;
 using EpicGames.Horde.Users;
-using Horde.Server.Perforce;
 using Horde.Server.Server;
 using Horde.Server.Utilities;
 using MongoDB.Bson;
diff --git a/Engine/Source/Programs/Horde/Horde.Server/Jobs/Bisect/IBisectTask.cs b/Engine/Source/Programs/Horde/Horde.Server/Jobs/Bisect/IBisectTask.cs
--- a/Engine/Source/Programs/Horde/Horde.Server/Jobs/Bisect/IBisectTask.cs
+++ b/Engine/Source/Programs/Horde/Horde.Server/Jobs/Bisect/IBisectTask.cs
@@ -1,12 +1,12 @@
 // Copyright Epic Games, Inc. All Rights Reserved.
 
 using System.Collections.Generic;
+using EpicGames.Horde.Commits;
 using EpicGames.Horde.Jobs;
 using EpicGames.Horde.Jobs.Bisect;
 using EpicGames.Horde.Jobs.Templates;
 using EpicGames.Horde.Streams;
 using EpicGames.Horde.Users;
-using Horde.Server.Perforce;
 
 namespace Horde.Server.Jobs.Bisect
 {
diff --git a/Engine/Source/Programs/Horde/Horde.Server/Jobs/Bisect/IBisectTaskCollection.cs b/Engine/Source/Programs/Horde/Horde.Server/Jobs/Bisect/IBisectTaskCollection.cs
--- a/Engine/Source/Programs/Horde/Horde.Server/Jobs/Bisect/IBisectTaskCollection.cs
+++ b/Engine/Source/Programs/Horde/Horde.Server/Jobs/Bisect/IBisectTaskCollection.cs
@@ -4,11 +4,11 @@ using System;
 using System.Collections.Generic;
 using System.Threading;
 using System.Threading.Tasks;
+using EpicGames.Horde.Commits;
 using EpicGames.Horde.Jobs;
 using EpicGames.Horde.Jobs.Bisect;
 using EpicGames.Horde.Users;
 using Horde.Server.Jobs.Graphs;
-using Horde.Server.Perforce;
 using Microsoft.Extensions.Logging;
 
 namespace Horde.Server.Jobs.Bisect
diff --git a/Engine/Source/Programs/Horde/Horde.Server/Jobs/JobService.cs b/Engine/Source/Programs/Horde/Horde.Server/Jobs/JobService.cs
--- a/Engine/Source/Programs/Horde/Horde.Server/Jobs/JobService.cs
+++ b/Engine/Source/Programs/Horde/Horde.Server/Jobs/JobService.cs
@@ -13,6 +13,7 @@ using EpicGames.Horde.Acls;
 using EpicGames.Horde.Agents;
 using EpicGames.Horde.Agents.Leases;
 using EpicGames.Horde.Agents.Sessions;
+using EpicGames.Horde.Commits;
 using EpicGames.Horde.Jobs;
 using EpicGames.Horde.Jobs.Templates;
 using EpicGames.Horde.Logs;
diff --git a/Engine/Source/Programs/Horde/Horde.Server/Jobs/Schedules/ScheduleService.cs b/Engine/Source/Programs/Horde/Horde.Server/Jobs/Schedules/ScheduleService.cs
--- a/Engine/Source/Programs/Horde/Horde.Server/Jobs/Schedules/ScheduleService.cs
+++ b/Engine/Source/Programs/Horde/Horde.Server/Jobs/Schedules/ScheduleService.cs
@@ -8,6 +8,7 @@ using System.Text.RegularExpressions;
 using System.Threading;
 using System.Threading.Tasks;
 using EpicGames.Core;
+using EpicGames.Horde.Commits;
 using EpicGames.Horde.Jobs;
 using EpicGames.Horde.Jobs.Templates;
 using EpicGames.Horde.Streams;
diff --git a/Engine/Source/Programs/Horde/Horde.Server/Jobs/Schedules/Schedules.cs b/Engine/Source/Programs/Horde/Horde.Server/Jobs/Schedules/Schedules.cs
--- a/Engine/Source/Programs/Horde/Horde.Server/Jobs/Schedules/Schedules.cs
+++ b/Engine/Source/Programs/Horde/Horde.Server/Jobs/Schedules/Schedules.cs
@@ -2,8 +2,8 @@
 
 using System;
 using System.Collections.Generic;
+using EpicGames.Horde.Commits;
 using EpicGames.Horde.Jobs;
-using Horde.Server.Perforce;
 using Horde.Server.Streams;
 
 namespace Horde.Server.Jobs.Schedules
diff --git a/Engine/Source/Programs/Horde/Horde.Server/Perforce/CommitTag.cs b/Engine/Source/Programs/Horde/Horde.Server/Perforce/CommitTag.cs
deleted file mode 100644
diff --git a/Engine/Source/Programs/Horde/Horde.Server/Perforce/ICommit.cs b/Engine/Source/Programs/Horde/Horde.Server/Perforce/ICommit.cs
--- a/Engine/Source/Programs/Horde/Horde.Server/Perforce/ICommit.cs
+++ b/Engine/Source/Programs/Horde/Horde.Server/Perforce/ICommit.cs
@@ -5,6 +5,7 @@ using System.Collections.Generic;
 using System.Threading;
 using System.Threading.Tasks;
 using EpicGames.Core;
+using EpicGames.Horde.Commits;
 using EpicGames.Horde.Streams;
 using EpicGames.Horde.Users;
 
diff --git a/Engine/Source/Programs/Horde/Horde.Server/Perforce/ICommitCollection.cs b/Engine/Source/Programs/Horde/Horde.Server/Perforce/ICommitCollection.cs
--- a/Engine/Source/Programs/Horde/Horde.Server/Perforce/ICommitCollection.cs
+++ b/Engine/Source/Programs/Horde/Horde.Server/Perforce/ICommitCollection.cs
@@ -5,6 +5,7 @@ using System.Collections.Generic;
 using System.Linq;
 using System.Threading;
 using System.Threading.Tasks;
+using EpicGames.Horde.Commits;
 using EpicGames.Perforce;
 using Horde.Server.Jobs.Templates;
 
diff --git a/Engine/Source/Programs/Horde/Horde.Server/Perforce/IPerforceService.cs b/Engine/Source/Programs/Horde/Horde.Server/Perforce/IPerforceService.cs
--- a/Engine/Source/Programs/Horde/Horde.Server/Perforce/IPerforceService.cs
+++ b/Engine/Source/Programs/Horde/Horde.Server/Perforce/IPerforceService.cs
@@ -5,6 +5,7 @@ using System.Collections.Generic;
 using System.Linq;
 using System.Threading;
 using System.Threading.Tasks;
+using EpicGames.Horde.Commits;
 using EpicGames.Perforce;
 using Horde.Server.Streams;
 using Horde.Server.Users;
diff --git a/Engine/Source/Programs/Horde/Horde.Server/Perforce/PerforceService.cs b/Engine/Source/Programs/Horde/Horde.Server/Perforce/PerforceService.cs
--- a/Engine/Source/Programs/Horde/Horde.Server/Perforce/PerforceService.cs
+++ b/Engine/Source/Programs/Horde/Horde.Server/Perforce/PerforceService.cs
@@ -14,6 +14,7 @@ using System.Text.RegularExpressions;
 using System.Threading;
 using System.Threading.Tasks;
 using EpicGames.Core;
+using EpicGames.Horde.Commits;
 using EpicGames.Horde.Streams;
 using EpicGames.Horde.Users;
 using EpicGames.Perforce;
diff --git a/Engine/Source/Programs/Horde/Horde.Server/Perforce/PerforceServiceCache.cs b/Engine/Source/Programs/Horde/Horde.Server/Perforce/PerforceServiceCache.cs
--- a/Engine/Source/Programs/Horde/Horde.Server/Perforce/PerforceServiceCache.cs
+++ b/Engine/Source/Programs/Horde/Horde.Server/Perforce/PerforceServiceCache.cs
@@ -9,6 +9,7 @@ using System.Text;
 using System.Threading;
 using System.Threading.Tasks;
 using EpicGames.Core;
+using EpicGames.Horde.Commits;
 using EpicGames.Horde.Streams;
 using EpicGames.Horde.Users;
 using EpicGames.Perforce;
diff --git a/Engine/Source/Programs/Horde/Horde.Server/Startup.cs b/Engine/Source/Programs/Horde/Horde.Server/Startup.cs
--- a/Engine/Source/Programs/Horde/Horde.Server/Startup.cs
+++ b/Engine/Source/Programs/Horde/Horde.Server/Startup.cs
@@ -26,6 +26,7 @@ using EpicGames.Horde;
 using EpicGames.Horde.Acls;
 using EpicGames.Horde.Agents;
 using EpicGames.Horde.Agents.Leases;
+using EpicGames.Horde.Commits;
 using EpicGames.Horde.Jobs;
 using EpicGames.Horde.Logs;
 using EpicGames.Horde.Projects;
@@ -1197,6 +1198,21 @@ namespace Horde.Server
 			}
 		}
 
+		sealed class CommitTagBsonSerializer : SerializerBase<CommitTag>
+		{
+			/// <inheritdoc/>
+			public override CommitTag Deserialize(BsonDeserializationContext context, BsonDeserializationArgs args)
+			{
+				return new CommitTag(context.Reader.ReadString());
+			}
+
+			/// <inheritdoc/>
+			public override void Serialize(BsonSerializationContext context, BsonSerializationArgs args, CommitTag value)
+			{
+				context.Writer.WriteString(value.ToString());
+			}
+		}
+
 		static int s_haveConfiguredMongoDb = 0;
 
 		public static void ConfigureMongoDbClient()
@@ -1219,6 +1235,7 @@ namespace Horde.Server
 				BsonSerializer.RegisterSerializer(new AclScopeNameBsonSerializer());
 				BsonSerializer.RegisterSerializer(new ConditionSerializer());
 				BsonSerializer.RegisterSerializer(new SubResourceIdSerializer());
+				BsonSerializer.RegisterSerializer(new CommitTagBsonSerializer());
 				BsonSerializer.RegisterSerializationProvider(new BsonSerializationProvider());
 				BsonSerializer.RegisterSerializationProvider(new StringIdBsonSerializationProvider());
 				BsonSerializer.RegisterSerializationProvider(new BinaryIdBsonSerializationProvider());
diff --git a/Engine/Source/Programs/Horde/Horde.Server/Streams/StreamConfig.cs b/Engine/Source/Programs/Horde/Horde.Server/Streams/StreamConfig.cs
--- a/Engine/Source/Programs/Horde/Horde.Server/Streams/StreamConfig.cs
+++ b/Engine/Source/Programs/Horde/Horde.Server/Streams/StreamConfig.cs
@@ -14,6 +14,7 @@ using EpicGames.Core;
 using EpicGames.Horde;
 using EpicGames.Horde.Acls;
 using EpicGames.Horde.Agents.Pools;
+using EpicGames.Horde.Commits;
 using EpicGames.Horde.Common;
 using EpicGames.Horde.Issues;
 using EpicGames.Horde.Jobs;
@@ -21,13 +22,13 @@ using EpicGames.Horde.Jobs.Templates;
 using EpicGames.Horde.Replicators;
 using EpicGames.Horde.Streams;
 using EpicGames.Horde.Telemetry;
+using EpicGames.Perforce;
 using Horde.Server.Acls;
 using Horde.Server.Configuration;
 using Horde.Server.Issues;
 using Horde.Server.Jobs;
 using Horde.Server.Jobs.Graphs;
 using Horde.Server.Jobs.Templates;
-using Horde.Server.Perforce;
 using Horde.Server.Projects;
 using Horde.Server.Replicators;
 using Horde.Server.Server;
@@ -1232,11 +1233,19 @@ namespace Horde.Server.Streams
 		/// <summary>
 		/// Default config for code filters
 		/// </summary>
-		public static CommitTagConfig CodeDefault { get; } = new CommitTagConfig { Name = CommitTag.Code, Filter = CommitTag.CodeFilter.ToList() };
+		public static CommitTagConfig CodeDefault { get; } = new CommitTagConfig
+		{
+			Name = CommitTag.Code,
+			Filter = PerforceUtils.CodeExtensions.Select(x => $"*{x}").ToList()
+		};
 
 		/// <summary>
 		/// Default config for content filters
 		/// </summary>
-		public static CommitTagConfig ContentDefault { get; } = new CommitTagConfig { Name = CommitTag.Content, Filter = CommitTag.ContentFilter.ToList() };
+		public static CommitTagConfig ContentDefault { get; } = new CommitTagConfig
+		{
+			Name = CommitTag.Content,
+			Filter = Enumerable.Concat(new[] { "*" }, PerforceUtils.CodeExtensions.Select(x => $"-*{x}")).ToList()
+		};
 	}
 }
diff --git a/Engine/Source/Programs/Horde/Horde.Server/Streams/StreamMessages.cs b/Engine/Source/Programs/Horde/Horde.Server/Streams/StreamMessages.cs
--- a/Engine/Source/Programs/Horde/Horde.Server/Streams/StreamMessages.cs
+++ b/Engine/Source/Programs/Horde/Horde.Server/Streams/StreamMessages.cs
@@ -4,6 +4,7 @@ using System;
 using System.Collections.Generic;
 using System.ComponentModel.DataAnnotations;
 using EpicGames.Core;
+using EpicGames.Horde.Commits;
 using EpicGames.Horde.Jobs;
 using EpicGames.Horde.Jobs.Templates;
 using EpicGames.Horde.Projects;
diff --git a/Engine/Source/Programs/Horde/Horde.Server/Streams/StreamsController.cs b/Engine/Source/Programs/Horde/Horde.Server/Streams/StreamsController.cs
--- a/Engine/Source/Programs/Horde/Horde.Server/Streams/StreamsController.cs
+++ b/Engine/Source/Programs/Horde/Horde.Server/Streams/StreamsController.cs
@@ -5,6 +5,7 @@ using System.Collections.Generic;
 using System.Linq;
 using System.Threading;
 using System.Threading.Tasks;
+using EpicGames.Horde.Commits;
 using EpicGames.Horde.Jobs.Templates;
 using EpicGames.Horde.Projects;
 using EpicGames.Horde.Streams;
diff --git a/Engine/Source/Programs/Shared/EpicGames.Horde/Commits/CommitTag.cs b/Engine/Source/Programs/Shared/EpicGames.Horde/Commits/CommitTag.cs
new file mode 100644
--- /dev/null
+++ b/Engine/Source/Programs/Shared/EpicGames.Horde/Commits/CommitTag.cs
@@ -0,0 +1,93 @@
+// Copyright Epic Games, Inc. All Rights Reserved.
+
+using System;
+using System.Text.Json;
+using System.Text.Json.Serialization;
+using EpicGames.Core;
+
+namespace EpicGames.Horde.Commits
+{
+	/// <summary>
+	/// Constants for known commit tags
+	/// </summary>
+	[JsonSchemaString]
+	[JsonConverter(typeof(CommitTagJsonConverter))]
+	public readonly struct CommitTag : IEquatable<CommitTag>
+	{
+		/// <summary>
+		/// Predefined filter name for commits containing code
+		/// </summary>
+		public static CommitTag Code { get; } = new CommitTag("code");
+
+		/// <summary>
+		/// Predefined filter name for commits containing content
+		/// </summary>
+		public static CommitTag Content { get; } = new CommitTag("content");
+
+		/// <summary>
+		/// The tag text
+		/// </summary>
+		public string Text { get; }
+
+		/// <summary>
+		/// Constructor
+		/// </summary>
+		/// <param name="text"></param>
+		public CommitTag(string text)
+		{
+			Text = text;
+			StringId.ValidateArgument(new Utf8String(text), nameof(text));
+		}
+
+		/// <summary>
+		/// Check if the current tag is empty
+		/// </summary>
+		public bool IsEmpty() => Text == null;
+
+		/// <inheritdoc/>
+		public override bool Equals(object? obj) => obj is CommitTag id && Equals(id);
+
+		/// <inheritdoc/>
+		public override int GetHashCode() => String.GetHashCode(Text, StringComparison.Ordinal);
+
+		/// <inheritdoc/>
+		public bool Equals(CommitTag other) => String.Equals(Text, other.Text, StringComparison.Ordinal);
+
+		/// <inheritdoc/>
+		public override string ToString() => Text;
+
+		/// <summary>
+		/// Compares two string ids for equality
+		/// </summary>
+		/// <param name="left">The first string id</param>
+		/// <param name="right">Second string id</param>
+		/// <returns>True if the two string ids are equal</returns>
+		public static bool operator ==(CommitTag left, CommitTag right) => left.Equals(right);
+
+		/// <summary>
+		/// Compares two string ids for inequality
+		/// </summary>
+		/// <param name="left">The first string id</param>
+		/// <param name="right">Second string id</param>
+		/// <returns>True if the two string ids are not equal</returns>
+		public static bool operator !=(CommitTag left, CommitTag right) => !left.Equals(right);
+	}
+
+	/// <summary>
+	/// Converts <see cref="CommitTag"/> values to and from JSON
+	/// </summary>
+	public class CommitTagJsonConverter : JsonConverter<CommitTag>
+	{
+		/// <inheritdoc/>
+		public override CommitTag Read(ref Utf8JsonReader reader, Type typeToConvert, JsonSerializerOptions options)
+		{
+			return new CommitTag(reader.GetString() ?? String.Empty);
+		}
+
+		/// <inheritdoc/>
+		public override void Write(Utf8JsonWriter writer, CommitTag value, JsonSerializerOptions options)
+		{
+			writer.WriteStringValue(value.Text);
+		}
+	}
+}
