Horde: Move graph messages into EpicGames.Horde.

#rnx

diff --git a/Engine/Source/Programs/Horde/Horde.Server/Jobs/Graphs/GraphMessages.cs b/Engine/Source/Programs/Horde/Horde.Server/Jobs/Graphs/GraphMessages.cs
deleted file mode 100644
diff --git a/Engine/Source/Programs/Horde/Horde.Server/Jobs/JobsController.cs b/Engine/Source/Programs/Horde/Horde.Server/Jobs/JobsController.cs
--- a/Engine/Source/Programs/Horde/Horde.Server/Jobs/JobsController.cs
+++ b/Engine/Source/Programs/Horde/Horde.Server/Jobs/JobsController.cs
@@ -8,6 +8,7 @@ using System.Threading.Tasks;
 using EpicGames.Core;
 using EpicGames.Horde.Artifacts;
 using EpicGames.Horde.Jobs;
+using EpicGames.Horde.Jobs.Graphs;
 using EpicGames.Horde.Jobs.Templates;
 using EpicGames.Horde.Logs;
 using EpicGames.Horde.Streams;
@@ -719,7 +720,68 @@ namespace Horde.Server.Jobs
 			}
 
 			IGraph graph = await _jobService.GetGraphAsync(job);
-			return PropertyFilter.Apply(new GetGraphResponse(graph), filter);
+			return PropertyFilter.Apply(CreateGetGraphResponse(graph), filter);
+		}
+
+		static GetGraphResponse CreateGetGraphResponse(IGraph graph)
+		{
+			GetGraphResponse response = new GetGraphResponse(graph.Id.ToString());
+			if (graph.Groups.Count > 0)
+			{
+				response.Groups = graph.Groups.ConvertAll(x => CreateGetGroupResponse(x, graph.Groups));
+			}
+			if (graph.Aggregates.Count > 0)
+			{
+				response.Aggregates = graph.Aggregates.ConvertAll(x => CreateGetAggregateResponse(x, graph.Groups));
+			}
+			if (graph.Labels.Count > 0)
+			{
+				response.Labels = graph.Labels.ConvertAll(x => CreateGetLabelResponse(x, graph.Groups));
+			}
+			return response;
+		}
+
+		static GetGroupResponse CreateGetGroupResponse(INodeGroup group, IReadOnlyList<INodeGroup> groups)
+		{
+			GetGroupResponse response = new GetGroupResponse(group.AgentType);
+			response.Nodes.AddRange(group.Nodes.Select(x => CreateGetNodeResponse(x, groups)));
+			return response;
+		}
+
+		static GetNodeResponse CreateGetNodeResponse(INode node, IReadOnlyList<INodeGroup> groups)
+		{
+			GetNodeResponse response = new GetNodeResponse(node.Name);
+			response.Inputs.AddRange(node.Inputs.Select(x => groups[x.NodeRef.GroupIdx].Nodes[x.NodeRef.NodeIdx].OutputNames[x.OutputIdx]));
+			response.Outputs.AddRange(node.OutputNames);
+			response.InputDependencies.AddRange(node.InputDependencies.Select(x => groups[x.GroupIdx].Nodes[x.NodeIdx].Name));
+			response.OrderDependencies.AddRange(node.OrderDependencies.Select(x => groups[x.GroupIdx].Nodes[x.NodeIdx].Name));
+			response.Priority = node.Priority;
+			response.AllowRetry = node.AllowRetry;
+			response.RunEarly = node.RunEarly;
+			response.Warnings = node.Warnings;
+			response.Credentials = node.Credentials;
+			response.Properties = node.Properties;
+			response.Annotations = node.Annotations;
+			return response;
+		}
+
+		static GetAggregateResponse CreateGetAggregateResponse(IAggregate aggregate, IReadOnlyList<INodeGroup> groups)
+		{
+			GetAggregateResponse response = new GetAggregateResponse(aggregate.Name);
+			response.Nodes = aggregate.Nodes.ConvertAll(x => x.ToNode(groups).Name);
+			return response;
+		}
+
+		static GetLabelResponse CreateGetLabelResponse(ILabel label, IReadOnlyList<INodeGroup> groups)
+		{
+			GetLabelResponse response = new GetLabelResponse();
+			response.DashboardName = label.DashboardName;
+			response.DashboardCategory = label.DashboardCategory;
+			response.UgsName = label.UgsName;
+			response.UgsProject = label.UgsProject;
+			response.RequiredNodes = label.RequiredNodes.ConvertAll(x => x.ToNode(groups).Name);
+			response.IncludedNodes = label.IncludedNodes.ConvertAll(x => x.ToNode(groups).Name);
+			return response;
 		}
 
 		/// <summary>
@@ -1100,7 +1162,7 @@ namespace Horde.Server.Jobs
 			}
 
 			IGraph graph = await _jobService.GetGraphAsync(job, cancellationToken);
-			return graph.Groups.ConvertAll(x => new GetGroupResponse(x, graph.Groups).ApplyFilter(filter));
+			return graph.Groups.ConvertAll(x => CreateGetGroupResponse(x, graph.Groups).ApplyFilter(filter));
 		}
 
 		/// <summary>
@@ -1139,7 +1201,7 @@ namespace Horde.Server.Jobs
 				return NotFound(jobId, groupIdx);
 			}
 
-			return new GetGroupResponse(graph.Groups[groupIdx], graph.Groups).ApplyFilter(filter);
+			return CreateGetGroupResponse(graph.Groups[groupIdx], graph.Groups).ApplyFilter(filter);
 		}
 
 		/// <summary>
@@ -1176,7 +1238,7 @@ namespace Horde.Server.Jobs
 				return NotFound(jobId, groupIdx);
 			}
 
-			return graph.Groups[groupIdx].Nodes.ConvertAll(x => new GetNodeResponse(x, graph.Groups).ApplyFilter(filter));
+			return graph.Groups[groupIdx].Nodes.ConvertAll(x => CreateGetNodeResponse(x, graph.Groups).ApplyFilter(filter));
 		}
 
 		/// <summary>
@@ -1218,7 +1280,7 @@ namespace Horde.Server.Jobs
 				return NotFound(jobId, groupIdx, nodeIdx);
 			}
 
-			return new GetNodeResponse(graph.Groups[groupIdx].Nodes[nodeIdx], graph.Groups).ApplyFilter(filter);
+			return CreateGetNodeResponse(graph.Groups[groupIdx].Nodes[nodeIdx], graph.Groups).ApplyFilter(filter);
 		}
 
 		/// <summary>
diff --git a/Engine/Source/Programs/Horde/Horde.Server/Server/DebugController.cs b/Engine/Source/Programs/Horde/Horde.Server/Server/DebugController.cs
--- a/Engine/Source/Programs/Horde/Horde.Server/Server/DebugController.cs
+++ b/Engine/Source/Programs/Horde/Horde.Server/Server/DebugController.cs
@@ -13,7 +13,6 @@ using System.Text;
 using System.Threading;
 using System.Threading.Tasks;
 using System.Web;
-using EpicGames.Core;
 using EpicGames.Horde.Agents.Leases;
 using EpicGames.Horde.Compute;
 using EpicGames.Horde.Jobs;
@@ -23,7 +22,6 @@ using Horde.Common.Rpc;
 using Horde.Server.Agents.Relay;
 using Horde.Server.Configuration;
 using Horde.Server.Jobs;
-using Horde.Server.Jobs.Graphs;
 using Horde.Server.Logs;
 using Horde.Server.Projects;
 using Horde.Server.Streams;
@@ -260,7 +258,6 @@ namespace Horde.Server.Server
 		private readonly AgentRelayService _agentRelayService;
 		private readonly JobService _jobService;
 		private readonly JobTaskSource _jobTaskSource;
-		private readonly IGraphCollection _graphCollection;
 		private readonly ILogFileCollection _logFileCollection;
 		private readonly IOptionsSnapshot<GlobalConfig> _globalConfig;
 		private readonly ILogger<SecureDebugController> _logger;
@@ -274,7 +271,6 @@ namespace Horde.Server.Server
 			AgentRelayService agentRelayService,
 			JobService jobService,
 			JobTaskSource jobTaskSource,
-			IGraphCollection graphCollection,
 			ILogFileCollection logFileCollection, IOptionsSnapshot<GlobalConfig> globalConfig, ILogger<SecureDebugController> logger)
 		{
 			_mongoService = mongoService;
@@ -282,7 +278,6 @@ namespace Horde.Server.Server
 			_jobService = jobService;
 			_agentRelayService = agentRelayService;
 			_jobTaskSource = jobTaskSource;
-			_graphCollection = graphCollection;
 			_logFileCollection = logFileCollection;
 			_globalConfig = globalConfig;
 			_logger = logger;
@@ -499,42 +494,6 @@ namespace Horde.Server.Server
 			return Ok($"Log message generated logLevel={logLevelInternal} messageLen={messageLen} exceptionMessageLen={exceptionMessageLen} argCount={argCount} argLen={argLen}");
 		}
 
-		/// <summary>
-		/// Queries for all graphs
-		/// </summary>
-		/// <returns>The graph definitions</returns>
-		[HttpGet]
-		[Route("/api/v1/debug/graphs")]
-		[ProducesResponseType(200, Type = typeof(GetGraphResponse))]
-		public async Task<ActionResult<List<object>>> GetGraphsAsync([FromQuery] int? index = null, [FromQuery] int? count = null, [FromQuery] PropertyFilter? filter = null)
-		{
-			if (!_globalConfig.Value.Authorize(ServerAclAction.Debug, User))
-			{
-				return Forbid(ServerAclAction.Debug);
-			}
-
-			List<IGraph> graphs = await _graphCollection.FindAllAsync(null, index, count);
-			return graphs.ConvertAll(x => new GetGraphResponse(x).ApplyFilter(filter));
-		}
-
-		/// <summary>
-		/// Queries for a particular graph by hash
-		/// </summary>
-		/// <returns>The graph definition</returns>
-		[HttpGet]
-		[Route("/api/v1/debug/graphs/{GraphId}")]
-		[ProducesResponseType(200, Type = typeof(GetGraphResponse))]
-		public async Task<ActionResult<object>> GetGraphAsync(string graphId, [FromQuery] PropertyFilter? filter = null)
-		{
-			if (!_globalConfig.Value.Authorize(ServerAclAction.Debug, User))
-			{
-				return Forbid(ServerAclAction.Debug);
-			}
-
-			IGraph graph = await _graphCollection.GetAsync(ContentHash.Parse(graphId));
-			return new GetGraphResponse(graph).ApplyFilter(filter);
-		}
-
 		/// <summary>
 		/// Retrieve metadata about a specific log file
 		/// </summary>
diff --git a/Engine/Source/Programs/Shared/EpicGames.Horde/Jobs/Graphs/GraphMessages.cs b/Engine/Source/Programs/Shared/EpicGames.Horde/Jobs/Graphs/GraphMessages.cs
new file mode 100644
--- /dev/null
+++ b/Engine/Source/Programs/Shared/EpicGames.Horde/Jobs/Graphs/GraphMessages.cs
@@ -0,0 +1,216 @@
+// Copyright Epic Games, Inc. All Rights Reserved.
+
+using System.Collections.Generic;
+using System.Linq;
+using EpicGames.Core;
+
+namespace EpicGames.Horde.Jobs.Graphs
+{
+	/// <summary>
+	/// Information required to create a node
+	/// </summary>
+	public class GetNodeResponse
+	{
+		/// <summary>
+		/// The name of this node 
+		/// </summary>
+		public string Name { get; set; }
+
+		/// <summary>
+		/// Inputs for this node
+		/// </summary>
+		public List<string> Inputs { get; set; } = new List<string>();
+
+		/// <summary>
+		/// Output from this node
+		/// </summary>
+		public List<string> Outputs { get; set; } = new List<string>();
+
+		/// <summary>
+		/// Indices of nodes which must have succeeded for this node to run
+		/// </summary>
+		public List<string> InputDependencies { get; set; } = new List<string>();
+
+		/// <summary>
+		/// Indices of nodes which must have completed for this node to run
+		/// </summary>
+		public List<string> OrderDependencies { get; set; } = new List<string>();
+
+		/// <summary>
+		/// The priority of this node
+		/// </summary>
+		public Priority Priority { get; set; }
+
+		/// <summary>
+		/// Whether this node can be retried
+		/// </summary>
+		public bool AllowRetry { get; set; }
+
+		/// <summary>
+		/// This node can start running early, before dependencies of other nodes in the same group are complete
+		/// </summary>
+		public bool RunEarly { get; set; }
+
+		/// <summary>
+		/// Whether to include warnings in diagnostic output
+		/// </summary>
+		public bool Warnings { get; set; }
+
+		/// <summary>
+		/// Average time to execute this node based on historical trends
+		/// </summary>
+		public float? AverageDuration { get; set; }
+
+		/// <summary>
+		/// Credentials required for this node to run. This dictionary maps from environment variable names to a credential property in the format 'CredentialName.PropertyName'.
+		/// </summary>
+		public IReadOnlyDictionary<string, string>? Credentials { get; set; }
+
+		/// <summary>
+		/// Properties for this node
+		/// </summary>
+		public IReadOnlyDictionary<string, string>? Properties { get; set; }
+
+		/// <summary>
+		/// Annotations for this node
+		/// </summary>
+		public IReadOnlyDictionary<string, string>? Annotations { get; set; }
+
+		/// <summary>
+		/// Constructor
+		/// </summary>
+		/// <param name="name"></param>
+		public GetNodeResponse(string name)
+		{
+			Name = name;
+		}
+	}
+
+	/// <summary>
+	/// Information about a group of nodes
+	/// </summary>
+	public class GetGroupResponse
+	{
+		/// <summary>
+		/// The type of agent to execute this group
+		/// </summary>
+		public string AgentType { get; set; }
+
+		/// <summary>
+		/// Nodes in the group
+		/// </summary>
+		public List<GetNodeResponse> Nodes { get; set; } = new List<GetNodeResponse>();
+
+		/// <summary>
+		/// Constructor
+		/// </summary>
+		public GetGroupResponse(string agentType)
+		{
+			AgentType = agentType;
+		}
+	}
+
+	/// <summary>
+	/// Information about an aggregate
+	/// </summary>
+	public class GetAggregateResponse
+	{
+		/// <summary>
+		/// Name of the aggregate
+		/// </summary>
+		public string Name { get; set; }
+
+		/// <summary>
+		/// Nodes which must be part of the job for the aggregate to be shown
+		/// </summary>
+		public List<string> Nodes { get; set; } = new List<string>();
+
+		/// <summary>
+		/// Constructor
+		/// </summary>
+		public GetAggregateResponse(string name)
+		{
+			Name = name;
+		}
+	}
+
+	/// <summary>
+	/// Information about a label
+	/// </summary>
+	public class GetLabelResponse
+	{
+		/// <summary>
+		/// Category of the aggregate
+		/// </summary>
+		public string? Category => DashboardCategory;
+
+		/// <summary>
+		/// Label for this aggregate
+		/// </summary>
+		public string? Name => DashboardName;
+
+		/// <summary>
+		/// Name to show for this label on the dashboard
+		/// </summary>
+		public string? DashboardName { get; set; }
+
+		/// <summary>
+		/// Category to show this label in on the dashboard
+		/// </summary>
+		public string? DashboardCategory { get; set; }
+
+		/// <summary>
+		/// Name to show for this label in UGS
+		/// </summary>
+		public string? UgsName { get; set; }
+
+		/// <summary>
+		/// Project to display this label for in UGS
+		/// </summary>
+		public string? UgsProject { get; set; }
+
+		/// <summary>
+		/// Nodes which must be part of the job for the aggregate to be shown
+		/// </summary>
+		public List<string> RequiredNodes { get; set; } = new List<string>();
+
+		/// <summary>
+		/// Nodes to include in the status of this aggregate, if present in the job
+		/// </summary>
+		public List<string> IncludedNodes { get; set; } = new List<string>();
+	}
+
+	/// <summary>
+	/// Information about a graph
+	/// </summary>
+	public class GetGraphResponse
+	{
+		/// <summary>
+		/// The hash of the graph
+		/// </summary>
+		public string Hash { get; set; }
+
+		/// <summary>
+		/// Array of nodes for this job
+		/// </summary>
+		public List<GetGroupResponse>? Groups { get; set; }
+
+		/// <summary>
+		/// List of aggregates
+		/// </summary>
+		public List<GetAggregateResponse>? Aggregates { get; set; }
+
+		/// <summary>
+		/// List of labels for the graph
+		/// </summary>
+		public List<GetLabelResponse>? Labels { get; set; }
+
+		/// <summary>
+		/// Constructor
+		/// </summary>
+		public GetGraphResponse(string hash)
+		{
+			Hash = hash;
+		}
+	}
+}
