Horde: Move log messages into EpicGames.Horde.

#rnx

diff --git a/Engine/Source/Programs/Horde/Horde.Server.Tests/Jobs/TestData/TestDataTests.cs b/Engine/Source/Programs/Horde/Horde.Server.Tests/Jobs/TestData/TestDataTests.cs
--- a/Engine/Source/Programs/Horde/Horde.Server.Tests/Jobs/TestData/TestDataTests.cs
+++ b/Engine/Source/Programs/Horde/Horde.Server.Tests/Jobs/TestData/TestDataTests.cs
@@ -9,6 +9,7 @@ using System.Threading.Tasks;
 using EpicGames.Horde;
 using EpicGames.Horde.Jobs;
 using EpicGames.Horde.Jobs.Templates;
+using EpicGames.Horde.Logs;
 using EpicGames.Horde.Projects;
 using EpicGames.Horde.Streams;
 using Horde.Server.Jobs;
diff --git a/Engine/Source/Programs/Horde/Horde.Server.Tests/Logs/LogChunkTests.cs b/Engine/Source/Programs/Horde/Horde.Server.Tests/Logs/LogChunkTests.cs
--- a/Engine/Source/Programs/Horde/Horde.Server.Tests/Logs/LogChunkTests.cs
+++ b/Engine/Source/Programs/Horde/Horde.Server.Tests/Logs/LogChunkTests.cs
@@ -3,7 +3,7 @@
 using System;
 using System.Collections.Generic;
 using System.Text;
-using Horde.Server.Logs;
+using EpicGames.Horde.Logs;
 using Horde.Server.Logs.Data;
 using Microsoft.Extensions.Logging.Abstractions;
 using Microsoft.VisualStudio.TestTools.UnitTesting;
diff --git a/Engine/Source/Programs/Horde/Horde.Server/Logs/Data/LogChunkData.cs b/Engine/Source/Programs/Horde/Horde.Server/Logs/Data/LogChunkData.cs
--- a/Engine/Source/Programs/Horde/Horde.Server/Logs/Data/LogChunkData.cs
+++ b/Engine/Source/Programs/Horde/Horde.Server/Logs/Data/LogChunkData.cs
@@ -4,6 +4,7 @@ using System;
 using System.Collections.Generic;
 using System.Linq;
 using EpicGames.Core;
+using EpicGames.Horde.Logs;
 using Microsoft.Extensions.Logging;
 
 namespace Horde.Server.Logs.Data
diff --git a/Engine/Source/Programs/Horde/Horde.Server/Logs/Data/LogSubChunkData.cs b/Engine/Source/Programs/Horde/Horde.Server/Logs/Data/LogSubChunkData.cs
--- a/Engine/Source/Programs/Horde/Horde.Server/Logs/Data/LogSubChunkData.cs
+++ b/Engine/Source/Programs/Horde/Horde.Server/Logs/Data/LogSubChunkData.cs
@@ -2,6 +2,7 @@
 
 using System;
 using EpicGames.Core;
+using EpicGames.Horde.Logs;
 using Horde.Server.Utilities;
 using Microsoft.Extensions.Logging;
 
diff --git a/Engine/Source/Programs/Horde/Horde.Server/Logs/LogEventMessages.cs b/Engine/Source/Programs/Horde/Horde.Server/Logs/LogEventMessages.cs
--- a/Engine/Source/Programs/Horde/Horde.Server/Logs/LogEventMessages.cs
+++ b/Engine/Source/Programs/Horde/Horde.Server/Logs/LogEventMessages.cs
@@ -3,6 +3,7 @@
 using System.Collections.Generic;
 using System.Text.Json;
 using EpicGames.Core;
+using EpicGames.Horde.Logs;
 using HordeCommon;
 
 namespace Horde.Server.Logs
@@ -15,7 +16,7 @@ namespace Horde.Server.Logs
 		/// <summary>
 		/// Unique id of the log containing this event
 		/// </summary>
-		public string LogId { get; set; }
+		public LogId LogId { get; set; }
 
 		/// <summary>
 		/// Severity of this event
@@ -51,7 +52,7 @@ namespace Horde.Server.Logs
 		public GetLogEventResponse(ILogEvent logEvent, ILogEventData eventData, int? issueId)
 		{
 			Severity = logEvent.Severity;
-			LogId = logEvent.LogId.ToString();
+			LogId = logEvent.LogId;
 			LineIndex = logEvent.LineIndex;
 			LineCount = logEvent.LineCount;
 			IssueId = issueId;
diff --git a/Engine/Source/Programs/Horde/Horde.Server/Logs/LogFileMessages.cs b/Engine/Source/Programs/Horde/Horde.Server/Logs/LogFileMessages.cs
deleted file mode 100644
diff --git a/Engine/Source/Programs/Horde/Horde.Server/Logs/LogsController.cs b/Engine/Source/Programs/Horde/Horde.Server/Logs/LogsController.cs
--- a/Engine/Source/Programs/Horde/Horde.Server/Logs/LogsController.cs
+++ b/Engine/Source/Programs/Horde/Horde.Server/Logs/LogsController.cs
@@ -78,7 +78,7 @@ namespace Horde.Server.Logs
 		/// <returns>Information about the requested project</returns>
 		[HttpGet]
 		[Route("/api/v1/logs/{logFileId}")]
-		[ProducesResponseType(typeof(GetLogFileResponse), 200)]
+		[ProducesResponseType(typeof(GetLogResponse), 200)]
 		public async Task<ActionResult<object>> GetLogAsync(LogId logFileId, [FromQuery] PropertyFilter? filter = null, CancellationToken cancellationToken = default)
 		{
 			ILogFile? logFile = await _logFileService.GetLogFileAsync(logFileId, cancellationToken);
@@ -92,7 +92,19 @@ namespace Horde.Server.Logs
 			}
 
 			LogMetadata metadata = await _logFileService.GetMetadataAsync(logFile, cancellationToken);
-			return new GetLogFileResponse(logFile, metadata).ApplyFilter(filter);
+			return CreateGetLogFileResponse(logFile, metadata).ApplyFilter(filter);
+		}
+
+		static GetLogResponse CreateGetLogFileResponse(ILogFile logFile, LogMetadata metadata)
+		{
+			GetLogResponse response = new GetLogResponse();
+			response.Id = logFile.Id;
+			response.JobId = logFile.JobId;
+			response.LeaseId = logFile.LeaseId;
+			response.SessionId = logFile.SessionId;
+			response.Type = logFile.Type;
+			response.LineCount = metadata.MaxLineIndex;
+			return response;
 		}
 
 		/// <summary>
@@ -268,7 +280,7 @@ namespace Horde.Server.Logs
 		/// <returns>Raw log data for the requested range</returns>
 		[HttpGet]
 		[Route("/api/v1/logs/{logFileId}/search")]
-		public async Task<ActionResult<SearchLogFileResponse>> SearchLogFileAsync(
+		public async Task<ActionResult<SearchLogResponse>> SearchLogFileAsync(
 			LogId logFileId,
 			[FromQuery] string text,
 			[FromQuery] int firstLine = 0,
@@ -285,7 +297,7 @@ namespace Horde.Server.Logs
 				return Forbid();
 			}
 
-			SearchLogFileResponse response = new SearchLogFileResponse();
+			SearchLogResponse response = new SearchLogResponse();
 			response.Stats = new SearchStats();
 			response.Lines = await _logFileService.SearchLogDataAsync(logFile, text, firstLine, count, response.Stats, cancellationToken);
 			return response;
diff --git a/Engine/Source/Programs/Shared/EpicGames.Horde/Jobs/Graphs/GraphMessages.cs b/Engine/Source/Programs/Shared/EpicGames.Horde/Jobs/Graphs/GraphMessages.cs
--- a/Engine/Source/Programs/Shared/EpicGames.Horde/Jobs/Graphs/GraphMessages.cs
+++ b/Engine/Source/Programs/Shared/EpicGames.Horde/Jobs/Graphs/GraphMessages.cs
@@ -1,8 +1,8 @@
 // Copyright Epic Games, Inc. All Rights Reserved.
 
 using System.Collections.Generic;
-using System.Linq;
-using EpicGames.Core;
+
+#pragma warning disable CA2227 // Change 'InputDependencies' to be read-only by removing the property setter
 
 namespace EpicGames.Horde.Jobs.Graphs
 {
diff --git a/Engine/Source/Programs/Shared/EpicGames.Horde/Logs/LogMessages.cs b/Engine/Source/Programs/Shared/EpicGames.Horde/Logs/LogMessages.cs
new file mode 100644
--- /dev/null
+++ b/Engine/Source/Programs/Shared/EpicGames.Horde/Logs/LogMessages.cs
@@ -0,0 +1,103 @@
+// Copyright Epic Games, Inc. All Rights Reserved.
+
+using System;
+using System.Collections.Generic;
+using EpicGames.Horde.Agents.Leases;
+using EpicGames.Horde.Agents.Sessions;
+using EpicGames.Horde.Jobs;
+
+#pragma warning disable CA2227 // Change 'Lines' to be read-only by removing the property setter
+
+namespace EpicGames.Horde.Logs
+{
+	/// <summary>
+	/// The type of data stored in this log file
+	/// </summary>
+	public enum LogType
+	{
+		/// <summary>
+		/// Plain text data
+		/// </summary>
+		Text,
+
+		/// <summary>
+		/// Structured json objects, output as one object per line (without trailing commas)
+		/// </summary>
+		Json
+	}
+
+	/// <summary>
+	/// Creates a new log file
+	/// </summary>
+	public class CreateLogRequest
+	{
+		/// <summary>
+		/// Type of the log file
+		/// </summary>
+		public LogType Type { get; set; } = LogType.Json;
+	}
+
+	/// <summary>
+	/// Response from creating a log file
+	/// </summary>
+	public class CreateLogResponse
+	{
+		/// <summary>
+		/// Identifier for the created log file
+		/// </summary>
+		public string Id { get; set; } = String.Empty;
+	}
+
+	/// <summary>
+	/// Response describing a log file
+	/// </summary>
+	public class GetLogResponse
+	{
+		/// <summary>
+		/// Unique id of the log file
+		/// </summary>
+		public LogId Id { get; set; }
+
+		/// <summary>
+		/// Unique id of the job for this log file
+		/// </summary>
+		public JobId JobId { get; set; }
+
+		/// <summary>
+		/// The lease allowed to write to this log
+		/// </summary>
+		public LeaseId? LeaseId { get; set; }
+
+		/// <summary>
+		/// The session allowed to write to this log
+		/// </summary>
+		public SessionId? SessionId { get; set; }
+
+		/// <summary>
+		/// Type of events stored in this log
+		/// </summary>
+		public LogType Type { get; set; }
+
+		/// <summary>
+		/// Number of lines in the file
+		/// </summary>
+		public int LineCount { get; set; }
+	}
+
+	/// <summary>
+	/// Response describing a log file
+	/// </summary>
+	public class SearchLogResponse
+	{
+		/// <summary>
+		/// List of line numbers containing the search text
+		/// </summary>
+		public List<int> Lines { get; set; } = new List<int>();
+
+		/// <summary>
+		/// Stats for the search
+		/// </summary>
+		public SearchStats? Stats { get; set; }
+	}
+}
+
