Horde: Move notice messages into EpicGames.Horde.

#rnx

diff --git a/Engine/Source/Programs/Horde/Horde.Server/Server/Notices/NoticeMessages.cs b/Engine/Source/Programs/Horde/Horde.Server/Server/Notices/NoticeMessages.cs
deleted file mode 100644
diff --git a/Engine/Source/Programs/Horde/Horde.Server/Server/Notices/NoticesController.cs b/Engine/Source/Programs/Horde/Horde.Server/Server/Notices/NoticesController.cs
--- a/Engine/Source/Programs/Horde/Horde.Server/Server/Notices/NoticesController.cs
+++ b/Engine/Source/Programs/Horde/Horde.Server/Server/Notices/NoticesController.cs
@@ -6,6 +6,7 @@ using System.Linq;
 using System.Threading;
 using System.Threading.Tasks;
 using EpicGames.Core;
+using EpicGames.Horde.Server.Notices;
 using EpicGames.Horde.Users;
 using Horde.Server.Acls;
 using Horde.Server.Users;
diff --git a/Engine/Source/Programs/Shared/EpicGames.Horde/Server/Notices/NoticeMessages.cs b/Engine/Source/Programs/Shared/EpicGames.Horde/Server/Notices/NoticeMessages.cs
new file mode 100644
--- /dev/null
+++ b/Engine/Source/Programs/Shared/EpicGames.Horde/Server/Notices/NoticeMessages.cs
@@ -0,0 +1,113 @@
+// Copyright Epic Games, Inc. All Rights Reserved.
+
+using System;
+using System.ComponentModel.DataAnnotations;
+using EpicGames.Horde.Users;
+
+namespace EpicGames.Horde.Server.Notices
+{
+	/// <summary>
+	/// Parameters required to create a notice
+	/// </summary>
+	public class CreateNoticeRequest
+	{
+		/// <summary>
+		/// Start time to display this message
+		/// </summary>
+		public DateTime? StartTime { get; set; }
+
+		/// <summary>
+		/// Finish time to display this message
+		/// </summary>
+		public DateTime? FinishTime { get; set; }
+
+		/// <summary>
+		/// Message to display
+		/// </summary>
+		[Required]
+		public string Message { get; set; }
+
+		/// <summary>
+		/// Constructor
+		/// </summary>
+		public CreateNoticeRequest(string message)
+		{
+			Message = message;
+		}
+	}
+
+	/// <summary>
+	/// Parameters required to update a notice
+	/// </summary>
+	public class UpdateNoticeRequest
+	{
+		/// <summary>
+		/// The id of the notice to update
+		/// </summary>
+		public string Id { get; set; }
+
+		/// <summary>
+		/// Start time to display this message
+		/// </summary>
+		public DateTime? StartTime { get; set; }
+
+		/// <summary>
+		/// Finish time to display this message
+		/// </summary>
+		public DateTime? FinishTime { get; set; }
+
+		/// <summary>
+		/// Message to display
+		/// </summary>
+		public string? Message { get; set; }
+
+		/// <summary>
+		/// Constructor
+		/// </summary>
+		public UpdateNoticeRequest(string id)
+		{
+			Id = id;
+		}
+	}
+
+	/// <summary>
+	/// Notice information
+	/// </summary>
+	public class GetNoticeResponse
+	{
+		/// <summary>
+		/// The id of the notice to update
+		/// </summary>
+		public string? Id { get; set; }
+
+		/// <summary>
+		/// Start time to display this message
+		/// </summary>
+		public DateTime? StartTime { get; set; }
+
+		/// <summary>
+		/// Finish time to display this message
+		/// </summary>
+		public DateTime? FinishTime { get; set; }
+
+		/// <summary>
+		/// Whether this notice is for scheduled downtime
+		/// </summary>
+		public bool ScheduledDowntime { get; set; } = false;
+
+		/// <summary>
+		/// Whether the notice is currently active
+		/// </summary>
+		public bool Active { get; set; }
+
+		/// <summary>
+		/// Message to display
+		/// </summary>
+		public string? Message { get; set; }
+
+		/// <summary>
+		/// User id who created the notice, otherwise null if a system message
+		/// </summary>
+		public GetThinUserInfoResponse? CreatedByUser { get; set; }
+	}
+}
\ No newline at end of file
