Horde: Fix static analysis warnings.

#rnx

diff --git a/Engine/Source/Programs/Shared/EpicGames.Horde/Logs/LogEventMessages.cs b/Engine/Source/Programs/Shared/EpicGames.Horde/Logs/LogEventMessages.cs
--- a/Engine/Source/Programs/Shared/EpicGames.Horde/Logs/LogEventMessages.cs
+++ b/Engine/Source/Programs/Shared/EpicGames.Horde/Logs/LogEventMessages.cs
@@ -2,8 +2,8 @@
 
 using System.Collections.Generic;
 using System.Text.Json;
-using EpicGames.Core;
-using EpicGames.Horde.Logs;
+
+#pragma warning disable CA2227
 
 namespace EpicGames.Horde.Logs
 {
diff --git a/Engine/Source/Programs/Shared/EpicGames.Horde/Logs/LogMessages.cs b/Engine/Source/Programs/Shared/EpicGames.Horde/Logs/LogMessages.cs
--- a/Engine/Source/Programs/Shared/EpicGames.Horde/Logs/LogMessages.cs
+++ b/Engine/Source/Programs/Shared/EpicGames.Horde/Logs/LogMessages.cs
@@ -2,7 +2,6 @@
 
 using System;
 using System.Collections.Generic;
-using System.Text.Json;
 using EpicGames.Horde.Agents.Leases;
 using EpicGames.Horde.Agents.Sessions;
 using EpicGames.Horde.Jobs;
