EpicGames.Perforce: Recompute allocated data size whenever needed rather than trying to track as we go along (and missing a spot).

diff --git a/Engine/Source/Programs/Shared/EpicGames.Perforce/NativePerforceConnection.cs b/Engine/Source/Programs/Shared/EpicGames.Perforce/NativePerforceConnection.cs
--- a/Engine/Source/Programs/Shared/EpicGames.Perforce/NativePerforceConnection.cs
+++ b/Engine/Source/Programs/Shared/EpicGames.Perforce/NativePerforceConnection.cs
@@ -273,7 +273,6 @@ namespace EpicGames.Perforce
 		readonly ManualResetEvent _responseCompleteEvent;
 		readonly Stopwatch _stallTimer = new Stopwatch();
 		readonly HangMonitor _hangMonitor;
-		long _allocatedSize;
 		bool _disposed;
 
 		/// <inheritdoc/>
@@ -312,7 +311,6 @@ namespace EpicGames.Perforce
 			{
 				_buffers[idx] = new PinnedBuffer(bufferSize);
 				_writeBuffers.TryAdd(_buffers[idx]);
-				_allocatedSize += bufferSize;
 			}
 
 			_onBufferReadyInst = new OnBufferReadyFn(OnBufferReady);
@@ -375,7 +373,6 @@ namespace EpicGames.Perforce
 			if (buffer.MaxLength < minSize)
 			{
 				buffer.Resize(minSize);
-				_allocatedSize += minSize - buffer.MaxLength;
 			}
 
 			nativeWriteBuffer._data = buffer.BasePtr;
@@ -601,7 +598,7 @@ namespace EpicGames.Perforce
 			Logger.LogTrace("Conn {ConnectionId}: {Command} {Args}", _uniqueId, command, argList.ToString());
 
 			using IDisposable scope = _hangMonitor.Start($"{command} {argList}");
-			long initialAllocatedSize = _allocatedSize;
+			long initialAllocatedSize = GetAllocatedSize();
 
 			List<IntPtr> nativeArgs = new List<IntPtr>();
 			try
@@ -633,14 +630,26 @@ namespace EpicGames.Perforce
 			}
 
 			const long WarnSize = 16 * 1024 * 1024;
-			if (_allocatedSize > initialAllocatedSize && _allocatedSize > WarnSize)
+
+			long allocatedSize = GetAllocatedSize();
+			if (allocatedSize > initialAllocatedSize && allocatedSize > WarnSize)
 			{
-				Logger.LogTrace("Native P4 connection allocated {AllocatedSize} ({Command} {Args})", _allocatedSize, command, argList.ToString());
+				Logger.LogTrace("Native P4 connection allocated {AllocatedSize} ({Command} {Args})", allocatedSize, command, argList.ToString());
 			}
 
 			Logger.LogTrace("Conn {ConnectionId}: Request completed in {Time}ms", _uniqueId, timer.ElapsedMilliseconds);
 		}
 
+		long GetAllocatedSize()
+		{
+			long allocatedSize = 0;
+			foreach (PinnedBuffer buffer in _buffers)
+			{
+				allocatedSize += buffer.MaxLength;
+			}
+			return allocatedSize;
+		}
+
 		/// <summary>
 		/// Converts a Python marshalled data blob to a spec definition
 		/// </summary>
