EpicGames.Perforce.Native: Fix buffer being expanded whenever an empty flush packet is sent at the end of each command.

diff --git a/Engine/Source/Programs/Shared/EpicGames.Perforce.Native/Client.cpp b/Engine/Source/Programs/Shared/EpicGames.Perforce.Native/Client.cpp
--- a/Engine/Source/Programs/Shared/EpicGames.Perforce.Native/Client.cpp
+++ b/Engine/Source/Programs/Shared/EpicGames.Perforce.Native/Client.cpp
@@ -131,6 +131,15 @@ public:
 	}
 
 	void Flush()
+	{
+		// Writing an empty buffer indicates to the owning process that we need more data, so avoid that during a normal flush.
+		if (Length > 0)
+		{
+			FlushOrExpandBuffer();
+		}
+	}
+
+	void FlushOrExpandBuffer()
 	{
 		FReadBuffer ReadBuffer;
 		ReadBuffer.Data = Data;
@@ -151,7 +160,7 @@ public:
 	{
 		while (!TryOutputError(err))
 		{
-			Flush();
+			FlushOrExpandBuffer();
 		}
 	}
 
@@ -159,7 +168,7 @@ public:
 	{
 		while (!TryOutputError(Err))
 		{
-			Flush();
+			FlushOrExpandBuffer();
 		}
 	}
 
@@ -225,7 +234,7 @@ public:
 	{
 		while (!TryOutputIo(FileId, Command, Payload, PayloadLen))
 		{
-			Flush();
+			FlushOrExpandBuffer();
 		}
 	}
 
@@ -262,7 +271,7 @@ public:
 	{
 		while (!TryOutputInfo(InLevel, InData))
 		{
-			Flush();
+			FlushOrExpandBuffer();
 		}
 	}
 
@@ -307,7 +316,7 @@ public:
 	{
 		while (!TryWriteRecord(VarList))
 		{
-			Flush();
+			FlushOrExpandBuffer();
 		}
 		if (++Count > MaxCount)
 		{
