Horde: Add guard bytes to buffers allocated for native perforce connection to help track down segfaults on linux.

#rnx

diff --git a/Engine/Source/Programs/Shared/EpicGames.Perforce/NativePerforceConnection.cs b/Engine/Source/Programs/Shared/EpicGames.Perforce/NativePerforceConnection.cs
--- a/Engine/Source/Programs/Shared/EpicGames.Perforce/NativePerforceConnection.cs
+++ b/Engine/Source/Programs/Shared/EpicGames.Perforce/NativePerforceConnection.cs
@@ -84,28 +84,51 @@ namespace EpicGames.Perforce
 		/// </summary>
 		class PinnedBuffer : IDisposable
 		{
-			public byte[] Data { get; private set; }
+			static byte[] s_guardBytes = Enumerable.Repeat<byte>(0xfd, 16).ToArray();
+
+			public byte[] AllocatedData { get; private set; }
+			public Memory<byte> Data { get; private set; }
 			public GCHandle Handle { get; private set; }
 			public IntPtr BasePtr { get; private set; }
 			public int MaxLength => Data.Length;
 
 			public PinnedBuffer(int maxLength)
 			{
-				Data = new byte[maxLength];
-				Handle = GCHandle.Alloc(Data, GCHandleType.Pinned);
-				BasePtr = Handle.AddrOfPinnedObject();
+				AllocatedData = null!;
+				AllocateBuffer(maxLength);
+			}
+
+			void AllocateBuffer(int maxLength)
+			{
+				AllocatedData = new byte[maxLength + (s_guardBytes.Length * 2)];
+				Data = AllocatedData.AsMemory(s_guardBytes.Length, maxLength);
+				Handle = GCHandle.Alloc(AllocatedData, GCHandleType.Pinned);
+				BasePtr = Handle.AddrOfPinnedObject() + s_guardBytes.Length;
+
+				s_guardBytes.AsSpan().CopyTo(AllocatedData);
+				s_guardBytes.AsSpan().CopyTo(AllocatedData.AsSpan(AllocatedData.Length - s_guardBytes.Length));
+			}
+
+			public void CheckGuardBytes()
+			{
+				if (!AllocatedData.AsSpan(0, s_guardBytes.Length).SequenceEqual(s_guardBytes))
+				{
+					throw new InvalidOperationException("Guard bytes corrupted at start of memory region");
+				}
+				if (!AllocatedData.AsSpan(AllocatedData.Length - s_guardBytes.Length).SequenceEqual(s_guardBytes))
+				{
+					throw new InvalidOperationException("Guard bytes corrupted at end of memory region");
+				}
 			}
 
 			public void Resize(int maxLength)
 			{
-				byte[] oldData = Data;
+				byte[] oldData = AllocatedData;
 				Handle.Free();
 
-				Data = new byte[maxLength];
-				Handle = GCHandle.Alloc(Data, GCHandleType.Pinned);
-				BasePtr = Handle.AddrOfPinnedObject();
+				AllocateBuffer(maxLength);
 
-				oldData.CopyTo(Data, 0);
+				oldData.CopyTo(AllocatedData, 0);
 			}
 
 			public void Dispose()
@@ -131,7 +154,7 @@ namespace EpicGames.Perforce
 
 			public Channel<(PinnedBuffer Buffer, int Length)> _readBuffers = Channel.CreateUnbounded<(PinnedBuffer, int)>();
 
-			public ReadOnlyMemory<byte> Data => (_buffer == null) ? ReadOnlyMemory<byte>.Empty : _buffer.Data.AsMemory(_bufferPos, _bufferLen - _bufferPos);
+			public ReadOnlyMemory<byte> Data => (_buffer == null) ? ReadOnlyMemory<byte>.Empty : _buffer.Data.Slice(_bufferPos, _bufferLen - _bufferPos);
 
 			public Response(NativePerforceConnection outer)
 			{
@@ -211,7 +234,7 @@ namespace EpicGames.Perforce
 				{
 					if (_bufferPos > 0)
 					{
-						_buffer.Data.AsSpan(_bufferPos, _bufferLen - _bufferPos).CopyTo(_buffer.Data);
+						_buffer.Data.Slice(_bufferPos, _bufferLen - _bufferPos).CopyTo(_buffer.Data);
 						_bufferLen -= _bufferPos;
 						_bufferPos = 0;
 					}
@@ -234,7 +257,7 @@ namespace EpicGames.Perforce
 
 				// Try to copy some data from the next buffer
 				int copyLen = Math.Min(_nextBufferLen - _nextBufferPos, Math.Min(maxAppend, 16384));
-				_nextBuffer.Data.AsSpan(_nextBufferPos, copyLen).CopyTo(_buffer.Data.AsSpan(_bufferLen));
+				_nextBuffer.Data.Slice(_nextBufferPos, copyLen).CopyTo(_buffer.Data.Slice(_bufferLen));
 				_bufferLen += copyLen;
 				_nextBufferPos += copyLen;
 
@@ -287,7 +310,7 @@ namespace EpicGames.Perforce
 		/// <param name="settings">Settings for the connection</param>
 		/// <param name="logger">Logger for messages</param>
 		public NativePerforceConnection(IPerforceSettings settings, ILogger logger)
-			: this(settings, 2, 64 * 1024, logger)
+			: this(settings, 2, 128 * 1024, logger)
 		{
 		}
 
@@ -550,6 +573,8 @@ namespace EpicGames.Perforce
 		void OnBufferReady(NativeReadBuffer readBuffer, [In, Out] NativeWriteBuffer writeBuffer)
 		{
 			PinnedBuffer buffer = _buffers.First(x => x.BasePtr == readBuffer._data);
+			buffer.CheckGuardBytes();
+
 			_currentResponse!._readBuffers.Writer.TryWrite((buffer, readBuffer._length)); // Unbounded; will always succeed
 
 			int nextWriteSize = 0;
@@ -619,7 +644,15 @@ namespace EpicGames.Perforce
 					Encoding.UTF8.GetBytes(promptResponse, promptResponseBytes);
 				}
 
-				Client_Command(_client, command, nativeArgs.Count, nativeArgs.ToArray(), inputData, inputData?.Length ?? 0, promptResponseBytes, interceptIo);
+				try
+				{
+					Client_Command(_client, command, nativeArgs.Count, nativeArgs.ToArray(), inputData, inputData?.Length ?? 0, promptResponseBytes, interceptIo);
+				}
+				catch (Exception ex)
+				{
+					Logger.LogError(ex, "Exception while executing command {Command} {Args}: {Message}", command, argList.ToString(), ex.Message);
+					throw;
+				}
 			}
 			finally
 			{
@@ -634,7 +667,7 @@ namespace EpicGames.Perforce
 			long allocatedSize = GetAllocatedSize();
 			if (allocatedSize > initialAllocatedSize && allocatedSize > WarnSize)
 			{
-				Logger.LogTrace("Native P4 connection allocated {AllocatedSize} ({Command} {Args})", allocatedSize, command, argList.ToString());
+				Logger.LogTrace("Native P4 connection allocation increased {InitialAllocatedSize} -> {AllocatedSize} ({Command} {Args})", initialAllocatedSize, allocatedSize, command, argList.ToString());
 			}
 
 			Logger.LogTrace("Conn {ConnectionId}: Request completed in {Time}ms", _uniqueId, timer.ElapsedMilliseconds);
