Horde: Include response content in log when server returns an unauthorized response.

#rnx

diff --git a/Engine/Source/Programs/Shared/EpicGames.Horde/Compute/Clients/ServerComputeClient.cs b/Engine/Source/Programs/Shared/EpicGames.Horde/Compute/Clients/ServerComputeClient.cs
--- a/Engine/Source/Programs/Shared/EpicGames.Horde/Compute/Clients/ServerComputeClient.cs
+++ b/Engine/Source/Programs/Shared/EpicGames.Horde/Compute/Clients/ServerComputeClient.cs
@@ -280,7 +280,16 @@ namespace EpicGames.Horde.Compute.Clients
 
 				if (httpResponse.StatusCode == HttpStatusCode.Unauthorized)
 				{
-					throw new ComputeClientException($"Bad authentication credentials. Check or refresh token. (HTTP status {httpResponse.StatusCode})");
+					string? content;
+					try
+					{
+						content = await httpResponse.Content.ReadAsStringAsync(cancellationToken);
+					}
+					catch
+					{
+						content = "None";
+					}
+					throw new ComputeClientException($"Bad authentication credentials. Check or refresh token. (HTTP status {httpResponse.StatusCode}, response: {content})");
 				}
 
 				if (httpResponse.StatusCode == HttpStatusCode.Forbidden)
