EpicGames.Perforce: Fix static analysis warning.

#rnx

diff --git a/Engine/Source/Programs/Shared/EpicGames.Perforce/NativePerforceConnection.cs b/Engine/Source/Programs/Shared/EpicGames.Perforce/NativePerforceConnection.cs
--- a/Engine/Source/Programs/Shared/EpicGames.Perforce/NativePerforceConnection.cs
+++ b/Engine/Source/Programs/Shared/EpicGames.Perforce/NativePerforceConnection.cs
@@ -84,7 +84,7 @@ namespace EpicGames.Perforce
 		/// </summary>
 		class PinnedBuffer : IDisposable
 		{
-			static byte[] s_guardBytes = Enumerable.Repeat<byte>(0xfd, 16).ToArray();
+			static readonly byte[] s_guardBytes = Enumerable.Repeat<byte>(0xfd, 16).ToArray();
 
 			public byte[] AllocatedData { get; private set; }
 			public Memory<byte> Data { get; private set; }
