EpicGames.Perforce: Disable hang monitor output when PerforceSettings.EnableHangMonitor is set to false.

#rnx

diff --git a/Engine/Source/Programs/Shared/EpicGames.Perforce/NativePerforceConnection.cs b/Engine/Source/Programs/Shared/EpicGames.Perforce/NativePerforceConnection.cs
--- a/Engine/Source/Programs/Shared/EpicGames.Perforce/NativePerforceConnection.cs
+++ b/Engine/Source/Programs/Shared/EpicGames.Perforce/NativePerforceConnection.cs
@@ -13,6 +13,7 @@ using System.Threading.Channels;
 using System.Threading.Tasks;
 using EpicGames.Core;
 using Microsoft.Extensions.Logging;
+using Microsoft.Extensions.Logging.Abstractions;
 
 namespace EpicGames.Perforce
 {
@@ -327,7 +328,9 @@ namespace EpicGames.Perforce
 			Logger = logger;
 
 			_uniqueId = Interlocked.Increment(ref s_nextUniqueId);
-			_hangMonitor = new HangMonitor(TimeSpan.FromMinutes(1.0), $"Perforce connection ({_uniqueId})", logger);
+
+			ILogger hangLogger = settings.EnableHangMonitor ? logger : NullLogger.Instance;
+			_hangMonitor = new HangMonitor(TimeSpan.FromMinutes(1.0), $"Perforce connection ({_uniqueId})", hangLogger);
 
 			_buffers = new PinnedBuffer[bufferCount];
 			for (int idx = 0; idx < bufferCount; idx++)
