Modified UHT to generate an error when an IInterface is improperly used.

diff --git a/Engine/Source/Programs/Shared/EpicGames.UHT/Types/Properties/UhtInterfaceProperty.cs b/Engine/Source/Programs/Shared/EpicGames.UHT/Types/Properties/UhtInterfaceProperty.cs
--- a/Engine/Source/Programs/Shared/EpicGames.UHT/Types/Properties/UhtInterfaceProperty.cs
+++ b/Engine/Source/Programs/Shared/EpicGames.UHT/Types/Properties/UhtInterfaceProperty.cs
@@ -148,13 +148,12 @@ namespace EpicGames.UHT.Types
 		}
 
 		/// <inheritdoc/>
-		protected override void ValidateMember(UhtStruct structObj, UhtValidationOptions options)
+		public override void Validate(UhtStruct outerStruct, UhtProperty outermostProperty, UhtValidationOptions options)
 		{
-			base.ValidateMember(structObj, options);
-
+			base.Validate(outerStruct, outermostProperty, options);
 			if (PointerType == UhtPointerType.Native)
 			{
-				this.LogError($"UPROPERTY pointers cannot be interfaces - did you mean TScriptInterface<{InterfaceClass.SourceName}>?");
+				this.LogError($"Property and function argument pointers cannot be interfaces - did you mean TScriptInterface<{InterfaceClass.SourceName}>?");
 			}
 		}
 
