EpicGames.Horde add Job and Log query, fix serialization issues

diff --git a/Engine/Source/Programs/Shared/EpicGames.Horde/HordeHttpClient.cs b/Engine/Source/Programs/Shared/EpicGames.Horde/HordeHttpClient.cs
--- a/Engine/Source/Programs/Shared/EpicGames.Horde/HordeHttpClient.cs
+++ b/Engine/Source/Programs/Shared/EpicGames.Horde/HordeHttpClient.cs
@@ -13,6 +13,8 @@ using EpicGames.Core;
 using EpicGames.Horde.Agents.Telemetry;
 using EpicGames.Horde.Artifacts;
 using EpicGames.Horde.Dashboard;
+using EpicGames.Horde.Jobs;
+using EpicGames.Horde.Logs;
 using EpicGames.Horde.Projects;
 using EpicGames.Horde.Secrets;
 using EpicGames.Horde.Server;
@@ -108,6 +110,7 @@ namespace EpicGames.Horde
 			options.Converters.Add(new JsonStringEnumConverter());
 			options.Converters.Add(new StringIdJsonConverterFactory());
 			options.Converters.Add(new BinaryIdJsonConverterFactory());
+			options.Converters.Add(new SubResourceIdJsonConverterFactory());
 		}
 
 		#region Artifacts
@@ -431,6 +434,32 @@ namespace EpicGames.Horde
 
 		#endregion
 
+		#region Jobs
+		/// <summary>
+		/// Gets job information for given job ID. Fail response if jobID does not exist.
+		/// </summary>
+		/// <param name="id">Id of the job to get infomation for</param>
+		/// <param name="cancellationToken">Cancellation token for the operation</param>
+		/// <returns></returns>
+		public Task<GetJobResponse> GetJobAsync(JobId id, CancellationToken cancellationToken = default)
+		{
+			return GetAsync<GetJobResponse>(_httpClient, $"api/v1/jobs/{id}", cancellationToken);
+		}
+		#endregion
+
+		#region Log
+		/// <summary>
+		/// Get the given log file 
+		/// </summary>
+		/// <param name="logFileId">Id of the log file to retrieve</param>
+		/// <param name="cancellationToken">Cancellation token for the operation</param>
+		/// <returns></returns>
+		public Task<SearchLogResponse> GetSearchLogAsync(LogId logFileId, CancellationToken cancellationToken = default)
+		{
+			return GetAsync<SearchLogResponse>(_httpClient, $"/api/v1/logs/{logFileId}/search", cancellationToken);
+		}
+		#endregion
+
 		#region Utility Methods
 
 		/// <summary>
diff --git a/Engine/Source/Programs/Shared/EpicGames.Horde/Jobs/JobMessages.cs b/Engine/Source/Programs/Shared/EpicGames.Horde/Jobs/JobMessages.cs
--- a/Engine/Source/Programs/Shared/EpicGames.Horde/Jobs/JobMessages.cs
+++ b/Engine/Source/Programs/Shared/EpicGames.Horde/Jobs/JobMessages.cs
@@ -505,6 +505,15 @@ namespace EpicGames.Horde.Jobs
 			TemplateId = templateId;
 			Name = name;
 		}
+
+		/// <summary>
+		/// Default constructor needed for JsonSerializer
+		/// </summary>
+		[JsonConstructor]
+		public GetJobResponse()
+		{
+			Name = "";
+		}
 	}
 
 	/// <summary>
@@ -1009,6 +1018,15 @@ namespace EpicGames.Horde.Jobs
 			State = state;
 			Outcome = outcome;
 		}
+
+		/// <summary>
+		/// Default Constructor needed for JsonSerializer
+		/// </summary>
+		[JsonConstructor]
+		public GetLabelStateResponse()
+		{
+
+		}
 	}
 
 	/// <summary>
@@ -1032,6 +1050,15 @@ namespace EpicGames.Horde.Jobs
 		{
 			Nodes = nodes;
 		}
+
+		/// <summary>
+		/// Default Constructor needed for JsonSerializer
+		/// </summary>
+		[JsonConstructor]
+		public GetDefaultLabelStateResponse() : base()
+		{
+			Nodes = new List<string>();
+		}
 	}
 
 	/// <summary>
