Adding UE_DEPRECATED_FORGAME and UE_DEPRECATED_FORENGINE to known optional tokens parsable by UHT.  E.g. usable for enum entries.

diff --git a/Engine/Source/Programs/Shared/EpicGames.UHT/Tokenizer/UhtTokenReaderGeneralExtensions.cs b/Engine/Source/Programs/Shared/EpicGames.UHT/Tokenizer/UhtTokenReaderGeneralExtensions.cs
--- a/Engine/Source/Programs/Shared/EpicGames.UHT/Tokenizer/UhtTokenReaderGeneralExtensions.cs
+++ b/Engine/Source/Programs/Shared/EpicGames.UHT/Tokenizer/UhtTokenReaderGeneralExtensions.cs
@@ -224,7 +224,7 @@ namespace EpicGames.UHT.Tokenizer
 			for (; ; )
 			{
 				UhtToken token = tokenReader.PeekToken();
-				if (token.IsIdentifier("DEPRECATED") || token.IsIdentifier("UE_DEPRECATED"))
+				if (token.IsIdentifier("DEPRECATED") || token.IsIdentifier("UE_DEPRECATED") || token.IsIdentifier("UE_DEPRECATED_FORGAME") || token.IsIdentifier("UE_DEPRECATED_FORENGINE"))
 				{
 					tokenReader.ConsumeToken();
 					tokenReader.Require('(', "deprecation macro");
