EpicGames.Core: Use SafeProcessHandle instead of SafeFileHandle for child process in ManagedProcess

#rnx

diff --git a/Engine/Source/Programs/Shared/EpicGames.Core/ManagedProcess.cs b/Engine/Source/Programs/Shared/EpicGames.Core/ManagedProcess.cs
--- a/Engine/Source/Programs/Shared/EpicGames.Core/ManagedProcess.cs
+++ b/Engine/Source/Programs/Shared/EpicGames.Core/ManagedProcess.cs
@@ -335,7 +335,7 @@ namespace EpicGames.Core
 		);
 
 		[DllImport("kernel32.dll", SetLastError = true)]
-		static extern int GetExitCodeProcess(SafeFileHandle hProcess, out int lpExitCode);
+		static extern int GetExitCodeProcess(SafeProcessHandle hProcess, out int lpExitCode);
 
 		[DllImport("kernel32.dll", SetLastError = true)]
 		static extern int TerminateProcess(SafeHandleZeroOrMinusOneIsInvalid hProcess, uint uExitCode);
@@ -405,7 +405,7 @@ namespace EpicGames.Core
 		/// <summary>
 		/// Handle for the child process.
 		/// </summary>
-		SafeFileHandle? _processHandle;
+		SafeProcessHandle? _processHandle;
 
 		/// <summary>
 		/// The write end of the child process' stdin pipe.
@@ -837,7 +837,7 @@ namespace EpicGames.Core
 					}
 
 					// Wrap the process handle in a SafeFileHandle
-					_processHandle = new SafeFileHandle(processInfo.hProcess, true);
+					_processHandle = new SafeProcessHandle(processInfo.hProcess, true);
 				}
 				finally
 				{
@@ -1311,7 +1311,7 @@ namespace EpicGames.Core
 
 		class ProcessWaitHandle : WaitHandle
 		{
-			public ProcessWaitHandle(SafeFileHandle processHandle)
+			public ProcessWaitHandle(SafeProcessHandle processHandle)
 			{
 				SafeWaitHandle = new SafeWaitHandle(processHandle.DangerousGetHandle(), false);
 			}
