Speculative fix for crashes in native Perforce library. Also change build script to create output artifacts rather than submitting files directly.

See notes on these pages:

https://www.perforce.com/manuals/v23.1/p4api/Content/P4API/p4libraries.html
https://www.perforce.com/manuals/v23.1/p4api/Content/P4API/chapter.clientprogramming.signaler.html

#rnx

diff --git a/Engine/Source/Programs/Shared/EpicGames.Perforce.Native/BuildLibrary.xml b/Engine/Source/Programs/Shared/EpicGames.Perforce.Native/BuildLibrary.xml
--- a/Engine/Source/Programs/Shared/EpicGames.Perforce.Native/BuildLibrary.xml
+++ b/Engine/Source/Programs/Shared/EpicGames.Perforce.Native/BuildLibrary.xml
@@ -10,6 +10,7 @@
 			<Sign Files="#FilesToSign" Tag="#BinariesWin64"/>
 			<Log Message="Files:" Files="#BinariesWin64"/>
 		</Node>
+		<Artifact Name="Binaries (Win64)" Tag="#BinariesWin64"/>
 	</Agent>
 
 	<Agent Name="DotNetPerforceLib Mac" Type="Mac">
@@ -20,6 +21,7 @@
 			<Compile Target="DotNetPerforceLib" Configuration="Shipping" Platform="Mac" Tag="#BinariesMac"/>
 			<Log Message="Files:" Files="#BinariesMac"/>
 		</Node>
+		<Artifact Name="Binaries (Mac)" Tag="#BinariesMac"/>
 	</Agent>
 
 	<Agent Name="DotNetPerforceLib Linux" Type="Linux">
@@ -30,15 +32,9 @@
 			<Compile Target="DotNetPerforceLib" Configuration="Shipping" Platform="Linux" Tag="#BinariesLinux"/>
 			<Log Message="Files:" Files="#BinariesLinux"/>
 		</Node>
-	</Agent>
-	
-	<Agent Name="DotNetPerforceLib Submit" Type="Win64">
-		<Node Name="Submit DotNetPerforceLib Binaries" Requires="#BinariesWin64;#BinariesMac;#BinariesLinux">
-			<Tag Files="#BinariesWin64;#BinariesMac;#BinariesLinux" Filter="Engine/Binaries/DotNET/EpicGames.Perforce.Native/..." With="#FilesToSubmit"/>
-			<Submit Files="#FilesToSubmit" Description="Rebuilt DotNetPerforceLib binaries for CL $(Change)\n\n#okforgithub public\n#jira none\n#submittool bypass"/>
-		</Node>
+		<Artifact Name="Binaries (Linux)" Tag="#BinariesLinux"/>
 	</Agent>
 
-	<Aggregate Name="DotNetPerforceLib" Requires="Submit DotNetPerforceLib Binaries"/>
+	<Aggregate Name="DotNetPerforceLib" Requires="#BinariesWin64;#BinariesMac;#BinariesLinux"/>
 
 </BuildGraph>
diff --git a/Engine/Source/Programs/Shared/EpicGames.Perforce.Native/Client.cpp b/Engine/Source/Programs/Shared/EpicGames.Perforce.Native/Client.cpp
--- a/Engine/Source/Programs/Shared/EpicGames.Perforce.Native/Client.cpp
+++ b/Engine/Source/Programs/Shared/EpicGames.Perforce.Native/Client.cpp
@@ -13,8 +13,11 @@ THIRD_PARTY_INCLUDES_START
 #include "p4/i18napi.h"
 #include "p4/charset.h"
 #include "p4/md5.h"
+#include "p4/p4libs.h"
+#include "p4/signaler.h"
 #include "ThirdParty/gzip.h"
 #pragma warning(pop)
+#include <signal.h>
 #include <assert.h>
 THIRD_PARTY_INCLUDES_END
 
@@ -709,8 +712,21 @@ public:
 	}
 };
 
+struct FClientGlobalInit
+{
+	FClientGlobalInit()
+	{
+		Error e;
+		P4Libraries::Initialize(P4LIBRARIES_INIT_ALL, &e);
+		signal(SIGINT, SIG_DFL); // unset the default set by global signaler in C++ so it does not exit 
+		signaler.Disable(); // disable the global signaler memory tracking at runtime
+	}
+};
+
 extern "C" NATIVE_API FClient* Client_Create(const FSettings* Settings, FWriteBuffer* WriteBuffer, FOnBufferReadyFn* OnBufferReady)
 {
+	static FClientGlobalInit GlobalInit;
+
 	FClient* Client = new FClient(WriteBuffer, OnBufferReady);
 
 	if (Settings != nullptr)
