Include a type field for p4 artifacts.

#rnx

diff --git a/Engine/Source/Programs/Shared/EpicGames.Perforce.Native/BuildLibrary.xml b/Engine/Source/Programs/Shared/EpicGames.Perforce.Native/BuildLibrary.xml
--- a/Engine/Source/Programs/Shared/EpicGames.Perforce.Native/BuildLibrary.xml
+++ b/Engine/Source/Programs/Shared/EpicGames.Perforce.Native/BuildLibrary.xml
@@ -10,7 +10,7 @@
 			<Sign Files="#FilesToSign" Tag="#BinariesWin64"/>
 			<Log Message="Files:" Files="#BinariesWin64"/>
 		</Node>
-		<Artifact Name="Binaries (Win64)" Tag="#BinariesWin64"/>
+		<Artifact Name="Binaries (Win64)" Type="p4-binaries" Tag="#BinariesWin64"/>
 	</Agent>
 
 	<Agent Name="DotNetPerforceLib Mac" Type="Mac">
@@ -21,7 +21,7 @@
 			<Compile Target="DotNetPerforceLib" Configuration="Shipping" Platform="Mac" Tag="#BinariesMac"/>
 			<Log Message="Files:" Files="#BinariesMac"/>
 		</Node>
-		<Artifact Name="Binaries (Mac)" Tag="#BinariesMac"/>
+		<Artifact Name="Binaries (Mac)" Type="p4-binaries" Tag="#BinariesMac"/>
 	</Agent>
 
 	<Agent Name="DotNetPerforceLib Linux" Type="Linux">
@@ -32,7 +32,7 @@
 			<Compile Target="DotNetPerforceLib" Configuration="Shipping" Platform="Linux" Tag="#BinariesLinux"/>
 			<Log Message="Files:" Files="#BinariesLinux"/>
 		</Node>
-		<Artifact Name="Binaries (Linux)" Tag="#BinariesLinux"/>
+		<Artifact Name="Binaries (Linux)" Type="p4-binaries" Tag="#BinariesLinux"/>
 	</Agent>
 
 	<Aggregate Name="DotNetPerforceLib" Requires="#BinariesWin64;#BinariesMac;#BinariesLinux"/>
