EpicGames.Perforce: Include MSVC runtime dependencies to support EpicGames.Perforce.Native DLLs.

diff --git a/Engine/Source/Programs/Shared/EpicGames.Perforce/EpicGames.Perforce.csproj b/Engine/Source/Programs/Shared/EpicGames.Perforce/EpicGames.Perforce.csproj
--- a/Engine/Source/Programs/Shared/EpicGames.Perforce/EpicGames.Perforce.csproj
+++ b/Engine/Source/Programs/Shared/EpicGames.Perforce/EpicGames.Perforce.csproj
@@ -59,6 +59,15 @@
     <None Include="..\..\..\..\Binaries\DotNET\EpicGames.Perforce.Native\win-x64\$(WindowsNativeConfig)\EpicGames.Perforce.Native.pdb" Visible="false" Condition="Exists('..\..\..\..\Binaries\DotNET\EpicGames.Perforce.Native\win-x64\$(WindowsNativeConfig)\EpicGames.Perforce.Native.pdb')">
       <CopyToOutputDirectory>PreserveNewest</CopyToOutputDirectory>
     </None>
+    <None Include="..\..\..\..\Binaries\ThirdParty\AppLocalDependencies\Win64\Microsoft.VC.CRT\MSVCP140.dll" Visible="false" Condition="Exists('..\..\..\..\Binaries\ThirdParty\AppLocalDependencies\Win64\Microsoft.VC.CRT\MSVCP140.dll')">
+      <CopyToOutputDirectory>PreserveNewest</CopyToOutputDirectory>
+    </None>
+    <None Include="..\..\..\..\Binaries\ThirdParty\AppLocalDependencies\Win64\Microsoft.VC.CRT\VCRUNTIME140.dll" Visible="false" Condition="Exists('..\..\..\..\Binaries\ThirdParty\AppLocalDependencies\Win64\Microsoft.VC.CRT\VCRUNTIME140.dll')">
+      <CopyToOutputDirectory>PreserveNewest</CopyToOutputDirectory>
+    </None>
+    <None Include="..\..\..\..\Binaries\ThirdParty\AppLocalDependencies\Win64\Microsoft.VC.CRT\VCRUNTIME140_1.dll" Visible="false" Condition="Exists('..\..\..\..\Binaries\ThirdParty\AppLocalDependencies\Win64\Microsoft.VC.CRT\VCRUNTIME140_1.dll')">
+      <CopyToOutputDirectory>PreserveNewest</CopyToOutputDirectory>
+    </None>
     <None Include="..\..\..\..\Binaries\DotNET\EpicGames.Perforce.Native\linux-x64\Release\EpicGames.Perforce.Native.so" Visible="false">
       <CopyToOutputDirectory>PreserveNewest</CopyToOutputDirectory>
     </None>
