[EpicGames.UBA]
* Updated pinvoke code to use new function names in api

diff --git a/Engine/Source/Programs/Shared/EpicGames.UBA/Impl/ProcessImpl.cs b/Engine/Source/Programs/Shared/EpicGames.UBA/Impl/ProcessImpl.cs
--- a/Engine/Source/Programs/Shared/EpicGames.UBA/Impl/ProcessImpl.cs
+++ b/Engine/Source/Programs/Shared/EpicGames.UBA/Impl/ProcessImpl.cs
@@ -16,10 +16,10 @@ namespace EpicGames.UBA
 
 		#region DllImport
 		[DllImport("UbaHost", CharSet = CharSet.Auto)]
-		static extern IntPtr CreateProcessStartInfo(string application, string arguments, string workingDir, string description, uint priorityClass, ulong outputStatsThresholdMs, bool trackInputs, string logFile, ExitCallback? exit);
+		static extern IntPtr ProcessStartInfo_Create(string application, string arguments, string workingDir, string description, uint priorityClass, ulong outputStatsThresholdMs, bool trackInputs, string logFile, ExitCallback? exit);
 
 		[DllImport("UbaHost", CharSet = CharSet.Auto)]
-		static extern void DestroyProcessStartInfo(IntPtr server);
+		static extern void ProcessStartInfo_Destroy(IntPtr server);
 		#endregion
 
 		public ProcessStartInfoImpl(ProcessStartInfo info, bool needCallback)
@@ -29,7 +29,7 @@ namespace EpicGames.UBA
 				ExitCallbackDelegate = RaiseExited;
 			}
 
-			_handle = CreateProcessStartInfo(info.Application, info.Arguments, info.WorkingDirectory, info.Description, (uint)info.Priority, info.OutputStatsThresholdMs, info.TrackInputs, info.LogFile?? String.Empty, ExitCallbackDelegate);
+			_handle = ProcessStartInfo_Create(info.Application, info.Arguments, info.WorkingDirectory, info.Description, (uint)info.Priority, info.OutputStatsThresholdMs, info.TrackInputs, info.LogFile?? String.Empty, ExitCallbackDelegate);
 		}
 
 		void RaiseExited(IntPtr userData, IntPtr handle)
@@ -59,7 +59,7 @@ namespace EpicGames.UBA
 
 			if (_handle != IntPtr.Zero)
 			{
-				DestroyProcessStartInfo(_handle);
+				ProcessStartInfo_Destroy(_handle);
 				_handle = IntPtr.Zero;
 			}
 		}
diff --git a/Engine/Source/Programs/Shared/EpicGames.UBA/Impl/ServerImpl.cs b/Engine/Source/Programs/Shared/EpicGames.UBA/Impl/ServerImpl.cs
--- a/Engine/Source/Programs/Shared/EpicGames.UBA/Impl/ServerImpl.cs
+++ b/Engine/Source/Programs/Shared/EpicGames.UBA/Impl/ServerImpl.cs
@@ -12,23 +12,23 @@ namespace EpicGames.UBA
 
 		#region DllImport
 		[DllImport("UbaHost", CharSet = CharSet.Auto)]
-		static extern IntPtr CreateServer(IntPtr logger, int workerCount, int sendSize, int receiveTimeoutSeconds, bool useQuic);
+		static extern IntPtr NetworkServer_Create(IntPtr logger, int workerCount, int sendSize, int receiveTimeoutSeconds, bool useQuic);
 
 		[DllImport("UbaHost", CharSet = CharSet.Auto)]
-		static extern void DestroyServer(IntPtr server);
+		static extern void NetworkServer_Destroy(IntPtr server);
 
 		// Server Imports
 		[DllImport("UbaHost", CharSet = CharSet.Auto)]
-		static extern bool Server_StartListen(IntPtr server, int port, string ip, string crypto);
+		static extern bool NetworkServer_StartListen(IntPtr server, int port, string ip, string crypto);
 
 		[DllImport("UbaHost", CharSet = CharSet.Auto)]
-		static extern bool Server_AddClient(IntPtr server, string ip, int port, string crypto);
+		static extern bool NetworkServer_AddClient(IntPtr server, string ip, int port, string crypto);
 
 		[DllImport("UbaHost", CharSet = CharSet.Auto)]
-		static extern bool Server_AddNamedConnection(IntPtr server, string name);
+		static extern bool NetworkServer_AddNamedConnection(IntPtr server, string name);
 
 		[DllImport("UbaHost", CharSet = CharSet.Auto)]
-		static extern bool Server_Stop(IntPtr server);
+		static extern bool NetworkServer_Stop(IntPtr server);
 
 		const int DEFAULT_PORT = 1345;
 		#endregion
@@ -36,7 +36,7 @@ namespace EpicGames.UBA
 		public ServerImpl(int maxWorkers, int sendSize, ILogger logger, bool useQuic)
 		{
 			_logger = logger;
-			_handle = CreateServer(_logger.GetHandle(), maxWorkers, sendSize, 60, useQuic); // We know this code is using 5s pings.. so 60 seconds timeout on recv socket is good
+			_handle = NetworkServer_Create(_logger.GetHandle(), maxWorkers, sendSize, 60, useQuic); // We know this code is using 5s pings.. so 60 seconds timeout on recv socket is good
 		}
 
 		#region IDisposable
@@ -57,7 +57,7 @@ namespace EpicGames.UBA
 
 			if (_handle != IntPtr.Zero)
 			{
-				DestroyServer(_handle);
+				NetworkServer_Destroy(_handle);
 				_handle = IntPtr.Zero;
 			}
 		}
@@ -68,16 +68,16 @@ namespace EpicGames.UBA
 		public IntPtr GetHandle() => _handle;
 
 		/// <inheritdoc/>
-		public bool StartServer(string ip = "", int port = -1, string crypto = "") => Server_StartListen(_handle, port > 0 ? port : DEFAULT_PORT, ip, crypto);
+		public bool StartServer(string ip = "", int port = -1, string crypto = "") => NetworkServer_StartListen(_handle, port > 0 ? port : DEFAULT_PORT, ip, crypto);
 
 		/// <inheritdoc/>
-		public void StopServer() => Server_Stop(_handle);
+		public void StopServer() => NetworkServer_Stop(_handle);
 
 		/// <inheritdoc/>
-		public bool AddClient(string ip, int port, string crypto = "") => Server_AddClient(_handle, ip, port, crypto);
+		public bool AddClient(string ip, int port, string crypto = "") => NetworkServer_AddClient(_handle, ip, port, crypto);
 
 		/// <inheritdoc/>
-		public bool AddNamedConnection(string name) => Server_AddNamedConnection(_handle, name);
+		public bool AddNamedConnection(string name) => NetworkServer_AddNamedConnection(_handle, name);
 		#endregion
 	}
 }
diff --git a/Engine/Source/Programs/Shared/EpicGames.UBA/Impl/SessionServerImpl.cs b/Engine/Source/Programs/Shared/EpicGames.UBA/Impl/SessionServerImpl.cs
--- a/Engine/Source/Programs/Shared/EpicGames.UBA/Impl/SessionServerImpl.cs
+++ b/Engine/Source/Programs/Shared/EpicGames.UBA/Impl/SessionServerImpl.cs
@@ -16,11 +16,11 @@ namespace EpicGames.UBA
 
 		#region DllImport
 		[DllImport("UbaHost", CharSet = CharSet.Auto)]
-		static extern IntPtr CreateSessionServerCreateInfo(IntPtr storage, IntPtr client, IntPtr logger, string rootDir, string traceOutputFile,
+		static extern IntPtr SessionServerCreateInfo_Create(IntPtr storage, IntPtr client, IntPtr logger, string rootDir, string traceOutputFile,
 			bool disableCustomAllocator, bool launchVisualizer, bool resetCas, bool writeToDisk, bool detailedTrace, bool allowWaitOnMem, bool allowKillOnMem, bool storeObjFilesCompressed);
 
 		[DllImport("UbaHost", CharSet = CharSet.Auto)]
-		static extern void DestroySessionServerCreateInfo(IntPtr server);
+		static extern void SessionServerCreateInfo_Destroy(IntPtr server);
 		#endregion
 
 		public SessionServerCreateInfoImpl(IStorageServer storage, IServer client, ILogger logger, SessionServerCreateInfo info)
@@ -28,7 +28,7 @@ namespace EpicGames.UBA
 			_storage = storage;
 			_client = client;
 			_logger = logger;
-			_handle = CreateSessionServerCreateInfo(_storage.GetHandle(), _client.GetHandle(), _logger.GetHandle(), info.RootDirectory, info.TraceOutputFile, info.DisableCustomAllocator, info.LaunchVisualizer, info.ResetCas, info.WriteToDisk, info.DetailedTrace, info.AllowWaitOnMem, info.AllowKillOnMem, info.StoreObjFilesCompressed);
+			_handle = SessionServerCreateInfo_Create(_storage.GetHandle(), _client.GetHandle(), _logger.GetHandle(), info.RootDirectory, info.TraceOutputFile, info.DisableCustomAllocator, info.LaunchVisualizer, info.ResetCas, info.WriteToDisk, info.DetailedTrace, info.AllowWaitOnMem, info.AllowKillOnMem, info.StoreObjFilesCompressed);
 		}
 
 		#region IDisposable
@@ -48,7 +48,7 @@ namespace EpicGames.UBA
 
 			if (_handle != IntPtr.Zero)
 			{
-				DestroySessionServerCreateInfo(_handle);
+				SessionServerCreateInfo_Destroy(_handle);
 				_handle = IntPtr.Zero;
 			}
 		}
@@ -72,7 +72,7 @@ namespace EpicGames.UBA
 
 		#region DllImport
 		[DllImport("UbaHost", CharSet = CharSet.Auto)]
-		static extern IntPtr CreateSessionServer(IntPtr info);
+		static extern IntPtr SessionServer_Create(IntPtr info);
 
 		[DllImport("UbaHost", CharSet = CharSet.Auto)]
 		static extern void SessionServer_SetRemoteProcessAvailable(IntPtr server, RemoteProcessSlotAvailableCallback func);
@@ -114,7 +114,7 @@ namespace EpicGames.UBA
 		static extern void SessionServer_CancelAll(IntPtr server);
 
 		[DllImport("UbaHost", CharSet = CharSet.Auto)]
-		static extern void DestroySessionServer(IntPtr server);
+		static extern void SessionServer_Destroy(IntPtr server);
 		#endregion
 
 		public SessionServerImpl(ISessionServerCreateInfo info)
@@ -122,7 +122,7 @@ namespace EpicGames.UBA
 			_info = info;
 			_remoteProcessSlotAvailableCallbackDelegate = RaiseRemoteProcessSlotAvailable;
 			_remoteProcessReturnedCallbackDelegate = RaiseRemoteProcessReturned;
-			_handle = CreateSessionServer(_info.GetHandle());
+			_handle = SessionServer_Create(_info.GetHandle());
 			SessionServer_SetRemoteProcessAvailable(_handle, _remoteProcessSlotAvailableCallbackDelegate);
 			SessionServer_SetRemoteProcessReturned(_handle, _remoteProcessReturnedCallbackDelegate);
 		}
@@ -144,7 +144,7 @@ namespace EpicGames.UBA
 
 			if (_handle != IntPtr.Zero)
 			{
-				DestroySessionServer(_handle);
+				SessionServer_Destroy(_handle);
 				_handle = IntPtr.Zero;
 			}
 		}
diff --git a/Engine/Source/Programs/Shared/EpicGames.UBA/Impl/StorageServerImpl.cs b/Engine/Source/Programs/Shared/EpicGames.UBA/Impl/StorageServerImpl.cs
--- a/Engine/Source/Programs/Shared/EpicGames.UBA/Impl/StorageServerImpl.cs
+++ b/Engine/Source/Programs/Shared/EpicGames.UBA/Impl/StorageServerImpl.cs
@@ -14,10 +14,10 @@ namespace EpicGames.UBA
 
 		#region DllImport
 		[DllImport("UbaHost", CharSet = CharSet.Auto)]
-		static extern IntPtr CreateStorageServer(IntPtr server, string rootDir, ulong casCapacityBytes, bool storeCompressed, IntPtr logger, string zone);
+		static extern IntPtr StorageServer_Create(IntPtr server, string rootDir, ulong casCapacityBytes, bool storeCompressed, IntPtr logger, string zone);
 
 		[DllImport("UbaHost", CharSet = CharSet.Auto)]
-		static extern void DestroyStorageServer(IntPtr server);
+		static extern void StorageServer_Destroy(IntPtr server);
 
 		[DllImport("UbaHost", CharSet = CharSet.Auto)]
 		static extern void StorageServer_SaveCasTable(IntPtr server);
@@ -30,7 +30,7 @@ namespace EpicGames.UBA
 		{
 			_server = server;
 			_logger = logger;
-			_handle = CreateStorageServer(_server.GetHandle(), info.RootDirectory, info.CapacityBytes, info.StoreCompressed, _logger.GetHandle(), info.Zone);
+			_handle = StorageServer_Create(_server.GetHandle(), info.RootDirectory, info.CapacityBytes, info.StoreCompressed, _logger.GetHandle(), info.Zone);
 			Utils.DisallowedPaths().ToList().ForEach(x => RegisterDisallowedPath(x));
 		}
 
@@ -51,7 +51,7 @@ namespace EpicGames.UBA
 
 			if (_handle != IntPtr.Zero)
 			{
-				DestroyStorageServer(_handle);
+				StorageServer_Destroy(_handle);
 				_handle = IntPtr.Zero;
 			}
 		}
