[UBT] Fix  data race for platform directories access
#tests local UBT

diff --git a/Engine/Source/Programs/Shared/EpicGames.Build/Unreal.cs b/Engine/Source/Programs/Shared/EpicGames.Build/Unreal.cs
--- a/Engine/Source/Programs/Shared/EpicGames.Build/Unreal.cs
+++ b/Engine/Source/Programs/Shared/EpicGames.Build/Unreal.cs
@@ -242,42 +242,42 @@ namespace UnrealBuildBase
 
 		// A subset of the functionality in DataDrivenPlatformInfo.GetAllPlatformInfos() - finds the DataDrivenPlatformInfo.ini files and records their existence, but does not parse them
 		// (perhaps DataDrivenPlatformInfo.GetAllPlatformInfos() could be modified to use this data to avoid an additional search through the filesystem)
-		public static HashSet<string>? IniPresentForPlatform = null;
-		private static bool DataDrivenPlatformInfoIniIsPresent(string PlatformName)
+		private static Lazy<HashSet<string>> IniPresentForPlatform = new Lazy<HashSet<string>>(() =>
 		{
-			if (IniPresentForPlatform == null)
+			HashSet<string> Set = new HashSet<string>(StringComparer.OrdinalIgnoreCase);
+			// find all platform directories (skipping NFL/NoRedist)
+			foreach (DirectoryReference EngineConfigDir in GetExtensionDirs(Unreal.EngineDirectory, "Config", bIncludeRestrictedDirectories: false))
 			{
-				IniPresentForPlatform = new HashSet<string>(StringComparer.OrdinalIgnoreCase);
-
-				// find all platform directories (skipping NFL/NoRedist)
-				foreach (DirectoryReference EngineConfigDir in GetExtensionDirs(Unreal.EngineDirectory, "Config", bIncludeRestrictedDirectories: false))
+				// look through all config dirs looking for the data driven ini file
+				foreach (string FilePath in Directory.EnumerateFiles(EngineConfigDir.FullName, "DataDrivenPlatformInfo.ini", SearchOption.AllDirectories))
 				{
-					// look through all config dirs looking for the data driven ini file
-					foreach (string FilePath in Directory.EnumerateFiles(EngineConfigDir.FullName, "DataDrivenPlatformInfo.ini", SearchOption.AllDirectories))
-					{
-						FileReference FileRef = new FileReference(FilePath);
+					FileReference FileRef = new FileReference(FilePath);
 
-						// get the platform name from the path
-						string IniPlatformName;
-						if (FileRef.IsUnderDirectory(DirectoryReference.Combine(Unreal.EngineDirectory, "Config")))
-						{
-							// Foo/Engine/Config/<Platform>/DataDrivenPlatformInfo.ini
-							IniPlatformName = Path.GetFileName(Path.GetDirectoryName(FilePath))!;
-						}
-						else
-						{
-							// Foo/Engine/Platforms/<Platform>/Config/DataDrivenPlatformInfo.ini
-							IniPlatformName = Path.GetFileName(Path.GetDirectoryName(Path.GetDirectoryName(FilePath)))!;
-						}
+					// get the platform name from the path
+					string IniPlatformName;
+					if (FileRef.IsUnderDirectory(DirectoryReference.Combine(Unreal.EngineDirectory, "Config")))
+					{
+						// Foo/Engine/Config/<Platform>/DataDrivenPlatformInfo.ini
+						IniPlatformName = Path.GetFileName(Path.GetDirectoryName(FilePath))!;
+					}
+					else
+					{
+						// Foo/Engine/Platforms/<Platform>/Config/DataDrivenPlatformInfo.ini
+						IniPlatformName = Path.GetFileName(Path.GetDirectoryName(Path.GetDirectoryName(FilePath)))!;
+					}
 
-						// DataDrivenPlatformInfo.GetAllPlatformInfos() checks that [DataDrivenPlatformInfo] section exists as part of validating that the file exists
-						// This code should probably behave the same way.
+					// DataDrivenPlatformInfo.GetAllPlatformInfos() checks that [DataDrivenPlatformInfo] section exists as part of validating that the file exists
+					// This code should probably behave the same way.
 
-						IniPresentForPlatform.Add(IniPlatformName);
-					}
+					Set.Add(IniPlatformName);
 				}
 			}
-			return IniPresentForPlatform.Contains(PlatformName);
+
+			return Set;
+		});
+		private static bool DataDrivenPlatformInfoIniIsPresent(string PlatformName)
+		{
+			return IniPresentForPlatform.Value.Contains(PlatformName);
 		}
 
 		// cached dictionary of BaseDir to extension directories
