Horde: Handle zero byte reads due to socket shutdown when reading messages into a compute buffer.

diff --git a/Engine/Source/Programs/Shared/EpicGames.Horde/Compute/ComputeSocket.cs b/Engine/Source/Programs/Shared/EpicGames.Horde/Compute/ComputeSocket.cs
--- a/Engine/Source/Programs/Shared/EpicGames.Horde/Compute/ComputeSocket.cs
+++ b/Engine/Source/Programs/Shared/EpicGames.Horde/Compute/ComputeSocket.cs
@@ -579,6 +579,12 @@ namespace EpicGames.Horde.Compute
 						for (int offset = 0; offset < size;)
 						{
 							int read = await transport.RecvAsync(memory.Slice(offset, size - offset), cancellationToken);
+							if (read == 0)
+							{
+								// Return true to avoid logic in ReadPacketAsync() from trying to read the whole message from a closed stream
+								_logger.LogDebug("Unexpected end of stream while parsing message for channel {Id}; discarding message.", id);
+								return true;
+							}
 							offset += read;
 						}
 
