EpicGames.Perforce: Fix parsing of records with multiple arrays. All properties for the entire tree of possible records for a response are looked up by tag before identifying the target object via the array index, rather then recursively trying to parse complete records. It is possible for one record to have multiple child records, and also possible that tags for child records are serialized out of order.

One regression with this approach is that the required option is not validated for subelement types; since this is more of a validation of assumptions of the data format rather than an expected error case, it should be fine.

diff --git a/Engine/Source/Programs/Shared/EpicGames.Perforce/IPerforceOutput.cs b/Engine/Source/Programs/Shared/EpicGames.Perforce/IPerforceOutput.cs
--- a/Engine/Source/Programs/Shared/EpicGames.Perforce/IPerforceOutput.cs
+++ b/Engine/Source/Programs/Shared/EpicGames.Perforce/IPerforceOutput.cs
@@ -2,7 +2,6 @@
 
 using System;
 using System.Buffers.Binary;
-using System.Buffers.Text;
 using System.Collections.Generic;
 using System.Diagnostics.CodeAnalysis;
 using System.IO;
@@ -499,7 +498,7 @@ namespace EpicGames.Perforce
 			object? record;
 			if (code == ReadOnlyUtf8StringConstants.Stat && statRecordInfo != null)
 			{
-				if (!TryReadTypedRecord(buffer, ref bufferPos, Utf8String.Empty, statRecordInfo, out record))
+				if (!TryReadTypedRecord(buffer, ref bufferPos, statRecordInfo, out record))
 				{
 					response = null;
 					return false;
@@ -507,7 +506,7 @@ namespace EpicGames.Perforce
 			}
 			else if (code == ReadOnlyUtf8StringConstants.Info)
 			{
-				if (!TryReadTypedRecord(buffer, ref bufferPos, Utf8String.Empty, PerforceReflection.InfoRecordInfo, out record))
+				if (!TryReadTypedRecord(buffer, ref bufferPos, PerforceReflection.InfoRecordInfo, out record))
 				{
 					response = null;
 					return false;
@@ -515,7 +514,7 @@ namespace EpicGames.Perforce
 			}
 			else if (code == ReadOnlyUtf8StringConstants.Error)
 			{
-				if (!TryReadTypedRecord(buffer, ref bufferPos, Utf8String.Empty, PerforceReflection.ErrorRecordInfo, out record))
+				if (!TryReadTypedRecord(buffer, ref bufferPos, PerforceReflection.ErrorRecordInfo, out record))
 				{
 					response = null;
 					return false;
@@ -523,7 +522,7 @@ namespace EpicGames.Perforce
 			}
 			else if (code == ReadOnlyUtf8StringConstants.Io)
 			{
-				if (!TryReadTypedRecord(buffer, ref bufferPos, Utf8String.Empty, PerforceReflection.IoRecordInfo, out record))
+				if (!TryReadTypedRecord(buffer, ref bufferPos, PerforceReflection.IoRecordInfo, out record))
 				{
 					response = null;
 					return false;
@@ -551,11 +550,10 @@ namespace EpicGames.Perforce
 		/// </summary>
 		/// <param name="buffer">The buffer to read from</param>
 		/// <param name="bufferPos">Current read position within the buffer</param>
-		/// <param name="requiredSuffix">The required suffix for any subobject arrays.</param>
 		/// <param name="recordInfo">Reflection information for the type being serialized into.</param>
 		/// <param name="record">Receives the record on success</param>
 		/// <returns>The parsed object.</returns>
-		static bool TryReadTypedRecord(ReadOnlyMemory<byte> buffer, ref int bufferPos, Utf8String requiredSuffix, CachedRecordInfo recordInfo, [NotNullWhen(true)] out object? record)
+		static bool TryReadTypedRecord(ReadOnlyMemory<byte> buffer, ref int bufferPos, CachedRecordInfo recordInfo, [NotNullWhen(true)] out object? record)
 		{
 			// Create a bitmask for all the required tags
 			ulong requiredTagsBitMask = 0;
@@ -613,104 +611,81 @@ namespace EpicGames.Perforce
 				tag = tag.Slice(0, suffixIdx);
 
 				// Try to find the matching field
-				CachedTagInfo? tagInfo;
+				TaggedPropertyInfo? tagInfo;
 				if (recordInfo.NameToInfo.TryGetValue(tag, out tagInfo))
 				{
 					requiredTagsBitMask |= tagInfo.RequiredTagBitMask;
 				}
 
-				// Check whether it's a subobject or part of the current object.
-				if (suffix == requiredSuffix)
+				// Find the target object for this tag
+				object? targetRecord = null;
+				if (tagInfo != null)
 				{
-					if (!TryReadValue(buffer, ref bufferPos, newRecord, tagInfo))
-					{
-						record = null;
-						return false;
-					}
+					targetRecord = GetNestedRecord(newRecord, suffix, tagInfo.ParentRecords);
 				}
-				else if (suffix.StartsWith(requiredSuffix) && (requiredSuffix.Length == 0 || suffix[requiredSuffix.Length] == ','))
+
+				// Read the value
+				if (!TryReadValue(buffer, ref bufferPos, targetRecord, tagInfo))
 				{
-					// Part of a subobject. If this record doesn't have any listed subobject type, skip the field and continue.
-					if (tagInfo != null)
-					{
-						// Get the list field
-						System.Collections.IList? list = (System.Collections.IList?)tagInfo.PropertyInfo.GetValue(newRecord);
-						if (list == null)
-						{
-							throw new PerforceException($"Empty list for {tagInfo.PropertyInfo.Name}");
-						}
+					record = null;
+					return false;
+				}
+			}
 
-						// Check the suffix matches the index of the next element
-						ReadIndex(tag, suffix, requiredSuffix, list.Count);
+			// Make sure we've got all the required tags we need
+			if (requiredTagsBitMask != recordInfo.RequiredTagsBitMask)
+			{
+				string missingTagNames = String.Join(", ", recordInfo.NameToInfo.Where(x => (requiredTagsBitMask | x.Value.RequiredTagBitMask) != requiredTagsBitMask).Select(x => x.Key));
+				throw new PerforceException("Missing '{0}' tag when parsing '{1}'", missingTagNames, recordInfo.Type.Name);
+			}
 
-						// Add it to the list
-						if (!TryReadValue(buffer, ref bufferPos, newRecord, tagInfo))
-						{
-							record = null;
-							return false;
-						}
-					}
-					else if (recordInfo.SubElementProperty != null)
-					{
-						// Move back to the start of this tag
-						bufferPos = startBufferPos;
+			// Construct the response object with the record
+			record = newRecord;
+			return true;
+		}
 
-						// Get the list field
-						System.Collections.IList? list = (System.Collections.IList?)recordInfo.SubElementProperty.GetValue(newRecord);
-						if (list == null)
-						{
-							throw new PerforceException($"Invalid field for {recordInfo.SubElementProperty.Name}");
-						}
+		static object? GetNestedRecord(object? targetObject, Utf8String suffix, NestedRecordInfo[] nestedRecords)
+		{
+			object? nestedObject = targetObject;
+			foreach (NestedRecordInfo nestedRecord in nestedRecords)
+			{
+				if (suffix.Length == 0)
+				{
+					return null;
+				}
 
-						// Find the next index
-						int subElementSuffixLength = ReadIndex(tag, suffix, requiredSuffix, list.Count);
+				// Parse the next index
+				int index = 0;
+				while (suffix.Length > 0)
+				{
+					byte character = suffix[0];
+					suffix = suffix[1..];
 
-						// Parse the subobject and add it to the list
-						object? subRecord;
-						if (!TryReadTypedRecord(buffer, ref bufferPos, suffix.Substring(0, subElementSuffixLength), recordInfo.SubElementRecordInfo!, out subRecord))
-						{
-							record = null;
-							return false;
-						}
-						list.Add(subRecord);
+					if (character >= '0' && character <= '9')
+					{
+						index = (index * 10) + (character - '0');
 					}
 					else
 					{
-						// Just discard the value
-						if (!TryReadValue(buffer, ref bufferPos, newRecord, tagInfo))
-						{
-							record = null;
-							return false;
-						}
+						break;
 					}
 				}
-				else if (tag == ReadOnlyUtf8StringConstants.Func || tag == ReadOnlyUtf8StringConstants.IsSparse)
+
+				// Get the appropriate object
+				System.Collections.IList? list = (System.Collections.IList?)nestedRecord.PropertyInfo.GetValue(nestedObject);
+				if (list == null)
 				{
-					// Not sure why these fields are in the client output, but they are peppered into filelog results without an element index breaking the parser.
-					if (!TryReadValue(buffer, ref bufferPos, newRecord, null))
-					{
-						record = null;
-						return false;
-					}
+					throw new Exception();
 				}
-				else
+				while (index >= list.Count)
 				{
-					// Roll back
-					bufferPos = startBufferPos;
-					break;
+					list.Add(nestedRecord.CreateInstance());
 				}
-			}
 
-			// Make sure we've got all the required tags we need
-			if (requiredTagsBitMask != recordInfo.RequiredTagsBitMask)
-			{
-				string missingTagNames = String.Join(", ", recordInfo.NameToInfo.Where(x => (requiredTagsBitMask | x.Value.RequiredTagBitMask) != requiredTagsBitMask).Select(x => x.Key));
-				throw new PerforceException("Missing '{0}' tag when parsing '{1}'", missingTagNames, recordInfo.Type.Name);
+				// Move to the nested object
+				nestedObject = list[index];
 			}
-
-			// Construct the response object with the record
-			record = newRecord;
-			return true;
+			return nestedObject;
 		}
 
 		/// <summary>
@@ -721,7 +696,7 @@ namespace EpicGames.Perforce
 		/// <param name="newRecord">The new record</param>
 		/// <param name="tagInfo">The current tag</param>
 		/// <returns></returns>
-		static bool TryReadValue(ReadOnlyMemory<byte> buffer, ref int bufferPos, object newRecord, CachedTagInfo? tagInfo)
+		static bool TryReadValue(ReadOnlyMemory<byte> buffer, ref int bufferPos, object? newRecord, TaggedPropertyInfo? tagInfo)
 		{
 			ReadOnlySpan<byte> bufferSpan = buffer.Span;
 
@@ -740,7 +715,10 @@ namespace EpicGames.Perforce
 				{
 					return false;
 				}
-				tagInfo?.ReadFromString(newRecord, stringValue);
+				if (newRecord != null)
+				{
+					tagInfo?.ReadFromString(newRecord, stringValue);
+				}
 			}
 			else if (valueType == 'i')
 			{
@@ -749,7 +727,10 @@ namespace EpicGames.Perforce
 				{
 					return false;
 				}
-				tagInfo?.ReadFromInteger(newRecord, integerValue);
+				if (newRecord != null)
+				{
+					tagInfo?.ReadFromInteger(newRecord, integerValue);
+				}
 			}
 			else
 			{
@@ -804,56 +785,26 @@ namespace EpicGames.Perforce
 		/// </summary>
 		/// <param name="buffer">The buffer to read from</param>
 		/// <param name="bufferPos">Current read position within the buffer</param>
-		/// <param name="string">Receives the value that was read</param>
+		/// <param name="result">Receives the value that was read</param>
 		/// <returns>True if a string was read from the buffer, false if there was not enough data</returns>
-		static bool TryReadString(ReadOnlyMemory<byte> buffer, ref int bufferPos, out Utf8String @string)
+		static bool TryReadString(ReadOnlyMemory<byte> buffer, ref int bufferPos, out Utf8String result)
 		{
 			int length;
 			if (!TryReadInt(buffer.Span, ref bufferPos, out length))
 			{
-				@string = Utf8String.Empty;
+				result = Utf8String.Empty;
 				return false;
 			}
 
 			if (bufferPos + length > buffer.Length)
 			{
-				@string = Utf8String.Empty;
+				result = Utf8String.Empty;
 				return false;
 			}
 
-			@string = new Utf8String(buffer.Slice(bufferPos, length));
+			result = new Utf8String(buffer.Slice(bufferPos, length));
 			bufferPos += length;
 			return true;
 		}
-
-		/// <summary>
-		/// Checks that the given field suffix starts with the given required suffix and index.
-		/// </summary>
-		/// <param name="tag">The tag name</param>
-		/// <param name="suffix">The text to check</param>
-		/// <param name="requiredSuffix">The required prefix</param>
-		/// <param name="index">The required index</param>
-		/// <returns>True if the index is correct</returns>
-		static int ReadIndex(Utf8String tag, Utf8String suffix, Utf8String requiredSuffix, int index)
-		{
-			int valueStart = 0;
-			if (requiredSuffix.Length > 0)
-			{
-				valueStart = requiredSuffix.Length + 1;
-			}
-
-			int value;
-			int valueLength;
-			if (Utf8Parser.TryParse(suffix.Substring(valueStart), out value, out valueLength))
-			{
-				if (value == index && (valueStart + valueLength == suffix.Length || suffix[valueStart + valueLength] == (byte)','))
-				{
-					return valueStart + valueLength;
-				}
-			}
-
-			string expectedSuffix = (requiredSuffix.Length > 0) ? $"{requiredSuffix},{index}" : $"{index}";
-			throw new PerforceException("Subobject element received out of order: got {0}{1}, expected {0}{2}", tag, suffix, expectedSuffix);
-		}
 	}
 }
diff --git a/Engine/Source/Programs/Shared/EpicGames.Perforce/PerforceReflection.cs b/Engine/Source/Programs/Shared/EpicGames.Perforce/PerforceReflection.cs
--- a/Engine/Source/Programs/Shared/EpicGames.Perforce/PerforceReflection.cs
+++ b/Engine/Source/Programs/Shared/EpicGames.Perforce/PerforceReflection.cs
@@ -26,21 +26,37 @@ namespace EpicGames.Perforce
 		public static readonly Utf8String Default = new Utf8String("default");
 	}
 
+	/// <summary>
+	/// Delegate type for creating a record instance
+	/// </summary>
+	/// <returns>New instance</returns>
+	public delegate object CreateRecordDelegate();
+
+	/// <summary>
+	/// Information about a nested record
+	/// </summary>
+	record struct NestedRecordInfo(PropertyInfo PropertyInfo, CreateRecordDelegate CreateInstance);
+
 	/// <summary>
 	/// Stores cached information about a property with a <see cref="PerforceTagAttribute"/> attribute.
 	/// </summary>
-	class CachedTagInfo
+	class TaggedPropertyInfo
 	{
 		/// <summary>
 		/// Name of the tag. Specified in the attribute or inferred from the field name.
 		/// </summary>
-		public Utf8String Name { get; }
+		public Utf8String Tag { get; }
 
 		/// <summary>
 		/// Whether this tag is optional or not.
 		/// </summary>
 		public bool Optional { get; }
 
+		/// <summary>
+		/// List of nested records before getting to the target object
+		/// </summary>
+		public NestedRecordInfo[] ParentRecords { get; }
+
 		/// <summary>
 		/// The property containing the value of this data.
 		/// </summary>
@@ -71,18 +87,35 @@ namespace EpicGames.Perforce
 		/// </summary>
 		/// <param name="name"></param>
 		/// <param name="optional"></param>
+		/// <param name="parentRecords"></param>
 		/// <param name="propertyInfo"></param>
 		/// <param name="requiredTagBitMask"></param>
-		public CachedTagInfo(Utf8String name, bool optional, PropertyInfo propertyInfo, ulong requiredTagBitMask)
+		public TaggedPropertyInfo(Utf8String name, bool optional, NestedRecordInfo[] parentRecords, PropertyInfo propertyInfo, ulong requiredTagBitMask)
 		{
-			Name = name;
+			Tag = name;
 			Optional = optional;
+			ParentRecords = parentRecords;
 			PropertyInfo = propertyInfo;
 			RequiredTagBitMask = requiredTagBitMask;
 			Write = (obj, writer) => throw new PerforceException($"Field {name} does not have a serializer.");
 			ReadFromInteger = (obj, value) => throw new PerforceException($"Field {name} was not expecting an integer value.");
 			ReadFromString = (obj, str) => throw new PerforceException($"Field {name} was not expecting a string value.");
 		}
+
+		/// <inheritdoc/>
+		public override string ToString()
+		{
+			StringBuilder fullName = new StringBuilder(Tag.ToString());
+			if (ParentRecords.Length > 0)
+			{
+				fullName.Append('0');
+				for (int idx = 1; idx < ParentRecords.Length; idx++)
+				{
+					fullName.Append(",0");
+				}
+			}
+			return fullName.ToString();
+		}
 	}
 
 	/// <summary>
@@ -90,12 +123,6 @@ namespace EpicGames.Perforce
 	/// </summary>
 	class CachedRecordInfo
 	{
-		/// <summary>
-		/// Delegate type for creating a record instance
-		/// </summary>
-		/// <returns>New instance</returns>
-		public delegate object CreateRecordDelegate();
-
 		/// <summary>
 		/// Type of the record
 		/// </summary>
@@ -109,52 +136,27 @@ namespace EpicGames.Perforce
 		/// <summary>
 		/// List of fields in the record. These should be ordered to match P4 output for maximum efficiency.
 		/// </summary>
-		public List<CachedTagInfo> Properties { get; } = new List<CachedTagInfo>();
+		public List<TaggedPropertyInfo> Properties { get; } = new List<TaggedPropertyInfo>();
 
 		/// <summary>
 		/// Map of name to tag info
 		/// </summary>
-		public Dictionary<Utf8String, CachedTagInfo> NameToInfo { get; set; } = new Dictionary<Utf8String, CachedTagInfo>();
+		public Dictionary<Utf8String, TaggedPropertyInfo> NameToInfo { get; set; } = new Dictionary<Utf8String, TaggedPropertyInfo>();
 
 		/// <summary>
 		/// Bitmask of all the required tags. Formed by bitwise-or'ing the RequiredTagBitMask fields for each required CachedTagInfo.
 		/// </summary>
 		public ulong RequiredTagsBitMask { get; set; }
 
-		/// <summary>
-		/// The type of records to create for subelements
-		/// </summary>
-		public Type? SubElementType { get; set; }
-
-		/// <summary>
-		/// The cached record info for the subelement type
-		/// </summary>
-		public CachedRecordInfo? SubElementRecordInfo { get; set; }
-
-		/// <summary>
-		/// Property containing subelements
-		/// </summary>
-		public PropertyInfo? SubElementProperty { get; set; }
-
 		/// <summary>
 		/// Constructor
 		/// </summary>
 		/// <param name="type">The record type</param>
-		public CachedRecordInfo(Type type)
+		/// <param name="createInstance">Method used to create an instance of this record</param>
+		public CachedRecordInfo(Type type, CreateRecordDelegate createInstance)
 		{
 			Type = type;
-
-			ConstructorInfo? constructor = type.GetConstructor(BindingFlags.Public | BindingFlags.NonPublic | BindingFlags.Instance, null, Type.EmptyTypes, null);
-			if (constructor == null)
-			{
-				throw new PerforceException($"Unable to find default constructor for {type}");
-			}
-
-			DynamicMethod dynamicMethod = new DynamicMethod("_", type, null);
-			ILGenerator generator = dynamicMethod.GetILGenerator();
-			generator.Emit(OpCodes.Newobj, constructor);
-			generator.Emit(OpCodes.Ret);
-			CreateInstance = (CreateRecordDelegate)dynamicMethod.CreateDelegate(typeof(CreateRecordDelegate));
+			CreateInstance = createInstance;
 		}
 	}
 
@@ -387,12 +389,12 @@ namespace EpicGames.Perforce
 		public static void Serialize(object obj, IMemoryWriter writer)
 		{
 			CachedRecordInfo recordInfo = GetCachedRecordInfo(obj.GetType());
-			foreach (CachedTagInfo tagInfo in recordInfo.Properties)
+			foreach (TaggedPropertyInfo tagInfo in recordInfo.Properties)
 			{
 				object? value = tagInfo.PropertyInfo.GetValue(obj);
 				if (value != null)
 				{
-					WriteUtf8StringWithTag(writer, tagInfo.Name);
+					WriteUtf8StringWithTag(writer, tagInfo.Tag);
 					tagInfo.Write(writer, value!);
 				}
 			}
@@ -462,115 +464,135 @@ namespace EpicGames.Perforce
 			CachedRecordInfo? record;
 			if (!s_recordTypeToInfo.TryGetValue(recordType, out record))
 			{
-				record = new CachedRecordInfo(recordType);
+				record = new CachedRecordInfo(recordType, GetCreateRecordDelegate(recordType));
 
-				// Get all the fields for this type
-				PropertyInfo[] properties = recordType.GetProperties(BindingFlags.Public | BindingFlags.NonPublic | BindingFlags.Instance);
+				// Find all the properties in the record
+				AddRecordProperties(recordType, record, Array.Empty<NestedRecordInfo>());
+				record.NameToInfo = record.Properties.ToDictionary(x => x.Tag, x => x);
 
-				// Build the map of all tags for this record
-				foreach (PropertyInfo property in properties)
+				// Try to save the record info, or get the version that's already in the cache
+				if (!s_recordTypeToInfo.TryAdd(recordType, record))
 				{
-					PerforceTagAttribute? tagAttribute = property.GetCustomAttribute<PerforceTagAttribute>();
-					if (tagAttribute != null)
-					{
-						string tagName = tagAttribute.Name ?? property.Name;
-
-						ulong requiredTagBitMask = 0;
-						if (!tagAttribute.Optional)
-						{
-							requiredTagBitMask = record.RequiredTagsBitMask + 1;
-							if (requiredTagBitMask == 0)
-							{
-								throw new PerforceException("Too many required tags in {0}; max is {1}", recordType.Name, sizeof(ulong) * 8);
-							}
-							record.RequiredTagsBitMask |= requiredTagBitMask;
-						}
+					record = s_recordTypeToInfo[recordType];
+				}
+			}
+			return record;
+		}
 
-						CachedTagInfo tagInfo = new CachedTagInfo(new Utf8String(tagName), tagAttribute.Optional, property, requiredTagBitMask);
+		static void AddRecordProperties(Type recordType, CachedRecordInfo rootRecord, NestedRecordInfo[] parentRecords)
+		{
+			// Get all the fields for this type
+			PropertyInfo[] properties = recordType.GetProperties(BindingFlags.Public | BindingFlags.NonPublic | BindingFlags.Instance);
 
-						Type fieldType = property.PropertyType;
+			// Build the map of all tags for this record
+			foreach (PropertyInfo property in properties)
+			{
+				PerforceTagAttribute? tagAttribute = property.GetCustomAttribute<PerforceTagAttribute>();
+				if (tagAttribute != null)
+				{
+					string tagName = tagAttribute.Name ?? property.Name;
 
-						PropertyInfo propertyCopy = property;
-						if (fieldType == typeof(DateTime))
-						{
-							tagInfo.Write = (writer, value) => WriteUtf8StringWithTag(writer, new Utf8String(((long)((DateTime)value - PerforceReflection.UnixEpoch).TotalSeconds).ToString()));
-							tagInfo.ReadFromString = (obj, value) => propertyCopy.SetValue(obj, ParseStringAsDateTime(value));
-						}
-						else if (fieldType == typeof(bool))
-						{
-							tagInfo.Write = (writer, value) => WriteUtf8StringWithTag(writer, ((bool)value) ? StringConstants.True : StringConstants.False);
-							tagInfo.ReadFromString = (obj, value) => propertyCopy.SetValue(obj, ParseStringAsBool(value));
-						}
-						else if (fieldType == typeof(Nullable<bool>))
-						{
-							tagInfo.ReadFromString = (obj, value) => propertyCopy.SetValue(obj, ParseStringAsNullableBool(value));
-						}
-						else if (fieldType == typeof(int))
-						{
-							tagInfo.Write = (writer, value) => WriteIntegerWithTag(writer, (int)value);
-							tagInfo.ReadFromInteger = (obj, value) => propertyCopy.SetValue(obj, value);
-							tagInfo.ReadFromString = (obj, value) => propertyCopy.SetValue(obj, ParseStringAsInt(value));
-						}
-						else if (fieldType == typeof(long))
-						{
-							tagInfo.Write = (writer, value) => WriteUtf8StringWithTag(writer, new Utf8String(((long)value).ToString()));
-							tagInfo.ReadFromString = (obj, value) => propertyCopy.SetValue(obj, ParseStringAsLong(value));
-						}
-						else if (fieldType == typeof(string))
-						{
-							tagInfo.Write = (writer, value) => WriteStringWithTag(writer, (string)value);
-							tagInfo.ReadFromString = (obj, value) => propertyCopy.SetValue(obj, ParseString(value));
-						}
-						else if (fieldType == typeof(Utf8String))
-						{
-							tagInfo.Write = (writer, value) => WriteUtf8StringWithTag(writer, (Utf8String)value);
-							tagInfo.ReadFromString = (obj, str) => propertyCopy.SetValue(obj, str.Clone());
-						}
-						else if (fieldType.IsEnum)
-						{
-							CachedEnumInfo enumInfo = GetCachedEnumInfo(fieldType);
-							tagInfo.Write = (writer, value) => WriteUtf8StringWithTag(writer, enumInfo.GetName((int)value));
-							tagInfo.ReadFromInteger = (obj, value) => propertyCopy.SetValue(obj, enumInfo.ParseInteger(value));
-							tagInfo.ReadFromString = (obj, value) => propertyCopy.SetValue(obj, enumInfo.ParseString(value));
-						}
-						else if (fieldType == typeof(DateTimeOffset?))
-						{
-							tagInfo.ReadFromString = (obj, value) => propertyCopy.SetValue(obj, ParseStringAsNullableDateTimeOffset(value));
-						}
-						else if (fieldType == typeof(List<string>))
-						{
-							tagInfo.ReadFromString = (obj, value) => ((List<string>)propertyCopy.GetValue(obj)!).Add(value.ToString());
-						}
-						else if (fieldType == typeof(ReadOnlyMemory<byte>))
-						{
-							tagInfo.ReadFromString = (obj, value) => propertyCopy.SetValue(obj, value.Memory);
-						}
-						else
+					ulong requiredTagBitMask = 0;
+					if (!tagAttribute.Optional && parentRecords.Length == 0)
+					{
+						requiredTagBitMask = rootRecord.RequiredTagsBitMask + 1;
+						if (requiredTagBitMask == 0)
 						{
-							throw new PerforceException("Unsupported type of {0}.{1} for tag '{2}'", recordType.Name, fieldType.Name, tagName);
+							throw new PerforceException("Too many required tags in {0}; max is {1}", recordType.Name, sizeof(ulong) * 8);
 						}
-
-						record.Properties.Add(tagInfo);
+						rootRecord.RequiredTagsBitMask |= requiredTagBitMask;
 					}
+					
+					TaggedPropertyInfo tagInfo = new TaggedPropertyInfo(new Utf8String(tagName), tagAttribute.Optional, parentRecords, property, requiredTagBitMask);
 
-					record.NameToInfo = record.Properties.ToDictionary(x => x.Name, x => x);
+					Type fieldType = property.PropertyType;
 
-					PerforceRecordListAttribute? subElementAttribute = property.GetCustomAttribute<PerforceRecordListAttribute>();
-					if (subElementAttribute != null)
+					PropertyInfo propertyCopy = property;
+					if (fieldType == typeof(DateTime))
+					{
+						tagInfo.Write = (writer, value) => WriteUtf8StringWithTag(writer, new Utf8String(((long)((DateTime)value - PerforceReflection.UnixEpoch).TotalSeconds).ToString()));
+						tagInfo.ReadFromString = (obj, value) => propertyCopy.SetValue(obj, ParseStringAsDateTime(value));
+					}
+					else if (fieldType == typeof(bool))
+					{
+						tagInfo.Write = (writer, value) => WriteUtf8StringWithTag(writer, ((bool)value) ? StringConstants.True : StringConstants.False);
+						tagInfo.ReadFromString = (obj, value) => propertyCopy.SetValue(obj, ParseStringAsBool(value));
+					}
+					else if (fieldType == typeof(Nullable<bool>))
+					{
+						tagInfo.ReadFromString = (obj, value) => propertyCopy.SetValue(obj, ParseStringAsNullableBool(value));
+					}
+					else if (fieldType == typeof(int))
+					{
+						tagInfo.Write = (writer, value) => WriteIntegerWithTag(writer, (int)value);
+						tagInfo.ReadFromInteger = (obj, value) => propertyCopy.SetValue(obj, value);
+						tagInfo.ReadFromString = (obj, value) => propertyCopy.SetValue(obj, ParseStringAsInt(value));
+					}
+					else if (fieldType == typeof(long))
+					{
+						tagInfo.Write = (writer, value) => WriteUtf8StringWithTag(writer, new Utf8String(((long)value).ToString()));
+						tagInfo.ReadFromString = (obj, value) => propertyCopy.SetValue(obj, ParseStringAsLong(value));
+					}
+					else if (fieldType == typeof(string))
+					{
+						tagInfo.Write = (writer, value) => WriteStringWithTag(writer, (string)value);
+						tagInfo.ReadFromString = (obj, value) => propertyCopy.SetValue(obj, ParseString(value));
+					}
+					else if (fieldType == typeof(Utf8String))
 					{
-						record.SubElementProperty = property;
-						record.SubElementType = property.PropertyType.GenericTypeArguments[0];
-						record.SubElementRecordInfo = GetCachedRecordInfo(record.SubElementType);
+						tagInfo.Write = (writer, value) => WriteUtf8StringWithTag(writer, (Utf8String)value);
+						tagInfo.ReadFromString = (obj, str) => propertyCopy.SetValue(obj, str.Clone());
 					}
+					else if (fieldType.IsEnum)
+					{
+						CachedEnumInfo enumInfo = GetCachedEnumInfo(fieldType);
+						tagInfo.Write = (writer, value) => WriteUtf8StringWithTag(writer, enumInfo.GetName((int)value));
+						tagInfo.ReadFromInteger = (obj, value) => propertyCopy.SetValue(obj, enumInfo.ParseInteger(value));
+						tagInfo.ReadFromString = (obj, value) => propertyCopy.SetValue(obj, enumInfo.ParseString(value));
+					}
+					else if (fieldType == typeof(DateTimeOffset?))
+					{
+						tagInfo.ReadFromString = (obj, value) => propertyCopy.SetValue(obj, ParseStringAsNullableDateTimeOffset(value));
+					}
+					else if (fieldType == typeof(List<string>))
+					{
+						tagInfo.ReadFromString = (obj, value) => ((List<string>)propertyCopy.GetValue(obj)!).Add(value.ToString());
+					}
+					else if (fieldType == typeof(ReadOnlyMemory<byte>))
+					{
+						tagInfo.ReadFromString = (obj, value) => propertyCopy.SetValue(obj, value.Memory);
+					}
+					else
+					{
+						throw new PerforceException("Unsupported type of {0}.{1} for tag '{2}'", recordType.Name, fieldType.Name, tagName);
+					}
+
+					rootRecord.Properties.Add(tagInfo);
 				}
 
-				// Try to save the record info, or get the version that's already in the cache
-				if (!s_recordTypeToInfo.TryAdd(recordType, record))
+				PerforceRecordListAttribute? subElementAttribute = property.GetCustomAttribute<PerforceRecordListAttribute>();
+				if (subElementAttribute != null)
 				{
-					record = s_recordTypeToInfo[recordType];
+					Type newRecordType = property.PropertyType.GenericTypeArguments[0];
+					NestedRecordInfo newParentRecord = new NestedRecordInfo(property, GetCreateRecordDelegate(newRecordType));
+					AddRecordProperties(newRecordType, rootRecord, parentRecords.Append(newParentRecord).ToArray());
 				}
 			}
-			return record;
+		}
+
+		static CreateRecordDelegate GetCreateRecordDelegate(Type type)
+		{
+			ConstructorInfo? constructor = type.GetConstructor(BindingFlags.Public | BindingFlags.NonPublic | BindingFlags.Instance, null, Type.EmptyTypes, null);
+			if (constructor == null)
+			{
+				throw new PerforceException($"Unable to find default constructor for {type}");
+			}
+
+			DynamicMethod dynamicMethod = new DynamicMethod("_", type, null);
+			ILGenerator generator = dynamicMethod.GetILGenerator();
+			generator.Emit(OpCodes.Newobj, constructor);
+			generator.Emit(OpCodes.Ret);
+			return (CreateRecordDelegate)dynamicMethod.CreateDelegate(typeof(CreateRecordDelegate));
 		}
 
 		static object ParseString(Utf8String str)
