Horde: Prevent hypens before a period in a sanitized string id.

diff --git a/Engine/Source/Programs/Shared/EpicGames.Horde/StringId.cs b/Engine/Source/Programs/Shared/EpicGames.Horde/StringId.cs
--- a/Engine/Source/Programs/Shared/EpicGames.Horde/StringId.cs
+++ b/Engine/Source/Programs/Shared/EpicGames.Horde/StringId.cs
@@ -99,11 +99,25 @@ namespace EpicGames.Horde
 				{
 					result.Append((char)('a' + (character - 'A')));
 				}
-				else if (IsValidCharacter(character) && (result.Length > 0 || character != '.'))
+				else if (character == '.')
+				{
+					if (result.Length > 0)
+					{
+						if (result[^1] == '-')
+						{
+							result[^1] = character;
+						}
+						else
+						{
+							result.Append(character);
+						}
+					}
+				}
+				else if (IsValidCharacter(character))
 				{
 					result.Append(character);
 				}
-				else if (result.Length > 0 && result[^1] != '-')
+				else if (result.Length > 0 && result[^1] != '-' && result[^1] != '.')
 				{
 					result.Append('-');
 				}
