Hordehttpclient changes adding GetGraph and log search calls.

diff --git a/Engine/Source/Programs/Shared/EpicGames.Horde/HordeHttpClient.cs b/Engine/Source/Programs/Shared/EpicGames.Horde/HordeHttpClient.cs
--- a/Engine/Source/Programs/Shared/EpicGames.Horde/HordeHttpClient.cs
+++ b/Engine/Source/Programs/Shared/EpicGames.Horde/HordeHttpClient.cs
@@ -14,6 +14,7 @@ using EpicGames.Horde.Agents.Telemetry;
 using EpicGames.Horde.Artifacts;
 using EpicGames.Horde.Dashboard;
 using EpicGames.Horde.Jobs;
+using EpicGames.Horde.Jobs.Graphs;
 using EpicGames.Horde.Logs;
 using EpicGames.Horde.Projects;
 using EpicGames.Horde.Secrets;
@@ -452,12 +453,43 @@ namespace EpicGames.Horde
 		/// Get the given log file 
 		/// </summary>
 		/// <param name="logFileId">Id of the log file to retrieve</param>
+		/// <param name="searchText">Text to search for in the log</param>
+		/// <param name="count">Number of lines to return (default 5)</param>
 		/// <param name="cancellationToken">Cancellation token for the operation</param>
 		/// <returns></returns>
-		public Task<SearchLogResponse> GetSearchLogAsync(LogId logFileId, CancellationToken cancellationToken = default)
+		public Task<SearchLogResponse> GetSearchLogAsync(LogId logFileId,string searchText, int count = 5, CancellationToken cancellationToken = default)
 		{
-			return GetAsync<SearchLogResponse>(_httpClient, $"/api/v1/logs/{logFileId}/search", cancellationToken);
+			return GetAsync<SearchLogResponse>(_httpClient, $"/api/v1/logs/{logFileId}/search?Text={Uri.EscapeDataString(searchText)}&count={count}", cancellationToken);
 		}
+
+		/// <summary>
+		/// Get the requested number of lines from given logFileId, starting at index
+		/// </summary>
+		/// <param name="logFileId">Id of log file to retrieve lines from</param>
+		/// <param name="startIndex">Start index of lines to retrieve</param>
+		/// <param name="count">Number of lines to retrieve</param>
+		/// <param name="cancellationToken">Cancellation token for the operation</param>
+		/// <returns></returns>
+		public Task<LogLinesResponse> GetLogLinesAsync(LogId logFileId, int startIndex, int count, CancellationToken cancellationToken = default)
+		{
+			return GetAsync<LogLinesResponse>(_httpClient, $"/api/v1/logs/{logFileId}/lines?index={startIndex}&count={count}", cancellationToken);
+		}
+
+		#endregion
+
+		#region Graph
+
+		/// <summary>
+		/// Get graph of the given job
+		/// </summary>
+		/// <param name="jobId"></param>
+		/// <param name="cancellationToken"></param>
+		/// <returns>Contains buildgraph information for the job</returns>
+		public Task<GetGraphResponse> GetGraphAsync(JobId jobId, CancellationToken cancellationToken = default)
+		{
+			return GetAsync<GetGraphResponse>(_httpClient, $"/api/v1/jobs/{jobId}/graph", cancellationToken);
+		}
+
 		#endregion
 
 		#region Utility Methods
diff --git a/Engine/Source/Programs/Shared/EpicGames.Horde/Logs/LogMessages.cs b/Engine/Source/Programs/Shared/EpicGames.Horde/Logs/LogMessages.cs
--- a/Engine/Source/Programs/Shared/EpicGames.Horde/Logs/LogMessages.cs
+++ b/Engine/Source/Programs/Shared/EpicGames.Horde/Logs/LogMessages.cs
@@ -99,5 +99,67 @@ namespace EpicGames.Horde.Logs
 		/// </summary>
 		public SearchStats? Stats { get; set; }
 	}
+
+	/// <summary>
+	/// Response when querying for specific lines from a log file
+	/// </summary>
+	public class LogLinesResponse
+	{
+		/// <summary>
+		/// start index of the lines returned
+		/// </summary>
+		public int Index;
+
+		/// <summary>
+		/// Number of lines returned
+		/// </summary>
+		public int Count;
+
+		/// <summary>
+		/// Last index of the returned messages
+		/// </summary>
+		public int MaxLineIndex;
+
+		/// <summary>
+		/// Type of response, Json or Text
+		/// </summary>
+		public LogType Format;
+
+		/// <summary>
+		/// List of lines received
+		/// </summary>
+		public List<LogLineResponse> Lines = new List<LogLineResponse>();
+	}
+
+	/// <summary>
+	/// Response object for individual lines
+	/// </summary>
+	public class LogLineResponse
+	{
+		/// <summary>
+		/// Timestamp for log line
+		/// </summary>
+		public DateTime Time;
+
+		/// <summary>
+		/// Level of Message (Information, Warning, Error)
+		/// </summary>
+		public string? Level;
+
+		/// <summary>
+		/// Message itself
+		/// </summary>
+		public string? Message;
+
+		/// <summary>
+		/// Format string for the message
+		/// </summary>
+		public string? Format;
+
+		/// <summary>
+		/// User-defined properties for this jobstep.
+		/// </summary>
+		public Dictionary<string, string>? Properties { get; set; }
+	}
 }
 
