.NET static analysis warnings in LogMessages.cs

diff --git a/Engine/Source/Programs/Shared/EpicGames.Horde/Logs/LogMessages.cs b/Engine/Source/Programs/Shared/EpicGames.Horde/Logs/LogMessages.cs
--- a/Engine/Source/Programs/Shared/EpicGames.Horde/Logs/LogMessages.cs
+++ b/Engine/Source/Programs/Shared/EpicGames.Horde/Logs/LogMessages.cs
@@ -108,27 +108,27 @@ namespace EpicGames.Horde.Logs
 		/// <summary>
 		/// start index of the lines returned
 		/// </summary>
-		public int Index;
+		public int Index { get; set; }
 
 		/// <summary>
 		/// Number of lines returned
 		/// </summary>
-		public int Count;
+		public int Count { get; set; }
 
 		/// <summary>
 		/// Last index of the returned messages
 		/// </summary>
-		public int MaxLineIndex;
+		public int MaxLineIndex { get; set; }
 
 		/// <summary>
 		/// Type of response, Json or Text
 		/// </summary>
-		public LogType Format;
+		public LogType Format { get; set; }
 
 		/// <summary>
 		/// List of lines received
 		/// </summary>
-		public List<LogLineResponse> Lines = new List<LogLineResponse>();
+		public List<LogLineResponse> Lines { get; set; } = new List<LogLineResponse>();
 	}
 
 	/// <summary>
@@ -139,22 +139,22 @@ namespace EpicGames.Horde.Logs
 		/// <summary>
 		/// Timestamp for log line
 		/// </summary>
-		public DateTime Time;
+		public DateTime Time { get; set; }
 
 		/// <summary>
 		/// Level of Message (Information, Warning, Error)
 		/// </summary>
-		public string? Level;
+		public string? Level { get; set; }
 
 		/// <summary>
 		/// Message itself
 		/// </summary>
-		public string? Message;
+		public string? Message { get; set; }
 
 		/// <summary>
 		/// Format string for the message
 		/// </summary>
-		public string? Format;
+		public string? Format { get; set; }
 
 		/// <summary>
 		/// User-defined properties for this jobstep.
