DirectoryItem.ContainsFiles

#rnx

diff --git a/Engine/Source/Programs/Shared/EpicGames.Build/System/DirectoryItem.cs b/Engine/Source/Programs/Shared/EpicGames.Build/System/DirectoryItem.cs
--- a/Engine/Source/Programs/Shared/EpicGames.Build/System/DirectoryItem.cs
+++ b/Engine/Source/Programs/Shared/EpicGames.Build/System/DirectoryItem.cs
@@ -5,6 +5,7 @@ using System.Collections.Concurrent;
 using System.Collections.Generic;
 using System.Diagnostics.CodeAnalysis;
 using System.IO;
+using System.Linq;
 using EpicGames.Core;
 
 namespace UnrealBuildBase
@@ -318,6 +319,16 @@ namespace UnrealBuildBase
 			return Files!.Values;
 		}
 
+		/// <summary>
+		/// Check if this directory contains any files
+		/// </summary>
+		/// <param name="searchOption">Directory search options</param>
+		/// <returns>True if this directory has files</returns>
+		public bool ContainsFiles(SearchOption searchOption = SearchOption.TopDirectoryOnly)
+		{
+			return searchOption == SearchOption.TopDirectoryOnly ? EnumerateFiles().Any(x => x.Exists) : (EnumerateFiles().Any(x => x.Exists) || EnumerateDirectories().Any(x => x.ContainsFiles(searchOption)));
+		}
+
 		/// <summary>
 		/// Attempts to get a file from this directory by name. Unlike creating a file item and checking whether it exists, this will
 		/// not create a permanent FileItem object if it does not exist.
