Horde: Tweak order of properties in job responses.

#rnx

diff --git a/Engine/Source/Programs/Horde/Horde.Server/Jobs/IJob.cs b/Engine/Source/Programs/Horde/Horde.Server/Jobs/IJob.cs
--- a/Engine/Source/Programs/Horde/Horde.Server/Jobs/IJob.cs
+++ b/Engine/Source/Programs/Horde/Horde.Server/Jobs/IJob.cs
@@ -915,11 +915,13 @@ namespace Horde.Server.Jobs
 				}
 
 				// Create the response
-				GetLabelStateResponse response = new GetLabelStateResponse(newState, newOutcome);
+				GetLabelStateResponse response = new GetLabelStateResponse();
 				response.DashboardName = label.DashboardName;
 				response.DashboardCategory = label.DashboardCategory;
 				response.UgsName = label.UgsName;
 				response.UgsProject = label.UgsProject;
+				response.State = newState;
+				response.Outcome = newOutcome;
 
 				foreach (NodeRef includedNodeRef in label.IncludedNodes)
 				{
diff --git a/Engine/Source/Programs/Shared/EpicGames.Horde/Jobs/JobMessages.cs b/Engine/Source/Programs/Shared/EpicGames.Horde/Jobs/JobMessages.cs
--- a/Engine/Source/Programs/Shared/EpicGames.Horde/Jobs/JobMessages.cs
+++ b/Engine/Source/Programs/Shared/EpicGames.Horde/Jobs/JobMessages.cs
@@ -992,11 +992,6 @@ namespace EpicGames.Horde.Jobs
 		/// </summary>
 		public JobStepBatchError Error { get; set; }
 
-		/// <summary>
-		/// Steps within this run
-		/// </summary>
-		public List<GetStepResponse> Steps { get; set; } = new List<GetStepResponse>();
-
 		/// <summary>
 		/// The agent assigned to execute this group
 		/// </summary>
@@ -1036,6 +1031,11 @@ namespace EpicGames.Horde.Jobs
 		/// Time at which the group became ready (UTC).
 		/// </summary>
 		public DateTimeOffset? ReadyTime { get; set; }
+
+		/// <summary>
+		/// Steps within this run
+		/// </summary>
+		public List<GetStepResponse> Steps { get; set; } = new List<GetStepResponse>();
 	}
 
 	/// <summary>
@@ -1110,11 +1110,6 @@ namespace EpicGames.Horde.Jobs
 		/// </summary>
 		public string? UgsProject { get; set; }
 
-		/// <summary>
-		/// Steps to include in the status of this label
-		/// </summary>
-		public List<JobStepId> Steps { get; set; } = new List<JobStepId>();
-
 		/// <summary>
 		/// State of the label
 		/// </summary>
@@ -1126,24 +1121,9 @@ namespace EpicGames.Horde.Jobs
 		public LabelOutcome? Outcome { get; set; }
 
 		/// <summary>
-		/// Constructor
-		/// </summary>
-		/// <param name="state">State of the label</param>
-		/// <param name="outcome">Outcome of the label</param>
-		public GetLabelStateResponse(LabelState state, LabelOutcome outcome)
-		{
-			State = state;
-			Outcome = outcome;
-		}
-
-		/// <summary>
-		/// Default Constructor needed for JsonSerializer
+		/// Steps to include in the status of this label
 		/// </summary>
-		[JsonConstructor]
-		public GetLabelStateResponse()
-		{
-
-		}
+		public List<JobStepId> Steps { get; set; } = new List<JobStepId>();
 	}
 
 	/// <summary>
@@ -1154,28 +1134,7 @@ namespace EpicGames.Horde.Jobs
 		/// <summary>
 		/// List of nodes covered by default label
 		/// </summary>
-		public List<string> Nodes { get; set; }
-
-		/// <summary>
-		/// Constructor
-		/// </summary>
-		/// <param name="state">State of the label</param>
-		/// <param name="outcome">Outcome of the label</param>
-		/// <param name="nodes">List of nodes that are covered by the default label</param>
-		public GetDefaultLabelStateResponse(LabelState state, LabelOutcome outcome, List<string> nodes)
-			: base(state, outcome)
-		{
-			Nodes = nodes;
-		}
-
-		/// <summary>
-		/// Default Constructor needed for JsonSerializer
-		/// </summary>
-		[JsonConstructor]
-		public GetDefaultLabelStateResponse() : base()
-		{
-			Nodes = new List<string>();
-		}
+		public List<string> Nodes { get; set; } = new List<string>();
 	}
 
 	/// <summary>
