Horde: Use the shared registry location to configure the Horde command line tool on Windows, and use a .horde.json in the user folder on Mac and Linux.

diff --git a/Engine/Source/Programs/Horde/Horde/CmdApp.cs b/Engine/Source/Programs/Horde/Horde/CmdApp.cs
--- a/Engine/Source/Programs/Horde/Horde/CmdApp.cs
+++ b/Engine/Source/Programs/Horde/Horde/CmdApp.cs
@@ -29,18 +29,14 @@ namespace Horde
 				.AddEnvironmentVariables()
 				.Build();
 
-			CmdConfig cmdConfig = CmdConfig.Read();
-
 			using ILoggerFactory loggerFactory = CreateLoggerFactory(configuration, arguments);
 
 			IServiceCollection services = new ServiceCollection();
 			services.Configure<HordeOptions>(options => configuration.Bind("Horde", options));
-			services.Configure<HordeOptions>(options => options.ServerUrl = cmdConfig.Server);
 			services.AddCommandsFromAssembly(Assembly.GetExecutingAssembly());
 			services.AddSingleton(loggerFactory);
 			services.AddLogging();
 			services.AddMemoryCache();
-			services.AddSingleton(Options.Create(cmdConfig));
 			services.AddHorde();
 
 			// Execute all the commands
diff --git a/Engine/Source/Programs/Horde/Horde/CmdConfig.cs b/Engine/Source/Programs/Horde/Horde/CmdConfig.cs
deleted file mode 100644
diff --git a/Engine/Source/Programs/Horde/Horde/Commands/Bundles/BundleCreate.cs b/Engine/Source/Programs/Horde/Horde/Commands/Bundles/BundleCreate.cs
--- a/Engine/Source/Programs/Horde/Horde/Commands/Bundles/BundleCreate.cs
+++ b/Engine/Source/Programs/Horde/Horde/Commands/Bundles/BundleCreate.cs
@@ -35,8 +35,8 @@ namespace Horde.Commands.Bundles
 		[Description("Clean the output folder before writing any data")]
 		public bool CleanOutput { get; set; }
 
-		public BundleCreate(HttpStorageClientFactory storageClientFactory, BundleCache bundleCache, IOptions<CmdConfig> config)
-			: base(storageClientFactory, bundleCache, config)
+		public BundleCreate(HttpStorageClientFactory storageClientFactory, BundleCache bundleCache)
+			: base(storageClientFactory, bundleCache)
 		{
 		}
 
diff --git a/Engine/Source/Programs/Horde/Horde/Commands/Bundles/BundleExtract.cs b/Engine/Source/Programs/Horde/Horde/Commands/Bundles/BundleExtract.cs
--- a/Engine/Source/Programs/Horde/Horde/Commands/Bundles/BundleExtract.cs
+++ b/Engine/Source/Programs/Horde/Horde/Commands/Bundles/BundleExtract.cs
@@ -40,8 +40,8 @@ namespace Horde.Commands.Bundles
 		[Description("If set, deletes the contents of the output directory before extraction.")]
 		public bool CleanOutput { get; set; }
 
-		public BundleExtract(HttpStorageClientFactory storageClientFactory, BundleCache bundleCache, IOptions<CmdConfig> config)
-			: base(storageClientFactory, bundleCache, config)
+		public BundleExtract(HttpStorageClientFactory storageClientFactory, BundleCache bundleCache)
+			: base(storageClientFactory, bundleCache)
 		{
 		}
 
diff --git a/Engine/Source/Programs/Horde/Horde/Commands/ConfigCommand.cs b/Engine/Source/Programs/Horde/Horde/Commands/ConfigCommand.cs
--- a/Engine/Source/Programs/Horde/Horde/Commands/ConfigCommand.cs
+++ b/Engine/Source/Programs/Horde/Horde/Commands/ConfigCommand.cs
@@ -2,6 +2,7 @@
 
 using System.ComponentModel;
 using EpicGames.Core;
+using EpicGames.Horde;
 using Microsoft.Extensions.Logging;
 using Microsoft.Extensions.Options;
 
@@ -14,24 +15,16 @@ namespace Horde.Commands
 		[Description("Updates the server URL")]
 		public string? Server { get; set; }
 
-		readonly CmdConfig _config;
-
-		public ConfigCommand(IOptions<CmdConfig> config)
-		{
-			_config = config.Value;
-		}
-
 		/// <inheritdoc/>
-		public override async Task<int> ExecuteAsync(ILogger logger)
+		public override Task<int> ExecuteAsync(ILogger logger)
 		{
 			if (Server != null)
 			{
-				_config.Server = new Uri(Server);
-				await _config.WriteAsync();
+				HordeOptions.SetDefaultServerUrl(new Uri(Server));
 			}
 
-			logger.LogInformation("Server: {Server}", _config.Server);
-			return 0;
+			logger.LogInformation("Server: {Server}", HordeOptions.GetDefaultServerUrl());
+			return Task.FromResult<int>(0);
 		}
 	}
 }
diff --git a/Engine/Source/Programs/Horde/Horde/Commands/LoginCommand.cs b/Engine/Source/Programs/Horde/Horde/Commands/LoginCommand.cs
--- a/Engine/Source/Programs/Horde/Horde/Commands/LoginCommand.cs
+++ b/Engine/Source/Programs/Horde/Horde/Commands/LoginCommand.cs
@@ -21,34 +21,22 @@ namespace Horde.Commands
 			public string[]? RedirectUrls { get; set; }
 		}
 
-		[CommandLine("-Server=")]
-		[Description("The server to connect to")]
-		public string? Server { get; set; }
-
 		[CommandLine("-Token")]
 		[Description("Echo the bearer token acquired from the server to stdout")]
 		public bool Token { get; set; }
 
 		readonly IServiceProvider _serviceProvider;
 		readonly IHordeClient _hordeClient;
-		readonly CmdConfig _config;
 
-		public LoginCommand(IServiceProvider serviceProvider, IHordeClient hordeClient, IOptions<CmdConfig> config)
+		public LoginCommand(IServiceProvider serviceProvider, IHordeClient hordeClient)
 		{
 			_serviceProvider = serviceProvider;
 			_hordeClient = hordeClient;
-			_config = config.Value;
 		}
 
 		/// <inheritdoc/>
 		public override async Task<int> ExecuteAsync(ILogger logger)
 		{
-			if (Server != null)
-			{
-				_config.Server = new Uri(Server);
-				await _config.WriteAsync();
-			}
-
 			using HordeHttpClient httpClient = _hordeClient.CreateHttpClient();
 
 			GetServerInfoResponse serverInfo = await httpClient.GetServerInfoAsync();
diff --git a/Engine/Source/Programs/Horde/Horde/Commands/StorageCommandBase.cs b/Engine/Source/Programs/Horde/Horde/Commands/StorageCommandBase.cs
--- a/Engine/Source/Programs/Horde/Horde/Commands/StorageCommandBase.cs
+++ b/Engine/Source/Programs/Horde/Horde/Commands/StorageCommandBase.cs
@@ -33,22 +33,16 @@ namespace Horde.Commands
 		/// </summary>
 		public BundleCache BundleCache { get; }
 
-		/// <summary>
-		/// Configuration for the tool
-		/// </summary>
-		public CmdConfig Config { get; }
-
 		readonly HttpStorageClientFactory _storageClientFactory;
 
 		/// <summary>
 		/// Constructor
 		/// </summary>
-		public StorageCommandBase(HttpStorageClientFactory storageClientFactory, BundleCache bundleCache, IOptions<CmdConfig> config)
+		public StorageCommandBase(HttpStorageClientFactory storageClientFactory, BundleCache bundleCache)
 		{
 			_storageClientFactory = storageClientFactory;
 
 			BundleCache = bundleCache;
-			Config = config.Value;
 		}
 
 		/// <summary>
diff --git a/Engine/Source/Programs/Horde/Horde/Commands/Workspace/WorkspaceSync.cs b/Engine/Source/Programs/Horde/Horde/Commands/Workspace/WorkspaceSync.cs
--- a/Engine/Source/Programs/Horde/Horde/Commands/Workspace/WorkspaceSync.cs
+++ b/Engine/Source/Programs/Horde/Horde/Commands/Workspace/WorkspaceSync.cs
@@ -42,8 +42,8 @@ namespace Horde.Commands.Workspace
 		[Description("Outputs stats for the extraction operation.")]
 		public bool Stats { get; set; }
 
-		public WorkspaceSync(HttpStorageClientFactory storageClientFactory, BundleCache bundleCache, IOptions<CmdConfig> config)
-			: base(storageClientFactory, bundleCache, config)
+		public WorkspaceSync(HttpStorageClientFactory storageClientFactory, BundleCache bundleCache)
+			: base(storageClientFactory, bundleCache)
 		{
 		}
 
diff --git a/Engine/Source/Programs/Shared/EpicGames.Horde/HordeOptions.cs b/Engine/Source/Programs/Shared/EpicGames.Horde/HordeOptions.cs
--- a/Engine/Source/Programs/Shared/EpicGames.Horde/HordeOptions.cs
+++ b/Engine/Source/Programs/Shared/EpicGames.Horde/HordeOptions.cs
@@ -1,15 +1,21 @@
 // Copyright Epic Games, Inc. All Rights Reserved.
 
 using System;
+using System.IO;
 using System.Linq;
 using System.Net.Http;
 using System.Runtime.InteropServices;
 using System.Runtime.Versioning;
+using System.Text.Json;
+using System.Text.Json.Nodes;
+using EpicGames.Core;
 using EpicGames.Horde.Storage.Bundles;
 using Microsoft.Win32;
 
 namespace EpicGames.Horde
 {
+	using JsonObject = System.Text.Json.Nodes.JsonObject;
+
 	/// <summary>
 	/// Options for configuring the Horde connection
 	/// </summary>
@@ -49,7 +55,7 @@ namespace EpicGames.Horde
 		/// Gets the default server URL for the current user
 		/// </summary>
 		/// <returns>Default URL</returns>
-		public static Uri? GetDefaultServerUrl()
+		public static Uri GetDefaultServerUrl()
 		{
 			if (RuntimeInformation.IsOSPlatform(OSPlatform.Windows))
 			{
@@ -68,7 +74,34 @@ namespace EpicGames.Horde
 					}
 				}
 			}
-			return null;
+			else
+			{
+				FileReference? configFile = GetConfigFile();
+				if (configFile != null)
+				{
+					byte[] data = FileReference.ReadAllBytes(configFile);
+
+					JsonObject? root = JsonNode.Parse(data, new JsonNodeOptions { PropertyNameCaseInsensitive = true }, new JsonDocumentOptions { AllowTrailingCommas = true }) as JsonObject;
+					root ??= new JsonObject();
+
+					JsonNode? value;
+					if (root.TryGetPropertyValue("server", out value))
+					{
+						string? stringValue = (string?)value;
+						if (stringValue != null)
+						{
+							try
+							{
+								return new Uri(stringValue);
+							}
+							catch (UriFormatException)
+							{
+							}
+						}
+					}
+				}
+			}
+			return new Uri("http://localhost:5000/");
 		}
 
 		/// <summary>
@@ -77,6 +110,11 @@ namespace EpicGames.Horde
 		/// <param name="serverUrl">Horde server URL to use</param>
 		public static void SetDefaultServerUrl(Uri serverUrl)
 		{
+			if (!serverUrl.OriginalString.EndsWith("/", StringComparison.Ordinal))
+			{
+				serverUrl = new Uri(serverUrl.OriginalString + "/");
+			}
+
 			if (RuntimeInformation.IsOSPlatform(OSPlatform.Windows))
 			{
 				string? newServerUrl = serverUrl.ToString();
@@ -92,6 +130,32 @@ namespace EpicGames.Horde
 					Registry.SetValue(@"HKEY_CURRENT_USER\SOFTWARE\Epic Games\Horde", "Url", serverUrl.ToString());
 				}
 			}
+			else
+			{
+				FileReference? configFile = GetConfigFile();
+				if (configFile != null)
+				{
+					byte[] data = FileReference.ReadAllBytes(configFile);
+
+					JsonObject? root = JsonNode.Parse(data, new JsonNodeOptions { PropertyNameCaseInsensitive = true }, new JsonDocumentOptions { AllowTrailingCommas = true }) as JsonObject;
+					root ??= new JsonObject();
+					root["server"] = serverUrl.ToString();
+
+					using (FileStream stream = FileReference.Open(configFile, FileMode.Create, FileAccess.ReadWrite, FileShare.Read))
+					{
+						using Utf8JsonWriter writer = new Utf8JsonWriter(stream, new JsonWriterOptions { Indented = true });
+						root.WriteTo(writer);
+					}
+				}
+			}
+		}
+
+		static FileReference? GetConfigFile()
+		{
+			DirectoryReference? userFolder = DirectoryReference.GetSpecialFolder(Environment.SpecialFolder.UserProfile);
+			userFolder ??= DirectoryReference.GetSpecialFolder(Environment.SpecialFolder.LocalApplicationData);
+			userFolder ??= DirectoryReference.GetCurrentDirectory();
+			return FileReference.Combine(userFolder, ".horde.json");
 		}
 
 		[SupportedOSPlatform("windows")]
