EpicGames.Perforce: Fix errors parsing change views from labels.

diff --git a/Engine/Source/Programs/Shared/EpicGames.Perforce/PerforceChangeView.cs b/Engine/Source/Programs/Shared/EpicGames.Perforce/PerforceChangeView.cs
--- a/Engine/Source/Programs/Shared/EpicGames.Perforce/PerforceChangeView.cs
+++ b/Engine/Source/Programs/Shared/EpicGames.Perforce/PerforceChangeView.cs
@@ -47,7 +47,15 @@ namespace EpicGames.Perforce
 			{
 				if (!String.IsNullOrWhiteSpace(line))
 				{
-					changeView.Entries.Add(new PerforceChangeViewEntry(line, ignoreCase));
+					// Skip over lines with unsupported formatting (generally @label instead of @changenumber)
+					try
+					{
+						changeView.Entries.Add(new PerforceChangeViewEntry(line, ignoreCase));
+					}
+					catch (FormatException)
+					{
+						// don't do anything
+					}
 				}
 			}
 			return changeView;
