Horde: Do not return a default localhost server URL if nothing is configured.

diff --git a/Engine/Source/Programs/Shared/EpicGames.Horde/HordeHttpAuthHandler.cs b/Engine/Source/Programs/Shared/EpicGames.Horde/HordeHttpAuthHandler.cs
--- a/Engine/Source/Programs/Shared/EpicGames.Horde/HordeHttpAuthHandler.cs
+++ b/Engine/Source/Programs/Shared/EpicGames.Horde/HordeHttpAuthHandler.cs
@@ -280,7 +280,7 @@ namespace EpicGames.Horde
 			{
 				if (httpClient.BaseAddress == null)
 				{
-					throw new Exception("No http client is configured for Horde. Call IServiceCollection.AddHordeHttpClient().");
+					throw new Exception("Horde base address is not configured.");
 				}
 
 				serverUrl = httpClient.BaseAddress;
diff --git a/Engine/Source/Programs/Shared/EpicGames.Horde/HordeOptions.cs b/Engine/Source/Programs/Shared/EpicGames.Horde/HordeOptions.cs
--- a/Engine/Source/Programs/Shared/EpicGames.Horde/HordeOptions.cs
+++ b/Engine/Source/Programs/Shared/EpicGames.Horde/HordeOptions.cs
@@ -55,7 +55,7 @@ namespace EpicGames.Horde
 		/// Gets the default server URL for the current user
 		/// </summary>
 		/// <returns>Default URL</returns>
-		public static Uri GetDefaultServerUrl()
+		public static Uri? GetDefaultServerUrl()
 		{
 			if (RuntimeInformation.IsOSPlatform(OSPlatform.Windows))
 			{
@@ -101,7 +101,7 @@ namespace EpicGames.Horde
 					}
 				}
 			}
-			return new Uri("http://localhost:5000/");
+			return null;
 		}
 
 		/// <summary>
