Remove unnecessary using directives from threadsanitizerissuehandler. Fixing issue=460206

diff --git a/Engine/Source/Programs/Shared/EpicGames.Horde/Issues/Handlers/ThreadSanitizerIssueHandler.cs b/Engine/Source/Programs/Shared/EpicGames.Horde/Issues/Handlers/ThreadSanitizerIssueHandler.cs
--- a/Engine/Source/Programs/Shared/EpicGames.Horde/Issues/Handlers/ThreadSanitizerIssueHandler.cs
+++ b/Engine/Source/Programs/Shared/EpicGames.Horde/Issues/Handlers/ThreadSanitizerIssueHandler.cs
@@ -1,8 +1,6 @@
 // Copyright Epic Games, Inc. All Rights Reserved.
 
-using System;
 using System.Collections.Generic;
-using System.IO;
 using System.Linq;
 using System.Text.Json;
 using EpicGames.Core;
