OidcToken - Fixed static code analysis warning about not passing cancellation token

diff --git a/Engine/Source/Programs/Shared/EpicGames.OIDC/OidcTokenManager.cs b/Engine/Source/Programs/Shared/EpicGames.OIDC/OidcTokenManager.cs
--- a/Engine/Source/Programs/Shared/EpicGames.OIDC/OidcTokenManager.cs
+++ b/Engine/Source/Programs/Shared/EpicGames.OIDC/OidcTokenManager.cs
@@ -471,7 +471,7 @@ namespace EpicGames.OIDC
 			byte[] buffer = Encoding.UTF8.GetBytes(httpPage);
 			response.ContentLength64 = buffer.Length;
 			Stream responseOutput = response.OutputStream;
-			await responseOutput.WriteAsync(buffer, 0, buffer.Length);
+			await responseOutput.WriteAsync(buffer, 0, buffer.Length, cancellationToken);
 			responseOutput.Close();
 
 			return loginResult;
