Horde: Allow passing SortedSetWhen arguments to SortedSetAdd methods.

#rnx

diff --git a/Engine/Source/Programs/Horde/Horde.Server/Notifications/Sinks/SlackNotificationSink.cs b/Engine/Source/Programs/Horde/Horde.Server/Notifications/Sinks/SlackNotificationSink.cs
--- a/Engine/Source/Programs/Horde/Horde.Server/Notifications/Sinks/SlackNotificationSink.cs
+++ b/Engine/Source/Programs/Horde/Horde.Server/Notifications/Sinks/SlackNotificationSink.cs
@@ -50,6 +50,7 @@ using Microsoft.Extensions.Options;
 using MongoDB.Bson;
 using MongoDB.Bson.Serialization.Attributes;
 using MongoDB.Driver;
+using StackExchange.Redis;
 
 namespace Horde.Server.Notifications.Sinks
 {
@@ -1212,7 +1213,7 @@ namespace Horde.Server.Notifications.Sinks
 				if (workflow.EscalateAlias != null && workflow.EscalateTimes.Count > 0)
 				{
 					DateTime escalateTime = issue.CreatedAt.AddMinutes(workflow.EscalateTimes[0]);
-					if (await _redisService.GetDatabase().SortedSetAddAsync(s_escalateIssues, issue.Id, (escalateTime - DateTime.UnixEpoch).TotalSeconds, StackExchange.Redis.When.NotExists))
+					if (await _redisService.GetDatabase().SortedSetAddAsync(s_escalateIssues, issue.Id, (escalateTime - DateTime.UnixEpoch).TotalSeconds, SortedSetWhen.NotExists))
 					{
 						_logger.LogInformation("First escalation time for issue {IssueId} is {Time}", issue.Id, escalateTime);
 					}
diff --git a/Engine/Source/Programs/Shared/EpicGames.Redis/RedisSortedSet.cs b/Engine/Source/Programs/Shared/EpicGames.Redis/RedisSortedSet.cs
--- a/Engine/Source/Programs/Shared/EpicGames.Redis/RedisSortedSet.cs
+++ b/Engine/Source/Programs/Shared/EpicGames.Redis/RedisSortedSet.cs
@@ -105,13 +105,13 @@ namespace EpicGames.Redis
 		#region SortedSetAddAsync
 
 		/// <inheritdoc cref="IDatabaseAsync.SortedSetAddAsync(RedisKey, RedisValue, Double, CommandFlags)"/>
-		public static Task<bool> SortedSetAddAsync<TElement>(this IDatabaseAsync target, RedisSortedSetKey<TElement> key, TElement value, double score, When when = When.Always, CommandFlags flags = CommandFlags.None)
+		public static Task<bool> SortedSetAddAsync<TElement>(this IDatabaseAsync target, RedisSortedSetKey<TElement> key, TElement value, double score, SortedSetWhen when = SortedSetWhen.Always, CommandFlags flags = CommandFlags.None)
 		{
 			return target.SortedSetAddAsync(key.Inner, RedisSerializer.Serialize(value), score, when, flags);
 		}
 
 		/// <inheritdoc cref="IDatabaseAsync.SortedSetAddAsync(RedisKey, RedisValue, Double, When, CommandFlags)"/>
-		public static Task<long> SortedSetAddAsync<TElement>(this IDatabaseAsync target, RedisSortedSetKey<TElement> key, SortedSetEntry<TElement>[] values, When when = When.Always, CommandFlags flags = CommandFlags.None)
+		public static Task<long> SortedSetAddAsync<TElement>(this IDatabaseAsync target, RedisSortedSetKey<TElement> key, SortedSetEntry<TElement>[] values, SortedSetWhen when = SortedSetWhen.Always, CommandFlags flags = CommandFlags.None)
 		{
 			return target.SortedSetAddAsync(key.Inner, values.ConvertAll(x => new SortedSetEntry(x.ElementValue, x.Score)), when, flags);
 		}
