Horde: Move IClock extension methods into EpicGames.Core.

#rnx

diff --git a/Engine/Source/Programs/Horde/Horde.Server/Utilities/Clock.cs b/Engine/Source/Programs/Horde/Horde.Server/Utilities/Clock.cs
--- a/Engine/Source/Programs/Horde/Horde.Server/Utilities/Clock.cs
+++ b/Engine/Source/Programs/Horde/Horde.Server/Utilities/Clock.cs
@@ -37,57 +37,6 @@ namespace HordeCommon
 	/// </summary>
 	public static class ClockExtensions
 	{
-		/// <summary>
-		/// Create an event that will trigger after the given time
-		/// </summary>
-		/// <param name="clock">Clock to schedule the event on</param>
-		/// <param name="name">Name of the ticker</param>
-		/// <param name="interval">Interval for the callback</param>
-		/// <param name="tickAsync">Trigger callback</param>
-		/// <param name="logger">Logger for any error messages</param>
-		/// <returns>Handle to the event</returns>
-		public static ITicker AddTicker(this IClock clock, string name, TimeSpan interval, Func<CancellationToken, ValueTask> tickAsync, ILogger logger)
-		{
-			async ValueTask<TimeSpan?> WrappedTrigger(CancellationToken token)
-			{
-				Stopwatch timer = Stopwatch.StartNew();
-				await tickAsync(token);
-				return interval - timer.Elapsed;
-			}
-
-			return clock.AddTicker(name, interval, WrappedTrigger, logger);
-		}
-
-		/// <summary>
-		/// Create an event that will trigger after the given time
-		/// </summary>
-		/// <param name="clock">Clock to schedule the event on</param>
-		/// <param name="interval">Time after which the event will trigger</param>
-		/// <param name="tickAsync">Callback for the tick. Returns the time interval until the next tick, or null to cancel the tick.</param>
-		/// <param name="logger">Logger for error messages</param>
-		/// <returns>Handle to the event</returns>
-		public static ITicker AddTicker<T>(this IClock clock, TimeSpan interval, Func<CancellationToken, ValueTask<TimeSpan?>> tickAsync, ILogger logger) => clock.AddTicker(typeof(T).Name, interval, tickAsync, logger);
-
-		/// <summary>
-		/// Create an event that will trigger after the given time
-		/// </summary>
-		/// <param name="clock">Clock to schedule the event on</param>
-		/// <param name="interval">Interval for the callback</param>
-		/// <param name="tickAsync">Trigger callback</param>
-		/// <param name="logger">Logger for any error messages</param>
-		/// <returns>Handle to the event</returns>
-		public static ITicker AddTicker<T>(this IClock clock, TimeSpan interval, Func<CancellationToken, ValueTask> tickAsync, ILogger logger) => clock.AddTicker(typeof(T).Name, interval, tickAsync, logger);
-
-		/// <summary>
-		/// Create a ticker shared between all server pods
-		/// </summary>
-		/// <param name="clock">Clock to schedule the event on</param>
-		/// <param name="interval">Time after which the event will trigger</param>
-		/// <param name="tickAsync">Callback for the tick. Returns the time interval until the next tick, or null to cancel the tick.</param>
-		/// <param name="logger">Logger for error messages</param>
-		/// <returns>New ticker instance</returns>
-		public static ITicker AddSharedTicker<T>(this IClock clock, TimeSpan interval, Func<CancellationToken, ValueTask> tickAsync, ILogger logger) => clock.AddSharedTicker(typeof(T).Name, interval, tickAsync, logger);
-
 		/// <summary>
 		/// Gets the start of the day for the given datetime in UTC, respecting the configured timezone.
 		/// </summary>
diff --git a/Engine/Source/Programs/Shared/EpicGames.Core/Clock.cs b/Engine/Source/Programs/Shared/EpicGames.Core/Clock.cs
--- a/Engine/Source/Programs/Shared/EpicGames.Core/Clock.cs
+++ b/Engine/Source/Programs/Shared/EpicGames.Core/Clock.cs
@@ -1,6 +1,7 @@
 // Copyright Epic Games, Inc. All Rights Reserved.
 
 using System;
+using System.Diagnostics;
 using System.Threading;
 using System.Threading.Tasks;
 using Microsoft.Extensions.Logging;
@@ -123,3 +124,62 @@ public class StubClock : IClock
 	}
 }
 
+/// <summary>
+/// Extension methods for <see cref="IClock"/>
+/// </summary>
+public static class ClockExtensions
+{
+	/// <summary>
+	/// Create an event that will trigger after the given time
+	/// </summary>
+	/// <param name="clock">Clock to schedule the event on</param>
+	/// <param name="name">Name of the ticker</param>
+	/// <param name="interval">Interval for the callback</param>
+	/// <param name="tickAsync">Trigger callback</param>
+	/// <param name="logger">Logger for any error messages</param>
+	/// <returns>Handle to the event</returns>
+	public static ITicker AddTicker(this IClock clock, string name, TimeSpan interval, Func<CancellationToken, ValueTask> tickAsync, ILogger logger)
+	{
+		async ValueTask<TimeSpan?> WrappedTrigger(CancellationToken token)
+		{
+			Stopwatch timer = Stopwatch.StartNew();
+			await tickAsync(token);
+			return interval - timer.Elapsed;
+		}
+
+		return clock.AddTicker(name, interval, WrappedTrigger, logger);
+	}
+
+	/// <summary>
+	/// Create an event that will trigger after the given time
+	/// </summary>
+	/// <param name="clock">Clock to schedule the event on</param>
+	/// <param name="interval">Time after which the event will trigger</param>
+	/// <param name="tickAsync">Callback for the tick. Returns the time interval until the next tick, or null to cancel the tick.</param>
+	/// <param name="logger">Logger for error messages</param>
+	/// <returns>Handle to the event</returns>
+	public static ITicker AddTicker<T>(this IClock clock, TimeSpan interval, Func<CancellationToken, ValueTask<TimeSpan?>> tickAsync, ILogger logger) 
+		=> clock.AddTicker(typeof(T).Name, interval, tickAsync, logger);
+
+	/// <summary>
+	/// Create an event that will trigger after the given time
+	/// </summary>
+	/// <param name="clock">Clock to schedule the event on</param>
+	/// <param name="interval">Interval for the callback</param>
+	/// <param name="tickAsync">Trigger callback</param>
+	/// <param name="logger">Logger for any error messages</param>
+	/// <returns>Handle to the event</returns>
+	public static ITicker AddTicker<T>(this IClock clock, TimeSpan interval, Func<CancellationToken, ValueTask> tickAsync, ILogger logger) 
+		=> clock.AddTicker(typeof(T).Name, interval, tickAsync, logger);
+
+	/// <summary>
+	/// Create a ticker shared between all server pods
+	/// </summary>
+	/// <param name="clock">Clock to schedule the event on</param>
+	/// <param name="interval">Time after which the event will trigger</param>
+	/// <param name="tickAsync">Callback for the tick. Returns the time interval until the next tick, or null to cancel the tick.</param>
+	/// <param name="logger">Logger for error messages</param>
+	/// <returns>New ticker instance</returns>
+	public static ITicker AddSharedTicker<T>(this IClock clock, TimeSpan interval, Func<CancellationToken, ValueTask> tickAsync, ILogger logger) 
+		=> clock.AddSharedTicker(typeof(T).Name, interval, tickAsync, logger);
+}
