Add CheckConnectionAsync to HordeHttpClient.

diff --git a/Engine/Source/Programs/Shared/EpicGames.Horde/HordeHttpClient.cs b/Engine/Source/Programs/Shared/EpicGames.Horde/HordeHttpClient.cs
--- a/Engine/Source/Programs/Shared/EpicGames.Horde/HordeHttpClient.cs
+++ b/Engine/Source/Programs/Shared/EpicGames.Horde/HordeHttpClient.cs
@@ -114,6 +114,18 @@ namespace EpicGames.Horde
 			options.Converters.Add(new SubResourceIdJsonConverterFactory());
 		}
 
+		#region Connection
+		/// <summary>
+		/// Check account login status.
+		/// </summary>
+		public async Task<bool> CheckConnectionAsync(CancellationToken cancellationToken = default)
+		{
+			HttpResponseMessage response = await _httpClient.GetAsync("account", cancellationToken);
+
+			return response.IsSuccessStatusCode;
+		}
+
+		#endregion
 		#region Artifacts
 
 		/// <summary>
