Various fixes to allow for compiling shared projects with target framework newer than .net6 where apis have changed

#rnx

diff --git a/Engine/Source/Programs/Shared/EpicGames.Core/BackgroundTask.cs b/Engine/Source/Programs/Shared/EpicGames.Core/BackgroundTask.cs
--- a/Engine/Source/Programs/Shared/EpicGames.Core/BackgroundTask.cs
+++ b/Engine/Source/Programs/Shared/EpicGames.Core/BackgroundTask.cs
@@ -82,7 +82,11 @@ namespace EpicGames.Core
 			{
 				try
 				{
+#if NET8_0_OR_GREATER
+					await _cancellationTokenSource.CancelAsync();
+#else
 					_cancellationTokenSource.Cancel();
+#endif
 					await _task.WaitAsync(cancellationToken).ConfigureAwait(false);
 				}
 				catch (OperationCanceledException)
@@ -157,7 +161,11 @@ namespace EpicGames.Core
 			{
 				try
 				{
+#if NET8_0_OR_GREATER
+					await _cancellationTokenSource.CancelAsync();
+#else
 					_cancellationTokenSource.Cancel();
+#endif
 					await _task.WaitAsync(cancellationToken);
 				}
 				catch (OperationCanceledException)
diff --git a/Engine/Source/Programs/Shared/EpicGames.Core/JsonExtensions.cs b/Engine/Source/Programs/Shared/EpicGames.Core/JsonExtensions.cs
--- a/Engine/Source/Programs/Shared/EpicGames.Core/JsonExtensions.cs
+++ b/Engine/Source/Programs/Shared/EpicGames.Core/JsonExtensions.cs
@@ -4,8 +4,6 @@ using System;
 using System.Diagnostics;
 using System.Text.Json;
 
-#pragma warning disable CA1045 // Do not pass types by reference
-
 namespace EpicGames.Core
 {
 	/// <summary>
@@ -19,7 +17,7 @@ namespace EpicGames.Core
 		/// <param name="reader">Token reader</param>
 		/// <param name="propertyName">Receives the property name on success</param>
 		/// <returns>True if the read succeeded</returns>
-		public static bool TryReadNextPropertyName(ref Utf8JsonReader reader, out ReadOnlySpan<byte> propertyName)
+		public static bool TryReadNextPropertyName(this Utf8JsonReader reader, out ReadOnlySpan<byte> propertyName)
 		{
 			if (reader.Read() && reader.TokenType == JsonTokenType.PropertyName)
 			{
diff --git a/Engine/Source/Programs/Shared/EpicGames.Core/LogEvent.cs b/Engine/Source/Programs/Shared/EpicGames.Core/LogEvent.cs
--- a/Engine/Source/Programs/Shared/EpicGames.Core/LogEvent.cs
+++ b/Engine/Source/Programs/Shared/EpicGames.Core/LogEvent.cs
@@ -220,7 +220,7 @@ namespace EpicGames.Core
 			LogException? exception = null;
 
 			ReadOnlySpan<byte> propertyName;
-			for (; JsonExtensions.TryReadNextPropertyName(ref reader, out propertyName); reader.Skip())
+			for (; reader.TryReadNextPropertyName(out propertyName); reader.Skip())
 			{
 				if (Utf8StringComparer.OrdinalIgnoreCase.Equals(propertyName, LogEventPropertyName.Time.Span))
 				{
@@ -264,7 +264,7 @@ namespace EpicGames.Core
 			Dictionary<string, object> properties = new Dictionary<string, object>();
 
 			ReadOnlySpan<byte> propertyName;
-			for (; JsonExtensions.TryReadNextPropertyName(ref reader, out propertyName); reader.Skip())
+			for (; reader.TryReadNextPropertyName(out propertyName); reader.Skip())
 			{
 				string name = Encoding.UTF8.GetString(propertyName);
 				object value = ReadPropertyValue(ref reader);
@@ -313,7 +313,7 @@ namespace EpicGames.Core
 			Dictionary<Utf8String, object>? properties = null;
 
 			ReadOnlySpan<byte> propertyName;
-			for (; JsonExtensions.TryReadNextPropertyName(ref reader, out propertyName); reader.Skip())
+			for (; reader.TryReadNextPropertyName(out propertyName); reader.Skip())
 			{
 				if (Utf8StringComparer.OrdinalIgnoreCase.Equals(propertyName, LogEventPropertyName.Type.Span))
 				{
diff --git a/Engine/Source/Programs/Shared/EpicGames.Horde/Compute/Clients/ServerComputeClient.cs b/Engine/Source/Programs/Shared/EpicGames.Horde/Compute/Clients/ServerComputeClient.cs
--- a/Engine/Source/Programs/Shared/EpicGames.Horde/Compute/Clients/ServerComputeClient.cs
+++ b/Engine/Source/Programs/Shared/EpicGames.Horde/Compute/Clients/ServerComputeClient.cs
@@ -401,7 +401,9 @@ namespace EpicGames.Horde.Compute.Clients
 			await writer.WriteLineAsync(request.ToCharArray(), cancellationToken);
 
 			string exceptionMetadata = $"Connection: {response.ConnectionAddress} Target: {response.Ip}:{response.Port}";
+#pragma warning disable CA2016 // Forward the 'CancellationToken' parameter to methods
 			Task<string?> readTask = reader.ReadLineAsync();
+#pragma warning restore CA2016 // Forward the 'CancellationToken' parameter to methods
 			Task timeoutTask = Task.Delay(15000, cancellationToken);
 			if (await Task.WhenAny(readTask, timeoutTask) == timeoutTask)
 			{
diff --git a/Engine/Source/Programs/Shared/EpicGames.OIDC/OidcTokenManager.cs b/Engine/Source/Programs/Shared/EpicGames.OIDC/OidcTokenManager.cs
--- a/Engine/Source/Programs/Shared/EpicGames.OIDC/OidcTokenManager.cs
+++ b/Engine/Source/Programs/Shared/EpicGames.OIDC/OidcTokenManager.cs
@@ -339,7 +339,11 @@ namespace EpicGames.OIDC
 								{
 									// timed out
 									loginResult = new LoginResult($"Login timed out after: {_loginTimeout.TotalMinutes} minutes");
+#if NET8_0_OR_GREATER
+									await tokenSource.CancelAsync();
+#else
 									tokenSource.Cancel();
+#endif
 									// we do not await the processHttpTask as the GetContext method does not provide any way to cancel, as such this task will be stuck until httplistener is disposed
 								}
 							}
@@ -437,7 +441,11 @@ namespace EpicGames.OIDC
 
 							await using Stream body = request.InputStream;
 							using StreamReader reader = new StreamReader(body, request.ContentEncoding);
+#if NET7_0_OR_GREATER
+							responseData = await reader.ReadToEndAsync(cancellationToken);
+#else
 							responseData = await reader.ReadToEndAsync();
+#endif
 
 							loginResult = await oidcClient!.ProcessResponseAsync(responseData, loginState, cancellationToken: cancellationToken);
 							break;
diff --git a/Engine/Source/Programs/Shared/EpicGames.Redis/Utility/RedisLock.cs b/Engine/Source/Programs/Shared/EpicGames.Redis/Utility/RedisLock.cs
--- a/Engine/Source/Programs/Shared/EpicGames.Redis/Utility/RedisLock.cs
+++ b/Engine/Source/Programs/Shared/EpicGames.Redis/Utility/RedisLock.cs
@@ -51,7 +51,14 @@ namespace EpicGames.Redis.Utility
 		/// <inheritdoc/>
 		public async ValueTask DisposeAsync()
 		{
-			_cancellationSource?.Cancel();
+			if (_cancellationSource != null)
+			{
+#if NET8_0_OR_GREATER
+				await _cancellationSource.CancelAsync();
+#else
+				_cancellationSource.Cancel();
+#endif
+			}
 			if (_backgroundTask != null)
 			{
 				await _backgroundTask;
