EpicGames.Tracing: Fix remaining warnings

#rnx

diff --git a/Engine/Source/Programs/Shared/EpicGames.Tracing.Tests/UnrealInsights/GenericEventTest.cs b/Engine/Source/Programs/Shared/EpicGames.Tracing.Tests/UnrealInsights/GenericEventTest.cs
--- a/Engine/Source/Programs/Shared/EpicGames.Tracing.Tests/UnrealInsights/GenericEventTest.cs
+++ b/Engine/Source/Programs/Shared/EpicGames.Tracing.Tests/UnrealInsights/GenericEventTest.cs
@@ -418,6 +418,13 @@ namespace EpicGames.Tracing.Tests.UnrealInsights
 			using MemoryStream ms = new MemoryStream();
 			using BinaryWriter writer = new BinaryWriter(ms);
 			@event.Serialize(uid, writer);
+
+			if (isImportant)
+			{
+				byte[] expectedData = HexStringToBytes(expectedHexData);
+				int expectedEventSize = expectedData.Length - TraceImportantEventHeader.HeaderSize;
+				Assert.AreEqual(expectedData.Length, @event.Size);
+			}
 			
 			AssertHexString(expectedHexData, ms.ToArray());
 		}
diff --git a/Engine/Source/Programs/Shared/EpicGames.Tracing.Tests/UnrealInsights/UnrealInsightsTest.cs b/Engine/Source/Programs/Shared/EpicGames.Tracing.Tests/UnrealInsights/UnrealInsightsTest.cs
--- a/Engine/Source/Programs/Shared/EpicGames.Tracing.Tests/UnrealInsights/UnrealInsightsTest.cs
+++ b/Engine/Source/Programs/Shared/EpicGames.Tracing.Tests/UnrealInsights/UnrealInsightsTest.cs
@@ -131,13 +131,10 @@ namespace EpicGames.Tracing.Tests.UnrealInsights
 		}
 		
 		[TestMethod]
-		[Ignore]
 		public void WriteUtraceExample()
 		{
-			string fileName = "d:\\Temp\\mytrace.utrace";
-			
-			FileStream fs = File.Open(fileName, FileMode.Create);
-			using BinaryWriter binaryWriter = new BinaryWriter(fs);
+			using MemoryStream ms = new MemoryStream();
+			using BinaryWriter binaryWriter = new BinaryWriter(ms);
 			UnrealInsightsWriter writer = new UnrealInsightsWriter();
 			
 			byte[] cpuBatchData1 = GenericEventTest.HexStringToBytes("87 A7 EB AE 8A D1 0B 01 A7 02 02 04 FB 23 03 A4 0A 61 04 0B 05 F2 80 42 AA 1C EB 5C 06 C0 BC 03 17 07 02 47 08 C2 02 AB 15 09 AA 03 07 0A AE 01 05 0B D8 0D 05 0C B7 01 0D D8 03 00 03 0E A4 03 6F 0F 89 FF 03 10 F3 C4 04 11 60 06 A7 02 10 AB 19 12 B3 04 10 9F 11 13 0A 04 02 04 ED 03 10 91 11 14 80 01 02 97 02 10 A9 D6 0B 15 10 06 DD 03 10 CB 85 04 16 12 06 D3 01 10 89 16 17 0C 02 D6 02 07 18 DA 84 28 99 01 19 18 05 1A 0A A5 10 1B CD 01 1C 06 88 6A C5 49 1D B8 EB 24 A9 17 20 AF 01 10 CD F9 01 21 71 22 F2 AE 12 85 04 23 D1 0D 24 85 74 25 EA 58 91 0D 25 80 E6 02 A9 1C 25 88 CD 06 A7 27 25 D0 DA 09 8B 20 25 DE AD 0B AF 33 25 C0 A9 09 ED 23 25 C4 81 0A F5 2D 25 9A EC 0A F3 35 25 D4 CF 03 FF 19 25 BE 92 01 B1 0D 25 98 22");
@@ -186,11 +183,10 @@ namespace EpicGames.Tracing.Tests.UnrealInsights
 			}
 			writer.Write(binaryWriter);
 
-			fs.Close();
+			ms.Seek(0, SeekOrigin.Begin);
 
 			UnrealInsightsReader reader = new UnrealInsightsReader();
-			using FileStream stream = File.Open(fileName, FileMode.Open);
-			reader.Read(stream);
+			reader.Read(ms);
 			reader.PrintEventSummary();
 
 			Assert.AreEqual(0, reader.EventsPerThread[TransportPacket.ThreadIdEvents].Count);
diff --git a/Engine/Source/Programs/Shared/EpicGames.Tracing/UnrealInsights/Events.cs b/Engine/Source/Programs/Shared/EpicGames.Tracing/UnrealInsights/Events.cs
--- a/Engine/Source/Programs/Shared/EpicGames.Tracing/UnrealInsights/Events.cs
+++ b/Engine/Source/Programs/Shared/EpicGames.Tracing/UnrealInsights/Events.cs
@@ -189,6 +189,7 @@ namespace EpicGames.Tracing.UnrealInsights
 	{
 		public ushort Size => 7;
 		public EventType Type => EventType.WellKnown(PredefinedEventUid.EnterScope_T, "LeaveScopeTimestamp");
+		[System.Diagnostics.CodeAnalysis.SuppressMessage("CodeQuality", "IDE0052:Remove unread private members", Justification = "Serialize() unimplemented")]
 		readonly ulong _timestamp;
 
 		public LeaveScopeEventTimestamp(ulong timestamp)
diff --git a/Engine/Source/Programs/Shared/EpicGames.Tracing/UnrealInsights/Events/CpuProfiler.cs b/Engine/Source/Programs/Shared/EpicGames.Tracing/UnrealInsights/Events/CpuProfiler.cs
--- a/Engine/Source/Programs/Shared/EpicGames.Tracing/UnrealInsights/Events/CpuProfiler.cs
+++ b/Engine/Source/Programs/Shared/EpicGames.Tracing/UnrealInsights/Events/CpuProfiler.cs
@@ -29,8 +29,8 @@ namespace EpicGames.Tracing.UnrealInsights.Events
 			
 			Field[] fields =
 			{
-				Field.FromInt((int) id),
-				Field.FromString(name),
+				Field.FromInt((int) Id),
+				Field.FromString(_name),
 			};
 			
 			_genericEvent = new GenericEvent(0, fields, EventType);
diff --git a/Engine/Source/Programs/Shared/EpicGames.Tracing/UnrealInsights/Events/GenericEvent.cs b/Engine/Source/Programs/Shared/EpicGames.Tracing/UnrealInsights/Events/GenericEvent.cs
--- a/Engine/Source/Programs/Shared/EpicGames.Tracing/UnrealInsights/Events/GenericEvent.cs
+++ b/Engine/Source/Programs/Shared/EpicGames.Tracing/UnrealInsights/Events/GenericEvent.cs
@@ -98,6 +98,7 @@ namespace EpicGames.Tracing.UnrealInsights.Events
 			}
 		}
 
+		[System.Diagnostics.CodeAnalysis.SuppressMessage("Style", "IDE0060:Remove unused parameter")]
 		public void Serialize(ushort uid, BinaryWriter writer)
 		{
 			if (_eventType.HasSerial)
@@ -210,6 +211,7 @@ namespace EpicGames.Tracing.UnrealInsights.Events
 			}
 		}
 
+		[System.Diagnostics.CodeAnalysis.SuppressMessage("Style", "IDE0060:Remove unused parameter")]
 		public static GenericEvent Deserialize(ushort uid, BinaryReader reader, EventType eventType)
 		{
 			Field[] fields = new Field[eventType.Fields.Count];
diff --git a/Engine/Source/Programs/Shared/EpicGames.Tracing/UnrealInsights/Events/NewEvent.cs b/Engine/Source/Programs/Shared/EpicGames.Tracing/UnrealInsights/Events/NewEvent.cs
--- a/Engine/Source/Programs/Shared/EpicGames.Tracing/UnrealInsights/Events/NewEvent.cs
+++ b/Engine/Source/Programs/Shared/EpicGames.Tracing/UnrealInsights/Events/NewEvent.cs
@@ -8,6 +8,7 @@ using System.Text;
 
 namespace EpicGames.Tracing.UnrealInsights.Events
 {
+	[System.Diagnostics.CodeAnalysis.SuppressMessage("CodeQuality", "IDE0052:Remove unread private members")]
 	[System.Diagnostics.CodeAnalysis.SuppressMessage("Performance", "CA1823:Avoid unused private fields")]
 	[System.Diagnostics.CodeAnalysis.SuppressMessage("Style", "IDE1006:Naming Styles")]
 	public class EventTypeField
diff --git a/Engine/Source/Programs/Shared/EpicGames.Tracing/UnrealInsights/StreamHeader.cs b/Engine/Source/Programs/Shared/EpicGames.Tracing/UnrealInsights/StreamHeader.cs
--- a/Engine/Source/Programs/Shared/EpicGames.Tracing/UnrealInsights/StreamHeader.cs
+++ b/Engine/Source/Programs/Shared/EpicGames.Tracing/UnrealInsights/StreamHeader.cs
@@ -33,7 +33,21 @@ namespace EpicGames.Tracing.UnrealInsights
 
 		public void Serialize(BinaryWriter writer)
 		{
-			writer.Write(MagicTrc);
+			if (_magicFourCc != MagicTrc && _magicFourCc != MagicTrc2)
+			{
+				throw new ArgumentException("Only support magic number TRCE and TRC2");
+			}
+			writer.Write(_magicFourCc);
+			if (_magicFourCc == MagicTrc2)
+			{
+				if (_metadataSize != 4)
+				{
+					throw new ArgumentException("Only support metadata size of 4 bytes (got " + _metadataSize + ")");
+				}
+				writer.Write(_metadataSize);
+				writer.Write(_metadataField0);
+				writer.Write(_controlPort);
+			}
 			writer.Write(_transportVersion);
 			writer.Write(_protocolVersion);
 		}
