sh -> /bin/sh

#rnx

diff --git a/Engine/Source/Programs/Shared/EpicGames.Build/Platform/Mac/ApplePlatformSDK.cs b/Engine/Source/Programs/Shared/EpicGames.Build/Platform/Mac/ApplePlatformSDK.cs
--- a/Engine/Source/Programs/Shared/EpicGames.Build/Platform/Mac/ApplePlatformSDK.cs
+++ b/Engine/Source/Programs/Shared/EpicGames.Build/Platform/Mac/ApplePlatformSDK.cs
@@ -52,14 +52,14 @@ namespace UnrealBuildBase
 				int ExitCode;
 				// xcode-select -p gives the currently selected Xcode location (xcodebuild -version may fail if Xcode.app is broken)
 				// Example output: /Applications/Xcode.app/Contents/Developer
-				string Output = RunLocalProcessAndReturnStdOut("sh", "-c 'xcode-select -p'", out ExitCode);
+				string Output = RunLocalProcessAndReturnStdOut("/bin/sh", "-c 'xcode-select -p'", out ExitCode);
 
 				if (ExitCode == 0)
 				{
 					DirectoryReference DeveloperDir = new DirectoryReference(Output);
 					FileReference Plist = FileReference.Combine(DeveloperDir.ParentDirectory!, "Info.plist");
 					// Find out the version number in Xcode.app/Contents/Info.plist
-					Output = RunLocalProcessAndReturnStdOut("sh",
+					Output = RunLocalProcessAndReturnStdOut("/bin/sh",
 						  $"-c 'plutil -extract CFBundleShortVersionString raw {Plist}'", out ExitCode);
 					if (ExitCode == 0)
 					{
diff --git a/Engine/Source/Programs/UnrealBuildTool/ToolChain/AppleToolChain.cs b/Engine/Source/Programs/UnrealBuildTool/ToolChain/AppleToolChain.cs
--- a/Engine/Source/Programs/UnrealBuildTool/ToolChain/AppleToolChain.cs
+++ b/Engine/Source/Programs/UnrealBuildTool/ToolChain/AppleToolChain.cs
@@ -1139,7 +1139,7 @@ namespace UnrealBuildTool
 			}
 
 			// run the script
-			Utils.RunLocalProcessAndReturnStdOut("sh", $"\"{SignProjectScript.FullName}\"");
+			Utils.RunLocalProcessAndReturnStdOut("/bin/sh", $"\"{SignProjectScript.FullName}\"");
 
 			// Set parameters to make sure it uses the correct identity and keychain
 			// pass back the comandline arguments to xcodebuild to use these certicicates
@@ -1162,7 +1162,7 @@ namespace UnrealBuildTool
 				CleanWriter.WriteLine("security list-keychain -s login.keychain");
 			}
 
-			Utils.RunLocalProcessAndReturnStdOut("sh", $"\"{CleanProjectScript.FullName}\"");
+			Utils.RunLocalProcessAndReturnStdOut("/bin/sh", $"\"{CleanProjectScript.FullName}\"");
 		}
 
 		private static FileItem GetPostBuildOutputFile(FileReference Executable, string TargetName, UnrealTargetPlatform Platform)
