[UBA]
* Changed so environment variable is sent from c# to native in UTF8 formatted memory buffer.. Reason is because they are not automatically synced so in this case we want to provide it

diff --git a/Engine/Source/Programs/Shared/EpicGames.UBA/Impl/SessionServerImpl.cs b/Engine/Source/Programs/Shared/EpicGames.UBA/Impl/SessionServerImpl.cs
--- a/Engine/Source/Programs/Shared/EpicGames.UBA/Impl/SessionServerImpl.cs
+++ b/Engine/Source/Programs/Shared/EpicGames.UBA/Impl/SessionServerImpl.cs
@@ -1,8 +1,10 @@
 // Copyright Epic Games, Inc. All Rights Reserved.
 
 using System;
+using System.Collections;
 using System.Collections.Concurrent;
 using System.Collections.Generic;
+using System.IO;
 using System.Runtime.InteropServices;
 
 namespace EpicGames.UBA
@@ -72,7 +74,7 @@ namespace EpicGames.UBA
 
 		#region DllImport
 		[DllImport("UbaHost", CharSet = CharSet.Auto)]
-		static extern IntPtr SessionServer_Create(IntPtr info);
+		static extern IntPtr SessionServer_Create(IntPtr info, byte[] environment, uint environmentSize);
 
 		[DllImport("UbaHost", CharSet = CharSet.Auto)]
 		static extern void SessionServer_SetRemoteProcessAvailable(IntPtr server, RemoteProcessSlotAvailableCallback func);
@@ -122,7 +124,21 @@ namespace EpicGames.UBA
 			_info = info;
 			_remoteProcessSlotAvailableCallbackDelegate = RaiseRemoteProcessSlotAvailable;
 			_remoteProcessReturnedCallbackDelegate = RaiseRemoteProcessReturned;
-			_handle = SessionServer_Create(_info.GetHandle());
+
+			// We need to manually transfer environment variables on non-windows platforms since they are not automatically propagated from c# to native.
+			using MemoryStream environmentMemory = new();
+			{
+				if (!RuntimeInformation.IsOSPlatform(OSPlatform.Windows))
+				{
+					using (BinaryWriter writer = new(environmentMemory, System.Text.Encoding.UTF8, true))
+					{
+						foreach (DictionaryEntry de in Environment.GetEnvironmentVariables())
+							writer.Write($"{de.Key}={de.Value}");
+					}
+				}
+			}
+
+			_handle = SessionServer_Create(_info.GetHandle(), environmentMemory.GetBuffer(), (uint)environmentMemory.Position);
 			SessionServer_SetRemoteProcessAvailable(_handle, _remoteProcessSlotAvailableCallbackDelegate);
 			SessionServer_SetRemoteProcessReturned(_handle, _remoteProcessReturnedCallbackDelegate);
 		}
