dotnet8 ref fix

#rnx

diff --git a/Engine/Source/Programs/Shared/EpicGames.Core/JsonExtensions.cs b/Engine/Source/Programs/Shared/EpicGames.Core/JsonExtensions.cs
--- a/Engine/Source/Programs/Shared/EpicGames.Core/JsonExtensions.cs
+++ b/Engine/Source/Programs/Shared/EpicGames.Core/JsonExtensions.cs
@@ -19,7 +19,11 @@ namespace EpicGames.Core
 		/// <param name="reader">Token reader</param>
 		/// <param name="propertyName">Receives the property name on success</param>
 		/// <returns>True if the read succeeded</returns>
+#if NET7_0_OR_GREATER
+		public static bool TryReadNextPropertyName(scoped ref Utf8JsonReader reader, out ReadOnlySpan<byte> propertyName)
+#else
 		public static bool TryReadNextPropertyName(ref Utf8JsonReader reader, out ReadOnlySpan<byte> propertyName)
+#endif
 		{
 			if (reader.Read() && reader.TokenType == JsonTokenType.PropertyName)
 			{
diff --git a/Engine/Source/Programs/Shared/EpicGames.Horde/Jobs/Templates/ParameterId.cs b/Engine/Source/Programs/Shared/EpicGames.Horde/Jobs/Templates/ParameterId.cs
--- a/Engine/Source/Programs/Shared/EpicGames.Horde/Jobs/Templates/ParameterId.cs
+++ b/Engine/Source/Programs/Shared/EpicGames.Horde/Jobs/Templates/ParameterId.cs
@@ -19,7 +19,7 @@ namespace EpicGames.Horde.Jobs.Templates
 		{
 		}
 
-		/// <inheritdoc cref="StringId.Sanitize(string)"/>
+		/// <inheritdoc cref="StringId.Sanitize"/>
 		public static ParameterId Sanitize(string name)
 			=> new ParameterId(StringId.Sanitize(name));
 
