Horde: Improve exception message for too long StringId

#rnx

diff --git a/Engine/Source/Programs/Shared/EpicGames.Horde/StringId.cs b/Engine/Source/Programs/Shared/EpicGames.Horde/StringId.cs
--- a/Engine/Source/Programs/Shared/EpicGames.Horde/StringId.cs
+++ b/Engine/Source/Programs/Shared/EpicGames.Horde/StringId.cs
@@ -139,7 +139,7 @@ namespace EpicGames.Horde
 		{
 			if (text.Length > MaxLength)
 			{
-				throw new ArgumentException($"String id may not be longer than {MaxLength} characters", paramName);
+				throw new ArgumentException($"String id '{text}' may not be longer than {MaxLength} characters", paramName);
 			}
 
 			if (text.Length > 0 && (text[0] == '.' || text[^1] == '.'))
