OidcToken - Merging state of the token store, avoiding issues were long lived applications (like UGS) could overwrite the state written by other applications that it was not aware of.

This is somewhat speculative as I have been unable to repro this actually happening locally but we have a lot of users reporting issues with their login expiring matching what a bug like this would cause.

diff --git a/Engine/Source/Programs/Shared/EpicGames.OIDC/TokenStore.cs b/Engine/Source/Programs/Shared/EpicGames.OIDC/TokenStore.cs
--- a/Engine/Source/Programs/Shared/EpicGames.OIDC/TokenStore.cs
+++ b/Engine/Source/Programs/Shared/EpicGames.OIDC/TokenStore.cs
@@ -40,12 +40,13 @@ namespace EpicGames.OIDC
 			}
 		}
 
-		private readonly Dictionary<string, byte[]> _providerToRefreshToken = new Dictionary<string, byte[]>();
+		private readonly Dictionary<string, byte[]> _providerToRefreshToken;
+		private readonly List<string> _dirtyProviders = new List<string>();
 		private readonly SymmetricAlgorithm _crypt = Aes.Create();
 
 		public FilesystemTokenStore()
 		{
-			ReadStoreFromDisk();
+			_providerToRefreshToken = ReadStoreFromDisk();
 		}
 
 		private static FileInfo GetStorePath()
@@ -53,13 +54,13 @@ namespace EpicGames.OIDC
 			return new FileInfo(Path.Combine(Environment.GetFolderPath(Environment.SpecialFolder.ApplicationData), "Epic", "UnrealEngine", "Common", "OidcToken", "tokenStore.dat"));
 		}
 
-		private void ReadStoreFromDisk()
+		private Dictionary<string, byte[]> ReadStoreFromDisk()
 		{
 			FileInfo fi = GetStorePath();
 			if (!fi.Exists)
 			{
 				// if we have no store on disk then we just initialize it to empty
-				return;
+				return new Dictionary<string, byte[]>();
 			}
 
 			using FileStream fs = fi.OpenRead();
@@ -77,15 +78,18 @@ namespace EpicGames.OIDC
 			if (state == null)
 			{
 				// if we fail to deserialize the state just drop it, will mean users will need to login again
-				return;
+				return new Dictionary<string, byte[]>();
 			}
+			Dictionary<string, byte[]> providers = new Dictionary<string, byte[]>();
 
 			foreach ((string key, string value) in state.Providers)
 			{
-				_providerToRefreshToken[key] = Convert.FromBase64String(value);
+				providers[key] = Convert.FromBase64String(value);
 			}
-
+			
 			_crypt.IV = Convert.FromBase64String(state.Key);
+
+			return providers;
 		}
 
 		private void SaveStoreToDisk()
@@ -97,14 +101,33 @@ namespace EpicGames.OIDC
 				Directory.CreateDirectory(fi.Directory!.FullName);
 			}
 
-			string tempFile = Path.GetTempFileName();
+			lock (_dirtyProviders)
 			{
-				using FileStream fs = new FileStream(tempFile, FileMode.Create, FileAccess.Write);
-				using Utf8JsonWriter writer = new Utf8JsonWriter(fs);
-				JsonSerializer.Serialize<TokenStoreState>(writer, new TokenStoreState(_crypt.IV, _providerToRefreshToken));
+				// no providers have changed, do not touch the state file
+				if (_dirtyProviders.Count == 0)
+				{
+					return;
+				}
+
+				// read back the state of all providers but only overwrite the state of the ones we have actually got new state for (are dirty)
+				Dictionary<string, byte[]> providers = ReadStoreFromDisk();
+
+				foreach (string providerId in _dirtyProviders)
+				{
+					providers[providerId] = _providerToRefreshToken[providerId];
+				}
+
+				string tempFile = Path.GetTempFileName();
+				{
+					using FileStream fs = new FileStream(tempFile, FileMode.Create, FileAccess.Write);
+					using Utf8JsonWriter writer = new Utf8JsonWriter(fs);
+					JsonSerializer.Serialize<TokenStoreState>(writer, new TokenStoreState(_crypt.IV, providers));
+				}
+
+				File.Move(tempFile, fi.FullName, true);
+				_dirtyProviders.Clear();
 			}
 
-			File.Move(tempFile, fi.FullName, true);
 		}
 
 		public bool TryGetRefreshToken(string oidcProvider, out string refreshToken)
@@ -130,6 +153,11 @@ namespace EpicGames.OIDC
 
 			byte[] encryptedToken = encryptor.TransformFinalBlock(bytes, 0, bytes.Length);
 			_providerToRefreshToken[providerIdentifier] = encryptedToken;
+
+			lock (_dirtyProviders)
+			{
+				_dirtyProviders.Add(providerIdentifier);
+			}
 		}
 
 		public void Save()
diff --git a/Engine/Source/Programs/Shared/EpicGames.OIDC/WindowsTokenStore.cs b/Engine/Source/Programs/Shared/EpicGames.OIDC/WindowsTokenStore.cs
--- a/Engine/Source/Programs/Shared/EpicGames.OIDC/WindowsTokenStore.cs
+++ b/Engine/Source/Programs/Shared/EpicGames.OIDC/WindowsTokenStore.cs
@@ -4,6 +4,7 @@ using System;
 using System.Collections.Generic;
 using System.ComponentModel;
 using System.IO;
+using System.Linq;
 using System.Runtime.InteropServices;
 using System.Text;
 using System.Text.Json;
@@ -40,7 +41,7 @@ namespace EpicGames.OIDC
 		}
 
 		private readonly Dictionary<string, byte[]> _providerToRefreshToken = new Dictionary<string, byte[]>();
-		private bool _isDirty;
+		private readonly List<string> _dirtyProviders = new List<string>();
 
 		public WindowsTokenStore()
 		{
@@ -52,7 +53,7 @@ namespace EpicGames.OIDC
 		{
 			_logger = logger;
 
-			ReadStoreFromDisk();
+			_providerToRefreshToken = ReadStoreFromDisk();
 		}
 
 		private static FileInfo GetStorePath()
@@ -60,14 +61,14 @@ namespace EpicGames.OIDC
 			return new FileInfo(Path.Combine(Environment.GetFolderPath(Environment.SpecialFolder.LocalApplicationData), "UnrealEngine", "Common", "OidcToken", "oidcTokenStore.dat"));
 		}
 
-		private void ReadStoreFromDisk()
+		private Dictionary<string, byte[]> ReadStoreFromDisk()
 		{
 			FileInfo fi = GetStorePath();
 			if (!fi.Exists)
 			{
 				_logger?.LogDebug("No existing token store found at {Path}. Assuming empty store.", fi.FullName);
 				// if we have no store on disk then we just initialize it to empty
-				return;
+				return new Dictionary<string, byte[]>();
 			}
 
 			using FileStream fs = fi.Open(FileMode.Open, FileAccess.Read, FileShare.Read);
@@ -87,22 +88,21 @@ namespace EpicGames.OIDC
 			{
 				_logger?.LogDebug("Failed to deserialize state. Dropping the existing state.");
 				// if we fail to deserialize the state just drop it, will mean users will need to login again
-				return;
+				return new Dictionary<string, byte[]>();
 			}
 
+			Dictionary<string, byte[]> providers = new Dictionary<string, byte[]>();
+
 			foreach ((string key, string value) in state.Providers)
 			{
-				_providerToRefreshToken[key] = Convert.FromBase64String(value);
+				providers[key] = Convert.FromBase64String(value);
 			}
+
+			return providers;
 		}
 
 		private void SaveStoreToDisk()
 		{
-			if (!_isDirty)
-			{
-				return;
-			}
-
 			FileInfo fi = GetStorePath();
 
 			if (!fi.Directory?.Exists ?? false)
@@ -110,14 +110,14 @@ namespace EpicGames.OIDC
 				Directory.CreateDirectory(fi.Directory!.FullName);
 			}
 
-			string tempFile = Path.GetTempFileName();
+			lock (_dirtyProviders)
 			{
-				using FileStream fs = new FileStream(tempFile, FileMode.Create, FileAccess.Write);
-				using Utf8JsonWriter writer = new Utf8JsonWriter(fs);
-				JsonSerializer.Serialize<TokenStoreState>(writer, new TokenStoreState(_providerToRefreshToken));
-			}
+				// no providers have changed, do not touch the state file
+				if (_dirtyProviders.Count == 0)
+				{
+					return;
+				}
 
-			{
 				using Mutex mutex = new Mutex(false, "oidcTokenStoreDat");
 
 				try
@@ -128,11 +128,27 @@ namespace EpicGames.OIDC
 				{
 
 				}
+
+				// read back the state of all providers but only overwrite the state of the ones we have actually got new state for (are dirty)
+				Dictionary<string, byte[]> providers = ReadStoreFromDisk();
+
+				foreach (string providerId in _dirtyProviders)
+				{
+					providers[providerId] = _providerToRefreshToken[providerId];
+				}
+				string tempFile = Path.GetTempFileName();
+				{
+					using FileStream fs = new FileStream(tempFile, FileMode.Create, FileAccess.Write);
+					using Utf8JsonWriter writer = new Utf8JsonWriter(fs);
+					JsonSerializer.Serialize<TokenStoreState>(writer, new TokenStoreState(providers));
+				}
+
 				File.Move(tempFile, fi.FullName, true);
 
 				mutex.ReleaseMutex();
+
+				_dirtyProviders.Clear();
 			}
-			_isDirty = false;
 		}
 
 		public bool TryGetRefreshToken(string oidcProvider, out string refreshToken)
@@ -184,7 +200,10 @@ namespace EpicGames.OIDC
 
 			_providerToRefreshToken[providerIdentifier] = encryptedToken;
 
-			_isDirty = true;
+			lock (_dirtyProviders)
+			{
+				_dirtyProviders.Add(providerIdentifier);
+			}
 		}
 
 		public void Save()
