EpicGames.OIDC - Fixed issue were the token store on windows would not correctly read the state when constructed outside of a dependency injected setup (primarily a issue for UBT)

diff --git a/Engine/Source/Programs/Shared/EpicGames.OIDC/WindowsTokenStore.cs b/Engine/Source/Programs/Shared/EpicGames.OIDC/WindowsTokenStore.cs
--- a/Engine/Source/Programs/Shared/EpicGames.OIDC/WindowsTokenStore.cs
+++ b/Engine/Source/Programs/Shared/EpicGames.OIDC/WindowsTokenStore.cs
@@ -45,7 +45,7 @@ namespace EpicGames.OIDC
 
 		public WindowsTokenStore()
 		{
-			ReadStoreFromDisk();
+			_providerToRefreshToken = ReadStoreFromDisk();
 		}
 
 		[ActivatorUtilitiesConstructor]
