[UBA]
* Added code to handle requesting shutting down uba cache service

diff --git a/Engine/Source/Programs/Shared/EpicGames.UBA/ICacheClient.cs b/Engine/Source/Programs/Shared/EpicGames.UBA/ICacheClient.cs
--- a/Engine/Source/Programs/Shared/EpicGames.UBA/ICacheClient.cs
+++ b/Engine/Source/Programs/Shared/EpicGames.UBA/ICacheClient.cs
@@ -61,6 +61,12 @@ namespace EpicGames.UBA
 		/// <returns>True if successful</returns>
 		public abstract bool FetchFromCache(IRootPaths rootPaths, uint bucket, ProcessStartInfo info);
 
+		/// <summary>
+		/// Request the connected server to shutdown
+		/// </summary>
+		/// <param name="reason">Reason for shutdown</param>
+		public abstract void RequestServerShutdown(string reason);
+
 		/// <summary>
 		/// Create a ICacheClient object
 		/// </summary>
diff --git a/Engine/Source/Programs/Shared/EpicGames.UBA/Impl/CacheClientImpl.cs b/Engine/Source/Programs/Shared/EpicGames.UBA/Impl/CacheClientImpl.cs
--- a/Engine/Source/Programs/Shared/EpicGames.UBA/Impl/CacheClientImpl.cs
+++ b/Engine/Source/Programs/Shared/EpicGames.UBA/Impl/CacheClientImpl.cs
@@ -88,6 +88,9 @@ namespace EpicGames.UBA
 		[DllImport("UbaHost", CharSet = CharSet.Auto)]
 		static extern bool CacheClient_FetchFromCache(IntPtr cacheClient, IntPtr rootPaths, uint bucket, IntPtr info);
 
+		[DllImport("UbaHost", CharSet = CharSet.Auto)]
+		static extern void CacheClient_RequestServerShutdown(IntPtr cacheClient, string reason);
+
 		[DllImport("UbaHost", CharSet = CharSet.Auto)]
 		static extern void CacheClient_Destroy(IntPtr cacheClient);
 		#endregion
@@ -138,6 +141,11 @@ namespace EpicGames.UBA
 			return result;
 		}
 
+		public void RequestServerShutdown(string reason)
+		{
+			CacheClient_RequestServerShutdown(_handle, reason);
+		}
+
 		public IntPtr GetHandle() => _handle;
 	}
 }
diff --git a/Engine/Source/Programs/UnrealBuildTool/Executors/UnrealBuildAccelerator/UBAExecutor.cs b/Engine/Source/Programs/UnrealBuildTool/Executors/UnrealBuildAccelerator/UBAExecutor.cs
--- a/Engine/Source/Programs/UnrealBuildTool/Executors/UnrealBuildAccelerator/UBAExecutor.cs
+++ b/Engine/Source/Programs/UnrealBuildTool/Executors/UnrealBuildAccelerator/UBAExecutor.cs
@@ -328,6 +328,8 @@ namespace UnrealBuildTool
 
 						bool success = ExecuteActionsInternal(inputActions, _session, logger, actionArtifactCache, () => ActionQueueCanceled(ubaStorage));
 
+						//_cacheClient.RequestServerShutdown("");
+
 						if (!UBAConfig.bDisableRemote)
 						{
 							Server.StopServer();
