Horde: Allow creating GrpcChannel from a Horde client instance.

#rnx

diff --git a/Engine/Source/Programs/Shared/EpicGames.Horde/EpicGames.Horde.csproj b/Engine/Source/Programs/Shared/EpicGames.Horde/EpicGames.Horde.csproj
--- a/Engine/Source/Programs/Shared/EpicGames.Horde/EpicGames.Horde.csproj
+++ b/Engine/Source/Programs/Shared/EpicGames.Horde/EpicGames.Horde.csproj
@@ -40,6 +40,7 @@
   <ItemGroup>
     <PackageReference Include="BitFaster.Caching" Version="2.4.1" />
     <PackageReference Include="Dapper" Version="2.1.24" />
+    <PackageReference Include="Grpc.Net.Client" Version="2.59.0" />
     <PackageReference Include="K4os.Compression.LZ4" Version="1.2.16" />
     <PackageReference Include="Microsoft.CodeAnalysis.CSharp.Workspaces" Version="4.2.0" />
     <PackageReference Include="Microsoft.Data.Sqlite" Version="8.0.0" />
diff --git a/Engine/Source/Programs/Shared/EpicGames.Horde/HordeClient.cs b/Engine/Source/Programs/Shared/EpicGames.Horde/HordeClient.cs
--- a/Engine/Source/Programs/Shared/EpicGames.Horde/HordeClient.cs
+++ b/Engine/Source/Programs/Shared/EpicGames.Horde/HordeClient.cs
@@ -2,11 +2,17 @@
 
 using System;
 using System.Net.Http;
+using System.Net.Http.Json;
 using System.Threading;
 using System.Threading.Tasks;
+using EpicGames.Core;
+using EpicGames.Horde.Server;
 using EpicGames.Horde.Storage;
 using EpicGames.Horde.Storage.Backends;
 using EpicGames.Horde.Storage.Bundles;
+using Grpc.Core;
+using Grpc.Net.Client;
+using Grpc.Net.Client.Configuration;
 using Microsoft.Extensions.Logging;
 using Microsoft.Extensions.Options;
 
@@ -15,7 +21,7 @@ namespace EpicGames.Horde
 	/// <summary>
 	/// Default implementation of <see cref="IHordeClient"/>
 	/// </summary>
-	class HordeClient : IHordeClient
+	sealed class HordeClient : IHordeClient, IAsyncDisposable
 	{
 		readonly IHttpClientFactory _httpClientFactory;
 		readonly HordeHttpAuthHandlerState _authHandlerState;
@@ -23,6 +29,7 @@ namespace EpicGames.Horde
 		readonly HordeOptions _hordeOptions;
 		readonly ILoggerFactory _loggerFactory;
 
+		BackgroundTask<GrpcChannel>? _grpcChannel;
 		Uri? _serverUrl;
 
 		/// <inheritdoc/>
@@ -40,6 +47,16 @@ namespace EpicGames.Horde
 			_loggerFactory = loggerFactory;
 		}
 
+		/// <inheritdoc/>
+		public async ValueTask DisposeAsync()
+		{
+			if (_grpcChannel != null)
+			{
+				await _grpcChannel.DisposeAsync();
+				_grpcChannel = null;
+			}
+		}
+
 		Uri GetServerUrl()
 		{
 			if (_serverUrl == null)
@@ -69,6 +86,60 @@ namespace EpicGames.Horde
 			}
 		}
 
+		/// <inheritdoc/>
+		public async Task<GrpcChannel> GetGrpcChannelAsync(CancellationToken cancellationToken)
+		{
+			_grpcChannel ??= BackgroundTask.StartNew(ctx => CreateGrpcChannelInternalAsync(ctx));
+			return await _grpcChannel.WaitAsync(cancellationToken);
+		}
+
+		async Task<GrpcChannel> CreateGrpcChannelInternalAsync(CancellationToken cancellationToken)
+		{
+			HttpClient httpClient = _httpClientFactory.CreateClient(HordeHttpClient.HttpClientName);
+
+			// Get the server URL for gRPC traffic. If we're using an unencrpyted connection we need to use a different port for http/2, so 
+			// send a http1 request to the server to query it.
+			Uri serverUri = httpClient.BaseAddress ?? throw new InvalidOperationException("Horde server base address is not configured");
+			if (serverUri.Scheme.Equals("http", StringComparison.Ordinal))
+			{
+				using (HttpResponseMessage response = await httpClient.GetAsync("api/v1/server/ports", cancellationToken))
+				{
+					GetPortsResponse? ports = await response.Content.ReadFromJsonAsync<GetPortsResponse>(HordeHttpClient.JsonSerializerOptions, cancellationToken);
+					if (ports != null && ports.UnencryptedHttp2.HasValue && ports.UnencryptedHttp2 != 0)
+					{
+						UriBuilder builder = new UriBuilder(serverUri);
+						builder.Port = ports.UnencryptedHttp2.Value;
+						serverUri = builder.Uri;
+					}
+				}
+			}
+
+			ServiceConfig serviceConfig = new ServiceConfig();
+			serviceConfig.MethodConfigs.Add(new MethodConfig
+			{
+				Names = { MethodName.Default },
+				RetryPolicy = new RetryPolicy
+				{
+					MaxAttempts = 3,
+					InitialBackoff = TimeSpan.FromSeconds(1),
+					MaxBackoff = TimeSpan.FromSeconds(10),
+					BackoffMultiplier = 2.0,
+					RetryableStatusCodes = { StatusCode.Unavailable },
+				}
+			});
+
+			return GrpcChannel.ForAddress(serverUri, new GrpcChannelOptions
+			{
+				// Required payloads coming from CAS service can be large
+				MaxReceiveMessageSize = 1024 * 1024 * 1024, // 1 GB
+				MaxSendMessageSize = 1024 * 1024 * 1024, // 1 GB
+				LoggerFactory = _loggerFactory,
+				HttpClient = httpClient,
+				DisposeHttpClient = true,
+				ServiceConfig = serviceConfig
+			});
+		}
+
 		/// <inheritdoc/>
 		public HordeHttpClient CreateHttpClient()
 			=> _httpClientFactory.CreateHordeHttpClient();
diff --git a/Engine/Source/Programs/Shared/EpicGames.Horde/IHordeClient.cs b/Engine/Source/Programs/Shared/EpicGames.Horde/IHordeClient.cs
--- a/Engine/Source/Programs/Shared/EpicGames.Horde/IHordeClient.cs
+++ b/Engine/Source/Programs/Shared/EpicGames.Horde/IHordeClient.cs
@@ -5,6 +5,8 @@ using System.Threading;
 using System.Threading.Tasks;
 using EpicGames.Horde.Storage;
 using EpicGames.Horde.Tools;
+using Grpc.Core;
+using Grpc.Net.Client;
 
 namespace EpicGames.Horde
 {
@@ -32,9 +34,13 @@ namespace EpicGames.Horde
 		bool IsConnected();
 
 		/// <summary>
-		/// Creates a http client 
+		/// Gets a grpc channel for communication with the server. This should NOT be disposed by the caller.
+		/// </summary>
+		Task<GrpcChannel> GetGrpcChannelAsync(CancellationToken cancellationToken = default);
+
+		/// <summary>
+		/// Creates a Horde HTTP client 
 		/// </summary>
-		/// <returns></returns>
 		HordeHttpClient CreateHttpClient();
 
 		/// <summary>
@@ -59,5 +65,15 @@ namespace EpicGames.Horde
 		/// </summary>
 		public static IStorageClient CreateStorageClient(this IHordeClient hordeClient, ToolId toolId)
 			=> hordeClient.CreateStorageClient($"api/v1/tools/{toolId}");
+
+		/// <summary>
+		/// Attempts to get a client reference, returning immediately if there's not one available
+		/// </summary>
+		public static async Task<TClient> GetGrpcClientAsync<TClient>(this IHordeClient hordeClient, CancellationToken cancellationToken = default) 
+			where TClient : ClientBase<TClient>
+		{
+			GrpcChannel channel = await hordeClient.GetGrpcChannelAsync(cancellationToken);
+			return (TClient)Activator.CreateInstance(typeof(TClient), channel)!;
+		}
 	}
 }
