Horde: Suppress build warning.

#rnx

diff --git a/Engine/Source/Programs/Shared/EpicGames.Horde/HordeClient.cs b/Engine/Source/Programs/Shared/EpicGames.Horde/HordeClient.cs
--- a/Engine/Source/Programs/Shared/EpicGames.Horde/HordeClient.cs
+++ b/Engine/Source/Programs/Shared/EpicGames.Horde/HordeClient.cs
@@ -16,6 +16,8 @@ using Grpc.Net.Client.Configuration;
 using Microsoft.Extensions.Logging;
 using Microsoft.Extensions.Options;
 
+#pragma warning disable CA2234 // Use URIs instead of strings
+
 namespace EpicGames.Horde
 {
 	/// <summary>
