Horde: Add grouping keys and platforms to tools. Intent is for the dashboard to only show one tool for each grouping key by default (the one with the closest matching platform based on the browser's user-agent), but users can manually expand the list of tools to show other platforms if desired.

Platforms should be NET RIDs, eg. win-x64, osx-arm64, etc... (https://learn.microsoft.com/en-us/dotnet/core/rid-catalog).

diff --git a/Engine/Source/Programs/Horde/Horde.Server/Tools/ITool.cs b/Engine/Source/Programs/Horde/Horde.Server/Tools/ITool.cs
--- a/Engine/Source/Programs/Horde/Horde.Server/Tools/ITool.cs
+++ b/Engine/Source/Programs/Horde/Horde.Server/Tools/ITool.cs
@@ -38,6 +38,16 @@ namespace Horde.Server.Tools
 		/// </summary>
 		string? Category { get; }
 
+		/// <summary>
+		/// Grouping key for merging tool versions together on the dashboard
+		/// </summary>
+		string? Group { get; }
+
+		/// <summary>
+		/// Supported platforms, as a NET runtime identifiers
+		/// </summary>
+		IReadOnlyList<string>? Platforms { get; }
+
 		/// <summary>
 		/// Whether the tool is available to authenticated users
 		/// </summary>
diff --git a/Engine/Source/Programs/Horde/Horde.Server/Tools/ToolCollection.cs b/Engine/Source/Programs/Horde/Horde.Server/Tools/ToolCollection.cs
--- a/Engine/Source/Programs/Horde/Horde.Server/Tools/ToolCollection.cs
+++ b/Engine/Source/Programs/Horde/Horde.Server/Tools/ToolCollection.cs
@@ -104,6 +104,8 @@ namespace Horde.Server.Tools
 			public string Name => _config.Name;
 			public string Description => _config.Description;
 			public string? Category => _config.Category;
+			public string? Group => _config.Group;
+			public IReadOnlyList<string>? Platforms => _config.Platforms;
 			public bool Public => _config.Public;
 			public bool ShowInUgs => _config.ShowInUgs;
 			public bool ShowInDashboard => _config.ShowInDashboard;
diff --git a/Engine/Source/Programs/Horde/Horde.Server/Tools/ToolConfig.cs b/Engine/Source/Programs/Horde/Horde.Server/Tools/ToolConfig.cs
--- a/Engine/Source/Programs/Horde/Horde.Server/Tools/ToolConfig.cs
+++ b/Engine/Source/Programs/Horde/Horde.Server/Tools/ToolConfig.cs
@@ -1,6 +1,7 @@
 // Copyright Epic Games, Inc. All Rights Reserved.
 
 using System;
+using System.Collections.Generic;
 using System.ComponentModel.DataAnnotations;
 using System.Diagnostics;
 using System.Security.Claims;
@@ -48,6 +49,16 @@ namespace Horde.Server.Tools
 		/// </summary>
 		public string? Category { get; set; }
 
+		/// <summary>
+		/// Grouping key for different variations of the same tool. The dashboard will show these together.
+		/// </summary>
+		public string? Group { get; set; }
+
+		/// <summary>
+		/// Platforms for this tool. Takes the form of a NET RID (https://learn.microsoft.com/en-us/dotnet/core/rid-catalog).
+		/// </summary>
+		public List<string>? Platforms { get; set; }
+
 		/// <summary>
 		/// Whether this tool should be exposed for download on a public endpoint without authentication
 		/// </summary>
diff --git a/Engine/Source/Programs/Horde/Horde.Server/Tools/ToolsController.cs b/Engine/Source/Programs/Horde/Horde.Server/Tools/ToolsController.cs
--- a/Engine/Source/Programs/Horde/Horde.Server/Tools/ToolsController.cs
+++ b/Engine/Source/Programs/Horde/Horde.Server/Tools/ToolsController.cs
@@ -210,7 +210,7 @@ namespace Horde.Server.Tools
 		static GetToolSummaryResponse CreateGetToolSummaryResponse(ITool tool)
 		{
 			IToolDeployment? deployment = (tool.Deployments.Count == 0) ? null : tool.Deployments[^1];
-			return new GetToolSummaryResponse(tool.Id, tool.Name, tool.Description, tool.Category, deployment?.Version, deployment?.Id, tool.ShowInUgs, tool.ShowInDashboard);
+			return new GetToolSummaryResponse(tool.Id, tool.Name, tool.Description, tool.Category, tool.Group, tool.Platforms?.ToList(), deployment?.Version, deployment?.Id, tool.ShowInUgs, tool.ShowInDashboard);
 		}
 
 		/// <summary>
@@ -254,7 +254,7 @@ namespace Horde.Server.Tools
 
 		static GetToolResponse CreateGetToolResponse(ITool tool, List<GetToolDeploymentResponse> deployments)
 		{
-			return new GetToolResponse(tool.Id, tool.Name, tool.Description, tool.Category, deployments, tool.Public, tool.ShowInUgs, tool.ShowInDashboard);
+			return new GetToolResponse(tool.Id, tool.Name, tool.Description, tool.Category, tool.Group, tool.Platforms?.ToList(), deployments, tool.Public, tool.ShowInUgs, tool.ShowInDashboard);
 		}
 
 		/// <summary>
diff --git a/Engine/Source/Programs/Shared/EpicGames.Horde/Tools/ToolMessages.cs b/Engine/Source/Programs/Shared/EpicGames.Horde/Tools/ToolMessages.cs
--- a/Engine/Source/Programs/Shared/EpicGames.Horde/Tools/ToolMessages.cs
+++ b/Engine/Source/Programs/Shared/EpicGames.Horde/Tools/ToolMessages.cs
@@ -14,11 +14,13 @@ namespace EpicGames.Horde.Tools
 	/// <param name="Name">Name of the tool</param>
 	/// <param name="Description">Description for the tool</param>
 	/// <param name="Category">Category to display the tool in on the dashboard</param>
+	/// <param name="Group">Grouping key to control how different tools should be merged on the dashboard</param>
+	/// <param name="Platforms">List of platforms that this tool supports, as NET runtime identifiers.</param>
 	/// <param name="Deployments">Current deployments of this tool, sorted by time.</param>
 	/// <param name="Public">Whether this tool should be exposed for download on a public endpoint without authentication</param>
 	/// <param name="ShowInUgs">Whether to show this tool for download inside UGS</param>
 	/// <param name="ShowInDashboard">Whether to show this tool for download on the dashboard</param>
-	public record class GetToolResponse(ToolId Id, string Name, string Description, string? Category, List<GetToolDeploymentResponse> Deployments, bool Public, bool ShowInUgs, bool ShowInDashboard);
+	public record class GetToolResponse(ToolId Id, string Name, string Description, string? Category, string? Group, List<string>? Platforms, List<GetToolDeploymentResponse> Deployments, bool Public, bool ShowInUgs, bool ShowInDashboard);
 
 	/// <summary>
 	/// Summary for a particular tool.
@@ -27,11 +29,13 @@ namespace EpicGames.Horde.Tools
 	/// <param name="Name">Name of the tool</param>
 	/// <param name="Description">Description for the tool</param>
 	/// <param name="Category">Category to display the tool in on the dashboard</param>
+	/// <param name="Group">Grouping key to control how different tools should be merged on the dashboard</param>
+	/// <param name="Platforms">List of platforms that this tool supports, as NET runtime identifiers.</param>
 	/// <param name="Version">Version number of the current deployment of this tool</param>
 	/// <param name="DeploymentId">Identifier for the current deployment</param>
 	/// <param name="ShowInUgs">Whether to show this tool for download inside UGS</param>
 	/// <param name="ShowInDashboard">Whether to show this tool for download on the dashboard</param>
-	public record class GetToolSummaryResponse(ToolId Id, string Name, string Description, string? Category, string? Version, ToolDeploymentId? DeploymentId, bool ShowInUgs, bool ShowInDashboard);
+	public record class GetToolSummaryResponse(ToolId Id, string Name, string Description, string? Category, string? Group, List<string>? Platforms, string? Version, ToolDeploymentId? DeploymentId, bool ShowInUgs, bool ShowInDashboard);
 
 	/// <summary>
 	/// Response when querying all tools
