Horde: Use IHordeClient within the agent for all RPC/storage traffic.

#rnx

diff --git a/Engine/Source/Programs/Horde/Horde.Agent.Driver/Execution/JobExecutor.cs b/Engine/Source/Programs/Horde/Horde.Agent.Driver/Execution/JobExecutor.cs
--- a/Engine/Source/Programs/Horde/Horde.Agent.Driver/Execution/JobExecutor.cs
+++ b/Engine/Source/Programs/Horde/Horde.Agent.Driver/Execution/JobExecutor.cs
@@ -12,10 +12,8 @@ using EpicGames.Horde.Artifacts;
 using EpicGames.Horde.Jobs;
 using EpicGames.Horde.Logs;
 using EpicGames.Horde.Storage;
-using EpicGames.Horde.Storage.Clients;
 using EpicGames.Horde.Storage.Nodes;
 using Grpc.Core;
-using Grpc.Net.Client;
 using Horde.Agent.Driver;
 using Horde.Agent.Parser;
 //using Horde.Agent.Services;
@@ -62,28 +60,22 @@ namespace Horde.Agent.Execution
 
 	class JobExecutorOptions
 	{
-		public Uri ServerUrl { get; }
+		public IHordeClient HordeClient { get; }
 		public DirectoryReference WorkingDir { get; }
-		public GrpcChannel GrpcChannel { get; }
 		public IReadOnlyList<ProcessToTerminate>? ProcessesToTerminate { get; }
-		public HttpStorageClientFactory StorageFactory { get; }
 		public JobId JobId { get; }
 		public JobStepBatchId BatchId { get; }
 		public RpcBeginBatchResponse Batch { get; }
-		public string Token { get; }
 		public RpcJobOptions JobOptions { get; }
 
-		public JobExecutorOptions(Uri serverUrl, DirectoryReference workingDir, GrpcChannel grpcChannel, IReadOnlyList<ProcessToTerminate>? processesToTerminate, HttpStorageClientFactory storageFactory, JobId jobId, JobStepBatchId batchId, RpcBeginBatchResponse batch, string token, RpcJobOptions jobOptions)
+		public JobExecutorOptions(IHordeClient hordeClient, DirectoryReference workingDir, IReadOnlyList<ProcessToTerminate>? processesToTerminate, JobId jobId, JobStepBatchId batchId, RpcBeginBatchResponse batch, RpcJobOptions jobOptions)
 		{
-			ServerUrl = serverUrl;
+			HordeClient = hordeClient;
 			WorkingDir = workingDir;
-			GrpcChannel = grpcChannel;
 			ProcessesToTerminate = processesToTerminate;
-			StorageFactory = storageFactory;
 			JobId = jobId;
 			BatchId = batchId;
 			Batch = batch;
-			Token = token;
 			JobOptions = jobOptions;
 		}
 	}
@@ -238,11 +230,10 @@ namespace Horde.Agent.Execution
 		protected bool _compileAutomationTool = true;
 
 		protected JobExecutorOptions Options { get; }
-		protected HttpStorageClientFactory StorageFactory { get; }
 		protected RpcJobOptions JobOptions { get; }
 
-		protected GrpcChannel GrpcChannel => Options.GrpcChannel;
-		protected JobRpc.JobRpcClient JobRpc { get; }
+		protected IHordeClient HordeClient => Options.HordeClient;
+		protected JobRpc.JobRpcClient JobRpc { get; private set; }
 		protected Dictionary<string, string> _remapAgentTypes = new Dictionary<string, string>(StringComparer.OrdinalIgnoreCase);
 
 		protected Dictionary<string, string> _envVars = new Dictionary<string, string>(StringComparer.OrdinalIgnoreCase);
@@ -250,17 +241,13 @@ namespace Horde.Agent.Execution
 		protected JobExecutor(JobExecutorOptions options, ILogger logger)
 		{
 			Options = options;
-			StorageFactory = options.StorageFactory;
 
 			JobId = options.JobId;
 			BatchId = options.BatchId;
 			Batch = options.Batch;
 
 			JobOptions = options.JobOptions;
-			JobRpc = new JobRpc.JobRpcClient(GrpcChannel);
-
-			_envVars[HordeHttpClient.HordeUrlEnvVarName] = options.ServerUrl.ToString();
-			_envVars[HordeHttpClient.HordeTokenEnvVarName] = options.Token;
+			JobRpc = null!; // Set in InitializeAsync()
 
 			Logger = logger;
 		}
@@ -276,8 +263,18 @@ namespace Horde.Agent.Execution
 		{
 		}
 
-		public virtual Task InitializeAsync(ILogger logger, CancellationToken cancellationToken)
+		public virtual async Task InitializeAsync(ILogger logger, CancellationToken cancellationToken)
 		{
+			JobRpc = await HordeClient.CreateGrpcClientAsync<JobRpc.JobRpcClient>(cancellationToken);
+
+			_envVars[HordeHttpClient.HordeUrlEnvVarName] = HordeClient.ServerUrl.ToString();
+
+			string? accessToken = await HordeClient.GetAccessTokenAsync(false, cancellationToken);
+			if (accessToken != null)
+			{
+				_envVars[HordeHttpClient.HordeTokenEnvVarName] = accessToken;
+			}
+
 			// Setup the agent type
 			foreach (KeyValuePair<string, string> envVar in Batch.Environment)
 			{
@@ -341,8 +338,6 @@ namespace Horde.Agent.Execution
 			{
 				_additionalArguments.Add($"-set:PreflightChange={Batch.PreflightChange}");
 			}
-
-			return Task.CompletedTask;
 		}
 
 		public static bool IsUserAdministrator()
@@ -454,7 +449,7 @@ namespace Horde.Agent.Execution
 
 		IStorageClient CreateStorageClient(NamespaceId namespaceId, string? token)
 		{
-			return StorageFactory.CreateClient(namespaceId, token);
+			return HordeClient.CreateStorageClient(namespaceId, token);
 		}
 
 		protected virtual async Task<bool> SetupAsync(JobStepInfo step, DirectoryReference workspaceDir, bool? useP4, ILogger logger, CancellationToken cancellationToken)
@@ -834,7 +829,7 @@ namespace Horde.Agent.Execution
 
 				RpcCreateJobArtifactResponseV2 artifact = await JobRpc.CreateArtifactV2Async(artifactRequest, cancellationToken: cancellationToken);
 				ArtifactId artifactId = ArtifactId.Parse(artifact.Id);
-				Logger.LogInformation("Created artifact {ArtifactId} '{ArtifactName}' ({ArtifactType}) with ref {RefName} ({Link})", artifactId, name, type, artifact.RefName, $"{Options.ServerUrl}/api/v1/storage/{artifact.NamespaceId}/refs/{artifact.RefName}");
+				Logger.LogInformation("Created artifact {ArtifactId} '{ArtifactName}' ({ArtifactType}) with ref {RefName} ({Link})", artifactId, name, type, artifact.RefName, $"{HordeClient.ServerUrl}/api/v1/storage/{artifact.NamespaceId}/refs/{artifact.RefName}");
 
 				using IStorageClient storage = CreateStorageClient(new NamespaceId(artifact.NamespaceId), artifact.Token);
 
@@ -884,7 +879,7 @@ namespace Horde.Agent.Execution
 				string nodeName = input.Substring(0, slashIdx);
 				string tagName = input.Substring(slashIdx + 1);
 
-				TempStorageTagManifest fileList = await TempStorage.RetrieveTagAsync(JobRpc, JobId, step.StepId, StorageFactory, nodeName, tagName, manifestDir, logger, cancellationToken);
+				TempStorageTagManifest fileList = await TempStorage.RetrieveTagAsync(HordeClient, JobId, step.StepId, nodeName, tagName, manifestDir, logger, cancellationToken);
 				tagNameToFileSet[tagName] = fileList.ToFileSet(workspaceDir);
 				inputStorageBlocks.UnionWith(fileList.Blocks);
 			}
@@ -897,7 +892,7 @@ namespace Horde.Agent.Execution
 				scope.Span.SetTag("blocks", inputStorageBlocks.Count);
 				foreach (TempStorageBlockRef inputStorageBlock in inputStorageBlocks)
 				{
-					TempStorageBlockManifest manifest = await TempStorage.RetrieveBlockAsync(JobRpc, JobId, step.StepId, StorageFactory, inputStorageBlock.NodeName, inputStorageBlock.OutputName, workspaceDir, manifestDir, logger, cancellationToken);
+					TempStorageBlockManifest manifest = await TempStorage.RetrieveBlockAsync(HordeClient, JobId, step.StepId, inputStorageBlock.NodeName, inputStorageBlock.OutputName, workspaceDir, manifestDir, logger, cancellationToken);
 					inputManifests[inputStorageBlock] = manifest;
 				}
 				scope.Span.SetTag("size", inputManifests.Sum(x => x.Value.GetTotalSize()));
@@ -1034,7 +1029,7 @@ namespace Horde.Agent.Execution
 
 				RpcCreateJobArtifactResponseV2 artifact = await JobRpc.CreateArtifactV2Async(artifactRequest, cancellationToken: cancellationToken);
 				ArtifactId artifactId = ArtifactId.Parse(artifact.Id);
-				logger.LogInformation("Created artifact {ArtifactId} '{ArtifactName}' ({ArtifactType}) with ref {RefName} ({RefUrl})", artifactId, artifactRequest.Name, ArtifactType.StepOutput, artifact.RefName, $"{Options.ServerUrl.ToString().TrimEnd('/')}/api/v1/storage/{artifact.NamespaceId}/refs/{artifact.RefName}");
+				logger.LogInformation("Created artifact {ArtifactId} '{ArtifactName}' ({ArtifactType}) with ref {RefName} ({RefUrl})", artifactId, artifactRequest.Name, ArtifactType.StepOutput, artifact.RefName, $"{HordeClient.ServerUrl.ToString().TrimEnd('/')}/api/v1/storage/{artifact.NamespaceId}/refs/{artifact.RefName}");
 
 				using IStorageClient storage = CreateStorageClient(new NamespaceId(artifact.NamespaceId), artifact.Token);
 
@@ -1113,7 +1108,7 @@ namespace Horde.Agent.Execution
 
 					RpcCreateJobArtifactResponseV2 artifact = await JobRpc.CreateArtifactV2Async(artifactRequest, cancellationToken: cancellationToken);
 					ArtifactId artifactId = ArtifactId.Parse(artifact.Id);
-					logger.LogInformation("Created artifact {ArtifactId} '{ArtifactName}' ({ArtifactType}) with ref {RefName} ({RefUrl})", artifactId, artifactRequest.Name, ArtifactType.StepOutput, artifact.RefName, $"{Options.ServerUrl.ToString().TrimEnd('/')}/api/v1/storage/{artifact.NamespaceId}/refs/{artifact.RefName}");
+					logger.LogInformation("Created artifact {ArtifactId} '{ArtifactName}' ({ArtifactType}) with ref {RefName} ({RefUrl})", artifactId, artifactRequest.Name, ArtifactType.StepOutput, artifact.RefName, $"{HordeClient.ServerUrl.ToString().TrimEnd('/')}/api/v1/storage/{artifact.NamespaceId}/refs/{artifact.RefName}");
 
 					using IStorageClient storage = CreateStorageClient(new NamespaceId(artifact.NamespaceId), artifact.Token);
 
diff --git a/Engine/Source/Programs/Horde/Horde.Agent.Driver/Utility/JsonRpcLogSink.cs b/Engine/Source/Programs/Horde/Horde.Agent.Driver/Utility/JsonRpcLogSink.cs
--- a/Engine/Source/Programs/Horde/Horde.Agent.Driver/Utility/JsonRpcLogSink.cs
+++ b/Engine/Source/Programs/Horde/Horde.Agent.Driver/Utility/JsonRpcLogSink.cs
@@ -2,12 +2,12 @@
 
 using System.Text;
 using EpicGames.Core;
+using EpicGames.Horde;
 using EpicGames.Horde.Jobs;
 using EpicGames.Horde.Logs;
 using EpicGames.Horde.Storage;
 using Google.Protobuf;
 using Grpc.Core;
-using Grpc.Net.Client;
 using Horde.Common.Rpc;
 using HordeCommon.Rpc;
 using Microsoft.Extensions.Logging;
@@ -26,7 +26,7 @@ namespace Horde.Agent.Utility
 	{
 		const int FlushLength = 1024 * 1024;
 
-		readonly GrpcChannel _connection;
+		readonly IHordeClient _hordeClient;
 		readonly JobId? _jobId;
 		readonly JobStepBatchId? _jobBatchId;
 		readonly JobStepId? _jobStepId;
@@ -43,16 +43,16 @@ namespace Horde.Agent.Utility
 		AsyncEvent _tailTaskStop;
 		readonly AsyncEvent _newTailDataEvent = new AsyncEvent();
 
-		public JsonRpcAndStorageLogSink(GrpcChannel connection, LogId logId, JobId? jobId, JobStepBatchId? jobBatchId, JobStepId? jobStepId, IStorageClient store, ILogger logger)
+		public JsonRpcAndStorageLogSink(IHordeClient hordeClient, LogId logId, JobId? jobId, JobStepBatchId? jobBatchId, JobStepId? jobStepId, ILogger logger)
 		{
-			_connection = connection;
+			_hordeClient = hordeClient;
 			_logId = logId;
 			_jobId = jobId;
 			_jobBatchId = jobBatchId;
 			_jobStepId = jobStepId;
 			_builder = new LogBuilder(LogFormat.Json, logger);
-			_store = store;
-			_writer = store.CreateBlobWriter();
+			_store = _hordeClient.CreateStorageClient(logId);
+			_writer = _store.CreateBlobWriter();
 			_logger = logger;
 
 			_tailTaskStop = new AsyncEvent();
@@ -174,7 +174,7 @@ namespace Horde.Agent.Utility
 			// Update the outcome of this jobstep
 			if (_jobId != null && _jobBatchId != null && _jobStepId != null)
 			{
-				JobRpc.JobRpcClient jobRpc = new JobRpc.JobRpcClient(_connection);
+				JobRpc.JobRpcClient jobRpc = await _hordeClient.CreateGrpcClientAsync<JobRpc.JobRpcClient>(cancellationToken);
 				try
 				{
 					await jobRpc.UpdateStepAsync(new RpcUpdateStepRequest(_jobId.Value, _jobBatchId.Value, _jobStepId.Value, JobStepState.Unspecified, outcome), cancellationToken: cancellationToken);
@@ -189,7 +189,7 @@ namespace Horde.Agent.Utility
 		/// <inheritdoc/>
 		public async Task WriteEventsAsync(List<RpcCreateEventRequest> events, CancellationToken cancellationToken)
 		{
-			JobRpc.JobRpcClient jobRpc = new JobRpc.JobRpcClient(_connection);
+			JobRpc.JobRpcClient jobRpc = await _hordeClient.CreateGrpcClientAsync<JobRpc.JobRpcClient>(cancellationToken);
 			await jobRpc.CreateEventsAsync(new RpcCreateEventsRequest(events), cancellationToken: cancellationToken);
 		}
 
@@ -222,7 +222,7 @@ namespace Horde.Agent.Utility
 			request.TargetLocator = target.GetLocator().ToString();
 			request.Complete = complete;
 
-			LogRpcClient clientRef = new LogRpcClient(_connection);
+			LogRpcClient clientRef = await _hordeClient.CreateGrpcClientAsync<LogRpcClient>(cancellationToken);
 			await clientRef.UpdateLogAsync(request, cancellationToken: cancellationToken);
 		}
 
@@ -231,7 +231,7 @@ namespace Horde.Agent.Utility
 			DateTime deadline = DateTime.UtcNow.AddMinutes(2.0);
 			try
 			{
-				LogRpcClient clientRef = new LogRpcClient(_connection);
+				LogRpcClient clientRef = await _hordeClient.CreateGrpcClientAsync<LogRpcClient>(cancellationToken);
 				using AsyncDuplexStreamingCall<UpdateLogTailRequest, UpdateLogTailResponse> call = clientRef.UpdateLogTail(deadline: deadline, cancellationToken: cancellationToken);
 
 				// Write the request to the server
diff --git a/Engine/Source/Programs/Horde/Horde.Agent.Driver/Utility/TempStorage.cs b/Engine/Source/Programs/Horde/Horde.Agent.Driver/Utility/TempStorage.cs
--- a/Engine/Source/Programs/Horde/Horde.Agent.Driver/Utility/TempStorage.cs
+++ b/Engine/Source/Programs/Horde/Horde.Agent.Driver/Utility/TempStorage.cs
@@ -10,7 +10,6 @@ using EpicGames.Horde;
 using EpicGames.Horde.Artifacts;
 using EpicGames.Horde.Jobs;
 using EpicGames.Horde.Storage;
-using EpicGames.Horde.Storage.Clients;
 using EpicGames.Horde.Storage.Nodes;
 using Horde.Common.Rpc;
 using Microsoft.Extensions.Logging;
@@ -579,17 +578,16 @@ namespace Horde.Storage.Utility
 		/// <summary>
 		/// Reads a set of tagged files from disk
 		/// </summary>
-		/// <param name="jobRpc">The job rpc interface</param>
+		/// <param name="hordeClient">The job rpc interface</param>
 		/// <param name="jobId"></param>
 		/// <param name="stepId"></param>
-		/// <param name="storageClientFactory">Reader for node data</param>
 		/// <param name="nodeName">Name of the node which produced the tag set</param>
 		/// <param name="tagName">Name of the tag, with a '#' prefix</param>
 		/// <param name="manifestDir">The local directory containing manifests</param>
 		/// <param name="logger">Logger for output</param>
 		/// <param name="cancellationToken"></param>
 		/// <returns>The set of files</returns>
-		public static async Task<TempStorageTagManifest> RetrieveTagAsync(JobRpc.JobRpcClient jobRpc, JobId jobId, JobStepId stepId, HttpStorageClientFactory storageClientFactory, string nodeName, string tagName, DirectoryReference manifestDir, ILogger logger, CancellationToken cancellationToken)
+		public static async Task<TempStorageTagManifest> RetrieveTagAsync(IHordeClient hordeClient, JobId jobId, JobStepId stepId, string nodeName, string tagName, DirectoryReference manifestDir, ILogger logger, CancellationToken cancellationToken)
 		{
 			// Try to read the tag set from the local directory
 			FileReference localFileListLocation = GetTagManifestLocation(manifestDir, nodeName, tagName);
@@ -608,6 +606,7 @@ namespace Horde.Storage.Utility
 				artifactRequest.Name = artifactName.ToString();
 				artifactRequest.Type = artifactType.ToString();
 
+				JobRpc.JobRpcClient jobRpc = await hordeClient.CreateGrpcClientAsync<JobRpc.JobRpcClient>(cancellationToken);
 				RpcGetJobArtifactResponse artifact = await jobRpc.GetArtifactAsync(artifactRequest, cancellationToken: cancellationToken);
 
 				NamespaceId namespaceId = new NamespaceId(artifact.NamespaceId);
@@ -615,7 +614,7 @@ namespace Horde.Storage.Utility
 
 				logger.LogInformation("Reading node \"{NodeName}\" tag \"{TagName}\" from temp storage (artifact: {ArtifactId} '{ArtifactName}' ({ArtifactType}), ns: {NamespaceId}, ref: {RefName}, localFile: {LocalFile})", nodeName, tagName, artifact.Id, artifactName, artifactType, namespaceId, refName, localFileListLocation);
 
-				using IStorageClient storageClient = storageClientFactory.CreateClient(namespaceId, artifact.Token);
+				using IStorageClient storageClient = hordeClient.CreateStorageClient(namespaceId, artifact.Token);
 				DirectoryNode node = await storageClient.ReadRefTargetAsync<DirectoryNode>(artifact.RefName, cancellationToken: cancellationToken);
 
 				FileEntry fileEntry = node.GetFileEntry(localFileListLocation.GetFileName());
@@ -708,10 +707,9 @@ namespace Horde.Storage.Utility
 		/// <summary>
 		/// Retrieve an output of the given node. Fetches and decompresses the files from shared storage if necessary, or validates the local files.
 		/// </summary>
-		/// <param name="jobRpc"></param>
+		/// <param name="hordeClient"></param>
 		/// <param name="jobId"></param>
 		/// <param name="stepId"></param>
-		/// <param name="storageClientFactory">Store to read data from</param>
 		/// <param name="nodeName">The node which created the storage block</param>
 		/// <param name="blockName">Name of the block to retrieve.</param>
 		/// <param name="rootDir">Local directory for extracting data to</param>
@@ -719,7 +717,7 @@ namespace Horde.Storage.Utility
 		/// <param name="logger">Logger for output</param>
 		/// <param name="cancellationToken"></param>
 		/// <returns>Manifest of the files retrieved</returns>
-		public static async Task<TempStorageBlockManifest> RetrieveBlockAsync(JobRpc.JobRpcClient jobRpc, JobId jobId, JobStepId stepId, HttpStorageClientFactory storageClientFactory, string nodeName, string blockName, DirectoryReference rootDir, DirectoryReference manifestDir, ILogger logger, CancellationToken cancellationToken)
+		public static async Task<TempStorageBlockManifest> RetrieveBlockAsync(IHordeClient hordeClient, JobId jobId, JobStepId stepId, string nodeName, string blockName, DirectoryReference rootDir, DirectoryReference manifestDir, ILogger logger, CancellationToken cancellationToken)
 		{
 			// Get the path to the local manifest
 			FileReference localManifestFile = GetBlockManifestLocation(manifestDir, nodeName, blockName);
@@ -746,13 +744,15 @@ namespace Horde.Storage.Utility
 				artifactRequest.Name = artifactName.ToString();
 				artifactRequest.Type = artifactType.ToString();
 
+				JobRpc.JobRpcClient jobRpc = await hordeClient.CreateGrpcClientAsync<JobRpc.JobRpcClient>(cancellationToken);
+
 				RpcGetJobArtifactResponse artifact = await jobRpc.GetArtifactAsync(artifactRequest, cancellationToken: cancellationToken);
 				NamespaceId namespaceId = new NamespaceId(artifact.NamespaceId);
 				RefName refName = new RefName(artifact.RefName);
 
 				logger.LogInformation("Reading node \"{NodeName}\" block \"{BlockName}\" from temp storage (artifact: {ArtifactId} '{ArtifactName}' ({ArtifactType}), ns: {NamespaceId}, ref: {RefName}, local: {LocalFile}, blockdir: {BlockDir})", nodeName, blockName, artifact.Id, artifactName, artifactType, namespaceId, refName, localManifestFile, blockDirectoryName);
 
-				using IStorageClient storageClient = storageClientFactory.CreateClient(namespaceId, artifact.Token);
+				using IStorageClient storageClient = hordeClient.CreateStorageClient(namespaceId, artifact.Token);
 				DirectoryNode node = await storageClient.ReadRefTargetAsync<DirectoryNode>(refName, cancellationToken: cancellationToken);
 
 				DirectoryEntry? rootDirEntry;
diff --git a/Engine/Source/Programs/Horde/Horde.Agent.Tests/Execution/WorkspaceExecutorTest.cs b/Engine/Source/Programs/Horde/Horde.Agent.Tests/Execution/WorkspaceExecutorTest.cs
--- a/Engine/Source/Programs/Horde/Horde.Agent.Tests/Execution/WorkspaceExecutorTest.cs
+++ b/Engine/Source/Programs/Horde/Horde.Agent.Tests/Execution/WorkspaceExecutorTest.cs
@@ -47,14 +47,14 @@ public sealed class WorkspaceExecutorTest : IAsyncDisposable
 		_server.AddStream(_streamId, "//Foo/Main");
 		_server.AddAgentType(_streamId, AgentType);
 		_server.AddJob(_jobId, _streamId, 1, 0);
-		_session = FakeServerSessionFactory.CreateSession(_server.GetGrpcChannel());
+		_session = FakeServerSessionFactory.CreateSession(_server.GetHordeClient());
 
 		_autoSdkWorkspace.SetFile(1, "HostWin64/Android/base.h", "base");
 		_workspace.SetFile(1, "main.cpp", "main");
 		_workspace.SetFile(1, "foo/bar/baz.h", "baz");
 
 		RpcBeginBatchResponse batch = new RpcBeginBatchResponse { Change = 1 };
-		JobExecutorOptions executorOptions = new JobExecutorOptions(_session.ServerUrl, _session.WorkingDir, _server.GetGrpcChannel(), null, null!, _jobId, _batchId, batch, null!, new RpcJobOptions());
+		JobExecutorOptions executorOptions = new JobExecutorOptions(_session.HordeClient, _session.WorkingDir, null, _jobId, _batchId, batch, new RpcJobOptions());
 		_executor = new(executorOptions, _workspace, null, NullLogger.Instance);
 	}
 
@@ -80,7 +80,7 @@ public sealed class WorkspaceExecutorTest : IAsyncDisposable
 	public async Task RegularAndAutoSdkWorkspaceAsync()
 	{
 		RpcBeginBatchResponse batch = new RpcBeginBatchResponse { Change = 1 };
-		JobExecutorOptions executorOptions = new JobExecutorOptions(_session.ServerUrl, _session.WorkingDir, _session.GrpcChannel, Array.Empty<ProcessToTerminate>(), null!, _jobId, _batchId, batch, null!, new RpcJobOptions());
+		JobExecutorOptions executorOptions = new JobExecutorOptions(_session.HordeClient, _session.WorkingDir, Array.Empty<ProcessToTerminate>(), _jobId, _batchId, batch, new RpcJobOptions());
 		using WorkspaceExecutor executor = new(executorOptions, _workspace, _autoSdkWorkspace, NullLogger.Instance);
 
 		await executor.InitializeAsync(_logger, CancellationToken.None);
@@ -94,7 +94,7 @@ public sealed class WorkspaceExecutorTest : IAsyncDisposable
 	public async Task EnvVarsAsync()
 	{
 		RpcBeginBatchResponse batch = new RpcBeginBatchResponse { Change = 1, StreamName = "//UE5/Main" };
-		JobExecutorOptions executorOptions = new JobExecutorOptions(_session.ServerUrl, _session.WorkingDir, _session.GrpcChannel, Array.Empty<ProcessToTerminate>(), null!, _jobId, _batchId, batch, null!, new RpcJobOptions());
+		JobExecutorOptions executorOptions = new JobExecutorOptions(_session.HordeClient, _session.WorkingDir, Array.Empty<ProcessToTerminate>(), _jobId, _batchId, batch, new RpcJobOptions());
 		using WorkspaceExecutor executor = new(executorOptions, _workspace, _autoSdkWorkspace, NullLogger.Instance);
 		await executor.InitializeAsync(_logger, CancellationToken.None);
 
@@ -121,7 +121,7 @@ public sealed class WorkspaceExecutorTest : IAsyncDisposable
 		_workspace.SetFile(1000, "New/Feature/Foo.cs", "foo");
 
 		RpcBeginBatchResponse batch = new RpcBeginBatchResponse { Change = 1, PreflightChange = 1000 };
-		JobExecutorOptions executorOptions = new JobExecutorOptions(_session.ServerUrl, _session.WorkingDir, _session.GrpcChannel, Array.Empty<ProcessToTerminate>(), null!, preflightJobId, _batchId, batch, null!, new RpcJobOptions());
+		JobExecutorOptions executorOptions = new JobExecutorOptions(_session.HordeClient, _session.WorkingDir, Array.Empty<ProcessToTerminate>(), preflightJobId, _batchId, batch, new RpcJobOptions());
 		using WorkspaceExecutor executor = new(executorOptions, _workspace, null, NullLogger.Instance);
 
 		await executor.InitializeAsync(_logger, CancellationToken.None);
@@ -138,7 +138,7 @@ public sealed class WorkspaceExecutorTest : IAsyncDisposable
 
 		_server.AddJob(noChangeJobId, _streamId, 0, 0);
 		RpcBeginBatchResponse batch = new RpcBeginBatchResponse { };
-		JobExecutorOptions executorOptions = new JobExecutorOptions(_session.ServerUrl, _session.WorkingDir, _session.GrpcChannel, Array.Empty<ProcessToTerminate>(), null!, noChangeJobId, _batchId, batch, null!, new RpcJobOptions());
+		JobExecutorOptions executorOptions = new JobExecutorOptions(_session.HordeClient, _session.WorkingDir, Array.Empty<ProcessToTerminate>(), noChangeJobId, _batchId, batch, new RpcJobOptions());
 		using WorkspaceExecutor executor = new(executorOptions, _workspace, null, NullLogger.Instance);
 		await Assert.ThrowsExceptionAsync<WorkspaceMaterializationException>(() => executor.InitializeAsync(_logger, CancellationToken.None));
 	}
diff --git a/Engine/Source/Programs/Horde/Horde.Agent.Tests/JsonRpcLoggerTests.cs b/Engine/Source/Programs/Horde/Horde.Agent.Tests/JsonRpcLoggerTests.cs
--- a/Engine/Source/Programs/Horde/Horde.Agent.Tests/JsonRpcLoggerTests.cs
+++ b/Engine/Source/Programs/Horde/Horde.Agent.Tests/JsonRpcLoggerTests.cs
@@ -7,11 +7,12 @@ using System.Text;
 using System.Threading;
 using System.Threading.Tasks;
 using EpicGames.Core;
+using EpicGames.Horde;
 using EpicGames.Horde.Jobs;
 using EpicGames.Horde.Logs;
 using EpicGames.Horde.Storage;
 using EpicGames.Horde.Storage.Bundles;
-using EpicGames.Horde.Storage.Bundles.V1;
+using Grpc.Core;
 using Grpc.Net.Client;
 using Horde.Agent.Utility;
 using Microsoft.Extensions.Logging;
@@ -127,8 +128,8 @@ namespace Horde.Agent.Tests
 		{
 			public IBlobRef? Target { get; private set; }
 
-			public FakeJsonRpcLoggerBackend(GrpcChannel connection, LogId logId, JobId? jobId, JobStepBatchId? batchId, JobStepId? stepId, IStorageClient store, ILogger logger)
-				: base(connection, logId, jobId, batchId, stepId, store, logger)
+			public FakeJsonRpcLoggerBackend(IHordeClient hordeClient, LogId logId, JobId? jobId, JobStepBatchId? batchId, JobStepId? stepId, ILogger logger)
+				: base(hordeClient, logId, jobId, batchId, stepId, logger)
 			{
 			}
 
@@ -145,18 +146,62 @@ namespace Horde.Agent.Tests
 			}
 		}
 
+		class FakeHordeClient : IHordeClient
+		{
+			public Dictionary<string, BundleStorageClient> StorageClients { get; } = new Dictionary<string, BundleStorageClient>();
+
+			public Uri ServerUrl => throw new NotImplementedException();
+
+			public Task<bool> LoginAsync(bool allowLogin, CancellationToken cancellationToken)
+				=> throw new NotImplementedException();
+
+			public HordeHttpClient CreateHttpClient()
+				=> throw new NotImplementedException();
+
+			public IStorageClient CreateStorageClient(string relativePath, string? accessToken = null)
+			{
+				BundleStorageClient? storageClient;
+				if (!StorageClients.TryGetValue(relativePath, out storageClient))
+				{
+					storageClient = BundleStorageClient.CreateInMemory(NullLogger.Instance);
+					StorageClients.Add(relativePath, storageClient);
+				}
+				return storageClient;
+			}
+
+			public ValueTask DisposeAsync()
+			{
+				foreach (BundleStorageClient storageClient in StorageClients.Values)
+				{
+					storageClient.Dispose();
+				}
+				StorageClients.Clear();
+				return default;
+			}
+
+			public Task<string?> GetAccessTokenAsync(bool interactive, CancellationToken cancellationToken = default)
+				=> throw new NotImplementedException();
+
+			public Task<GrpcChannel> CreateGrpcChannelAsync(CancellationToken cancellationToken = default)
+				=> throw new NotImplementedException();
+
+			public Task<TClient> CreateGrpcClientAsync<TClient>(CancellationToken cancellationToken = default) where TClient : ClientBase<TClient>
+				=> throw new NotImplementedException();
+
+			public bool HasValidAccessToken()
+				=> throw new NotImplementedException();
+		}
+
 		[TestMethod]
 		public async Task StorageLoggerTestAsync()
 		{
 			await using BundleCache cache = new BundleCache();
-			using BundleStorageClient store = BundleStorageClient.CreateInMemory(NullLogger.Instance);
-
-			BundleReader reader = new BundleReader(store, cache, NullLogger.Instance);
+			await using FakeHordeClient hordeClient = new FakeHordeClient();
 
 			const int Count = 20000;
 
 			LogNode file;
-			await using (FakeJsonRpcLoggerBackend sink = new FakeJsonRpcLoggerBackend(null!, default, null, null, null, store, NullLogger.Instance))
+			await using (FakeJsonRpcLoggerBackend sink = new FakeJsonRpcLoggerBackend(hordeClient, default, null, null, null, NullLogger.Instance))
 			{
 				await using (ServerLogger logger = new ServerLogger(sink, default, null, LogLevel.Information, NullLogger.Instance, NullLogger.Instance))
 				{
diff --git a/Engine/Source/Programs/Horde/Horde.Agent.Tests/Services/TelemetryServiceTest.cs b/Engine/Source/Programs/Horde/Horde.Agent.Tests/Services/TelemetryServiceTest.cs
--- a/Engine/Source/Programs/Horde/Horde.Agent.Tests/Services/TelemetryServiceTest.cs
+++ b/Engine/Source/Programs/Horde/Horde.Agent.Tests/Services/TelemetryServiceTest.cs
@@ -26,7 +26,7 @@ public sealed class TelemetryServiceTest : IDisposable
 		OptionsWrapper<AgentSettings> settingsOpt = new(settings);
 
 		using WorkerService workerService = new WorkerService(null!, null!, null!, null!);
-		JobHandler jobHandler = new JobHandler(new List<IJobExecutorFactory>(), settingsOpt, new OptionsWrapper<DriverSettings>(new DriverSettings()), null!, null!);
+		JobHandler jobHandler = new JobHandler(new List<IJobExecutorFactory>(), settingsOpt, new OptionsWrapper<DriverSettings>(new DriverSettings()), null!);
 		GrpcService grpcService = new(settingsOpt, NullLogger<GrpcService>.Instance, _loggerFactory);
 
 		_telemetryService = new TelemetryService(workerService, jobHandler, grpcService, null!, settingsOpt, NullLogger<TelemetryService>.Instance);
diff --git a/Engine/Source/Programs/Horde/Horde.Agent.Tests/WorkerServiceTest.cs b/Engine/Source/Programs/Horde/Horde.Agent.Tests/WorkerServiceTest.cs
--- a/Engine/Source/Programs/Horde/Horde.Agent.Tests/WorkerServiceTest.cs
+++ b/Engine/Source/Programs/Horde/Horde.Agent.Tests/WorkerServiceTest.cs
@@ -67,7 +67,7 @@ namespace Horde.Agent.Tests
 
 		class FakeServerLoggerFactory : IServerLoggerFactory
 		{
-			public IServerLogger CreateLogger(ISession session, LogId logId, ILogger localLogger, JobId? jobId, JobStepBatchId? batchId, JobStepId? stepId, bool? warnings = null, LogLevel outputLevel = LogLevel.Information) => new FakeServerLogger();
+			public IServerLogger CreateLogger(IHordeClient hordeClient, LogId logId, ILogger localLogger, JobId? jobId, JobStepBatchId? batchId, JobStepId? stepId, bool? warnings = null, LogLevel outputLevel = LogLevel.Information) => new FakeServerLogger();
 		}
 
 		internal static IJobExecutor NullExecutor = new SimpleTestExecutor(async (step, logger, cancellationToken) =>
@@ -80,7 +80,7 @@ namespace Horde.Agent.Tests
 		{
 			_serviceCollection = new ServiceCollection();
 			_serviceCollection.AddLogging();
-			_serviceCollection.AddHordeHttpClient();
+			_serviceCollection.AddHorde();
 			_serviceCollection.AddSingleton<IServerLoggerFactory, FakeServerLoggerFactory>();
 			_serviceCollection.AddSingleton<BundleCache>();
 			_serviceCollection.AddSingleton<StorageBackendCache>();
@@ -257,7 +257,7 @@ namespace Horde.Agent.Tests
 			cts.CancelAfter(20000);
 
 			await using FakeHordeRpcServer fakeServer = new();
-			await using ISession session = FakeServerSessionFactory.CreateSession(fakeServer.GetGrpcChannel());
+			await using ISession session = FakeServerSessionFactory.CreateSession(fakeServer.GetHordeClient());
 
 			LeaseManager manager = new LeaseManager(session, serviceProvider);
 
@@ -276,16 +276,15 @@ namespace Horde.Agent.Tests
 
 		public Task<ISession> CreateAsync(CancellationToken cancellationToken)
 		{
-			return Task.FromResult(CreateSession(_fakeServer.GetGrpcChannel()));
+			return Task.FromResult(CreateSession(_fakeServer.GetHordeClient()));
 		}
 
-		public static ISession CreateSession(GrpcChannel grpcChannel)
+		public static ISession CreateSession(IHordeClient hordeClient)
 		{
 			Mock<ISession> fakeSession = new Mock<ISession>(MockBehavior.Strict);
-			fakeSession.Setup(x => x.ServerUrl).Returns(new Uri("https://localhost:9999"));
+			fakeSession.Setup(x => x.HordeClient).Returns(hordeClient);
 			fakeSession.Setup(x => x.AgentId).Returns(new EpicGames.Horde.Agents.AgentId("LocalAgent"));
 			fakeSession.Setup(x => x.SessionId).Returns(new EpicGames.Horde.Agents.Sessions.SessionId(default));
-			fakeSession.Setup(x => x.GrpcChannel).Returns(grpcChannel);
 			fakeSession.Setup(x => x.DisposeAsync()).Returns(new ValueTask());
 			fakeSession.Setup(x => x.WorkingDir).Returns(DirectoryReference.Combine(DirectoryReference.GetCurrentDirectory(), Guid.NewGuid().ToString()));
 			return fakeSession.Object;
@@ -309,8 +308,51 @@ namespace Horde.Agent.Tests
 		public readonly TaskCompletionSource<bool> CreateSessionReceived = new();
 		public readonly TaskCompletionSource<bool> UpdateSessionReceived = new();
 
-		private readonly GrpcChannel _grpcChannel;
-//		private readonly FakeJobRpcClient _client;
+		private readonly FakeHordeClient _hordeClient;
+
+		private class FakeHordeClient : IHordeClient
+		{
+			readonly FakeHordeRpcServer _server;
+
+			public Uri ServerUrl => new Uri("http://horde-server");
+
+			public FakeHordeClient(FakeHordeRpcServer server)
+				=> _server = server;
+
+			public Task<bool> LoginAsync(bool allowLogin, CancellationToken cancellationToken)
+				=> throw new NotImplementedException();
+
+			public HordeHttpClient CreateHttpClient()
+				=> throw new NotImplementedException();
+
+			public IStorageClient CreateStorageClient(string relativePath, string? accessToken = null)
+				=> throw new NotImplementedException();
+
+			public ValueTask DisposeAsync()
+				=> default;
+
+			public Task<string?> GetAccessTokenAsync(bool interactive, CancellationToken cancellationToken = default)
+				=> Task.FromResult<string?>(null);
+
+			public Task<GrpcChannel> CreateGrpcChannelAsync(CancellationToken cancellationToken = default)
+				=> throw new NotImplementedException();
+
+			public Task<TClient> CreateGrpcClientAsync<TClient>(CancellationToken cancellationToken = default) where TClient : ClientBase<TClient>
+			{
+				if (typeof(TClient) == typeof(HordeRpc.HordeRpcClient))
+				{
+					return Task.FromResult<TClient>((TClient)(object)new FakeHordeRpcClient(_server));
+				}
+				if (typeof(TClient) == typeof(JobRpc.JobRpcClient))
+				{
+					return Task.FromResult<TClient>((TClient)(object)new FakeJobRpcClient(_server));
+				}
+				throw new NotImplementedException();
+			}
+
+			public bool HasValidAccessToken()
+				=> throw new NotImplementedException();
+		}
 
 		private class FakeHordeRpcClient : HordeRpc.HordeRpcClient
 		{
@@ -361,9 +403,7 @@ namespace Horde.Agent.Tests
 		{
 			_serverName = "FakeServer";
 			_logger = NullLogger<FakeHordeRpcServer>.Instance;
-//			FakeHordeRpcClient hordeClient = new FakeHordeRpcClient(this);
-//			_client = new FakeJobRpcClient(this);
-			_grpcChannel = GrpcChannel.ForAddress(new Uri("http://horde-agent-test"), new GrpcChannelOptions());
+			_hordeClient = new FakeHordeClient(this);
 		}
 
 		public void AddTestLease(string leaseId)
@@ -431,9 +471,9 @@ namespace Horde.Agent.Tests
 			};
 		}
 
-		public GrpcChannel GetGrpcChannel()
+		public IHordeClient GetHordeClient()
 		{
-			return _grpcChannel;
+			return _hordeClient;
 		}
 
 		public RpcCreateSessionResponse OnCreateSessionRequest(RpcCreateSessionRequest request)
@@ -519,9 +559,9 @@ namespace Horde.Agent.Tests
 				() => { });
 		}
 
-		public ValueTask DisposeAsync()
+		public async ValueTask DisposeAsync()
 		{
-			_grpcChannel.Dispose();
+			await _hordeClient.DisposeAsync();
 
 			foreach (RpcGetStreamResponse stream in _streamIdToStreamResponse.Values)
 			{
@@ -533,8 +573,6 @@ namespace Horde.Agent.Tests
 					}
 				}
 			}
-
-			return default;
 		}
 	}
 
diff --git a/Engine/Source/Programs/Horde/Horde.Agent/AgentApp.cs b/Engine/Source/Programs/Horde/Horde.Agent/AgentApp.cs
--- a/Engine/Source/Programs/Horde/Horde.Agent/AgentApp.cs
+++ b/Engine/Source/Programs/Horde/Horde.Agent/AgentApp.cs
@@ -11,10 +11,6 @@ using Datadog.Trace.Configuration;
 using Datadog.Trace.OpenTracing;
 using EpicGames.Core;
 using EpicGames.Horde;
-using EpicGames.Horde.Storage;
-using EpicGames.Horde.Storage.Backends;
-using EpicGames.Horde.Storage.Bundles;
-using EpicGames.Horde.Storage.Clients;
 using Horde.Agent.Driver;
 using Horde.Agent.Execution;
 using Horde.Agent.Leases;
@@ -24,7 +20,6 @@ using Horde.Agent.Utility;
 using Microsoft.Extensions.Configuration;
 using Microsoft.Extensions.DependencyInjection;
 using Microsoft.Extensions.Logging;
-using Microsoft.Extensions.Options;
 using Microsoft.Win32;
 using OpenTracing.Util;
 using Polly;
@@ -177,35 +172,15 @@ namespace Horde.Agent
 
 			Logging.SetEnv(serverProfile.Environment);
 
-			ILogger certificateLogger = loggerFactory.CreateLogger(typeof(CertificateHelper).FullName!);
-			services.AddHttpClient(AgentApp.HordeServerClientName, config =>
+			services.AddHorde(options =>
 			{
-				config.BaseAddress = serverProfile.Url;
-				config.DefaultRequestHeaders.Add("Accept", "application/json");
-				config.Timeout = TimeSpan.FromSeconds(300); // Need to make sure this doesn't cancel any long running gRPC streaming calls (eg. session update)
-			})
-			.ConfigurePrimaryHttpMessageHandler(() =>
-			{
-				HttpClientHandler handler = new HttpClientHandler();
-#pragma warning disable MA0039
-				handler.ServerCertificateCustomValidationCallback += (sender, cert, chain, errors) => CertificateHelper.CertificateValidationCallBack(certificateLogger, sender, cert, chain, errors, serverProfile);
-#pragma warning restore MA0039
-				return handler;
-			})
-			.AddTransientHttpErrorPolicy(builder =>
-			{
-				return builder.WaitAndRetryAsync(new[] { TimeSpan.FromSeconds(1), TimeSpan.FromSeconds(5), TimeSpan.FromSeconds(10) });
+				options.ServerUrl = serverProfile.Url;
+				options.AccessToken = serverProfile.Token;
+				options.AllowAuthPrompt = false;
+				options.BackendCache.CacheDir = DirectoryReference.Combine(settings.WorkingDir, "Saved", "Bundles").FullName;
+				options.BackendCache.MaxSize = settings.BundleCacheSize * 1024 * 1024;
 			});
 
-			services.Configure<HordeOptions>(options => options.AllowAuthPrompt = false);
-
-			services.AddHordeHttpClient(client => client.BaseAddress = serverProfile.Url)
-				.ConfigurePrimaryHttpMessageHandler(() => new SocketsHttpHandler
-				{
-					MaxConnectionsPerServer = 16,
-					PooledConnectionIdleTimeout = TimeSpan.FromMinutes(15),
-				});
-
 			services.AddHttpClient(AwsInstanceLifecycleService.HttpClientName)
 				.AddTransientHttpErrorPolicy(builder =>
 				{
@@ -246,11 +221,6 @@ namespace Horde.Agent
 			services.AddSingleton<LeaseLoggerFactory>();
 			services.AddHostedService(sp => sp.GetRequiredService<WorkerService>());
 
-			services.AddSingleton<BundleCache>();
-			services.AddSingleton<StorageBackendCache>(CreateStorageBackendCache);
-			services.AddSingleton<HttpStorageBackendFactory>();
-			services.AddSingleton<HttpStorageClientFactory>();
-
 			services.AddSingleton<ComputeListenerService>();
 			services.AddHostedService(sp => sp.GetRequiredService<ComputeListenerService>());
 
@@ -340,13 +310,6 @@ namespace Horde.Agent
 			}
 		}
 
-		static StorageBackendCache CreateStorageBackendCache(IServiceProvider serviceProvider)
-		{
-			AgentSettings settings = serviceProvider.GetRequiredService<IOptions<AgentSettings>>().Value;
-			DirectoryReference cacheDir = DirectoryReference.Combine(settings.WorkingDir, "Saved", "Bundles");
-			return new StorageBackendCache(cacheDir, settings.BundleCacheSize * 1024 * 1024, serviceProvider.GetRequiredService<ILogger<StorageBackendCache>>());
-		}
-
 		static void ConfigureTracing(string environment, string version)
 		{
 			TracerSettings settings = TracerSettings.FromDefaultSources();
diff --git a/Engine/Source/Programs/Horde/Horde.Agent/Commands/Execution/ExecuteJobCommand.cs b/Engine/Source/Programs/Horde/Horde.Agent/Commands/Execution/ExecuteJobCommand.cs
--- a/Engine/Source/Programs/Horde/Horde.Agent/Commands/Execution/ExecuteJobCommand.cs
+++ b/Engine/Source/Programs/Horde/Horde.Agent/Commands/Execution/ExecuteJobCommand.cs
@@ -1,15 +1,14 @@
 // Copyright Epic Games, Inc. All Rights Reserved.
 
 using EpicGames.Core;
+using EpicGames.Horde;
 using EpicGames.Horde.Agents;
 using EpicGames.Horde.Agents.Leases;
 using EpicGames.Horde.Agents.Sessions;
-using Grpc.Net.Client;
 using Horde.Agent.Leases.Handlers;
 using Horde.Agent.Services;
 using HordeCommon.Rpc.Tasks;
 using Microsoft.Extensions.Logging;
-using Microsoft.Extensions.Options;
 
 namespace Horde.Agent.Commands.Execution
 {
@@ -31,24 +30,21 @@ namespace Horde.Agent.Commands.Execution
 		[CommandLine("-WorkingDir=", Required = true)]
 		public DirectoryReference WorkingDir { get; set; } = null!;
 
-		readonly GrpcService _grpcService;
+		readonly IHordeClientFactory _hordeClientFactory;
 		readonly JobHandler _jobHandler;
-		readonly AgentSettings _agentSettings;
 
-		public ExecuteJobCommand(GrpcService grpcService, JobHandler jobHandler, IOptions<AgentSettings> agentSettings)
+		public ExecuteJobCommand(IHordeClientFactory hordeClientFactory, JobHandler jobHandler)
 		{
-			_grpcService = grpcService;
+			_hordeClientFactory = hordeClientFactory;
 			_jobHandler = jobHandler;
-			_agentSettings = agentSettings.Value;
 		}
 
 		public override async Task<int> ExecuteAsync(ILogger logger)
 		{
-			ServerProfile serverProfile = _agentSettings.GetCurrentServerProfile();
 			ExecuteJobTask executeTask = ExecuteJobTask.Parser.ParseFrom(Convert.FromBase64String(Task));
 
-			using GrpcChannel rpcConnection = await _grpcService.CreateGrpcChannelAsync(executeTask.Token, CancellationToken.None);
-			await using Session session = new Session(serverProfile.Url, AgentId, SessionId, executeTask.Token, rpcConnection, WorkingDir);
+			IHordeClient hordeClient = _hordeClientFactory.Create(executeTask.Token);
+			await using Session session = new Session(AgentId, SessionId, WorkingDir, hordeClient);
 
 			await _jobHandler.ExecuteInternalAsync(session, LeaseId, executeTask, logger, CancellationToken.None);
 			return 0;
diff --git a/Engine/Source/Programs/Horde/Horde.Agent/Leases/Handlers/ComputeHandler.cs b/Engine/Source/Programs/Horde/Horde.Agent/Leases/Handlers/ComputeHandler.cs
--- a/Engine/Source/Programs/Horde/Horde.Agent/Leases/Handlers/ComputeHandler.cs
+++ b/Engine/Source/Programs/Horde/Horde.Agent/Leases/Handlers/ComputeHandler.cs
@@ -36,7 +36,7 @@ namespace Horde.Agent.Leases.Handlers
 		/// <inheritdoc/>
 		public override async Task<LeaseResult> ExecuteAsync(ISession session, LeaseId leaseId, ComputeTask computeTask, ILogger localLogger, CancellationToken cancellationToken)
 		{
-			await using IServerLogger? serverLogger = (computeTask.LogId != null) ? _serverLoggerFactory.CreateLogger(session, LogId.Parse(computeTask.LogId), localLogger, null, LogLevel.Trace) : null;
+			await using IServerLogger? serverLogger = (computeTask.LogId != null) ? _serverLoggerFactory.CreateLogger(session.HordeClient, LogId.Parse(computeTask.LogId), localLogger, null, LogLevel.Trace) : null;
 			ILogger logger = serverLogger ?? localLogger;
 
 			if (!String.IsNullOrEmpty(computeTask.ParentLeaseId))
diff --git a/Engine/Source/Programs/Horde/Horde.Agent/Leases/Handlers/ConformHandler.cs b/Engine/Source/Programs/Horde/Horde.Agent/Leases/Handlers/ConformHandler.cs
--- a/Engine/Source/Programs/Horde/Horde.Agent/Leases/Handlers/ConformHandler.cs
+++ b/Engine/Source/Programs/Horde/Horde.Agent/Leases/Handlers/ConformHandler.cs
@@ -31,7 +31,7 @@ namespace Horde.Agent.Leases.Handlers
 		/// <inheritdoc/>
 		public override async Task<LeaseResult> ExecuteAsync(ISession session, LeaseId leaseId, ConformTask conformTask, ILogger localLogger, CancellationToken cancellationToken)
 		{
-			await using IServerLogger serverLogger = _serverLoggerFactory.CreateLogger(session, LogId.Parse(conformTask.LogId), localLogger, null);
+			await using IServerLogger serverLogger = _serverLoggerFactory.CreateLogger(session.HordeClient, LogId.Parse(conformTask.LogId), localLogger, null);
 			try
 			{
 				LeaseResult result = await ExecuteInternalAsync(session, leaseId, conformTask, serverLogger, cancellationToken);
@@ -78,7 +78,7 @@ namespace Horde.Agent.Leases.Handlers
 				request.Workspaces.AddRange(pendingWorkspaces);
 				request.RemoveUntrackedFiles = removeUntrackedFiles;
 
-				HordeRpc.HordeRpcClient hordeRpc = new HordeRpc.HordeRpcClient(session.GrpcChannel);
+				HordeRpc.HordeRpcClient hordeRpc = await session.HordeClient.CreateGrpcClientAsync<HordeRpc.HordeRpcClient>(cancellationToken);
 
 				RpcUpdateAgentWorkspacesResponse response = await hordeRpc.UpdateAgentWorkspacesAsync(request, cancellationToken: cancellationToken);
 				if (!response.Retry)
diff --git a/Engine/Source/Programs/Horde/Horde.Agent/Leases/Handlers/JobHandler.cs b/Engine/Source/Programs/Horde/Horde.Agent/Leases/Handlers/JobHandler.cs
--- a/Engine/Source/Programs/Horde/Horde.Agent/Leases/Handlers/JobHandler.cs
+++ b/Engine/Source/Programs/Horde/Horde.Agent/Leases/Handlers/JobHandler.cs
@@ -7,7 +7,6 @@ using EpicGames.Core;
 using EpicGames.Horde.Agents.Leases;
 using EpicGames.Horde.Jobs;
 using EpicGames.Horde.Logs;
-using EpicGames.Horde.Storage.Clients;
 using Google.Protobuf;
 using Grpc.Core;
 using Horde.Agent.Driver;
@@ -34,7 +33,7 @@ namespace Horde.Agent.Leases.Handlers
 		/// Exposed as internal to ease testing.
 		/// </summary>
 		internal TimeSpan _stepAbortPollInterval = TimeSpan.FromSeconds(5);
-
+		
 		/// <summary>
 		/// How long to wait before retrying a failed step abort check request
 		/// </summary>
@@ -58,18 +57,16 @@ namespace Horde.Agent.Leases.Handlers
 		readonly IEnumerable<IJobExecutorFactory> _executorFactories;
 		readonly AgentSettings _agentSettings;
 		readonly DriverSettings _driverSettings;
-		readonly HttpStorageClientFactory _serverStorageFactory;
 		readonly IServerLoggerFactory _serverLoggerFactory;
 
 		/// <summary>
 		/// Constructor
 		/// </summary>
-		public JobHandler(IEnumerable<IJobExecutorFactory> executorFactories, IOptions<AgentSettings> agentSettings, IOptions<DriverSettings> driverSettings, HttpStorageClientFactory storageClientFactory, IServerLoggerFactory serverLoggerFactory)
+		public JobHandler(IEnumerable<IJobExecutorFactory> executorFactories, IOptions<AgentSettings> agentSettings, IOptions<DriverSettings> driverSettings, IServerLoggerFactory serverLoggerFactory)
 		{
 			_executorFactories = executorFactories;
 			_agentSettings = agentSettings.Value;
 			_driverSettings = driverSettings.Value;
-			_serverStorageFactory = storageClientFactory;
 			_serverLoggerFactory = serverLoggerFactory;
 		}
 
@@ -157,11 +154,11 @@ namespace Horde.Agent.Leases.Handlers
 
 		internal async Task<LeaseResult> ExecuteInternalAsync(ISession session, LeaseId leaseId, JobTaskInfo executeTask, ILogger localLogger, CancellationToken cancellationToken)
 		{
-			JobRpc.JobRpcClient jobRpc = new JobRpc.JobRpcClient(session.GrpcChannel);
+			JobRpc.JobRpcClient jobRpc = await session.HordeClient.CreateGrpcClientAsync<JobRpc.JobRpcClient>(cancellationToken);
 
 			// Create a storage client for this session
 			RpcJobOptions jobOptions = executeTask.JobOptions;
-			await using IServerLogger logger = _serverLoggerFactory.CreateLogger(session, executeTask.LogId, localLogger, executeTask.JobId, executeTask.BatchId, null, null);
+			await using IServerLogger logger = _serverLoggerFactory.CreateLogger(session.HordeClient, executeTask.LogId, localLogger, executeTask.JobId, executeTask.BatchId, null, null);
 
 			logger.LogInformation("Executing job \"{JobName}\", jobId {JobId}, batchId {BatchId}, leaseId {LeaseId}, agentVersion {AgentVersion}", executeTask.JobName, executeTask.JobId, executeTask.BatchId, leaseId, AgentApp.Version);
 
@@ -175,14 +172,14 @@ namespace Horde.Agent.Leases.Handlers
 			RpcBeginBatchResponse batch = await jobRpc.BeginBatchAsync(new RpcBeginBatchRequest(executeTask.JobId, executeTask.BatchId, leaseId), cancellationToken: cancellationToken);
 			try
 			{
-				JobExecutorOptions options = new JobExecutorOptions(session.ServerUrl, session.WorkingDir, session.GrpcChannel, _driverSettings.ProcessesToTerminate, _serverStorageFactory, executeTask.JobId, executeTask.BatchId, batch, executeTask.Token, jobOptions);
+				JobExecutorOptions options = new JobExecutorOptions(session.HordeClient, session.WorkingDir, _driverSettings.ProcessesToTerminate, executeTask.JobId, executeTask.BatchId, batch, jobOptions);
 				await ExecuteBatchAsync(session, leaseId, executeTask.Workspace, executeTask.AutoSdkWorkspace, options, logger, localLogger, cancellationToken);
 			}
 			catch (Exception ex)
 			{
 				if (cancellationToken.IsCancellationRequested && ex.IsCancellationException())
 				{
-					if (session.GrpcChannel.State == ConnectivityState.TransientFailure)
+					if (!session.HordeClient.HasValidAccessToken())
 					{
 						logger.LogError(ex, "Connection to the server was lost; step aborted.");
 					}
@@ -217,7 +214,7 @@ namespace Horde.Agent.Leases.Handlers
 		/// </summary>
 		async Task ExecuteBatchAsync(ISession session, LeaseId leaseId, RpcAgentWorkspace workspaceInfo, RpcAgentWorkspace? autoSdkWorkspaceInfo, JobExecutorOptions options, ILogger logger, ILogger localLogger, CancellationToken cancellationToken)
 		{
-			JobRpc.JobRpcClient rpcClient = new JobRpc.JobRpcClient(session.GrpcChannel);
+			JobRpc.JobRpcClient rpcClient = await session.HordeClient.CreateGrpcClientAsync<JobRpc.JobRpcClient>(cancellationToken);
 
 			// Create an executor for this job
 			string executorName = String.IsNullOrEmpty(options.JobOptions.Executor) ? _driverSettings.Executor : options.JobOptions.Executor;
@@ -313,7 +310,7 @@ namespace Horde.Agent.Leases.Handlers
 					{
 						// Start writing to the log file
 #pragma warning disable CA2000 // Dispose objects before losing scope
-						await using (IServerLogger stepLogger = _serverLoggerFactory.CreateLogger(session, step.LogId, localLogger, options.JobId, options.BatchId, step.StepId, step.Warnings))
+						await using (IServerLogger stepLogger = _serverLoggerFactory.CreateLogger(session.HordeClient, step.LogId, localLogger, options.JobId, options.BatchId, step.StepId, step.Warnings))
 						{
 							// Execute the task
 							using CancellationTokenSource stepPollCancelSource = new CancellationTokenSource();
@@ -361,7 +358,7 @@ namespace Horde.Agent.Leases.Handlers
 			{
 				if (cancellationToken.IsCancellationRequested && ex.IsCancellationException())
 				{
-					if (session.GrpcChannel.State == ConnectivityState.TransientFailure)
+					if (!session.HordeClient.HasValidAccessToken())
 					{
 						logger.LogError(ex, "Exception while executing batch: {Ex}", ex);
 					}
diff --git a/Engine/Source/Programs/Horde/Horde.Agent/Leases/Handlers/UpgradeHandler.cs b/Engine/Source/Programs/Horde/Horde.Agent/Leases/Handlers/UpgradeHandler.cs
--- a/Engine/Source/Programs/Horde/Horde.Agent/Leases/Handlers/UpgradeHandler.cs
+++ b/Engine/Source/Programs/Horde/Horde.Agent/Leases/Handlers/UpgradeHandler.cs
@@ -28,7 +28,7 @@ namespace Horde.Agent.Leases.Handlers
 		/// <inheritdoc/>
 		public override async Task<LeaseResult> ExecuteAsync(ISession session, LeaseId leaseId, UpgradeTask task, ILogger localLogger, CancellationToken cancellationToken)
 		{
-			await using IServerLogger logger = _serverLoggerFactory.CreateLogger(session, LogId.Parse(task.LogId), localLogger, null);
+			await using IServerLogger logger = _serverLoggerFactory.CreateLogger(session.HordeClient, LogId.Parse(task.LogId), localLogger, null);
 
 			string requiredVersion = task.SoftwareId;
 
@@ -45,7 +45,7 @@ namespace Horde.Agent.Leases.Handlers
 				// Download the new software
 				FileInfo outputFile = new FileInfo(Path.Combine(upgradeDir.FullName, "Agent.zip"));
 
-				HordeRpc.HordeRpcClient hordeRpc = new HordeRpc.HordeRpcClient(session.GrpcChannel);
+				HordeRpc.HordeRpcClient hordeRpc = await session.HordeClient.CreateGrpcClientAsync<HordeRpc.HordeRpcClient>(cancellationToken);
 				using (AsyncServerStreamingCall<RpcDownloadSoftwareResponse> cursor = hordeRpc.DownloadSoftware(new RpcDownloadSoftwareRequest(requiredVersion), null, null, cancellationToken))
 				{
 					await using Stream outputStream = outputFile.Open(FileMode.Create);
diff --git a/Engine/Source/Programs/Horde/Horde.Agent/Leases/LeaseManager.cs b/Engine/Source/Programs/Horde/Horde.Agent/Leases/LeaseManager.cs
--- a/Engine/Source/Programs/Horde/Horde.Agent/Leases/LeaseManager.cs
+++ b/Engine/Source/Programs/Horde/Horde.Agent/Leases/LeaseManager.cs
@@ -6,7 +6,6 @@ using EpicGames.Horde.Agents.Leases;
 using EpicGames.Perforce.Managed;
 using Google.Protobuf.WellKnownTypes;
 using Grpc.Core;
-using Grpc.Net.Client;
 using Horde.Agent.Services;
 using HordeCommon.Rpc;
 using HordeCommon.Rpc.Messages;
@@ -229,7 +228,7 @@ namespace Horde.Agent.Leases
 
 		async Task<SessionResult> HandleSessionAsync(bool shutdownAfterFinishedLease, CancellationToken stoppingToken)
 		{
-			HordeRpc.HordeRpcClient hordeRpc = new HordeRpc.HordeRpcClient(_session.GrpcChannel);
+			HordeRpc.HordeRpcClient hordeRpc = await _session.HordeClient.CreateGrpcClientAsync<HordeRpc.HordeRpcClient>(stoppingToken);
 
 			// Terminate any remaining child processes from other instances
 			ProcessUtils.TerminateProcesses(x => x.IsUnderDirectory(_session.WorkingDir), _logger, stoppingToken);
@@ -308,7 +307,7 @@ namespace Horde.Agent.Leases
 				using (stopping ? (CancellationTokenRegistration?)null : stoppingToken.Register(() => _updateLeasesEvent.Set()))
 				{
 					// Update the state with the server
-					RpcUpdateSessionResponse? updateSessionResponse = await UpdateSessionAsync(_session.GrpcChannel, updateSessionRequest, waitTask);
+					RpcUpdateSessionResponse? updateSessionResponse = await UpdateSessionAsync(hordeRpc, updateSessionRequest, waitTask);
 
 					lock (_lockObject)
 					{
@@ -369,7 +368,7 @@ namespace Horde.Agent.Leases
 					}
 
 					// Update the current status
-					if (_session.GrpcChannel.State == ConnectivityState.TransientFailure)
+					if (!_session.HordeClient.HasValidAccessToken())
 					{
 						_statusService.Set(false, _activeLeases.Count, "Attempting to connect to server...");
 					}
@@ -413,14 +412,12 @@ namespace Horde.Agent.Leases
 		/// <summary>
 		/// Wrapper for <see cref="UpdateSessionInternalAsync"/> which filters/logs exceptions
 		/// </summary>
-		/// <param name="grpcChannel">The RPC client connection</param>
+		/// <param name="hordeRpc">The RPC client connection</param>
 		/// <param name="updateSessionRequest">The session update request</param>
 		/// <param name="waitTask">Task which can be used to jump out of the update early</param>
 		/// <returns>Response from the call</returns>
-		async Task<RpcUpdateSessionResponse?> UpdateSessionAsync(GrpcChannel grpcChannel, RpcUpdateSessionRequest updateSessionRequest, Task waitTask)
+		async Task<RpcUpdateSessionResponse?> UpdateSessionAsync(HordeRpc.HordeRpcClient hordeRpc, RpcUpdateSessionRequest updateSessionRequest, Task waitTask)
 		{
-			HordeRpc.HordeRpcClient hordeRpc = new HordeRpc.HordeRpcClient(grpcChannel);
-
 			RpcUpdateSessionResponse? updateSessionResponse = null;
 			try
 			{
diff --git a/Engine/Source/Programs/Horde/Horde.Agent/Services/ServerLoggerFactory.cs b/Engine/Source/Programs/Horde/Horde.Agent/Services/ServerLoggerFactory.cs
--- a/Engine/Source/Programs/Horde/Horde.Agent/Services/ServerLoggerFactory.cs
+++ b/Engine/Source/Programs/Horde/Horde.Agent/Services/ServerLoggerFactory.cs
@@ -1,9 +1,8 @@
 // Copyright Epic Games, Inc. All Rights Reserved.
 
+using EpicGames.Horde;
 using EpicGames.Horde.Jobs;
 using EpicGames.Horde.Logs;
-using EpicGames.Horde.Storage;
-using EpicGames.Horde.Storage.Clients;
 using Horde.Agent.Utility;
 using Microsoft.Extensions.Logging;
 
@@ -17,7 +16,7 @@ namespace Horde.Agent.Services
 		/// <summary>
 		/// Creates a logger which uploads data to the server
 		/// </summary>
-		/// <param name="session">The current session</param>
+		/// <param name="hordeClient">The Horde client</param>
 		/// <param name="logId">The log id</param>
 		/// <param name="localLogger">Local log output device</param>
 		/// <param name="jobId">Job containing the log</param>
@@ -26,7 +25,7 @@ namespace Horde.Agent.Services
 		/// <param name="warnings">Whether to suppress warnings</param>
 		/// <param name="outputLevel">Minimum output level for messages</param>
 		/// <returns>New logger instance</returns>
-		IServerLogger CreateLogger(ISession session, LogId logId, ILogger localLogger, JobId? jobId, JobStepBatchId? batchId, JobStepId? stepId, bool? warnings, LogLevel outputLevel = LogLevel.Information);
+		IServerLogger CreateLogger(IHordeClient hordeClient, LogId logId, ILogger localLogger, JobId? jobId, JobStepBatchId? batchId, JobStepId? stepId, bool? warnings, LogLevel outputLevel = LogLevel.Information);
 	}
 
 	/// <summary>
@@ -38,15 +37,15 @@ namespace Horde.Agent.Services
 		/// Creates a logger which uploads data to the server
 		/// </summary>
 		/// <param name="service">Service instance</param>
-		/// <param name="session">The current session</param>
+		/// <param name="hordeClient">The current session</param>
 		/// <param name="logId">The log identifier</param>
 		/// <param name="localLogger">Local log output device</param>
 		/// <param name="warnings">Whether to suppress warnings</param>
 		/// <param name="outputLevel">Minimum output level for messages</param>
 		/// <returns>New logger instance</returns>
-		public static IServerLogger CreateLogger(this IServerLoggerFactory service, ISession session, LogId logId, ILogger localLogger, bool? warnings, LogLevel outputLevel = LogLevel.Information)
+		public static IServerLogger CreateLogger(this IServerLoggerFactory service, IHordeClient hordeClient, LogId logId, ILogger localLogger, bool? warnings, LogLevel outputLevel = LogLevel.Information)
 		{
-			return service.CreateLogger(session, logId, localLogger, null, null, null, warnings, outputLevel);
+			return service.CreateLogger(hordeClient, logId, localLogger, null, null, null, warnings, outputLevel);
 		}
 	}
 
@@ -55,25 +54,21 @@ namespace Horde.Agent.Services
 	/// </summary>
 	class ServerLoggerFactory : IServerLoggerFactory
 	{
-		readonly HttpStorageClientFactory _storageClientFactory;
 		readonly ILogger _logger;
 
 		/// <summary>
 		/// Constructor
 		/// </summary>
-		public ServerLoggerFactory(HttpStorageClientFactory storageClientFactory, ILogger<ServerLoggerFactory> logger)
+		public ServerLoggerFactory(ILogger<ServerLoggerFactory> logger)
 		{
-			_storageClientFactory = storageClientFactory;
 			_logger = logger;
 		}
 
 		/// <inheritdoc/>
-		public IServerLogger CreateLogger(ISession session, LogId logId, ILogger localLogger, JobId? jobId, JobStepBatchId? batchId, JobStepId? stepId, bool? warnings, LogLevel outputLevel)
+		public IServerLogger CreateLogger(IHordeClient hordeClient, LogId logId, ILogger localLogger, JobId? jobId, JobStepBatchId? batchId, JobStepId? stepId, bool? warnings, LogLevel outputLevel)
 		{
 #pragma warning disable CA2000 // Dispose objects before losing scope
-			IStorageClient storageClient = _storageClientFactory.CreateClientWithPath($"api/v1/logs/{logId}", session.Token);
-			IJsonRpcLogSink sink = new JsonRpcAndStorageLogSink(session.GrpcChannel, logId, jobId, batchId, stepId, storageClient, _logger);
-
+			IJsonRpcLogSink sink = new JsonRpcAndStorageLogSink(hordeClient, logId, jobId, batchId, stepId, _logger);
 			return new ServerLogger(sink, logId, warnings, outputLevel, localLogger, _logger);
 #pragma warning restore CA2000 // Dispose objects before losing scope
 		}
diff --git a/Engine/Source/Programs/Horde/Horde.Agent/Services/SessionFactory.cs b/Engine/Source/Programs/Horde/Horde.Agent/Services/SessionFactory.cs
--- a/Engine/Source/Programs/Horde/Horde.Agent/Services/SessionFactory.cs
+++ b/Engine/Source/Programs/Horde/Horde.Agent/Services/SessionFactory.cs
@@ -4,6 +4,7 @@ using System.Runtime.InteropServices;
 using System.Security.Cryptography;
 using System.Text.Json;
 using EpicGames.Core;
+using EpicGames.Horde;
 using EpicGames.Horde.Agents;
 using EpicGames.Horde.Agents.Sessions;
 using Grpc.Core;
@@ -22,11 +23,6 @@ namespace Horde.Agent.Services
 	/// </summary>
 	interface ISession : IAsyncDisposable
 	{
-		/// <summary>
-		/// URL of the server
-		/// </summary>
-		Uri ServerUrl { get; }
-
 		/// <summary>
 		/// The agent identifier
 		/// </summary>
@@ -38,19 +34,14 @@ namespace Horde.Agent.Services
 		SessionId SessionId { get; }
 
 		/// <summary>
-		/// Token to use for connection to the server
-		/// </summary>
-		string Token { get; }
-
-		/// <summary>
-		/// Connection to the server
+		/// Working directory for sandboxes etc..
 		/// </summary>
-		GrpcChannel GrpcChannel { get; }
+		DirectoryReference WorkingDir { get; }
 
 		/// <summary>
-		/// Working directory for sandboxes etc..
+		/// Horde client instance
 		/// </summary>
-		DirectoryReference WorkingDir { get; }
+		IHordeClient HordeClient { get; }
 	}
 
 	/// <summary>
@@ -71,9 +62,6 @@ namespace Horde.Agent.Services
 	/// </summary>
 	sealed class Session : ISession
 	{
-		/// <inheritdoc/>
-		public Uri ServerUrl { get; }
-
 		/// <inheritdoc/>
 		public AgentId AgentId { get; }
 
@@ -81,25 +69,20 @@ namespace Horde.Agent.Services
 		public SessionId SessionId { get; }
 
 		/// <inheritdoc/>
-		public string Token { get; }
-
-		/// <inheritdoc/>
-		public GrpcChannel GrpcChannel { get; }
+		public DirectoryReference WorkingDir { get; }
 
 		/// <inheritdoc/>
-		public DirectoryReference WorkingDir { get; }
+		public IHordeClient HordeClient { get; }
 
 		/// <summary>
 		/// Constructor
 		/// </summary>
-		public Session(Uri serverUrl, AgentId agentId, SessionId sessionId, string token, GrpcChannel grpcChannel, DirectoryReference workingDir)
+		public Session(AgentId agentId, SessionId sessionId, DirectoryReference workingDir, IHordeClient hordeClient)
 		{
-			ServerUrl = serverUrl;
 			AgentId = agentId;
 			SessionId = sessionId;
-			Token = token;
-			GrpcChannel = grpcChannel;
 			WorkingDir = workingDir;
+			HordeClient = hordeClient;
 		}
 
 		public class AgentRegistrationList
@@ -112,7 +95,7 @@ namespace Horde.Agent.Services
 		/// <summary>
 		/// Creates a new agent session
 		/// </summary>
-		public static async Task<Session> CreateAsync(CapabilitiesService capabilitiesService, GrpcService grpcService, StatusService statusService, IOptions<AgentSettings> settings, ILogger logger, CancellationToken cancellationToken)
+		public static async Task<Session> CreateAsync(CapabilitiesService capabilitiesService, GrpcService grpcService, StatusService statusService, IOptions<AgentSettings> settings, IHordeClientFactory hordeClientFactory, ILogger logger, CancellationToken cancellationToken)
 		{
 			AgentSettings currentSettings = settings.Value;
 
@@ -169,7 +152,7 @@ namespace Horde.Agent.Services
 			AgentRegistrationList registrationList = await ReadRegistrationListAsync(registrationFile, cancellationToken);
 
 			// If they aren't valid, create a new agent registration
-			AgentRegistration? registrationInfo = registrationList.Entries.FirstOrDefault(x => x.Server == grpcService.ServerProfile.Url);
+			AgentRegistration? registrationInfo = registrationList.Entries.FirstOrDefault(x => x.Server == serverProfile.Url);
 			if (registrationInfo == null)
 			{
 				statusService.Set(AgentStatusMessage.WaitingForEnrollment);
@@ -185,9 +168,9 @@ namespace Horde.Agent.Services
 			statusService.Set(AgentStatusMessage.ConnectingToServer);
 
 			RpcCreateSessionResponse createSessionResponse;
-			using (GrpcChannel channel = await grpcService.CreateGrpcChannelAsync(registrationInfo.Token, cancellationToken))
+			await using (IHordeClient sessionClient = hordeClientFactory.Create(registrationInfo.Token))
 			{
-				HordeRpc.HordeRpcClient rpcClient = new HordeRpc.HordeRpcClient(channel);
+				HordeRpc.HordeRpcClient rpcClient = await sessionClient.CreateGrpcClientAsync<HordeRpc.HordeRpcClient>(cancellationToken);
 
 				// Create the session information
 				RpcCreateSessionRequest sessionRequest = new RpcCreateSessionRequest();
@@ -216,8 +199,8 @@ namespace Horde.Agent.Services
 
 			// Open a connection to the server
 #pragma warning disable CA2000 // False positive; ownership is transferred to new Session object.
-			GrpcChannel grpcChannel = await grpcService.CreateGrpcChannelAsync(createSessionResponse.Token, cancellationToken);
-			return new Session(serverProfile.Url, new AgentId(createSessionResponse.AgentId), SessionId.Parse(createSessionResponse.SessionId), createSessionResponse.Token, grpcChannel, workingDir);
+			IHordeClient client = hordeClientFactory.Create(createSessionResponse.Token);
+			return new Session(new AgentId(createSessionResponse.AgentId), SessionId.Parse(createSessionResponse.SessionId), workingDir, client);
 #pragma warning restore CA2000
 		}
 
@@ -225,10 +208,9 @@ namespace Horde.Agent.Services
 		/// Dispose of the current session
 		/// </summary>
 		/// <returns></returns>
-		public ValueTask DisposeAsync()
+		public async ValueTask DisposeAsync()
 		{
-			GrpcChannel.Dispose();
-			return default;
+			await HordeClient.DisposeAsync();
 		}
 
 		static FileReference GetRegistrationFile()
@@ -347,22 +329,24 @@ namespace Horde.Agent.Services
 		readonly CapabilitiesService _capabilitiesService;
 		readonly GrpcService _grpcService;
 		readonly StatusService _statusService;
+		readonly IHordeClientFactory _hordeClientFactory;
 		readonly IOptions<AgentSettings> _agentSettings;
 		readonly ILogger _logger;
 
 		/// <summary>
 		/// Constructor
 		/// </summary>
-		public SessionFactory(CapabilitiesService capabilitiesService, GrpcService grpcService, StatusService statusService, IOptions<AgentSettings> agentSettings, ILogger<SessionFactory> logger)
+		public SessionFactory(CapabilitiesService capabilitiesService, GrpcService grpcService, StatusService statusService, IHordeClientFactory hordeClientFactory, IOptions<AgentSettings> agentSettings, ILogger<SessionFactory> logger)
 		{
 			_capabilitiesService = capabilitiesService;
 			_grpcService = grpcService;
 			_statusService = statusService;
+			_hordeClientFactory = hordeClientFactory;
 			_agentSettings = agentSettings;
 			_logger = logger;
 		}
 
 		/// <inheritdoc/>
-		public async Task<ISession> CreateAsync(CancellationToken cancellationToken) => await Session.CreateAsync(_capabilitiesService, _grpcService, _statusService, _agentSettings, _logger, cancellationToken);
+		public async Task<ISession> CreateAsync(CancellationToken cancellationToken) => await Session.CreateAsync(_capabilitiesService, _grpcService, _statusService, _agentSettings, _hordeClientFactory, _logger, cancellationToken);
 	}
 }
diff --git a/Engine/Source/Programs/Horde/Samples/RemoteClient/RemoteClientApp.cs b/Engine/Source/Programs/Horde/Samples/RemoteClient/RemoteClientApp.cs
--- a/Engine/Source/Programs/Horde/Samples/RemoteClient/RemoteClientApp.cs
+++ b/Engine/Source/Programs/Horde/Samples/RemoteClient/RemoteClientApp.cs
@@ -58,7 +58,7 @@ namespace RemoteClient
 			}
 			else
 			{
-				services.AddHordeHttpClient(x => x.BaseAddress = new Uri(options.Server));
+				services.AddHorde(x => x.ServerUrl = new Uri(options.Server));
 				services.AddSingleton<IComputeClient, ServerComputeClient>();
 			}
 
diff --git a/Engine/Source/Programs/Shared/EpicGames.Horde/HordeClient.cs b/Engine/Source/Programs/Shared/EpicGames.Horde/HordeClient.cs
--- a/Engine/Source/Programs/Shared/EpicGames.Horde/HordeClient.cs
+++ b/Engine/Source/Programs/Shared/EpicGames.Horde/HordeClient.cs
@@ -2,6 +2,7 @@
 
 using System;
 using System.Net.Http;
+using System.Net.Http.Headers;
 using System.Net.Http.Json;
 using System.Threading;
 using System.Threading.Tasks;
@@ -23,30 +24,31 @@ namespace EpicGames.Horde
 	/// <summary>
 	/// Default implementation of <see cref="IHordeClient"/>
 	/// </summary>
-	sealed class HordeClient : IHordeClient, IAsyncDisposable
+	abstract class HordeClient : IHordeClient
 	{
+		readonly Uri _serverUrl;
 		readonly IHttpClientFactory _httpClientFactory;
-		readonly HordeHttpAuthHandlerState _authHandlerState;
 		readonly BundleCache _bundleCache;
 		readonly HordeOptions _hordeOptions;
 		readonly ILoggerFactory _loggerFactory;
+		readonly ILogger _logger;
 
 		BackgroundTask<GrpcChannel>? _grpcChannel;
-		Uri? _serverUrl;
 
 		/// <inheritdoc/>
-		public Uri ServerUrl => GetServerUrl();
+		public Uri ServerUrl => _serverUrl;
 
 		/// <summary>
 		/// Constructor
 		/// </summary>
-		public HordeClient(IHttpClientFactory httpClientFactory, HordeHttpAuthHandlerState authHandlerState, BundleCache bundleCache, IOptionsSnapshot<HordeOptions> hordeOptions, ILoggerFactory loggerFactory)
+		public HordeClient(Uri serverUrl, IHttpClientFactory httpClientFactory, BundleCache bundleCache, IOptionsSnapshot<HordeOptions> hordeOptions, ILoggerFactory loggerFactory)
 		{
+			_serverUrl = serverUrl;
 			_httpClientFactory = httpClientFactory;
-			_authHandlerState = authHandlerState;
 			_bundleCache = bundleCache;
 			_hordeOptions = hordeOptions.Value;
 			_loggerFactory = loggerFactory;
+			_logger = _loggerFactory.CreateLogger<HordeClient>();
 		}
 
 		/// <inheritdoc/>
@@ -59,37 +61,17 @@ namespace EpicGames.Horde
 			}
 		}
 
-		Uri GetServerUrl()
-		{
-			if (_serverUrl == null)
-			{
-				using HttpClient httpClient = _httpClientFactory.CreateClient(HordeHttpClient.HttpClientName);
-				_serverUrl = httpClient.BaseAddress ?? new Uri("http://horde-server");
-			}
-			return _serverUrl;
-		}
+		/// <inheritdoc/>
+		public abstract Task<bool> LoginAsync(bool allowLogin, CancellationToken cancellationToken);
 
 		/// <inheritdoc/>
-		public async Task<bool> ConnectAsync(bool allowLogin, CancellationToken cancellationToken)
-		{
-			return await _authHandlerState.LoginAsync(allowLogin, cancellationToken);
-		}
+		public abstract bool HasValidAccessToken();
 
 		/// <inheritdoc/>
-		public bool IsConnected()
-		{
-			try
-			{
-				return _authHandlerState.IsLoggedIn();
-			}
-			catch
-			{
-				return false;
-			}
-		}
+		public abstract Task<string?> GetAccessTokenAsync(bool interactive, CancellationToken cancellationToken);
 
 		/// <inheritdoc/>
-		public async Task<GrpcChannel> GetGrpcChannelAsync(CancellationToken cancellationToken)
+		public async Task<GrpcChannel> CreateGrpcChannelAsync(CancellationToken cancellationToken)
 		{
 			_grpcChannel ??= BackgroundTask.StartNew(ctx => CreateGrpcChannelInternalAsync(ctx));
 			return await _grpcChannel.WaitAsync(cancellationToken);
@@ -97,27 +79,44 @@ namespace EpicGames.Horde
 
 		async Task<GrpcChannel> CreateGrpcChannelInternalAsync(CancellationToken cancellationToken)
 		{
-			HttpClient httpClient = _httpClientFactory.CreateClient(HordeHttpClient.HttpClientName);
+			Uri serverUri = ServerUrl;
+			bool useInsecureConnection = serverUri.Scheme.Equals("http", StringComparison.Ordinal);
 
 			// Get the server URL for gRPC traffic. If we're using an unencrpyted connection we need to use a different port for http/2, so 
 			// send a http1 request to the server to query it.
-			Uri serverUri = httpClient.BaseAddress ?? throw new InvalidOperationException("Horde server base address is not configured");
-			if (serverUri.Scheme.Equals("http", StringComparison.Ordinal))
+			if (useInsecureConnection)
 			{
-				using (HttpResponseMessage response = await httpClient.GetAsync("api/v1/server/ports", cancellationToken))
+				_logger.LogInformation("Querying server {BaseUrl} for rpc port", serverUri);
+				using (HttpClient httpClient = _httpClientFactory.CreateClient())
 				{
-					GetPortsResponse? ports = await response.Content.ReadFromJsonAsync<GetPortsResponse>(HordeHttpClient.JsonSerializerOptions, cancellationToken);
-					if (ports != null && ports.UnencryptedHttp2.HasValue && ports.UnencryptedHttp2 != 0)
+					httpClient.DefaultRequestHeaders.Add("Accept", "application/json");
+					httpClient.Timeout = TimeSpan.FromSeconds(210); // Need to make sure this doesn't cancel any long running gRPC streaming calls (eg. session update)
+					using (HttpResponseMessage response = await httpClient.GetAsync(new Uri(serverUri, "api/v1/server/ports"), cancellationToken))
 					{
-						UriBuilder builder = new UriBuilder(serverUri);
-						builder.Port = ports.UnencryptedHttp2.Value;
-						serverUri = builder.Uri;
+						GetPortsResponse? ports = await response.Content.ReadFromJsonAsync<GetPortsResponse>(HordeHttpClient.JsonSerializerOptions, cancellationToken);
+						if (ports != null && ports.UnencryptedHttp2.HasValue && ports.UnencryptedHttp2.Value != 0)
+						{
+							UriBuilder builder = new UriBuilder(serverUri);
+							builder.Port = ports.UnencryptedHttp2.Value;
+							serverUri = builder.Uri;
+						}
 					}
 				}
 			}
 
-			ServiceConfig serviceConfig = new ServiceConfig();
-			serviceConfig.MethodConfigs.Add(new MethodConfig
+#pragma warning disable CA2000 // Dispose objects before losing scope
+			SocketsHttpHandler httpHandler = new SocketsHttpHandler();
+#pragma warning restore CA2000 // Dispose objects before losing scope
+
+			// Create options for the new channel
+			GrpcChannelOptions channelOptions = new GrpcChannelOptions();
+			channelOptions.MaxReceiveMessageSize = 1024 * 1024 * 1024; // 1 GB 		// Required payloads coming from CAS service can be large
+			channelOptions.MaxSendMessageSize = 1024 * 1024 * 1024; // 1 GB
+			channelOptions.LoggerFactory = _loggerFactory;
+			channelOptions.HttpHandler = httpHandler;
+			channelOptions.DisposeHttpClient = true;
+			channelOptions.ServiceConfig = new ServiceConfig();
+			channelOptions.ServiceConfig.MethodConfigs.Add(new MethodConfig
 			{
 				Names = { MethodName.Default },
 				RetryPolicy = new RetryPolicy
@@ -130,29 +129,202 @@ namespace EpicGames.Horde
 				}
 			});
 
-			return GrpcChannel.ForAddress(serverUri, new GrpcChannelOptions
+			// Configure requests to send the bearer token
+			string? bearerToken = await GetAccessTokenAsync(false, cancellationToken);
+			if (!String.IsNullOrEmpty(bearerToken))
 			{
-				// Required payloads coming from CAS service can be large
-				MaxReceiveMessageSize = 1024 * 1024 * 1024, // 1 GB
-				MaxSendMessageSize = 1024 * 1024 * 1024, // 1 GB
-				LoggerFactory = _loggerFactory,
-				HttpClient = httpClient,
-				DisposeHttpClient = true,
-				ServiceConfig = serviceConfig
-			});
+				CallCredentials callCredentials = CallCredentials.FromInterceptor((context, metadata) =>
+				{
+					metadata.Add("Authorization", $"Bearer {bearerToken}");
+					return Task.CompletedTask;
+				});
+
+				if (useInsecureConnection)
+				{
+					channelOptions.Credentials = ChannelCredentials.Create(ChannelCredentials.Insecure, callCredentials);
+				}
+				else
+				{
+					channelOptions.Credentials = ChannelCredentials.Create(ChannelCredentials.SecureSsl, callCredentials);
+				}
+
+				channelOptions.UnsafeUseInsecureChannelCallCredentials = useInsecureConnection;
+			}
+
+			// Create the channel
+			_logger.LogInformation("Connecting to rpc server {BaseUrl}", serverUri);
+			return GrpcChannel.ForAddress(serverUri, channelOptions);
+		}
+
+		public async Task<TClient> CreateGrpcClientAsync<TClient>(CancellationToken cancellationToken = default)
+			where TClient : ClientBase<TClient>
+		{
+			GrpcChannel channel = await CreateGrpcChannelAsync(cancellationToken);
+			return (TClient)Activator.CreateInstance(typeof(TClient), channel)!;
 		}
 
 		/// <inheritdoc/>
 		public HordeHttpClient CreateHttpClient()
-			=> _httpClientFactory.CreateHordeHttpClient();
+			=> new HordeHttpClient(CreateDefaultHttpClient());
 
 		/// <inheritdoc/>
-		public IStorageClient CreateStorageClient(string basePath)
+		public IStorageClient CreateStorageClient(string basePath, string? accessToken = null)
 		{
-			HttpClient CreateClient() => _httpClientFactory.CreateClient(HordeHttpClient.HttpClientName);
+			HttpClient CreateClient()
+			{
+				HttpClient httpClient = CreateDefaultHttpClient();
+				if (accessToken != null)
+				{
+					httpClient.DefaultRequestHeaders.Authorization = new AuthenticationHeaderValue("Bearer", accessToken);
+				}
+				return httpClient;
+			}
+
 			HttpClient CreateUploadRedirectClient() => _httpClientFactory.CreateClient(HordeHttpClient.UploadRedirectHttpClientName);
 			HttpStorageBackend httpStorageBackend = new HttpStorageBackend(basePath, CreateClient, CreateUploadRedirectClient, _loggerFactory.CreateLogger<HttpStorageBackend>());
 			return new BundleStorageClient(httpStorageBackend, _bundleCache, _hordeOptions.Bundle, _loggerFactory.CreateLogger<BundleStorageClient>());
 		}
+
+		/// <summary>
+		/// Creates an http client for satisfying requests
+		/// </summary>
+		protected abstract HttpClient CreateDefaultHttpClient();
+	}
+
+	/// <summary>
+	/// Default implementation of <see cref="IHordeClient"/>
+	/// </summary>
+	class HordeClientWithStaticCredentials : HordeClient
+	{
+		readonly string? _accessToken;
+		readonly IHttpClientFactory _httpClientFactory;
+
+		/// <summary>
+		/// Constructor
+		/// </summary>
+		public HordeClientWithStaticCredentials(Uri serverUrl, string? accessToken, IHttpClientFactory httpClientFactory, BundleCache bundleCache, IOptionsSnapshot<HordeOptions> hordeOptions, ILoggerFactory loggerFactory)
+			: base(serverUrl, httpClientFactory, bundleCache, hordeOptions, loggerFactory)
+		{
+			_accessToken = accessToken;
+			_httpClientFactory = httpClientFactory;
+		}
+
+		/// <inheritdoc/>
+		public override Task<bool> LoginAsync(bool allowLogin, CancellationToken cancellationToken)
+			=> Task.FromResult(true);
+
+		/// <inheritdoc/>
+		public override bool HasValidAccessToken()
+			=> true;
+
+		/// <inheritdoc/>
+		public override Task<string?> GetAccessTokenAsync(bool interactive, CancellationToken cancellationToken)
+			=> Task.FromResult(_accessToken);
+
+		protected override HttpClient CreateDefaultHttpClient()
+		{
+			HttpClient httpClient = _httpClientFactory.CreateClient(HordeHttpClient.HttpClientName);
+			httpClient.BaseAddress = ServerUrl;
+			if (!String.IsNullOrEmpty(_accessToken))
+			{
+				httpClient.DefaultRequestHeaders.Authorization = new AuthenticationHeaderValue("Bearer", _accessToken);
+			}
+			return httpClient;
+		}
+	}
+
+	/// <summary>
+	/// Default implementation of <see cref="IHordeClient"/>
+	/// </summary>
+	class HordeClientWithDynamicCredentials : HordeClient
+	{
+		readonly IHttpClientFactory _httpClientFactory;
+		readonly HordeHttpAuthHandlerState _authHandler;
+
+		/// <summary>
+		/// Constructor
+		/// </summary>
+		public HordeClientWithDynamicCredentials(Uri serverUrl, IHttpClientFactory httpClientFactory, HordeHttpAuthHandlerState authHandler, BundleCache bundleCache, IOptionsSnapshot<HordeOptions> hordeOptions, ILoggerFactory loggerFactory)
+			: base(serverUrl, httpClientFactory, bundleCache, hordeOptions, loggerFactory)
+		{
+			_httpClientFactory = httpClientFactory;
+			_authHandler = authHandler;
+		}
+
+		/// <inheritdoc/>
+		public override async Task<bool> LoginAsync(bool allowLogin, CancellationToken cancellationToken)
+		{
+			return await _authHandler.LoginAsync(allowLogin, cancellationToken);
+		}
+
+		/// <inheritdoc/>
+		public override bool HasValidAccessToken()
+		{
+			try
+			{
+				return _authHandler.IsLoggedIn();
+			}
+			catch
+			{
+				return false;
+			}
+		}
+
+		/// <inheritdoc/>
+		public override Task<string?> GetAccessTokenAsync(bool interactive, CancellationToken cancellationToken)
+			=> _authHandler.GetAccessTokenAsync(interactive, cancellationToken);
+
+		/// <inheritdoc/>
+		protected override HttpClient CreateDefaultHttpClient()
+		{
+			HttpClient httpClient = _httpClientFactory.CreateClient(HordeHttpClient.HttpClientName);
+			httpClient.BaseAddress = ServerUrl;
+			return httpClient;
+		}
+	}
+
+	/// <summary>
+	/// Implementation of <see cref="IHordeClientFactory"/>
+	/// </summary>
+	class HordeClientFactory : IHordeClientFactory
+	{
+		readonly IHttpClientFactory _httpClientFactory;
+		readonly BundleCache _bundleCache;
+		readonly HordeHttpAuthHandlerState _authHandlerState;
+		readonly IOptionsSnapshot<HordeOptions> _hordeOptions;
+		readonly ILoggerFactory _loggerFactory;
+
+		public HordeClientFactory(IHttpClientFactory httpClientFactory, BundleCache bundleCache, HordeHttpAuthHandlerState authHandlerState, IOptionsSnapshot<HordeOptions> hordeOptions, ILoggerFactory loggerFactory)
+		{
+			_httpClientFactory = httpClientFactory;
+			_bundleCache = bundleCache;
+			_authHandlerState = authHandlerState;
+			_hordeOptions = hordeOptions;
+			_loggerFactory = loggerFactory;
+		}
+
+		Uri GetServerUrl()
+		{
+			Uri? serverUrl = _hordeOptions.Value.GetServerUrlOrDefault();
+			if (serverUrl == null)
+			{
+				throw new Exception("No Horde server is configured, or can be detected from the environment. Consider specifying a URL when calling AddHordeHttpClient().");
+			}
+			return serverUrl;
+		}
+
+		/// <inheritdoc/>
+		public IHordeClient Create()
+		{
+			Uri serverUrl = GetServerUrl();
+			return new HordeClientWithDynamicCredentials(serverUrl, _httpClientFactory, _authHandlerState, _bundleCache, _hordeOptions, _loggerFactory);
+		}
+
+		/// <inheritdoc/>
+		public IHordeClient Create(string? accessToken)
+		{
+			Uri serverUrl = GetServerUrl();
+			return new HordeClientWithStaticCredentials(serverUrl, accessToken, _httpClientFactory, _bundleCache, _hordeOptions, _loggerFactory);
+		}
 	}
 }
diff --git a/Engine/Source/Programs/Shared/EpicGames.Horde/HordeExtensions.cs b/Engine/Source/Programs/Shared/EpicGames.Horde/HordeExtensions.cs
--- a/Engine/Source/Programs/Shared/EpicGames.Horde/HordeExtensions.cs
+++ b/Engine/Source/Programs/Shared/EpicGames.Horde/HordeExtensions.cs
@@ -1,12 +1,20 @@
 // Copyright Epic Games, Inc. All Rights Reserved.
 
 using System;
+using System.Net.Http;
+using System.Threading.Tasks;
+using EpicGames.Core;
 using EpicGames.Horde.Storage;
 using EpicGames.Horde.Storage.Backends;
 using EpicGames.Horde.Storage.Bundles;
 using EpicGames.Horde.Storage.Clients;
 using Microsoft.Extensions.DependencyInjection;
+using Microsoft.Extensions.Logging;
 using Microsoft.Extensions.Options;
+using Polly;
+using Polly.Extensions.Http;
+using Polly.Retry;
+using Polly.Timeout;
 
 namespace EpicGames.Horde
 {
@@ -21,12 +29,50 @@ namespace EpicGames.Horde
 		/// <param name="serviceCollection">Collection to register services with</param>
 		public static void AddHorde(this IServiceCollection serviceCollection)
 		{
-			serviceCollection.AddHordeHttpClient();
+			// Sets defaults from the environment before calling the user provided configuration method
+			void ConfigureClientFromEnvironment(IServiceProvider serviceProvider, HttpClient httpClient)
+			{
+				httpClient.Timeout = TimeSpan.FromSeconds(240); // Global timeout
+			}
+
+			// Register the HTTP client for handling login requests
+			void ConfigureClientFromOptions(IServiceProvider serviceProvider, HttpClient httpClient)
+			{
+				IOptions<HordeOptions> options = serviceProvider.GetRequiredService<IOptions<HordeOptions>>();
+				if (options.Value.ServerUrl != null)
+				{
+					httpClient.BaseAddress = options.Value.ServerUrl;
+				}
+			}
+
+			serviceCollection.AddSingleton<HordeHttpAuthHandlerState>();
+			serviceCollection.AddTransient<HordeHttpAuthHandler>();
+			serviceCollection.AddHttpClient(HordeHttpAuthHandlerState.HttpClientName, ConfigureClientFromOptions);
+
+			// Register the HTTP client for processing upload redirects
+			serviceCollection.AddHttpClient(HordeHttpClient.UploadRedirectHttpClientName)
+				.AddPolicyHandler((serviceProvider, request) => CreateDefaultTimeoutRetryPolicy(request, serviceProvider.GetRequiredService<ILogger<HttpStorageBackend>>()))
+				.AddPolicyHandler((serviceProvider, request) => CreateDefaultTransientErrorPolicy(request, serviceProvider.GetRequiredService<ILogger<HttpStorageBackend>>()));
+
+			// Create the HTTP client for handling Horde requests
+			serviceCollection.AddHttpClient<HordeHttpClient>(HordeHttpClient.HttpClientName, ConfigureClientFromEnvironment)
+				.AddHttpMessageHandler<HordeHttpAuthHandler>()
+				.AddPolicyHandler((serviceProvider, request) => CreateDefaultTimeoutRetryPolicy(request, serviceProvider.GetRequiredService<ILogger<HttpStorageBackend>>()))
+				.AddPolicyHandler((serviceProvider, request) => CreateDefaultTransientErrorPolicy(request, serviceProvider.GetRequiredService<ILogger<HttpStorageBackend>>()));
+
+			static StorageBackendCache CreateBackendCache(IServiceProvider serviceProvider)
+			{
+				StorageBackendCacheOptions options = serviceProvider.GetRequiredService<IOptions<HordeOptions>>().Value.BackendCache;
+				DirectoryReference? cacheDir = String.IsNullOrEmpty(options.CacheDir) ? null : new DirectoryReference(options.CacheDir);
+				return new StorageBackendCache(cacheDir, options.MaxSize, serviceProvider.GetRequiredService<ILogger<StorageBackendCache>>());
+			}
+
 			serviceCollection.AddSingleton<BundleCache>(sp => new BundleCache(sp.GetRequiredService<IOptions<HordeOptions>>().Value.BundleCache));
-			serviceCollection.AddSingleton<StorageBackendCache>();
+			serviceCollection.AddSingleton<StorageBackendCache>(CreateBackendCache);
 			serviceCollection.AddSingleton<HttpStorageBackendFactory>();
 			serviceCollection.AddSingleton<HttpStorageClientFactory>();
-			serviceCollection.AddSingleton<IHordeClient, HordeClient>();
+			serviceCollection.AddSingleton<IHordeClient>(sp => sp.GetRequiredService<IHordeClientFactory>().Create());
+			serviceCollection.AddSingleton<IHordeClientFactory, HordeClientFactory>();
 		}
 
 		/// <summary>
@@ -39,5 +85,45 @@ namespace EpicGames.Horde
 			serviceCollection.Configure(configureHorde);
 			AddHorde(serviceCollection);
 		}
+
+		/// <summary>
+		/// Create a default timeout retry policy
+		/// </summary>
+		public static IAsyncPolicy<HttpResponseMessage> CreateDefaultTimeoutRetryPolicy(HttpRequestMessage request, ILogger logger)
+		{
+			// Wait 30 seconds for operations to timeout
+			Task OnTimeoutAsync(Context context, TimeSpan timespan, Task timeoutTask)
+			{
+				logger.LogWarning(KnownLogEvents.Systemic_Horde_Http, "{Method} {Url} timed out after {Time}s.", request.Method, request.RequestUri, (int)timespan.TotalSeconds);
+				return Task.CompletedTask;
+			}
+
+			AsyncTimeoutPolicy<HttpResponseMessage> timeoutPolicy = Policy.TimeoutAsync<HttpResponseMessage>(30, OnTimeoutAsync);
+
+			// Retry twice after a timeout
+			void OnRetry(Exception ex, TimeSpan timespan)
+			{
+				logger.LogWarning(KnownLogEvents.Systemic_Horde_Http, ex, "{Method} {Url} retrying after {Time}s.", request.Method, request.RequestUri, timespan.TotalSeconds);
+			}
+
+			TimeSpan[] retryTimes = new[] { TimeSpan.FromSeconds(5.0), TimeSpan.FromSeconds(10.0) };
+			AsyncRetryPolicy retryPolicy = Policy.Handle<TimeoutRejectedException>().WaitAndRetryAsync(retryTimes, OnRetry);
+			return retryPolicy.WrapAsync(timeoutPolicy);
+		}
+
+		/// <summary>
+		/// Create a default timeout retry policy
+		/// </summary>
+		public static IAsyncPolicy<HttpResponseMessage> CreateDefaultTransientErrorPolicy(HttpRequestMessage request, ILogger logger)
+		{
+			Task OnTimeoutAsync(DelegateResult<HttpResponseMessage> outcome, TimeSpan timespan, int retryAttempt, Context context)
+			{
+				logger.LogWarning(KnownLogEvents.Systemic_Horde_Http, "{Method} {Url} failed ({Result}). Delaying for {DelayMs}ms (attempt #{RetryNum}).", request.Method, request.RequestUri, outcome.Result?.StatusCode, timespan.TotalMilliseconds, retryAttempt);
+				return Task.CompletedTask;
+			}
+
+			TimeSpan[] retryTimes = new[] { TimeSpan.FromSeconds(1.0), TimeSpan.FromSeconds(5.0), TimeSpan.FromSeconds(10.0), TimeSpan.FromSeconds(30.0), TimeSpan.FromSeconds(30.0) };
+			return HttpPolicyExtensions.HandleTransientHttpError().WaitAndRetryAsync(retryTimes, OnTimeoutAsync);
+		}
 	}
 }
diff --git a/Engine/Source/Programs/Shared/EpicGames.Horde/HordeHttpClient.cs b/Engine/Source/Programs/Shared/EpicGames.Horde/HordeHttpClient.cs
--- a/Engine/Source/Programs/Shared/EpicGames.Horde/HordeHttpClient.cs
+++ b/Engine/Source/Programs/Shared/EpicGames.Horde/HordeHttpClient.cs
@@ -20,16 +20,8 @@ using EpicGames.Horde.Projects;
 using EpicGames.Horde.Secrets;
 using EpicGames.Horde.Server;
 using EpicGames.Horde.Storage;
-using EpicGames.Horde.Storage.Backends;
 using EpicGames.Horde.Streams;
 using EpicGames.Horde.Tools;
-using Microsoft.Extensions.DependencyInjection;
-using Microsoft.Extensions.Logging;
-using Microsoft.Extensions.Options;
-using Polly;
-using Polly.Extensions.Http;
-using Polly.Retry;
-using Polly.Timeout;
 
 #pragma warning disable CA2234
 
@@ -57,6 +49,11 @@ namespace EpicGames.Horde
 		/// </summary>
 		public const string HttpClientName = "HordeHttpClient";
 
+		/// <summary>
+		/// Name of clients used for anonymous requests.
+		/// </summary>
+		public const string AnonymousHttpClientName = "HordeAnonymousHttpClient";
+
 		/// <summary>
 		/// Name of clients created from the http client factory for handling upload redirects. Should not contain Horde auth headers.
 		/// </summary>
@@ -600,146 +597,4 @@ namespace EpicGames.Horde
 
 		#endregion
 	}
-
-	/// <summary>
-	/// Extension methods for Horde HTTP clients
-	/// </summary>
-	public static class HordeHttpClientExtensions
-	{
-		/// <summary>
-		/// Creates a <see cref="HordeHttpClient"/> instance from an http client factory
-		/// </summary>
-		public static HordeHttpClient CreateHordeHttpClient(this IHttpClientFactory factory)
-		{
-			return new HordeHttpClient(factory.CreateClient(HordeHttpClient.HttpClientName));
-		}
-
-		/// <summary>
-		/// Registers a Horde HTTP client type, and configures it to use the default OIDC message handler.
-		/// </summary>
-		/// <param name="services">Service collection to add services to</param>
-		public static IHttpClientBuilder AddHordeHttpClient(this IServiceCollection services)
-		{
-			return services.AddHordeHttpClient((sp, client) => { });
-		}
-
-		/// <summary>
-		/// Registers a Horde HTTP client type, and configures it to use the default OIDC message handler.
-		/// </summary>
-		/// <param name="services">Service collection to add services to</param>
-		/// <param name="configureClient">Callback to modify options for the http client</param>
-		public static IHttpClientBuilder AddHordeHttpClient(this IServiceCollection services, Action<HttpClient> configureClient)
-		{
-			return services.AddHordeHttpClient((sp, client) => configureClient(client));
-		}
-
-		/// <summary>
-		/// Registers a Horde HTTP client type, and configures it to use the default OIDC message handler.
-		/// </summary>
-		/// <param name="services">Service collection to add services to</param>
-		/// <param name="configureClient">Callback to modify options for the http client</param>
-		public static IHttpClientBuilder AddHordeHttpClient(this IServiceCollection services, Action<IServiceProvider, HttpClient> configureClient)
-		{
-			// Sets defaults from the environment before calling the user provided configuration method
-			void ConfigureClientFromEnvironment(IServiceProvider serviceProvider, HttpClient httpClient)
-			{
-				IOptions<HordeOptions> options = serviceProvider.GetRequiredService<IOptions<HordeOptions>>();
-				if (options.Value.ServerUrl != null)
-				{
-					httpClient.BaseAddress = options.Value.ServerUrl;
-				}
-
-				httpClient.Timeout = TimeSpan.FromSeconds(240); // Global timeout
-
-				// Run the user callbacks
-				options.Value.ConfigureHttpClient?.Invoke(httpClient);
-				configureClient(serviceProvider, httpClient);
-
-				// If the server URL isn't set, take it from the environment
-				if (httpClient.BaseAddress == null)
-				{
-					string? hordeUrlEnvVar = Environment.GetEnvironmentVariable(HordeHttpClient.HordeUrlEnvVarName);
-					if (!String.IsNullOrEmpty(hordeUrlEnvVar))
-					{
-						httpClient.BaseAddress = new Uri(hordeUrlEnvVar);
-					}
-				}
-
-				// Try to get the default server address from the registry
-				httpClient.BaseAddress ??= HordeOptions.GetDefaultServerUrl();
-
-				// Make sure we have a base URL set
-				if (httpClient.BaseAddress == null)
-				{
-					throw new Exception("No Horde server is configured, or can be detected from the environment. Consider specifying a URL when calling AddHordeHttpClient().");
-				}
-			}
-
-			// Register the HTTP client for handling login requests
-			void ConfigureClientFromOptions(IServiceProvider serviceProvider, HttpClient httpClient)
-			{
-				IOptions<HordeOptions> options = serviceProvider.GetRequiredService<IOptions<HordeOptions>>();
-				if (options.Value.ServerUrl != null)
-				{
-					httpClient.BaseAddress = options.Value.ServerUrl;
-				}
-			}
-			services.AddSingleton<HordeHttpAuthHandlerState>();
-			services.AddTransient<HordeHttpAuthHandler>();
-			services.AddHttpClient(HordeHttpAuthHandlerState.HttpClientName, ConfigureClientFromOptions);
-
-			// Register the HTTP client for processing upload redirects
-			services.AddHttpClient(HordeHttpClient.UploadRedirectHttpClientName)
-				.AddPolicyHandler((serviceProvider, request) => CreateDefaultTimeoutRetryPolicy(request, serviceProvider.GetRequiredService<ILogger<HttpStorageBackend>>()))
-				.AddPolicyHandler((serviceProvider, request) => CreateDefaultTransientErrorPolicy(request, serviceProvider.GetRequiredService<ILogger<HttpStorageBackend>>()));
-
-			// Create the HTTP client for handling Horde requests
-			IHttpClientBuilder builder = services.AddHttpClient<HordeHttpClient>(HordeHttpClient.HttpClientName, ConfigureClientFromEnvironment)
-				.AddHttpMessageHandler<HordeHttpAuthHandler>()
-				.AddPolicyHandler((serviceProvider, request) => CreateDefaultTimeoutRetryPolicy(request, serviceProvider.GetRequiredService<ILogger<HttpStorageBackend>>()))
-				.AddPolicyHandler((serviceProvider, request) => CreateDefaultTransientErrorPolicy(request, serviceProvider.GetRequiredService<ILogger<HttpStorageBackend>>()));
-
-			return builder;
-		}
-
-		/// <summary>
-		/// Create a default timeout retry policy
-		/// </summary>
-		public static IAsyncPolicy<HttpResponseMessage> CreateDefaultTimeoutRetryPolicy(HttpRequestMessage request, ILogger logger)
-		{
-			// Wait 30 seconds for operations to timeout
-			Task OnTimeoutAsync(Context context, TimeSpan timespan, Task timeoutTask)
-			{
-				logger.LogWarning(KnownLogEvents.Systemic_Horde_Http, "{Method} {Url} timed out after {Time}s.", request.Method, request.RequestUri, (int)timespan.TotalSeconds);
-				return Task.CompletedTask;
-			}
-
-			AsyncTimeoutPolicy<HttpResponseMessage> timeoutPolicy = Policy.TimeoutAsync<HttpResponseMessage>(30, OnTimeoutAsync);
-
-			// Retry twice after a timeout
-			void OnRetry(Exception ex, TimeSpan timespan)
-			{
-				logger.LogWarning(KnownLogEvents.Systemic_Horde_Http, ex, "{Method} {Url} retrying after {Time}s.", request.Method, request.RequestUri, timespan.TotalSeconds);
-			}
-
-			TimeSpan[] retryTimes = new[] { TimeSpan.FromSeconds(5.0), TimeSpan.FromSeconds(10.0) };
-			AsyncRetryPolicy retryPolicy = Policy.Handle<TimeoutRejectedException>().WaitAndRetryAsync(retryTimes, OnRetry);
-			return retryPolicy.WrapAsync(timeoutPolicy);
-		}
-
-		/// <summary>
-		/// Create a default timeout retry policy
-		/// </summary>
-		public static IAsyncPolicy<HttpResponseMessage> CreateDefaultTransientErrorPolicy(HttpRequestMessage request, ILogger logger)
-		{
-			Task OnTimeoutAsync(DelegateResult<HttpResponseMessage> outcome, TimeSpan timespan, int retryAttempt, Context context)
-			{
-				logger.LogWarning(KnownLogEvents.Systemic_Horde_Http, "{Method} {Url} failed ({Result}). Delaying for {DelayMs}ms (attempt #{RetryNum}).", request.Method, request.RequestUri, outcome.Result?.StatusCode, timespan.TotalMilliseconds, retryAttempt);
-				return Task.CompletedTask;
-			}
-
-			TimeSpan[] retryTimes = new[] { TimeSpan.FromSeconds(1.0), TimeSpan.FromSeconds(5.0), TimeSpan.FromSeconds(10.0), TimeSpan.FromSeconds(30.0), TimeSpan.FromSeconds(30.0) };
-			return HttpPolicyExtensions.HandleTransientHttpError().WaitAndRetryAsync(retryTimes, OnTimeoutAsync);
-		}
-	}
 }
diff --git a/Engine/Source/Programs/Shared/EpicGames.Horde/HordeOptions.cs b/Engine/Source/Programs/Shared/EpicGames.Horde/HordeOptions.cs
--- a/Engine/Source/Programs/Shared/EpicGames.Horde/HordeOptions.cs
+++ b/Engine/Source/Programs/Shared/EpicGames.Horde/HordeOptions.cs
@@ -3,7 +3,6 @@
 using System;
 using System.IO;
 using System.Linq;
-using System.Net.Http;
 using System.Runtime.InteropServices;
 using System.Runtime.Versioning;
 using System.Text.Json;
@@ -36,11 +35,6 @@ namespace EpicGames.Horde
 		/// </summary>
 		public bool AllowAuthPrompt { get; set; } = true;
 
-		/// <summary>
-		/// Callback to allow configuring any HTTP client created for Horde
-		/// </summary>
-		public Action<HttpClient>? ConfigureHttpClient { get; set; }
-
 		/// <summary>
 		/// Options for creating new bundles
 		/// </summary>
@@ -51,6 +45,30 @@ namespace EpicGames.Horde
 		/// </summary>
 		public BundleCacheOptions BundleCache { get; } = new BundleCacheOptions();
 
+		/// <summary>
+		/// Options for the storage backend cache
+		/// </summary>
+		public StorageBackendCacheOptions BackendCache { get; } = new StorageBackendCacheOptions();
+
+		/// <summary>
+		/// Gets the configured server URL, or the default value
+		/// </summary>
+		public Uri? GetServerUrlOrDefault()
+			=> ServerUrl ?? GetServerUrlFromEnvironment() ?? GetDefaultServerUrl();
+
+		/// <summary>
+		/// Reads the server URL from the environment
+		/// </summary>
+		public static Uri? GetServerUrlFromEnvironment()
+		{
+			string? hordeUrlEnvVar = Environment.GetEnvironmentVariable(HordeHttpClient.HordeUrlEnvVarName);
+			if (String.IsNullOrEmpty(hordeUrlEnvVar))
+			{
+				return null;
+			}
+			return new Uri(hordeUrlEnvVar);
+		}
+
 		/// <summary>
 		/// Gets the default server URL for the current user
 		/// </summary>
@@ -174,4 +192,20 @@ namespace EpicGames.Horde
 			}
 		}
 	}
+
+	/// <summary>
+	/// Options for the storage backend cache
+	/// </summary>
+	public class StorageBackendCacheOptions
+	{
+		/// <summary>
+		/// Directory to store cached data
+		/// </summary>
+		public string? CacheDir { get; set; }
+
+		/// <summary>
+		/// Maximum size of the cache, in bytes
+		/// </summary>
+		public long MaxSize { get; set; }
+	}
 }
diff --git a/Engine/Source/Programs/Shared/EpicGames.Horde/IHordeClient.cs b/Engine/Source/Programs/Shared/EpicGames.Horde/IHordeClient.cs
--- a/Engine/Source/Programs/Shared/EpicGames.Horde/IHordeClient.cs
+++ b/Engine/Source/Programs/Shared/EpicGames.Horde/IHordeClient.cs
@@ -3,6 +3,7 @@
 using System;
 using System.Threading;
 using System.Threading.Tasks;
+using EpicGames.Horde.Logs;
 using EpicGames.Horde.Storage;
 using EpicGames.Horde.Tools;
 using Grpc.Core;
@@ -13,7 +14,7 @@ namespace EpicGames.Horde
 	/// <summary>
 	/// Base interface for Horde functionality.
 	/// </summary>
-	public interface IHordeClient
+	public interface IHordeClient : IAsyncDisposable
 	{
 		/// <summary>
 		/// Base URL of the horde server
@@ -23,20 +24,30 @@ namespace EpicGames.Horde
 		/// <summary>
 		/// Connect to the Horde server
 		/// </summary>
-		/// <param name="allowLogin">Whether to allow interactive logins</param>
+		/// <param name="allowPrompt">Whether to allow prompting for credentials</param>
 		/// <param name="cancellationToken">Cancellation token for the operation</param>
 		/// <returns>True if the connection succeded</returns>
-		Task<bool> ConnectAsync(bool allowLogin, CancellationToken cancellationToken);
+		Task<bool> LoginAsync(bool allowPrompt, CancellationToken cancellationToken);
 
 		/// <summary>
 		/// Gets the current connection state
 		/// </summary>
-		bool IsConnected();
+		bool HasValidAccessToken();
+
+		/// <summary>
+		/// Gets an access token for the server
+		/// </summary>
+		Task<string?> GetAccessTokenAsync(bool interactive, CancellationToken cancellationToken = default);
 
 		/// <summary>
 		/// Gets a grpc channel for communication with the server. This should NOT be disposed by the caller.
 		/// </summary>
-		Task<GrpcChannel> GetGrpcChannelAsync(CancellationToken cancellationToken = default);
+		Task<GrpcChannel> CreateGrpcChannelAsync(CancellationToken cancellationToken = default);
+
+		/// <summary>
+		/// Gets a gRPC client interface
+		/// </summary>
+		Task<TClient> CreateGrpcClientAsync<TClient>(CancellationToken cancellationToken = default) where TClient : ClientBase<TClient>;
 
 		/// <summary>
 		/// Creates a Horde HTTP client 
@@ -46,7 +57,23 @@ namespace EpicGames.Horde
 		/// <summary>
 		/// Creates a storage client for the given base path
 		/// </summary>
-		IStorageClient CreateStorageClient(string relativePath);
+		IStorageClient CreateStorageClient(string relativePath, string? accessToken = null);
+	}
+
+	/// <summary>
+	/// Interface to allow creating custom horde client instances. To obtain a default horde client, get an IHordeClient instance via dependency injection.
+	/// </summary>
+	public interface IHordeClientFactory
+	{
+		/// <summary>
+		/// Create a client using the user's default access token
+		/// </summary>
+		IHordeClient Create();
+
+		/// <summary>
+		/// Create a client with an explicit access token
+		/// </summary>
+		IHordeClient Create(string accessToken);
 	}
 
 	/// <summary>
@@ -57,8 +84,14 @@ namespace EpicGames.Horde
 		/// <summary>
 		/// Creates a storage client for a particular namespace
 		/// </summary>
-		public static IStorageClient CreateStorageClient(this IHordeClient hordeClient, NamespaceId namespaceId)
-			=> hordeClient.CreateStorageClient($"api/v1/storage/{namespaceId}");
+		public static IStorageClient CreateStorageClient(this IHordeClient hordeClient, NamespaceId namespaceId, string? accessToken = null)
+			=> hordeClient.CreateStorageClient($"api/v1/storage/{namespaceId}", accessToken);
+
+		/// <summary>
+		/// Creates a storage client for a particular log
+		/// </summary>
+		public static IStorageClient CreateStorageClient(this IHordeClient hordeClient, LogId logId)
+			=> hordeClient.CreateStorageClient($"api/v1/logs/{logId}");
 
 		/// <summary>
 		/// Creates a storage client for a particular tool
@@ -66,14 +99,5 @@ namespace EpicGames.Horde
 		public static IStorageClient CreateStorageClient(this IHordeClient hordeClient, ToolId toolId)
 			=> hordeClient.CreateStorageClient($"api/v1/tools/{toolId}");
 
-		/// <summary>
-		/// Attempts to get a client reference, returning immediately if there's not one available
-		/// </summary>
-		public static async Task<TClient> GetGrpcClientAsync<TClient>(this IHordeClient hordeClient, CancellationToken cancellationToken = default) 
-			where TClient : ClientBase<TClient>
-		{
-			GrpcChannel channel = await hordeClient.GetGrpcChannelAsync(cancellationToken);
-			return (TClient)Activator.CreateInstance(typeof(TClient), channel)!;
-		}
 	}
 }
diff --git a/Engine/Source/Programs/UnrealGameSync/UnrealGameSync/Controls/WorkspaceControl.cs b/Engine/Source/Programs/UnrealGameSync/UnrealGameSync/Controls/WorkspaceControl.cs
--- a/Engine/Source/Programs/UnrealGameSync/UnrealGameSync/Controls/WorkspaceControl.cs
+++ b/Engine/Source/Programs/UnrealGameSync/UnrealGameSync/Controls/WorkspaceControl.cs
@@ -177,7 +177,9 @@ namespace UnrealGameSync
 		readonly IServiceProvider _serviceProvider;
 		readonly ILogger _logger;
 		readonly IPerforceSettings _perforceSettings;
+#pragma warning disable CA2213 // Disposed by DI container
 		readonly IHordeClient? _hordeClient;
+#pragma warning restore CA2213
 		ProjectInfo ProjectInfo { get; }
 
 		readonly UserSettings _settings;
@@ -3493,7 +3495,7 @@ namespace UnrealGameSync
 			{
 				statusLine.AddText("  |  ");
 
-				if (_hordeClient.IsConnected())
+				if (_hordeClient.HasValidAccessToken())
 				{
 					statusLine.AddLink("Connected to Horde", FontStyle.Regular, (p, r) => Utility.OpenUrl(_hordeClient.ServerUrl.ToString()));
 				}
@@ -3901,7 +3903,7 @@ namespace UnrealGameSync
 			{
 				try
 				{
-					await hordeClient.ConnectAsync(true, cancellationToken);
+					await hordeClient.LoginAsync(true, cancellationToken);
 				}
 				catch when (cancellationToken.IsCancellationRequested)
 				{
diff --git a/Engine/Source/Programs/UnrealGameSync/UnrealGameSync/PerforceMonitor.cs b/Engine/Source/Programs/UnrealGameSync/UnrealGameSync/PerforceMonitor.cs
--- a/Engine/Source/Programs/UnrealGameSync/UnrealGameSync/PerforceMonitor.cs
+++ b/Engine/Source/Programs/UnrealGameSync/UnrealGameSync/PerforceMonitor.cs
@@ -52,7 +52,9 @@ namespace UnrealGameSync
 		readonly DirectoryReference _cacheFolder;
 		readonly List<KeyValuePair<FileReference, DateTime>> _localConfigFiles;
 		readonly IAsyncDisposer _asyncDisposeTasks;
+#pragma warning disable CA2213 // Dispose of _hordeClient
 		readonly IHordeClient? _hordeClient;
+#pragma warning restore CA2213
 
 		string[] prevCodeRules = Array.Empty<string>();
 
