Horde: Move log RPC service definitions into EpicGames.Horde.

#rnx

diff --git a/Engine/Source/Programs/Horde/Horde.Agent.Driver/Horde.Agent.Driver.csproj b/Engine/Source/Programs/Horde/Horde.Agent.Driver/Horde.Agent.Driver.csproj
--- a/Engine/Source/Programs/Horde/Horde.Agent.Driver/Horde.Agent.Driver.csproj
+++ b/Engine/Source/Programs/Horde/Horde.Agent.Driver/Horde.Agent.Driver.csproj
@@ -98,10 +98,8 @@
     <Protobuf Include="..\Horde.Server\Protos\horde\horde_rpc_messages.proto" ProtoRoot="..\Horde.Server\Protos" Access="Public" />
     <Protobuf Include="..\Horde.Server\Protos\horde\enrollment_rpc.proto" ProtoRoot="..\Horde.Server\Protos" GrpcServices="Client" Access="Public" />
     <Protobuf Include="..\Horde.Server\Protos\horde\enrollment_rpc_messages.proto" ProtoRoot="..\Horde.Server\Protos" Access="Public" />
-    <Protobuf Include="..\Horde.Server\Protos\horde\job_rpc.proto" ProtoRoot="..\Horde.Server\Protos" GrpcServices="Client" Access="Public" />
+    <Protobuf Include="..\Horde.Server\Protos\horde\job_rpc.proto" ProtoRoot="..\Horde.Server\Protos" GrpcServices="Client" Access="Public" AdditionalImportDirs="..\..\Shared\EpicGames.Horde\Protos"/>
     <Protobuf Include="..\Horde.Server\Protos\horde\job_rpc_messages.proto" ProtoRoot="..\Horde.Server\Protos" Access="Public" />
-    <Protobuf Include="..\Horde.Server\Protos\horde\log_rpc.proto" ProtoRoot="..\Horde.Server\Protos" GrpcServices="Client" Access="Public" />
-    <Protobuf Include="..\Horde.Server\Protos\horde\log_rpc_messages.proto" ProtoRoot="..\Horde.Server\Protos" Access="Public" />
     <Protobuf Include="..\Horde.Server\Protos\horde\common\agent_status.proto" ProtoRoot="..\Horde.Server\Protos" Access="Public" />
     <Protobuf Include="..\Horde.Server\Protos\horde\common\io_hash.proto" ProtoRoot="..\Horde.Server\Protos" Access="Public" />
     <Protobuf Include="..\Horde.Server\Protos\horde\common\label_change.proto" ProtoRoot="..\Horde.Server\Protos" Access="Public" />
diff --git a/Engine/Source/Programs/Horde/Horde.Agent.Driver/Utility/JsonRpcLogWriter.cs b/Engine/Source/Programs/Horde/Horde.Agent.Driver/Utility/JsonRpcLogWriter.cs
deleted file mode 100644
diff --git a/Engine/Source/Programs/Horde/Horde.Agent.Driver/Utility/ServerLogger.cs b/Engine/Source/Programs/Horde/Horde.Agent.Driver/Utility/ServerLogger.cs
deleted file mode 100644
diff --git a/Engine/Source/Programs/Horde/Horde.Agent.Tests/JsonRpcLoggerTests.cs b/Engine/Source/Programs/Horde/Horde.Agent.Tests/JsonRpcLoggerTests.cs
deleted file mode 100644
diff --git a/Engine/Source/Programs/Horde/Horde.Server/Horde.Server.csproj b/Engine/Source/Programs/Horde/Horde.Server/Horde.Server.csproj
--- a/Engine/Source/Programs/Horde/Horde.Server/Horde.Server.csproj
+++ b/Engine/Source/Programs/Horde/Horde.Server/Horde.Server.csproj
@@ -94,14 +94,12 @@
   </ItemGroup>
 
   <ItemGroup>
-    <Protobuf Include="Protos\horde\job_rpc.proto" ProtoRoot="Protos" GrpcServices="Server" />
+    <Protobuf Include="Protos\horde\job_rpc.proto" ProtoRoot="Protos" GrpcServices="Server" AdditionalImportDirs="..\..\Shared\EpicGames.Horde\Protos" />
     <Protobuf Include="Protos\horde\job_rpc_messages.proto" ProtoRoot="Protos" />
     <Protobuf Include="Protos\horde\enrollment_rpc.proto" ProtoRoot="Protos" GrpcServices="Server" />
     <Protobuf Include="Protos\horde\enrollment_rpc_messages.proto" ProtoRoot="Protos" />
     <Protobuf Include="Protos\horde\horde_rpc.proto" ProtoRoot="Protos" GrpcServices="Server" />
     <Protobuf Include="Protos\horde\horde_rpc_messages.proto" ProtoRoot="Protos" />
-    <Protobuf Include="Protos\horde\log_rpc.proto" ProtoRoot="Protos" GrpcServices="Server" />
-    <Protobuf Include="Protos\horde\log_rpc_messages.proto" ProtoRoot="Protos" />
     <Protobuf Include="Protos\horde\relay_rpc.proto" ProtoRoot="Protos" />
     <Protobuf Include="Protos\horde\common\agent_status.proto" ProtoRoot="Protos" />
     <Protobuf Include="Protos\horde\common\io_hash.proto" ProtoRoot="Protos" />
diff --git a/Engine/Source/Programs/Shared/EpicGames.Horde.Tests/ServerLogPacketBuilderTests.cs b/Engine/Source/Programs/Shared/EpicGames.Horde.Tests/ServerLogPacketBuilderTests.cs
new file mode 100644
--- /dev/null
+++ b/Engine/Source/Programs/Shared/EpicGames.Horde.Tests/ServerLogPacketBuilderTests.cs
@@ -0,0 +1,414 @@
+// Copyright Epic Games, Inc. All Rights Reserved.
+
+using System;
+using System.Collections.Generic;
+using System.Linq;
+using System.Text;
+using System.Threading;
+using System.Threading.Channels;
+using System.Threading.Tasks;
+using EpicGames.Core;
+using EpicGames.Horde;
+using EpicGames.Horde.Compute;
+using EpicGames.Horde.Logs;
+using EpicGames.Horde.Storage;
+using EpicGames.Horde.Storage.Bundles;
+using Grpc.Core;
+using Grpc.Net.Client;
+//using Horde.Agent.Utility;
+using Horde.Common.Rpc;
+using Microsoft.Extensions.Logging;
+using Microsoft.Extensions.Logging.Abstractions;
+using Microsoft.VisualStudio.TestTools.UnitTesting;
+
+namespace EpicGames.Horde.Tests
+{
+	[TestClass]
+	public class ServerLogPacketBuilderTests
+	{
+		class JsonLoggerImpl : ILogger
+		{
+			public List<LogEvent> _lines = new List<LogEvent>();
+
+			public IDisposable BeginScope<TState>(TState state) => null!;
+
+			public bool IsEnabled(LogLevel logLevel) => true;
+
+			public void Log<TState>(LogLevel logLevel, EventId eventId, TState state, Exception? exception, Func<TState, Exception?, string> formatter)
+			{
+				JsonLogEvent logEvent = JsonLogEvent.FromLoggerState(logLevel, eventId, state, exception, formatter);
+				_lines.Add(LogEvent.Read(logEvent.Data.Span));
+			}
+		}
+
+		[TestMethod]
+		public void JsonLogEventTest()
+		{
+			JsonLoggerImpl impl = new JsonLoggerImpl();
+			impl.LogInformation("Hello {Text}", "world");
+			impl.LogInformation("Hello {Text}", "world");
+
+			Assert.AreEqual(impl._lines.Count, 2);
+			Assert.AreEqual(impl._lines[0].ToString(), "Hello world");
+			Assert.AreEqual(impl._lines[0].GetProperty("Text").ToString()!, "world");
+			Assert.AreEqual(impl._lines[1].ToString(), "Hello world");
+			Assert.AreEqual(impl._lines[1].GetProperty("Text").ToString(), "world");
+		}
+
+		static string[] SplitMultiLineMessage(string message)
+		{
+			JsonLogEvent jsonLogEvent;
+			Assert.IsTrue(JsonLogEvent.TryParse(Encoding.UTF8.GetBytes(message), out jsonLogEvent));
+
+			ServerLogPacketBuilder writer = new ServerLogPacketBuilder();
+			int count = writer.SanitizeAndWriteEvent(jsonLogEvent);
+
+			string[] result = Encoding.UTF8.GetString(writer.CreatePacket().Item1.Span).Split('\n', StringSplitOptions.RemoveEmptyEntries);
+			Assert.AreEqual(count, result.Length);
+			return result;
+		}
+
+		[TestMethod]
+		public void MultiLineFormatTest()
+		{
+			const string Msg = @"{""level"":""Information"",""message"":""ignored"",""format"":""This\nis a\nmulti-{Line}-end\nlog message {Var1} {Var2}"",""properties"":{""Line"":""line\nsplit\nin\nfour"",""Var1"":123,""Var2"":{""$type"":""SourceFile"",""$text"":""D:\\build\\\u002B\u002BUE5\\Sync\\GenerateProjectFiles.bat"",""relativePath"":""GenerateProjectFiles.bat"",""depotPath"":""//UE5/Main/GenerateProjectFiles.bat@20392842""}}}";
+
+			string[] result = SplitMultiLineMessage(Msg);
+			Assert.AreEqual(7, result.Length);
+			Assert.AreEqual(@"{""level"":""Information"",""message"":""This"",""format"":""This"",""properties"":{""Var1"":123,""Var2"":{""$type"":""SourceFile"",""$text"":""D:\\build\\\u002B\u002BUE5\\Sync\\GenerateProjectFiles.bat"",""relativePath"":""GenerateProjectFiles.bat"",""depotPath"":""//UE5/Main/GenerateProjectFiles.bat@20392842""}},""line"":0,""lineCount"":7}", result[0]);
+			Assert.AreEqual(@"{""level"":""Information"",""message"":""is a"",""format"":""is a"",""properties"":{""Var1"":123,""Var2"":{""$type"":""SourceFile"",""$text"":""D:\\build\\\u002B\u002BUE5\\Sync\\GenerateProjectFiles.bat"",""relativePath"":""GenerateProjectFiles.bat"",""depotPath"":""//UE5/Main/GenerateProjectFiles.bat@20392842""}},""line"":1,""lineCount"":7}", result[1]);
+			Assert.AreEqual(@"{""level"":""Information"",""message"":""multi-line"",""format"":""multi-{Line$0}"",""properties"":{""Var1"":123,""Var2"":{""$type"":""SourceFile"",""$text"":""D:\\build\\\u002B\u002BUE5\\Sync\\GenerateProjectFiles.bat"",""relativePath"":""GenerateProjectFiles.bat"",""depotPath"":""//UE5/Main/GenerateProjectFiles.bat@20392842""},""Line$0"":""line""},""line"":2,""lineCount"":7}", result[2]);
+			Assert.AreEqual(@"{""level"":""Information"",""message"":""split"",""format"":""{Line$1}"",""properties"":{""Var1"":123,""Var2"":{""$type"":""SourceFile"",""$text"":""D:\\build\\\u002B\u002BUE5\\Sync\\GenerateProjectFiles.bat"",""relativePath"":""GenerateProjectFiles.bat"",""depotPath"":""//UE5/Main/GenerateProjectFiles.bat@20392842""},""Line$1"":""split""},""line"":3,""lineCount"":7}", result[3]);
+			Assert.AreEqual(@"{""level"":""Information"",""message"":""in"",""format"":""{Line$2}"",""properties"":{""Var1"":123,""Var2"":{""$type"":""SourceFile"",""$text"":""D:\\build\\\u002B\u002BUE5\\Sync\\GenerateProjectFiles.bat"",""relativePath"":""GenerateProjectFiles.bat"",""depotPath"":""//UE5/Main/GenerateProjectFiles.bat@20392842""},""Line$2"":""in""},""line"":4,""lineCount"":7}", result[4]);
+			Assert.AreEqual(@"{""level"":""Information"",""message"":""four-end"",""format"":""{Line$3}-end"",""properties"":{""Var1"":123,""Var2"":{""$type"":""SourceFile"",""$text"":""D:\\build\\\u002B\u002BUE5\\Sync\\GenerateProjectFiles.bat"",""relativePath"":""GenerateProjectFiles.bat"",""depotPath"":""//UE5/Main/GenerateProjectFiles.bat@20392842""},""Line$3"":""four""},""line"":5,""lineCount"":7}", result[5]);
+			Assert.AreEqual(@"{""level"":""Information"",""message"":""log message 123 D:\\build\\\u002B\u002BUE5\\Sync\\GenerateProjectFiles.bat"",""format"":""log message {Var1} {Var2}"",""properties"":{""Var1"":123,""Var2"":{""$type"":""SourceFile"",""$text"":""D:\\build\\\u002B\u002BUE5\\Sync\\GenerateProjectFiles.bat"",""relativePath"":""GenerateProjectFiles.bat"",""depotPath"":""//UE5/Main/GenerateProjectFiles.bat@20392842""}},""line"":6,""lineCount"":7}", result[6]);
+		}
+
+		[TestMethod]
+		public void InvalidCharactersText()
+		{
+			byte[] data = Encoding.UTF8.GetBytes(@"{""level"":""Information"",""message"":""test"",""format"":""GOOD\nBAD""}");
+			int idx = data.AsSpan().IndexOf(Encoding.UTF8.GetBytes("BAD"));
+			data[idx] = 0xef;
+			data[idx + 1] = 0xbb;
+
+			ServerLogPacketBuilder writer = new ServerLogPacketBuilder();
+			int count = writer.SanitizeAndWriteEvent(JsonLogEvent.Parse(data));
+			Assert.AreEqual(1, count);
+
+			string[] result = Encoding.UTF8.GetString(writer.CreatePacket().Item1.Span).Split('\n', StringSplitOptions.RemoveEmptyEntries);
+			Assert.AreEqual(1, result.Length);
+
+			LogEvent logEvent = LogEvent.Read(Encoding.UTF8.GetBytes(result[0]));
+
+			const string ExpectedError = "Invalid json log event: {\"level\":\"Information\",\"message\":\"test\",\"format\":\"GOOD\\n\\xef\\xbbD\"}";
+			Assert.AreEqual(ExpectedError, logEvent.ToString());
+		}
+
+		[TestMethod]
+		public void MultiLineMessageTest()
+		{
+			const string Msg = @"{""level"":""Information"",""message"":""This is \na multi-line string""}";
+
+			string[] result = SplitMultiLineMessage(Msg);
+			Assert.AreEqual(2, result.Length);
+			Assert.AreEqual(result[0], @"{""level"":""Information"",""message"":""This is "",""line"":0,""lineCount"":2}");
+			Assert.AreEqual(result[1], @"{""level"":""Information"",""message"":""a multi-line string"",""line"":1,""lineCount"":2}");
+		}
+
+		[TestMethod]
+		public void MultiLineDecoyTest()
+		{
+			const string Msg = @"{""level"":""Information"",""message"":""This is \\\\not a multi-line format string""}";
+
+			string[] result = SplitMultiLineMessage(Msg);
+			Assert.AreEqual(1, result.Length);
+			Assert.AreEqual(result[0], Msg);
+		}
+
+		class GrpcHelpers
+		{
+			class StreamReader<T> : IAsyncStreamReader<T> where T : class
+			{
+				readonly ChannelReader<T> _reader;
+				T? _current;
+
+				public T Current => _current ?? throw new InvalidOperationException();
+
+				public StreamReader(ChannelReader<T> reader)
+					=> _reader = reader;
+
+				public async Task<bool> MoveNext(CancellationToken cancellationToken)
+				{
+					while (await _reader.WaitToReadAsync(cancellationToken))
+					{
+						if (_reader.TryRead(out _current))
+						{
+							return true;
+						}
+					}
+					return false;
+				}
+			}
+
+			class StreamWriter<T> : IClientStreamWriter<T>
+			{
+				readonly ChannelWriter<T> _writer;
+
+				public StreamWriter(ChannelWriter<T> writer)
+					=> _writer = writer;
+
+				public WriteOptions? WriteOptions { get => throw new NotImplementedException(); set => throw new NotImplementedException(); }
+
+				public Task CompleteAsync()
+				{
+					_writer.Complete();
+					return Task.CompletedTask;
+				}
+
+				public async Task WriteAsync(T message)
+					=> await _writer.WriteAsync(message);
+			}
+
+			public static AsyncDuplexStreamingCall<TRequest, TResponse> CreateDuplexCall<TRequest, TResponse>(Func<ChannelReader<TRequest>, ChannelWriter<TResponse>, Task> func)
+				where TRequest : class
+				where TResponse : class
+			{
+				Channel<TRequest> requests = Channel.CreateUnbounded<TRequest>();
+				Channel<TResponse> responses = Channel.CreateUnbounded<TResponse>();
+
+				BackgroundTask runner = BackgroundTask.StartNew(async ctx =>
+				{
+					await func(requests.Reader, responses.Writer);
+					responses.Writer.Complete();
+				});
+
+				return new AsyncDuplexStreamingCall<TRequest, TResponse>(new StreamWriter<TRequest>(requests.Writer), new StreamReader<TResponse>(responses.Reader), Task.FromResult<Metadata>(null!), () => Status.DefaultSuccess, () => null!, () => runner.DisposeAsync().AsTask().Wait());
+			}
+		}
+
+		class FakeLogRpcClient : LogRpc.LogRpcClient
+		{
+			public Dictionary<LogId, BlobLocator> Logs { get; } = new Dictionary<LogId, BlobLocator>();
+
+			public override AsyncUnaryCall<RpcUpdateLogResponse> UpdateLogAsync(RpcUpdateLogRequest request, CallOptions options)
+			{
+				Logs[LogId.Parse(request.LogId)] = new BlobLocator(request.TargetLocator);
+				return new AsyncUnaryCall<RpcUpdateLogResponse>(Task.FromResult(new RpcUpdateLogResponse()), Task.FromResult(new Metadata()), () => Status.DefaultSuccess, () => new Metadata(), () => { });
+			}
+
+			public override AsyncDuplexStreamingCall<RpcUpdateLogTailRequest, RpcUpdateLogTailResponse> UpdateLogTail(CallOptions options)
+			{
+				return GrpcHelpers.CreateDuplexCall<RpcUpdateLogTailRequest, RpcUpdateLogTailResponse>(async (reader, writer) =>
+				{
+					while (await reader.WaitToReadAsync())
+					{
+						reader.TryRead(out RpcUpdateLogTailRequest? request);
+					}
+				});
+			}
+		}
+
+		class FakeHordeClient : IHordeClient
+		{
+			public FakeLogRpcClient LogRpc { get; } = new FakeLogRpcClient();
+			public Dictionary<string, BundleStorageClient> StorageClients { get; } = new Dictionary<string, BundleStorageClient>();
+
+			public Uri ServerUrl => throw new NotImplementedException();
+
+			public Task<bool> LoginAsync(bool allowLogin, CancellationToken cancellationToken)
+				=> throw new NotImplementedException();
+
+			public HordeHttpClient CreateHttpClient()
+				=> throw new NotImplementedException();
+
+			public IComputeClient CreateComputeClient()
+				=> throw new NotImplementedException();
+
+			public IStorageClient CreateStorageClient(string relativePath, string? accessToken = null)
+			{
+				BundleStorageClient? storageClient;
+				if (!StorageClients.TryGetValue(relativePath, out storageClient))
+				{
+					storageClient = BundleStorageClient.CreateInMemory(NullLogger.Instance);
+					StorageClients.Add(relativePath, storageClient);
+				}
+				return storageClient;
+			}
+
+			public ValueTask DisposeAsync()
+			{
+				foreach (BundleStorageClient storageClient in StorageClients.Values)
+				{
+					storageClient.Dispose();
+				}
+				StorageClients.Clear();
+				return default;
+			}
+
+			public Task<string?> GetAccessTokenAsync(bool interactive, CancellationToken cancellationToken = default)
+				=> throw new NotImplementedException();
+
+			public Task<GrpcChannel> CreateGrpcChannelAsync(CancellationToken cancellationToken = default)
+				=> throw new NotImplementedException();
+
+			public Task<TClient> CreateGrpcClientAsync<TClient>(CancellationToken cancellationToken = default) where TClient : ClientBase<TClient>
+			{
+				if (typeof(TClient) == typeof(LogRpc.LogRpcClient))
+				{
+					return Task.FromResult<TClient>((TClient)(object)LogRpc);
+				}
+				throw new NotImplementedException();
+			}
+
+			public bool HasValidAccessToken()
+				=> throw new NotImplementedException();
+		}
+
+		[TestMethod]
+		public async Task StorageLoggerTestAsync()
+		{
+			await using BundleCache cache = new BundleCache();
+			await using FakeHordeClient hordeClient = new FakeHordeClient();
+
+			LogId logId = default;
+
+			const int Count = 20000;
+			await using (ServerLogger logger = new ServerLogger(hordeClient, logId, LogLevel.Information, NullLogger.Instance))
+			{
+				for (int idx = 0; idx < Count; idx++)
+				{
+					logger.LogInformation("Testing {Number}", idx);
+				}
+			}
+
+			// Read the log
+			using IStorageClient storageClient = hordeClient.CreateStorageClient(logId);
+			LogNode file = await storageClient.CreateBlobHandle(hordeClient.LogRpc.Logs[logId]).ReadBlobAsync<LogNode>();
+
+			// Check the index text
+			List<Utf8String> extractedIndexText = new List<Utf8String>();
+
+			LogIndexNode index = await file.IndexRef.ReadBlobAsync();
+			foreach (LogChunkRef block in index.PlainTextChunkRefs)
+			{
+				LogChunkNode text = await block.Target.ReadBlobAsync();
+				extractedIndexText.AddRange(text.Lines);
+			}
+
+			Assert.AreEqual(Count, extractedIndexText.Count);
+			for (int idx = 0; idx < Count; idx++)
+			{
+				Assert.AreEqual(extractedIndexText[idx], new Utf8String($"Testing {idx}"));
+			}
+
+			// Check the body text
+			List<string> extractedBodyText = new List<string>();
+
+			foreach (LogChunkRef blockRef in file.TextChunkRefs)
+			{
+				LogChunkNode blockText = await blockRef.Target.ReadBlobAsync();
+				foreach (Utf8String line in blockText.Lines)
+				{
+					LogEvent logEvent = LogEvent.Read(line.Span);
+					extractedBodyText.Add(logEvent.ToString());
+				}
+			}
+
+			Assert.AreEqual(Count, extractedBodyText.Count);
+			for (int idx = 0; idx < Count; idx++)
+			{
+				Assert.AreEqual(extractedBodyText[idx], $"Testing {idx}");
+			}
+		}
+
+		[TestMethod]
+		public void LongMessageTest()
+		{
+			DateTime time = new DateTime(2023, 1, 1, 0, 0, 0, 0, DateTimeKind.Utc);
+
+			string longLineA = String.Join("\n", Enumerable.Range(0, 5).Select(x => $"A{x}"));
+			string longLineB = String.Join("\n", Enumerable.Range(0, 5).Select(x => $"B{x}"));
+			string longLineC = String.Join("\n", Enumerable.Range(0, 5).Select(x => $"C{x}"));
+
+			Dictionary<string, object> properties = new Dictionary<string, object>
+			{
+				["LongLineA"] = longLineA,
+				["LongLineB"] = longLineB,
+				["LongLineC"] = longLineC,
+			};
+
+			LogEvent baseEvent = new LogEvent(time, LogLevel.Information, default, $"start {longLineA} x {longLineB} y\nz {longLineC} end", "start {LongLineA} x {LongLineB} y\nz {LongLineC} end", properties, null);
+
+			ServerLogPacketBuilder writer = new ServerLogPacketBuilder();
+			writer.SanitizeAndWriteEvent(new JsonLogEvent(baseEvent));
+
+			string[] output = Encoding.UTF8.GetString(writer.CreatePacket().Item1.Span).Split('\n');
+			string[] expected =
+			{
+				@"{""time"":""2023-01-01T00:00:00"",""level"":""Information"",""message"":""start A0"",""format"":""start {LongLineA$0}"",""properties"":{""LongLineA$0"":""A0""},""line"":0,""lineCount"":14}",
+				@"{""time"":""2023-01-01T00:00:00"",""level"":""Information"",""message"":""A1"",""format"":""{LongLineA$1}"",""properties"":{""LongLineA$1"":""A1""},""line"":1,""lineCount"":14}",
+				@"{""time"":""2023-01-01T00:00:00"",""level"":""Information"",""message"":""A2"",""format"":""{LongLineA$2}"",""properties"":{""LongLineA$2"":""A2""},""line"":2,""lineCount"":14}",
+				@"{""time"":""2023-01-01T00:00:00"",""level"":""Information"",""message"":""A3"",""format"":""{LongLineA$3}"",""properties"":{""LongLineA$3"":""A3""},""line"":3,""lineCount"":14}",
+				@"{""time"":""2023-01-01T00:00:00"",""level"":""Information"",""message"":""A4 x B0"",""format"":""{LongLineA$4} x {LongLineB$0}"",""properties"":{""LongLineA$4"":""A4"",""LongLineB$0"":""B0""},""line"":4,""lineCount"":14}",
+				@"{""time"":""2023-01-01T00:00:00"",""level"":""Information"",""message"":""B1"",""format"":""{LongLineB$1}"",""properties"":{""LongLineB$1"":""B1""},""line"":5,""lineCount"":14}",
+				@"{""time"":""2023-01-01T00:00:00"",""level"":""Information"",""message"":""B2"",""format"":""{LongLineB$2}"",""properties"":{""LongLineB$2"":""B2""},""line"":6,""lineCount"":14}",
+				@"{""time"":""2023-01-01T00:00:00"",""level"":""Information"",""message"":""B3"",""format"":""{LongLineB$3}"",""properties"":{""LongLineB$3"":""B3""},""line"":7,""lineCount"":14}",
+				@"{""time"":""2023-01-01T00:00:00"",""level"":""Information"",""message"":""B4 y"",""format"":""{LongLineB$4} y"",""properties"":{""LongLineB$4"":""B4""},""line"":8,""lineCount"":14}",
+				@"{""time"":""2023-01-01T00:00:00"",""level"":""Information"",""message"":""z C0"",""format"":""z {LongLineC$0}"",""properties"":{""LongLineC$0"":""C0""},""line"":9,""lineCount"":14}",
+				@"{""time"":""2023-01-01T00:00:00"",""level"":""Information"",""message"":""C1"",""format"":""{LongLineC$1}"",""properties"":{""LongLineC$1"":""C1""},""line"":10,""lineCount"":14}",
+				@"{""time"":""2023-01-01T00:00:00"",""level"":""Information"",""message"":""C2"",""format"":""{LongLineC$2}"",""properties"":{""LongLineC$2"":""C2""},""line"":11,""lineCount"":14}",
+				@"{""time"":""2023-01-01T00:00:00"",""level"":""Information"",""message"":""C3"",""format"":""{LongLineC$3}"",""properties"":{""LongLineC$3"":""C3""},""line"":12,""lineCount"":14}",
+				@"{""time"":""2023-01-01T00:00:00"",""level"":""Information"",""message"":""C4 end"",""format"":""{LongLineC$4} end"",""properties"":{""LongLineC$4"":""C4""},""line"":13,""lineCount"":14}",
+				""
+			};
+
+			for (int idx = 0; idx < output.Length; idx++)
+			{
+				Assert.AreEqual(expected[idx], output[idx]);
+			}
+		}
+
+		[TestMethod]
+		public void MaxPacketLengthTest()
+		{
+			DateTime time = new DateTime(2023, 1, 1, 0, 0, 0, 0, DateTimeKind.Utc);
+
+			LogEvent baseEvent = new LogEvent(time, LogLevel.Information, default, "abcdefghijklmnopqrstuvwxyzABCDEFGHIJKLMNOPQRSTUVWXYZ", null, null, null);
+
+			ServerLogPacketBuilder writer = new ServerLogPacketBuilder(maxPacketLength: 85);
+			writer.SanitizeAndWriteEvent(new JsonLogEvent(baseEvent));
+			writer.SanitizeAndWriteEvent(new JsonLogEvent(baseEvent));
+
+			for (int idx = 0; idx < 2; idx++)
+			{
+				string output = Encoding.UTF8.GetString(writer.CreatePacket().Item1.Span);
+				string expected = @"{""time"":""2023-01-01T00:00:00"",""level"":""Information"",""message"":""abcdefghijklmnopqrstuvwxyzABCDEFGHIJKLMNOPQRSTUVWXYZ""}" + "\n";
+				Assert.AreEqual(expected, output);
+			}
+		}
+
+		[TestMethod]
+		public void MaxLineLengthTest()
+		{
+			DateTime time = new DateTime(2023, 1, 1, 0, 0, 0, 0, DateTimeKind.Utc);
+
+			LogEvent baseEvent = new LogEvent(time, LogLevel.Information, default, "abcdefghijklmnopqrstuvwxyzABCDEFGHIJKLMNOPQRSTUVWXYZ", null, null, null);
+
+			ServerLogPacketBuilder writer = new ServerLogPacketBuilder(85);
+			writer.SanitizeAndWriteEvent(new JsonLogEvent(baseEvent));
+
+			string[] output = Encoding.UTF8.GetString(writer.CreatePacket().Item1.Span).Split('\n');
+			string[] expected =
+			{
+				@"{""time"":""2023-01-01T00:00:00"",""level"":""Information"",""message"":""abcdefghijklmnopqrstuvwxyzABCDEFGHI [...]""}",
+				""
+			};
+
+			for (int idx = 0; idx < output.Length; idx++)
+			{
+				Assert.AreEqual(expected[idx], output[idx]);
+			}
+		}
+	}
+}
diff --git a/Engine/Source/Programs/Shared/EpicGames.Horde/EpicGames.Horde.csproj b/Engine/Source/Programs/Shared/EpicGames.Horde/EpicGames.Horde.csproj
--- a/Engine/Source/Programs/Shared/EpicGames.Horde/EpicGames.Horde.csproj
+++ b/Engine/Source/Programs/Shared/EpicGames.Horde/EpicGames.Horde.csproj
@@ -40,7 +40,12 @@
   <ItemGroup>
     <PackageReference Include="BitFaster.Caching" Version="2.4.1" />
     <PackageReference Include="Dapper" Version="2.1.24" />
+    <PackageReference Include="Google.Protobuf" Version="3.25.1" />
     <PackageReference Include="Grpc.Net.Client" Version="2.59.0" />
+    <PackageReference Include="Grpc.Tools" Version="2.59.0">
+      <PrivateAssets>all</PrivateAssets>
+      <IncludeAssets>runtime; build; native; contentfiles; analyzers; buildtransitive</IncludeAssets>
+    </PackageReference>
     <PackageReference Include="K4os.Compression.LZ4" Version="1.2.16" />
     <PackageReference Include="Microsoft.CodeAnalysis.CSharp.Workspaces" Version="4.2.0" />
     <PackageReference Include="Microsoft.Data.Sqlite" Version="8.0.0" />
@@ -62,4 +67,8 @@
     <ProjectReference Include="..\EpicGames.Oodle\EpicGames.Oodle.csproj" />
     <ProjectReference Include="..\EpicGames.Serialization\EpicGames.Serialization.csproj" />
   </ItemGroup>
+  <ItemGroup>
+    <Protobuf Include="Protos\horde\log_rpc.proto" ProtoRoot="Protos" GrpcServices="Both" Access="Public" />
+    <Protobuf Include="Protos\horde\log_rpc_messages.proto" ProtoRoot="Protos" Access="Public" />
+  </ItemGroup>
 </Project>
\ No newline at end of file
diff --git a/Engine/Source/Programs/Shared/EpicGames.Horde/Logs/ServerLogPacketBuilder.cs b/Engine/Source/Programs/Shared/EpicGames.Horde/Logs/ServerLogPacketBuilder.cs
new file mode 100644
--- /dev/null
+++ b/Engine/Source/Programs/Shared/EpicGames.Horde/Logs/ServerLogPacketBuilder.cs
@@ -0,0 +1,354 @@
+// Copyright Epic Games, Inc. All Rights Reserved.
+
+using System;
+using System.Buffers;
+using System.Collections.Generic;
+using System.Diagnostics;
+using System.Diagnostics.CodeAnalysis;
+using System.Linq;
+using System.Text;
+using System.Text.Json;
+using System.Text.Json.Nodes;
+using EpicGames.Core;
+using Microsoft.Extensions.Logging;
+
+namespace EpicGames.Horde.Logs
+{
+	using JsonObject = System.Text.Json.Nodes.JsonObject;
+
+	/// <summary>
+	/// Utility class to split log events into separate lines and buffer them for writing to the server
+	/// </summary>
+	public class ServerLogPacketBuilder
+	{
+		[DebuggerDisplay("{Format}")]
+		class FormattedLine
+		{
+			public string Format { get; set; }
+			public Dictionary<string, string> Properties { get; } = new Dictionary<string, string>(StringComparer.OrdinalIgnoreCase);
+
+			public FormattedLine(string format)
+			{
+				Format = format;
+			}
+		}
+
+		static readonly string s_messagePropertyName = LogEventPropertyName.Message.ToString();
+		static readonly string s_formatPropertyName = LogEventPropertyName.Format.ToString();
+		static readonly string s_linePropertyName = LogEventPropertyName.Line.ToString();
+		static readonly string s_lineCountPropertyName = LogEventPropertyName.LineCount.ToString();
+
+		static readonly Utf8String s_escapedNewline = new Utf8String("\\n");
+
+		readonly List<JsonLogEvent> _logEvents = new List<JsonLogEvent>();
+		readonly ArrayBufferWriter<byte> _lineWriter;
+		readonly ArrayBufferWriter<byte> _packetWriter;
+		int _packetLength;
+
+		/// <summary>
+		/// Current packet length
+		/// </summary>
+		public int PacketLength => _packetLength;
+
+		/// <summary>
+		/// Maximum length of an individual line
+		/// </summary>
+		public int MaxLineLength { get; }
+
+		/// <summary>
+		/// Maximum size of a packet
+		/// </summary>
+		public int MaxPacketLength { get; }
+
+		/// <summary>
+		/// Constructor
+		/// </summary>
+		/// <param name="maxLineLength">Maximum length for an individual line</param>
+		/// <param name="maxPacketLength">Maximum length for a packet</param>
+		public ServerLogPacketBuilder(int maxLineLength = 64 * 1024, int maxPacketLength = 256 * 1024)
+		{
+			MaxLineLength = maxLineLength;
+			_lineWriter = new ArrayBufferWriter<byte>(maxLineLength);
+
+			MaxPacketLength = maxPacketLength;
+			_packetWriter = new ArrayBufferWriter<byte>(maxPacketLength);
+		}
+
+		/// <summary>
+		/// Creates a packet from the current data
+		/// </summary>
+		/// <returns>Packet data and number of lines written</returns>
+		public (ReadOnlyMemory<byte>, int) CreatePacket()
+		{
+			_packetWriter.Clear();
+
+			int eventCount = 0;
+			for (; eventCount < _logEvents.Count && (eventCount == 0 || _packetWriter.WrittenCount + _logEvents[eventCount].Data.Length + 1 < MaxPacketLength); eventCount++)
+			{
+				JsonLogEvent jsonLogEvent = _logEvents[eventCount];
+
+				Span<byte> span = _packetWriter.GetSpan(jsonLogEvent.Data.Length + 1);
+				jsonLogEvent.Data.Span.CopyTo(span);
+				span[jsonLogEvent.Data.Length] = (byte)'\n';
+				_packetWriter.Advance(jsonLogEvent.Data.Length + 1);
+
+				_packetLength -= jsonLogEvent.Data.Length + 1;
+			}
+			_logEvents.RemoveRange(0, eventCount);
+
+			return (_packetWriter.WrittenMemory, eventCount);
+		}
+
+		/// <summary>
+		/// Writes an event
+		/// </summary>
+		/// <param name="jsonLogEvent">Event to write</param>
+		public int SanitizeAndWriteEvent(JsonLogEvent jsonLogEvent)
+		{
+			try
+			{
+				return SanitizeAndWriteEventInternal(jsonLogEvent);
+			}
+			catch (Exception ex)
+			{
+				StringBuilder escapedLineBuilder = new StringBuilder();
+
+				ReadOnlySpan<byte> span = jsonLogEvent.Data.Span;
+				for (int idx = 0; idx < span.Length; idx++)
+				{
+					if (span[idx] >= 32 && span[idx] <= 127)
+					{
+						escapedLineBuilder.Append((char)span[idx]);
+					}
+					else
+					{
+						escapedLineBuilder.Append($"\\x{span[idx]:x2}");
+					}
+				}
+
+				string escapedLine = escapedLineBuilder.ToString();
+				KeyValuePair<string, object>[] properties = new[] { new KeyValuePair<string, object>("Text", escapedLine) };
+				LogEvent newLogEvent = new LogEvent(DateTime.UtcNow, LogLevel.Error, default, $"Invalid json log event: {escapedLineBuilder}", "Invalid json log event: {Text}", properties, LogException.FromException(ex));
+				JsonLogEvent newJsonLogEvent = new JsonLogEvent(newLogEvent);
+
+				return SanitizeAndWriteEventInternal(newJsonLogEvent);
+			}
+		}
+
+		int SanitizeAndWriteEventInternal(JsonLogEvent jsonLogEvent)
+		{
+			ReadOnlySpan<byte> span = jsonLogEvent.Data.Span;
+			if (jsonLogEvent.LineCount == 1 && span.IndexOf(s_escapedNewline) != -1)
+			{
+				JsonObject obj = (JsonObject)JsonNode.Parse(span)!;
+
+				JsonValue? formatValue = obj["format"] as JsonValue;
+				if (formatValue != null && formatValue.TryGetValue(out string? format))
+				{
+					return WriteEventWithFormat(jsonLogEvent, obj, format);
+				}
+
+				JsonValue? messageValue = obj["message"] as JsonValue;
+				if (messageValue != null && messageValue.TryGetValue(out string? message))
+				{
+					return WriteEventWithMessage(jsonLogEvent, obj, message);
+				}
+			}
+
+			WriteEventInternal(jsonLogEvent);
+			return 1;
+		}
+
+		/// <summary>
+		/// Writes an event with a format string, splitting it into multiple lines if necessary
+		/// </summary>
+		int WriteEventWithFormat(JsonLogEvent jsonLogEvent, JsonObject obj, string format)
+		{
+			// There is an object containing "common" properties
+			IEnumerable<KeyValuePair<string, object?>> propertyValueList = Enumerable.Empty<KeyValuePair<string, object?>>();
+
+			// Split the format string into lines
+			string[] formatLines = format.Split('\n');
+			List<FormattedLine> lines = new List<FormattedLine>();
+
+			// Split all the multi-line properties into separate properties
+			JsonObject? properties = obj["properties"] as JsonObject;
+			if (properties == null)
+			{
+				lines.AddRange(formatLines.Select(x => new FormattedLine(x)));
+			}
+			else
+			{
+				// Get all the current property values
+				Dictionary<string, string> propertyValues = new Dictionary<string, string>(StringComparer.OrdinalIgnoreCase);
+				foreach ((string name, JsonNode? node) in properties)
+				{
+					string value = String.Empty;
+					if (node != null)
+					{
+						if (node is JsonObject valueObject)
+						{
+							value = valueObject["$text"]?.ToString() ?? String.Empty;
+						}
+						else
+						{
+							value = node.ToString();
+						}
+					}
+					propertyValues[name] = value;
+				}
+
+				// Keep splitting properties in the last line
+				foreach (string formatLine in formatLines)
+				{
+					lines.Add(new FormattedLine(formatLine));
+					for (; ; )
+					{
+						FormattedLine line = lines[^1];
+						if (!TryGetNextMultiLineProperty(line.Format, propertyValues, out string? prefix, out string? suffix, out string? propertyName, out string[]? propertyLines))
+						{
+							break;
+						}
+
+						properties!.Remove(propertyName);
+
+						StringBuilder builder = new StringBuilder();
+						builder.Append(prefix);
+
+						for (int lineNum = 0; ; lineNum++)
+						{
+							string newPropertyName = $"{propertyName}${lineNum}";
+							builder.Append($"{{{newPropertyName}}}");
+
+							line.Properties.Add(newPropertyName, propertyLines[lineNum]);
+
+							if (lineNum + 1 >= propertyLines.Length)
+							{
+								break;
+							}
+
+							line.Format = builder.ToString();
+							builder.Clear();
+
+							line = new FormattedLine(String.Empty);
+							lines.Add(line);
+						}
+
+						builder.Append(suffix);
+						line.Format = builder.ToString();
+					}
+				}
+
+				// Get the enumerable property list for formatting
+				propertyValueList = propertyValues.Select(x => new KeyValuePair<string, object?>(x.Key, x.Value));
+			}
+
+			// Finally split the format string into multiple lines
+			for (int idx = 0; idx < lines.Count; idx++)
+			{
+				FormattedLine line = lines[idx];
+				foreach ((string name, string value) in line.Properties)
+				{
+					properties![name] = value;
+				}
+
+				string message = MessageTemplate.Render(lines[idx].Format, propertyValueList.Concat(line.Properties.Select(x => new KeyValuePair<string, object?>(x.Key, x.Value))));
+				WriteSingleEvent(jsonLogEvent, obj, message, lines[idx].Format, idx, lines.Count);
+
+				foreach ((string name, _) in line.Properties)
+				{
+					properties!.Remove(name);
+				}
+			}
+			return lines.Count;
+		}
+
+		static bool TryGetNextMultiLineProperty(string format, Dictionary<string, string> properties, [NotNullWhen(true)] out string? prefix, [NotNullWhen(true)] out string? suffix, [NotNullWhen(true)] out string? propertyName, [NotNullWhen(true)] out string[]? propertyLines)
+		{
+			int nameStart = -1;
+			for (int idx = 0; idx < format.Length; idx++)
+			{
+				if (format[idx] == '{')
+				{
+					nameStart = idx + 1;
+				}
+				else if (format[idx] == '}' && nameStart != -1)
+				{
+					string name = format.Substring(nameStart, idx - nameStart);
+					if (properties.TryGetValue(name, out string? text))
+					{
+						if (text.Contains('\n', StringComparison.Ordinal))
+						{
+							prefix = format.Substring(0, nameStart - 1);
+							suffix = format.Substring(idx + 1);
+							propertyName = name;
+							propertyLines = text.Split('\n');
+							return true;
+						}
+					}
+				}
+			}
+
+			prefix = null;
+			suffix = null;
+			propertyName = null;
+			propertyLines = null;
+			return false;
+		}
+
+		int WriteEventWithMessage(JsonLogEvent jsonLogEvent, JsonObject obj, string message)
+		{
+			string[] lines = message.Split('\n');
+			for (int idx = 0; idx < lines.Length; idx++)
+			{
+				WriteSingleEvent(jsonLogEvent, obj, lines[idx], null, idx, lines.Length);
+			}
+			return lines.Length;
+		}
+
+		void WriteSingleEvent(JsonLogEvent jsonLogEvent, JsonObject obj, string message, string? format, int line, int lineCount)
+		{
+			obj[s_messagePropertyName] = message;
+			if (format != null)
+			{
+				obj[s_formatPropertyName] = format;
+			}
+			if (lineCount > 1)
+			{
+				obj[s_linePropertyName] = line;
+				obj[s_lineCountPropertyName] = lineCount;
+			}
+
+			_lineWriter.Clear();
+			using (Utf8JsonWriter jsonWriter = new Utf8JsonWriter(_lineWriter))
+			{
+				obj.WriteTo(jsonWriter);
+			}
+
+			JsonLogEvent nextEvent = new JsonLogEvent(jsonLogEvent.Level, jsonLogEvent.EventId, line, lineCount, _lineWriter.WrittenMemory.ToArray());
+			WriteEventInternal(nextEvent);
+		}
+
+		void WriteEventInternal(JsonLogEvent jsonLogEvent)
+		{
+			if (jsonLogEvent.Data.Length > MaxLineLength)
+			{
+				LogEvent logEvent = LogEvent.Read(jsonLogEvent.Data.Span);
+
+				int maxMessageLength = Math.Max(10, MaxLineLength - 50);
+
+				string message = logEvent.Message;
+				if (message.Length > maxMessageLength)
+				{
+					message = message.Substring(0, maxMessageLength);
+				}
+
+				logEvent = new LogEvent(logEvent.Time, logEvent.Level, logEvent.Id, logEvent.LineIndex, logEvent.LineCount, $"{message} [...]", null, null, logEvent.Exception);
+				jsonLogEvent = new JsonLogEvent(logEvent);
+			}
+
+			_logEvents.Add(jsonLogEvent);
+			_packetLength += jsonLogEvent.Data.Length + 1;
+		}
+	}
+}
diff --git a/Engine/Source/Programs/Shared/EpicGames.Horde/Logs/ServerLogger.cs b/Engine/Source/Programs/Shared/EpicGames.Horde/Logs/ServerLogger.cs
new file mode 100644
--- /dev/null
+++ b/Engine/Source/Programs/Shared/EpicGames.Horde/Logs/ServerLogger.cs
@@ -0,0 +1,442 @@
+// Copyright Epic Games, Inc. All Rights Reserved.
+
+using System;
+using System.Collections.Generic;
+using System.Text;
+using System.Threading;
+using System.Threading.Channels;
+using System.Threading.Tasks;
+using EpicGames.Core;
+using EpicGames.Horde;
+using EpicGames.Horde.Logs;
+using EpicGames.Horde.Storage;
+using Google.Protobuf;
+using Grpc.Core;
+using Horde.Common.Rpc;
+using Microsoft.Extensions.Logging;
+
+namespace EpicGames.Horde.Logs
+{
+	using ByteString = Google.Protobuf.ByteString;
+
+	/// <summary>
+	/// Interface for a log device
+	/// </summary>
+	public interface IServerLogger : ILogger, IAsyncDisposable
+	{
+		/// <summary>
+		/// Flushes the logger with the server and stops the background work
+		/// </summary>
+		Task StopAsync();
+	}
+
+	/// <summary>
+	/// Class to handle uploading log data to the server in the background
+	/// </summary>
+	public sealed class ServerLogger : IServerLogger
+	{
+		const int FlushLength = 1024 * 1024;
+
+		readonly IHordeClient _hordeClient;
+		readonly LogId _logId;
+		readonly LogBuilder _builder;
+		readonly IStorageClient _store;
+		readonly IBlobWriter _writer;
+
+		int _bufferLength;
+
+		// Tailing task
+		readonly Task _tailTask;
+		AsyncEvent _tailTaskStop;
+		readonly AsyncEvent _newTailDataEvent = new AsyncEvent();
+
+		readonly LogLevel _outputLevel;
+		readonly ILogger _internalLogger;
+		readonly Channel<JsonLogEvent> _dataChannel;
+		Task? _dataWriter;
+
+		/// <summary>
+		/// Constructor
+		/// </summary>
+		/// <param name="hordeClient">Horde instance to write to</param>
+		/// <param name="logId">The log id to write to</param>
+		/// <param name="minimumLevel">Minimum level for output</param>
+		/// <param name="internalLogger">Logger for systemic messages</param>
+		public ServerLogger(IHordeClient hordeClient, LogId logId, LogLevel minimumLevel, ILogger internalLogger)
+		{
+			_hordeClient = hordeClient;
+			_logId = logId;
+			_builder = new LogBuilder(LogFormat.Json, internalLogger);
+			_store = _hordeClient.CreateStorageClient(logId);
+			_writer = _store.CreateBlobWriter();
+
+			_tailTaskStop = new AsyncEvent();
+			_tailTask = Task.Run(() => TickTailAsync());
+
+			_logId = logId;
+			_outputLevel = minimumLevel;
+			_internalLogger = internalLogger;
+			_dataChannel = Channel.CreateUnbounded<JsonLogEvent>();
+			_dataWriter = Task.Run(() => RunDataWriterAsync());
+		}
+
+		/// <inheritdoc/>
+		public void Log<TState>(LogLevel logLevel, EventId eventId, TState state, Exception? exception, Func<TState, Exception?, string> formatter)
+		{
+			JsonLogEvent jsonLogEvent = JsonLogEvent.FromLoggerState(logLevel, eventId, state, exception, formatter);
+			WriteFormattedEvent(jsonLogEvent);
+		}
+
+		/// <inheritdoc/>
+		public bool IsEnabled(LogLevel logLevel) => logLevel >= _outputLevel;
+
+		/// <inheritdoc/>
+		public IDisposable BeginScope<TState>(TState state)
+			=> null!;
+
+		private void WriteFormattedEvent(JsonLogEvent jsonLogEvent)
+		{
+			if (!_dataChannel.Writer.TryWrite(jsonLogEvent))
+			{
+				throw new InvalidOperationException("Expected unbounded writer to complete immediately");
+			}
+		}
+
+		/// <summary>
+		/// Stops the log writer's background task
+		/// </summary>
+		/// <returns>Async task</returns>
+		public async Task StopAsync()
+		{
+			if (_dataWriter != null)
+			{
+				_dataChannel.Writer.TryComplete();
+				await _dataWriter;
+				_dataWriter = null;
+			}
+		}
+
+		/// <summary>
+		/// Dispose of this object. Call StopAsync() to stop asynchronously.
+		/// </summary>
+		public async ValueTask DisposeAsync()
+		{
+			await StopAsync();
+
+			_internalLogger.LogInformation("Disposing json log task");
+
+			if (_tailTaskStop != null)
+			{
+				_tailTaskStop.Latch();
+				_newTailDataEvent.Latch();
+
+				await _tailTask;
+				_tailTaskStop = null!;
+			}
+
+			if (_writer != null)
+			{
+				await _writer.DisposeAsync();
+			}
+
+			if (_store != null)
+			{
+				_store.Dispose();
+			}
+		}
+
+		/// <summary>
+		/// Upload the log data to the server in the background
+		/// </summary>
+		/// <returns>Async task</returns>
+		async Task RunDataWriterAsync()
+		{
+			// Current position and line number in the log file
+			long packetOffset = 0;
+			int packetLineIndex = 0;
+
+			// Index of the next line to write to the log
+			int nextLineIndex = 0;
+
+			// Total number of errors and warnings
+			const int MaxErrors = 200;
+			int numErrors = 0;
+			const int MaxWarnings = 200;
+			int numWarnings = 0;
+
+			// Buffer for events read in a single iteration
+			ServerLogPacketBuilder writer = new ServerLogPacketBuilder();
+			List<RpcCreateLogEventRequest> events = new List<RpcCreateLogEventRequest>();
+
+			// Whether we've written the flush command
+			for (; ; )
+			{
+				events.Clear();
+
+				// Get the next data
+				Task waitTask = Task.Delay(TimeSpan.FromSeconds(2.0));
+				while (writer.PacketLength < writer.MaxPacketLength)
+				{
+					JsonLogEvent jsonLogEvent;
+					if (_dataChannel.Reader.TryRead(out jsonLogEvent))
+					{
+						int lineCount = writer.SanitizeAndWriteEvent(jsonLogEvent);
+						if (jsonLogEvent.LineIndex == 0)
+						{
+							if (jsonLogEvent.Level == LogLevel.Warning && ++numWarnings <= MaxWarnings)
+							{
+								AddEvent(jsonLogEvent.Data.Span, nextLineIndex, Math.Max(lineCount, jsonLogEvent.LineCount), LogEventSeverity.Warning, events);
+							}
+							else if ((jsonLogEvent.Level == LogLevel.Error || jsonLogEvent.Level == LogLevel.Critical) && ++numErrors <= MaxErrors)
+							{
+								AddEvent(jsonLogEvent.Data.Span, nextLineIndex, Math.Max(lineCount, jsonLogEvent.LineCount), LogEventSeverity.Error, events);
+							}
+						}
+						nextLineIndex += lineCount;
+					}
+					{
+						Task<bool> readTask = _dataChannel.Reader.WaitToReadAsync().AsTask();
+						if (await Task.WhenAny(readTask, waitTask) == waitTask)
+						{
+							break;
+						}
+						if (!await readTask)
+						{
+							break;
+						}
+					}
+				}
+
+				// Upload it to the server
+				if (writer.PacketLength > 0)
+				{
+					(ReadOnlyMemory<byte> packet, int packetLineCount) = writer.CreatePacket();
+					try
+					{
+						await WriteOutputAsync(_logId, packetOffset, packetLineIndex, packet, false, CancellationToken.None);
+						packetOffset += packet.Length;
+						packetLineIndex += packetLineCount;
+					}
+					catch (Exception ex)
+					{
+						_internalLogger.LogWarning(ex, "Unable to write data to server (log {LogId}, offset {Offset}, length {Length}, lines {StartLine}-{EndLine})", _logId, packetOffset, packet.Length, packetLineIndex, packetLineIndex + packetLineCount);
+					}
+				}
+
+				// Write all the events
+				if (events.Count > 0)
+				{
+					try
+					{
+						await WriteEventsAsync(events, CancellationToken.None);
+					}
+					catch (Exception ex)
+					{
+						_internalLogger.LogWarning(ex, "Unable to create events");
+					}
+				}
+
+				// Wait for more data to be available
+				if (writer.PacketLength <= 0 && !await _dataChannel.Reader.WaitToReadAsync())
+				{
+					try
+					{
+						await WriteOutputAsync(_logId, packetOffset, packetLineIndex, ReadOnlyMemory<byte>.Empty, true, CancellationToken.None);
+					}
+					catch (Exception ex)
+					{
+						_internalLogger.LogWarning(ex, "Unable to flush data to server (log {LogId}, offset {Offset})", _logId, packetOffset);
+					}
+					break;
+				}
+			}
+		}
+
+		void AddEvent(ReadOnlySpan<byte> span, int lineIndex, int lineCount, LogEventSeverity severity, List<RpcCreateLogEventRequest> events)
+		{
+			try
+			{
+				events.Add(new RpcCreateLogEventRequest{ Severity = (int) severity, LogId = _logId.ToString(), LineIndex = lineIndex, LineCount = lineCount });
+			}
+			catch (Exception ex)
+			{
+				_internalLogger.LogError(ex, "Exception while trying to parse line count from data ({Message})", Encoding.UTF8.GetString(span));
+			}
+		}
+
+		async Task TickTailAsync()
+		{
+			for (; ; )
+			{
+				try
+				{
+					await TickTailInternalAsync();
+					break;
+				}
+				catch (OperationCanceledException ex)
+				{
+					_internalLogger.LogInformation(ex, "Cancelled log tailing task");
+					break;
+				}
+				catch (Exception ex)
+				{
+					_internalLogger.LogError(ex, "Exception on log tailing task ({LogId}): {Message}", _logId, ex.Message);
+					await Task.Delay(TimeSpan.FromSeconds(10.0));
+				}
+			}
+		}
+
+		async Task TickTailInternalAsync()
+		{
+			int tailNext = -1;
+			Task tickTask = Task.CompletedTask;
+			while (!_tailTaskStop.IsSet())
+			{
+				Task newTailDataTask = _newTailDataEvent.Task;
+				int initialTailNext = tailNext;
+
+				// Get the data to send to the server
+				ReadOnlyMemory<byte> tailData = ReadOnlyMemory<byte>.Empty;
+				if (tailNext != -1)
+				{
+					(tailNext, tailData) = _builder.ReadTailData(tailNext, 16 * 1024);
+				}
+
+				// If we don't have any updates for the server, wait until we do. We need to ensure
+				// we keep pumping the RPC with the server in case the requested tail next value changes,
+				// and to make sure that we don't expire the existing tail data.
+				if (tailNext != -1 && tailData.IsEmpty && tailNext == initialTailNext && !tickTask.IsCompleted)
+				{
+					_internalLogger.LogInformation("No tail data available for log {LogId} after line {TailNext}; waiting for more...", _logId, tailNext);
+					await Task.WhenAny(newTailDataTask, tickTask);
+					continue;
+				}
+
+				string start = "";
+				if (tailData.Length > 0)
+				{
+					start = Encoding.UTF8.GetString(tailData.Slice(0, Math.Min(tailData.Length, 256)).Span);
+				}
+
+				// Update the next tailing position
+				int numLines = CountLines(tailData.Span);
+				_internalLogger.LogInformation("Setting log {LogId} tail = {TailNext}, data = {TailDataSize} bytes, {NumLines} lines ('{Start}')", _logId, tailNext, tailData.Length, numLines, start);
+
+				int newTailNext = await UpdateLogTailAsync(tailNext, tailData, CancellationToken.None);
+				_internalLogger.LogInformation("Log {LogId} tail next = {TailNext}", _logId, newTailNext);
+
+				if (newTailNext != tailNext)
+				{
+					tailNext = newTailNext;
+					_internalLogger.LogInformation("Modified tail position for log {LogId} to {TailNext}", _logId, tailNext);
+				}
+
+				tickTask = Task.Delay(TimeSpan.FromSeconds(10.0));
+			}
+			_internalLogger.LogInformation("Finishing log tail task");
+		}
+
+		static int CountLines(ReadOnlySpan<byte> data)
+		{
+			int lines = 0;
+			for (int idx = 0; idx < data.Length; idx++)
+			{
+				if (data[idx] == '\n')
+				{
+					lines++;
+				}
+			}
+			return lines;
+		}
+
+		/// <inheritdoc/>
+		public async Task WriteEventsAsync(List<RpcCreateLogEventRequest> events, CancellationToken cancellationToken)
+		{
+			LogRpc.LogRpcClient logRpc = await _hordeClient.CreateGrpcClientAsync<LogRpc.LogRpcClient>(cancellationToken);
+
+			RpcCreateLogEventsRequest request = new RpcCreateLogEventsRequest();
+			request.Events.AddRange(events);
+
+			await logRpc.CreateLogEventsAsync(request, cancellationToken: cancellationToken);
+		}
+
+		/// <inheritdoc/>
+		public async Task WriteOutputAsync(LogId logId, long offset, int lineIndex, ReadOnlyMemory<byte> data, bool flush, CancellationToken cancellationToken)
+		{
+			_builder.WriteData(data);
+			_bufferLength += data.Length;
+
+			if (flush || _bufferLength > FlushLength)
+			{
+				IBlobRef<LogNode> target = await _builder.FlushAsync(_writer, flush, cancellationToken);
+				await UpdateLogAsync(target, _builder.LineCount, flush, cancellationToken);
+				_bufferLength = 0;
+			}
+
+			_newTailDataEvent.Set();
+		}
+
+		#region RPC calls
+
+		async Task UpdateLogAsync(IBlobRef target, int lineCount, bool complete, CancellationToken cancellationToken)
+		{
+			_internalLogger.LogInformation("Updating log {LogId} to line {LineCount}, target {Locator}", _logId, lineCount, target.GetLocator());
+
+			RpcUpdateLogRequest request = new RpcUpdateLogRequest();
+			request.LogId = _logId.ToString();
+			request.LineCount = lineCount;
+			request.TargetHash = target.Hash.ToString();
+			request.TargetLocator = target.GetLocator().ToString();
+			request.Complete = complete;
+
+			LogRpc.LogRpcClient clientRef = await _hordeClient.CreateGrpcClientAsync<LogRpc.LogRpcClient>(cancellationToken);
+			await clientRef.UpdateLogAsync(request, cancellationToken: cancellationToken);
+		}
+
+		async Task<int> UpdateLogTailAsync(int tailNext, ReadOnlyMemory<byte> tailData, CancellationToken cancellationToken)
+		{
+			DateTime deadline = DateTime.UtcNow.AddMinutes(2.0);
+			try
+			{
+				LogRpc.LogRpcClient clientRef = await _hordeClient.CreateGrpcClientAsync<LogRpc.LogRpcClient>(cancellationToken);
+				using AsyncDuplexStreamingCall<RpcUpdateLogTailRequest, RpcUpdateLogTailResponse> call = clientRef.UpdateLogTail(deadline: deadline, cancellationToken: cancellationToken);
+
+				// Write the request to the server
+				RpcUpdateLogTailRequest request = new RpcUpdateLogTailRequest();
+				request.LogId = _logId.ToString();
+				request.TailNext = tailNext;
+				request.TailData = UnsafeByteOperations.UnsafeWrap(tailData);
+				await call.RequestStream.WriteAsync(request, cancellationToken);
+				_internalLogger.LogInformation("Writing log data: {LogId}, {TailNext}, {TailData} bytes", _logId, tailNext, tailData.Length);
+
+				// Wait until the server responds or we need to trigger a new update
+				Task<bool> moveNextAsync = call.ResponseStream.MoveNext();
+
+				Task task = await Task.WhenAny(moveNextAsync, _tailTaskStop.Task, Task.Delay(TimeSpan.FromMinutes(1.0), CancellationToken.None));
+				if (task == _tailTaskStop.Task)
+				{
+					_internalLogger.LogInformation("Cancelling long poll from client side (complete)");
+				}
+
+				// Close the request stream to indicate that we're finished
+				await call.RequestStream.CompleteAsync();
+
+				// Wait for a response or a new update to come in, then close the request stream
+				RpcUpdateLogTailResponse? response = null;
+				while (await moveNextAsync)
+				{
+					response = call.ResponseStream.Current;
+					moveNextAsync = call.ResponseStream.MoveNext();
+				}
+				return response?.TailNext ?? -1;
+			}
+			catch (RpcException ex) when (ex.StatusCode == StatusCode.DeadlineExceeded)
+			{
+				_internalLogger.LogDebug(ex, "Log tail deadline exceeded, ignoring.");
+				return -1;
+			}
+		}
+
+		#endregion
+	}
+}
