Horde: Remove IServerLoggerFactory, and allow creating IServerLogger instances from IHordeClient.

#rnx

diff --git a/Engine/Source/Programs/Horde/Horde.Agent.Driver/Utilities/HordeRpc.cs b/Engine/Source/Programs/Horde/Horde.Agent.Driver/Utilities/HordeRpc.cs
--- a/Engine/Source/Programs/Horde/Horde.Agent.Driver/Utilities/HordeRpc.cs
+++ b/Engine/Source/Programs/Horde/Horde.Agent.Driver/Utilities/HordeRpc.cs
@@ -2,8 +2,6 @@
 
 using EpicGames.Horde.Agents.Leases;
 using EpicGames.Horde.Jobs;
-using EpicGames.Horde.Logs;
-using Google.Protobuf;
 
 // These partial classes/extensions operate on generated gRPC and Protobuf code.
 // Warnings below are disabled to avoid documenting every class touched.
diff --git a/Engine/Source/Programs/Horde/Horde.Agent.Driver/Utility/JobStepLogger.cs b/Engine/Source/Programs/Horde/Horde.Agent.Driver/Utility/JobStepLogger.cs
--- a/Engine/Source/Programs/Horde/Horde.Agent.Driver/Utility/JobStepLogger.cs
+++ b/Engine/Source/Programs/Horde/Horde.Agent.Driver/Utility/JobStepLogger.cs
@@ -3,7 +3,6 @@
 using EpicGames.Horde;
 using EpicGames.Horde.Jobs;
 using EpicGames.Horde.Logs;
-using Horde.Agent.Utility;
 using Horde.Common.Rpc;
 using HordeCommon.Rpc;
 using Microsoft.Extensions.Logging;
@@ -26,7 +25,7 @@ namespace Horde.Agent.Driver.Utility
 		readonly JobStepId _jobStepId;
 		readonly bool _warnings;
 
-		readonly ServerLogger _serverLogger;
+		readonly IServerLogger _serverLogger;
 		readonly ILogger _localLogger;
 		JobStepOutcome _outcome;
 		Task _updateOutcomeTask;
@@ -43,7 +42,7 @@ namespace Horde.Agent.Driver.Utility
 			_jobBatchId = jobBatchId;
 			_jobStepId = jobStepId;
 			_warnings = warnings ?? true;
-			_serverLogger = new ServerLogger(hordeClient, logId, outputLevel, internalLogger);
+			_serverLogger = hordeClient.CreateServerLogger(logId, outputLevel);
 			_localLogger = localLogger;
 			_outcome = JobStepOutcome.Success;
 			_updateOutcomeTask = Task.CompletedTask;
diff --git a/Engine/Source/Programs/Horde/Horde.Agent.Tests/Services/TelemetryServiceTest.cs b/Engine/Source/Programs/Horde/Horde.Agent.Tests/Services/TelemetryServiceTest.cs
--- a/Engine/Source/Programs/Horde/Horde.Agent.Tests/Services/TelemetryServiceTest.cs
+++ b/Engine/Source/Programs/Horde/Horde.Agent.Tests/Services/TelemetryServiceTest.cs
@@ -26,7 +26,7 @@ public sealed class TelemetryServiceTest : IDisposable
 		OptionsWrapper<AgentSettings> settingsOpt = new(settings);
 
 		using WorkerService workerService = new WorkerService(null!, null!, null!, null!);
-		JobHandler jobHandler = new JobHandler(new List<IJobExecutorFactory>(), settingsOpt, new OptionsWrapper<DriverSettings>(new DriverSettings()), null!);
+		JobHandler jobHandler = new JobHandler(new List<IJobExecutorFactory>(), settingsOpt, new OptionsWrapper<DriverSettings>(new DriverSettings()));
 		GrpcService grpcService = new(settingsOpt, NullLogger<GrpcService>.Instance, _loggerFactory);
 
 		_telemetryService = new TelemetryService(workerService, jobHandler, grpcService, null!, settingsOpt, NullLogger<TelemetryService>.Instance);
diff --git a/Engine/Source/Programs/Horde/Horde.Agent.Tests/WorkerServiceTest.cs b/Engine/Source/Programs/Horde/Horde.Agent.Tests/WorkerServiceTest.cs
--- a/Engine/Source/Programs/Horde/Horde.Agent.Tests/WorkerServiceTest.cs
+++ b/Engine/Source/Programs/Horde/Horde.Agent.Tests/WorkerServiceTest.cs
@@ -64,11 +64,6 @@ namespace Horde.Agent.Tests
 			public Task StopAsync() => Task.CompletedTask;
 		}
 
-		class FakeServerLoggerFactory : IServerLoggerFactory
-		{
-			public IServerLogger CreateLogger(IHordeClient hordeClient, LogId logId, LogLevel outputLevel = LogLevel.Information) => new FakeServerLogger();
-		}
-
 		internal static IJobExecutor NullExecutor = new SimpleTestExecutor(async (step, logger, cancellationToken) =>
 		{
 			await Task.Delay(1, cancellationToken);
@@ -80,7 +75,6 @@ namespace Horde.Agent.Tests
 			_serviceCollection = new ServiceCollection();
 			_serviceCollection.AddLogging();
 			_serviceCollection.AddHorde();
-			_serviceCollection.AddSingleton<IServerLoggerFactory, FakeServerLoggerFactory>();
 			_serviceCollection.AddSingleton<BundleCache>();
 			_serviceCollection.AddSingleton<StorageBackendCache>();
 			_serviceCollection.AddSingleton<HttpStorageBackendFactory>();
@@ -354,6 +348,9 @@ namespace Horde.Agent.Tests
 
 			public bool HasValidAccessToken()
 				=> throw new NotImplementedException();
+
+			public IServerLogger CreateServerLogger(LogId logId, LogLevel minimumLevel = LogLevel.Information)
+				=> throw new NotImplementedException();
 		}
 
 		private class FakeHordeRpcClient : HordeRpc.HordeRpcClient
diff --git a/Engine/Source/Programs/Horde/Horde.Agent/AgentApp.cs b/Engine/Source/Programs/Horde/Horde.Agent/AgentApp.cs
--- a/Engine/Source/Programs/Horde/Horde.Agent/AgentApp.cs
+++ b/Engine/Source/Programs/Horde/Horde.Agent/AgentApp.cs
@@ -216,7 +216,6 @@ namespace Horde.Agent
 
 			services.AddSingleton<CapabilitiesService>();
 			services.AddSingleton<ISessionFactory, SessionFactory>();
-			services.AddSingleton<IServerLoggerFactory, ServerLoggerFactory>();
 			services.AddSingleton<WorkerService>();
 			services.AddSingleton<LeaseLoggerFactory>();
 			services.AddHostedService(sp => sp.GetRequiredService<WorkerService>());
diff --git a/Engine/Source/Programs/Horde/Horde.Agent/Leases/Handlers/ComputeHandler.cs b/Engine/Source/Programs/Horde/Horde.Agent/Leases/Handlers/ComputeHandler.cs
--- a/Engine/Source/Programs/Horde/Horde.Agent/Leases/Handlers/ComputeHandler.cs
+++ b/Engine/Source/Programs/Horde/Horde.Agent/Leases/Handlers/ComputeHandler.cs
@@ -7,7 +7,6 @@ using EpicGames.Horde.Compute;
 using EpicGames.Horde.Compute.Transports;
 using EpicGames.Horde.Logs;
 using Horde.Agent.Services;
-using Horde.Agent.Utility;
 using HordeCommon.Rpc.Tasks;
 using Microsoft.Extensions.Logging;
 using Microsoft.Extensions.Options;
@@ -20,23 +19,21 @@ namespace Horde.Agent.Leases.Handlers
 	class ComputeHandler : LeaseHandler<ComputeTask>
 	{
 		readonly ComputeListenerService _listenerService;
-		readonly IServerLoggerFactory _serverLoggerFactory;
 		readonly AgentSettings _settings;
 
 		/// <summary>
 		/// Constructor
 		/// </summary>
-		public ComputeHandler(ComputeListenerService listenerService, IServerLoggerFactory serverLoggerFactory, IOptions<AgentSettings> settings)
+		public ComputeHandler(ComputeListenerService listenerService, IOptions<AgentSettings> settings)
 		{
 			_listenerService = listenerService;
-			_serverLoggerFactory = serverLoggerFactory;
 			_settings = settings.Value;
 		}
 
 		/// <inheritdoc/>
 		public override async Task<LeaseResult> ExecuteAsync(ISession session, LeaseId leaseId, ComputeTask computeTask, ILogger localLogger, CancellationToken cancellationToken)
 		{
-			await using IServerLogger? serverLogger = (computeTask.LogId != null) ? _serverLoggerFactory.CreateLogger(session.HordeClient, LogId.Parse(computeTask.LogId), LogLevel.Trace).WithLocalLogger(localLogger) : null;
+			await using IServerLogger? serverLogger = (computeTask.LogId != null) ? session.HordeClient.CreateServerLogger(LogId.Parse(computeTask.LogId), LogLevel.Trace).WithLocalLogger(localLogger) : null;
 			ILogger logger = serverLogger ?? localLogger;
 
 			if (!String.IsNullOrEmpty(computeTask.ParentLeaseId))
diff --git a/Engine/Source/Programs/Horde/Horde.Agent/Leases/Handlers/ConformHandler.cs b/Engine/Source/Programs/Horde/Horde.Agent/Leases/Handlers/ConformHandler.cs
--- a/Engine/Source/Programs/Horde/Horde.Agent/Leases/Handlers/ConformHandler.cs
+++ b/Engine/Source/Programs/Horde/Horde.Agent/Leases/Handlers/ConformHandler.cs
@@ -17,21 +17,19 @@ namespace Horde.Agent.Leases.Handlers
 	class ConformHandler : LeaseHandler<ConformTask>
 	{
 		readonly DriverSettings _driverSettings;
-		readonly IServerLoggerFactory _serverLoggerFactory;
 
 		/// <summary>
 		/// Constructor
 		/// </summary>
-		public ConformHandler(IOptions<DriverSettings> driverSettings, IServerLoggerFactory serverLoggerFactory)
+		public ConformHandler(IOptions<DriverSettings> driverSettings)
 		{
 			_driverSettings = driverSettings.Value;
-			_serverLoggerFactory = serverLoggerFactory;
 		}
 
 		/// <inheritdoc/>
 		public override async Task<LeaseResult> ExecuteAsync(ISession session, LeaseId leaseId, ConformTask conformTask, ILogger localLogger, CancellationToken cancellationToken)
 		{
-			await using IServerLogger serverLogger = _serverLoggerFactory.CreateLogger(session.HordeClient, LogId.Parse(conformTask.LogId)).WithLocalLogger(localLogger);
+			await using IServerLogger serverLogger = session.HordeClient.CreateServerLogger(LogId.Parse(conformTask.LogId)).WithLocalLogger(localLogger);
 			try
 			{
 				LeaseResult result = await ExecuteInternalAsync(session, leaseId, conformTask, serverLogger, cancellationToken);
diff --git a/Engine/Source/Programs/Horde/Horde.Agent/Leases/Handlers/JobHandler.cs b/Engine/Source/Programs/Horde/Horde.Agent/Leases/Handlers/JobHandler.cs
--- a/Engine/Source/Programs/Horde/Horde.Agent/Leases/Handlers/JobHandler.cs
+++ b/Engine/Source/Programs/Horde/Horde.Agent/Leases/Handlers/JobHandler.cs
@@ -58,17 +58,15 @@ namespace Horde.Agent.Leases.Handlers
 		readonly IEnumerable<IJobExecutorFactory> _executorFactories;
 		readonly AgentSettings _agentSettings;
 		readonly DriverSettings _driverSettings;
-		readonly IServerLoggerFactory _serverLoggerFactory;
 
 		/// <summary>
 		/// Constructor
 		/// </summary>
-		public JobHandler(IEnumerable<IJobExecutorFactory> executorFactories, IOptions<AgentSettings> agentSettings, IOptions<DriverSettings> driverSettings, IServerLoggerFactory serverLoggerFactory)
+		public JobHandler(IEnumerable<IJobExecutorFactory> executorFactories, IOptions<AgentSettings> agentSettings, IOptions<DriverSettings> driverSettings)
 		{
 			_executorFactories = executorFactories;
 			_agentSettings = agentSettings.Value;
 			_driverSettings = driverSettings.Value;
-			_serverLoggerFactory = serverLoggerFactory;
 		}
 
 		/// <inheritdoc/>
@@ -159,7 +157,7 @@ namespace Horde.Agent.Leases.Handlers
 
 			// Create a storage client for this session
 			RpcJobOptions jobOptions = executeTask.JobOptions;
-			await using IServerLogger logger = _serverLoggerFactory.CreateLogger(session.HordeClient, executeTask.LogId).WithLocalLogger(localLogger);
+			await using IServerLogger logger = session.HordeClient.CreateServerLogger(executeTask.LogId).WithLocalLogger(localLogger);
 
 			logger.LogInformation("Executing job \"{JobName}\", jobId {JobId}, batchId {BatchId}, leaseId {LeaseId}, agentVersion {AgentVersion}", executeTask.JobName, executeTask.JobId, executeTask.BatchId, leaseId, AgentApp.Version);
 
diff --git a/Engine/Source/Programs/Horde/Horde.Agent/Leases/Handlers/UpgradeHandler.cs b/Engine/Source/Programs/Horde/Horde.Agent/Leases/Handlers/UpgradeHandler.cs
--- a/Engine/Source/Programs/Horde/Horde.Agent/Leases/Handlers/UpgradeHandler.cs
+++ b/Engine/Source/Programs/Horde/Horde.Agent/Leases/Handlers/UpgradeHandler.cs
@@ -9,7 +9,6 @@ using EpicGames.Horde.Agents.Leases;
 using EpicGames.Horde.Logs;
 using Grpc.Core;
 using Horde.Agent.Services;
-using Horde.Agent.Utility;
 using HordeCommon.Rpc;
 using HordeCommon.Rpc.Tasks;
 using Microsoft.Extensions.Logging;
@@ -18,17 +17,14 @@ namespace Horde.Agent.Leases.Handlers
 {
 	class UpgradeHandler : LeaseHandler<UpgradeTask>
 	{
-		readonly IServerLoggerFactory _serverLoggerFactory;
-
-		public UpgradeHandler(IServerLoggerFactory serverLoggerFactory)
+		public UpgradeHandler()
 		{
-			_serverLoggerFactory = serverLoggerFactory;
 		}
 
 		/// <inheritdoc/>
 		public override async Task<LeaseResult> ExecuteAsync(ISession session, LeaseId leaseId, UpgradeTask task, ILogger localLogger, CancellationToken cancellationToken)
 		{
-			await using IServerLogger logger = _serverLoggerFactory.CreateLogger(session.HordeClient, LogId.Parse(task.LogId)).WithLocalLogger(localLogger);
+			await using IServerLogger logger = session.HordeClient.CreateServerLogger(LogId.Parse(task.LogId)).WithLocalLogger(localLogger);
 
 			string requiredVersion = task.SoftwareId;
 
diff --git a/Engine/Source/Programs/Horde/Horde.Agent/Leases/LeaseLoggerFactory.cs b/Engine/Source/Programs/Horde/Horde.Agent/Leases/LeaseLoggerFactory.cs
--- a/Engine/Source/Programs/Horde/Horde.Agent/Leases/LeaseLoggerFactory.cs
+++ b/Engine/Source/Programs/Horde/Horde.Agent/Leases/LeaseLoggerFactory.cs
@@ -2,6 +2,7 @@
 
 using EpicGames.Core;
 using EpicGames.Horde.Agents.Leases;
+using EpicGames.Horde.Logs;
 using Microsoft.Extensions.Logging;
 using Microsoft.Extensions.Options;
 
@@ -106,4 +107,53 @@ namespace Horde.Agent.Leases
 			}
 		}
 	}
+
+	static class LeaseLoggerExtensions
+	{
+		sealed class ServerLoggerWithLeaseLogger : IServerLogger
+		{
+			readonly IServerLogger _serverLogger;
+			readonly ILogger _localLogger;
+
+			public ServerLoggerWithLeaseLogger(IServerLogger serverLogger, ILogger localLogger)
+			{
+				_serverLogger = serverLogger;
+				_localLogger = localLogger;
+			}
+
+			/// <inheritdoc/>
+			public IDisposable? BeginScope<TState>(TState state) where TState : notnull
+				=> _localLogger.BeginScope(state);
+
+			/// <inheritdoc/>
+			public ValueTask DisposeAsync()
+				=> _serverLogger.DisposeAsync();
+
+			/// <inheritdoc/>
+			public bool IsEnabled(LogLevel logLevel)
+				=> _serverLogger.IsEnabled(logLevel) || _localLogger.IsEnabled(logLevel);
+
+			/// <inheritdoc/>
+			public void Log<TState>(LogLevel logLevel, EventId eventId, TState state, Exception? exception, Func<TState, Exception?, string> formatter)
+			{
+				_serverLogger.Log(logLevel, eventId, state, exception, formatter);
+				_localLogger.Log(logLevel, eventId, state, exception, formatter);
+			}
+
+			/// <inheritdoc/>
+			public Task StopAsync()
+				=> _serverLogger.StopAsync();
+		}
+
+		/// <summary>
+		/// Creates a logger which uploads data to the server
+		/// </summary>
+		/// <param name="serverLogger">Logger for the server</param>
+		/// <param name="localLogger">Local log output device</param>
+		/// <returns>New logger instance</returns>
+		public static IServerLogger WithLocalLogger(this IServerLogger serverLogger, ILogger localLogger)
+		{
+			return new ServerLoggerWithLeaseLogger(serverLogger, localLogger);
+		}
+	}
 }
diff --git a/Engine/Source/Programs/Horde/Horde.Agent/Services/ServerLoggerFactory.cs b/Engine/Source/Programs/Horde/Horde.Agent/Services/ServerLoggerFactory.cs
deleted file mode 100644
diff --git a/Engine/Source/Programs/Shared/EpicGames.Horde.Tests/ServerLogPacketBuilderTests.cs b/Engine/Source/Programs/Shared/EpicGames.Horde.Tests/ServerLogPacketBuilderTests.cs
--- a/Engine/Source/Programs/Shared/EpicGames.Horde.Tests/ServerLogPacketBuilderTests.cs
+++ b/Engine/Source/Programs/Shared/EpicGames.Horde.Tests/ServerLogPacketBuilderTests.cs
@@ -15,7 +15,6 @@ using EpicGames.Horde.Storage;
 using EpicGames.Horde.Storage.Bundles;
 using Grpc.Core;
 using Grpc.Net.Client;
-//using Horde.Agent.Utility;
 using Horde.Common.Rpc;
 using Microsoft.Extensions.Logging;
 using Microsoft.Extensions.Logging.Abstractions;
@@ -263,6 +262,9 @@ namespace EpicGames.Horde.Tests
 
 			public bool HasValidAccessToken()
 				=> throw new NotImplementedException();
+
+			public IServerLogger CreateServerLogger(LogId logId, LogLevel minimumLevel = LogLevel.Information)
+				=> throw new NotImplementedException();
 		}
 
 		[TestMethod]
diff --git a/Engine/Source/Programs/Shared/EpicGames.Horde/HordeClient.cs b/Engine/Source/Programs/Shared/EpicGames.Horde/HordeClient.cs
--- a/Engine/Source/Programs/Shared/EpicGames.Horde/HordeClient.cs
+++ b/Engine/Source/Programs/Shared/EpicGames.Horde/HordeClient.cs
@@ -9,6 +9,7 @@ using System.Threading.Tasks;
 using EpicGames.Core;
 using EpicGames.Horde.Compute;
 using EpicGames.Horde.Compute.Clients;
+using EpicGames.Horde.Logs;
 using EpicGames.Horde.Server;
 using EpicGames.Horde.Storage;
 using EpicGames.Horde.Storage.Backends;
@@ -202,6 +203,10 @@ namespace EpicGames.Horde
 			return new BundleStorageClient(httpStorageBackend, _bundleCache, _hordeOptions.Bundle, _loggerFactory.CreateLogger<BundleStorageClient>());
 		}
 
+		/// <inheritdoc/>
+		public IServerLogger CreateServerLogger(LogId logId, LogLevel minimumLevel = LogLevel.Information)
+			=> new ServerLogger(this, logId, minimumLevel, _logger);
+
 		/// <summary>
 		/// Creates an http client for satisfying requests
 		/// </summary>
diff --git a/Engine/Source/Programs/Shared/EpicGames.Horde/IHordeClient.cs b/Engine/Source/Programs/Shared/EpicGames.Horde/IHordeClient.cs
--- a/Engine/Source/Programs/Shared/EpicGames.Horde/IHordeClient.cs
+++ b/Engine/Source/Programs/Shared/EpicGames.Horde/IHordeClient.cs
@@ -9,6 +9,7 @@ using EpicGames.Horde.Storage;
 using EpicGames.Horde.Tools;
 using Grpc.Core;
 using Grpc.Net.Client;
+using Microsoft.Extensions.Logging;
 
 namespace EpicGames.Horde
 {
@@ -64,6 +65,11 @@ namespace EpicGames.Horde
 		/// Creates a storage client for the given base path
 		/// </summary>
 		IStorageClient CreateStorageClient(string relativePath, string? accessToken = null);
+
+		/// <summary>
+		/// Creates a logger device that writes data to the server
+		/// </summary>
+		IServerLogger CreateServerLogger(LogId logId, LogLevel minimumLevel = LogLevel.Information);
 	}
 
 	/// <summary>
diff --git a/Engine/Source/Programs/Shared/EpicGames.Horde/Logs/ServerLogger.cs b/Engine/Source/Programs/Shared/EpicGames.Horde/Logs/ServerLogger.cs
--- a/Engine/Source/Programs/Shared/EpicGames.Horde/Logs/ServerLogger.cs
+++ b/Engine/Source/Programs/Shared/EpicGames.Horde/Logs/ServerLogger.cs
@@ -7,8 +7,6 @@ using System.Threading;
 using System.Threading.Channels;
 using System.Threading.Tasks;
 using EpicGames.Core;
-using EpicGames.Horde;
-using EpicGames.Horde.Logs;
 using EpicGames.Horde.Storage;
 using Google.Protobuf;
 using Grpc.Core;
@@ -17,23 +15,10 @@ using Microsoft.Extensions.Logging;
 
 namespace EpicGames.Horde.Logs
 {
-	using ByteString = Google.Protobuf.ByteString;
-
-	/// <summary>
-	/// Interface for a log device
-	/// </summary>
-	public interface IServerLogger : ILogger, IAsyncDisposable
-	{
-		/// <summary>
-		/// Flushes the logger with the server and stops the background work
-		/// </summary>
-		Task StopAsync();
-	}
-
 	/// <summary>
 	/// Class to handle uploading log data to the server in the background
 	/// </summary>
-	public sealed class ServerLogger : IServerLogger
+	sealed class ServerLogger : IServerLogger
 	{
 		const int FlushLength = 1024 * 1024;
 
@@ -213,7 +198,7 @@ namespace EpicGames.Horde.Logs
 					(ReadOnlyMemory<byte> packet, int packetLineCount) = writer.CreatePacket();
 					try
 					{
-						await WriteOutputAsync(_logId, packetOffset, packetLineIndex, packet, false, CancellationToken.None);
+						await WriteOutputAsync(packet, false, CancellationToken.None);
 						packetOffset += packet.Length;
 						packetLineIndex += packetLineCount;
 					}
@@ -241,7 +226,7 @@ namespace EpicGames.Horde.Logs
 				{
 					try
 					{
-						await WriteOutputAsync(_logId, packetOffset, packetLineIndex, ReadOnlyMemory<byte>.Empty, true, CancellationToken.None);
+						await WriteOutputAsync(ReadOnlyMemory<byte>.Empty, true, CancellationToken.None);
 					}
 					catch (Exception ex)
 					{
@@ -256,7 +241,7 @@ namespace EpicGames.Horde.Logs
 		{
 			try
 			{
-				events.Add(new RpcCreateLogEventRequest{ Severity = (int) severity, LogId = _logId.ToString(), LineIndex = lineIndex, LineCount = lineCount });
+				events.Add(new RpcCreateLogEventRequest { Severity = (int)severity, LogId = _logId.ToString(), LineIndex = lineIndex, LineCount = lineCount });
 			}
 			catch (Exception ex)
 			{
@@ -361,7 +346,7 @@ namespace EpicGames.Horde.Logs
 		}
 
 		/// <inheritdoc/>
-		public async Task WriteOutputAsync(LogId logId, long offset, int lineIndex, ReadOnlyMemory<byte> data, bool flush, CancellationToken cancellationToken)
+		public async Task WriteOutputAsync(ReadOnlyMemory<byte> data, bool flush, CancellationToken cancellationToken)
 		{
 			_builder.WriteData(data);
 			_bufferLength += data.Length;
