Horde: Missing file.

#rnx

diff --git a/Engine/Source/Programs/Shared/EpicGames.Horde/Logs/IServerLogger.cs b/Engine/Source/Programs/Shared/EpicGames.Horde/Logs/IServerLogger.cs
new file mode 100644
--- /dev/null
+++ b/Engine/Source/Programs/Shared/EpicGames.Horde/Logs/IServerLogger.cs
@@ -0,0 +1,19 @@
+// Copyright Epic Games, Inc. All Rights Reserved.
+
+using System;
+using System.Threading.Tasks;
+using Microsoft.Extensions.Logging;
+
+namespace EpicGames.Horde.Logs
+{
+	/// <summary>
+	/// Interface for a log device
+	/// </summary>
+	public interface IServerLogger : ILogger, IAsyncDisposable
+	{
+		/// <summary>
+		/// Flushes the logger with the server and stops the background work
+		/// </summary>
+		Task StopAsync();
+	}
+}
