Horde: Change initialization order of fields in ServerLogger to prevent race over accessing _internalLogger.

diff --git a/Engine/Source/Programs/Shared/EpicGames.Horde/Logs/ServerLogger.cs b/Engine/Source/Programs/Shared/EpicGames.Horde/Logs/ServerLogger.cs
--- a/Engine/Source/Programs/Shared/EpicGames.Horde/Logs/ServerLogger.cs
+++ b/Engine/Source/Programs/Shared/EpicGames.Horde/Logs/ServerLogger.cs
@@ -24,6 +24,8 @@ namespace EpicGames.Horde.Logs
 
 		readonly IHordeClient _hordeClient;
 		readonly LogId _logId;
+		readonly LogLevel _minimumLevel;
+		readonly ILogger _internalLogger;
 		readonly LogBuilder _builder;
 		readonly IStorageClient _store;
 		readonly IBlobWriter _writer;
@@ -35,8 +37,6 @@ namespace EpicGames.Horde.Logs
 		AsyncEvent _tailTaskStop;
 		readonly AsyncEvent _newTailDataEvent = new AsyncEvent();
 
-		readonly LogLevel _outputLevel;
-		readonly ILogger _internalLogger;
 		readonly Channel<JsonLogEvent> _dataChannel;
 		Task? _dataWriter;
 
@@ -51,6 +51,8 @@ namespace EpicGames.Horde.Logs
 		{
 			_hordeClient = hordeClient;
 			_logId = logId;
+			_minimumLevel = minimumLevel;
+			_internalLogger = internalLogger;
 			_builder = new LogBuilder(LogFormat.Json, internalLogger);
 			_store = _hordeClient.CreateStorageClient(logId);
 			_writer = _store.CreateBlobWriter();
@@ -58,9 +60,6 @@ namespace EpicGames.Horde.Logs
 			_tailTaskStop = new AsyncEvent();
 			_tailTask = Task.Run(() => TickTailAsync());
 
-			_logId = logId;
-			_outputLevel = minimumLevel;
-			_internalLogger = internalLogger;
 			_dataChannel = Channel.CreateUnbounded<JsonLogEvent>();
 			_dataWriter = Task.Run(() => RunDataWriterAsync());
 		}
@@ -73,7 +72,7 @@ namespace EpicGames.Horde.Logs
 		}
 
 		/// <inheritdoc/>
-		public bool IsEnabled(LogLevel logLevel) => logLevel >= _outputLevel;
+		public bool IsEnabled(LogLevel logLevel) => logLevel >= _minimumLevel;
 
 		/// <inheritdoc/>
 		public IDisposable BeginScope<TState>(TState state)
