Horde: Missing file.

#rnx

diff --git a/Engine/Source/Programs/Shared/EpicGames.Horde/Jobs/JobMessages.cs b/Engine/Source/Programs/Shared/EpicGames.Horde/Jobs/JobMessages.cs
--- a/Engine/Source/Programs/Shared/EpicGames.Horde/Jobs/JobMessages.cs
+++ b/Engine/Source/Programs/Shared/EpicGames.Horde/Jobs/JobMessages.cs
@@ -73,6 +73,11 @@ namespace EpicGames.Horde.Jobs
 		/// </summary>
 		public int? ExpireAfterDays { get; set; }
 
+		/// <summary>
+		/// Name of the driver to use
+		/// </summary>
+		public string? Driver { get; set; }
+
 		/// <summary>
 		/// Merge defaults from another options object
 		/// </summary>
@@ -84,6 +89,7 @@ namespace EpicGames.Horde.Jobs
 			WorkspaceMaterializer ??= other.WorkspaceMaterializer;
 			Container.MergeDefaults(other.Container);
 			ExpireAfterDays ??= other.ExpireAfterDays;
+			Driver ??= other.Driver;
 		}
 	}
 
