EpicGames.Perforce: Fix parsing of records which have duplicate field names in child records.

diff --git a/Engine/Source/Programs/Shared/EpicGames.Perforce/IPerforceOutput.cs b/Engine/Source/Programs/Shared/EpicGames.Perforce/IPerforceOutput.cs
--- a/Engine/Source/Programs/Shared/EpicGames.Perforce/IPerforceOutput.cs
+++ b/Engine/Source/Programs/Shared/EpicGames.Perforce/IPerforceOutput.cs
@@ -610,9 +610,23 @@ namespace EpicGames.Perforce
 				Utf8String suffix = tag.Slice(suffixIdx);
 				tag = tag.Slice(0, suffixIdx);
 
+				// Count the dimensions in the suffix
+				int rank = 0;
+				if (suffix.Length > 0)
+				{
+					rank++;
+					for (int idx = 0; idx < suffix.Length; idx++)
+					{
+						if (suffix[idx] == (byte)',')
+						{
+							rank++;
+						}
+					}
+				}
+
 				// Try to find the matching field
 				TaggedPropertyInfo? tagInfo;
-				if (recordInfo.NameToInfo.TryGetValue(tag, out tagInfo))
+				if (recordInfo.NameAndRankToInfo.TryGetValue(new TaggedPropertyNameAndRank(tag, rank), out tagInfo))
 				{
 					requiredTagsBitMask |= tagInfo.RequiredTagBitMask;
 				}
@@ -635,7 +649,7 @@ namespace EpicGames.Perforce
 			// Make sure we've got all the required tags we need
 			if (requiredTagsBitMask != recordInfo.RequiredTagsBitMask)
 			{
-				string missingTagNames = String.Join(", ", recordInfo.NameToInfo.Where(x => (requiredTagsBitMask | x.Value.RequiredTagBitMask) != requiredTagsBitMask).Select(x => x.Key));
+				string missingTagNames = String.Join(", ", recordInfo.NameAndRankToInfo.Where(x => (requiredTagsBitMask | x.Value.RequiredTagBitMask) != requiredTagsBitMask).Select(x => x.Key));
 				throw new PerforceException("Missing '{0}' tag when parsing '{1}'", missingTagNames, recordInfo.Type.Name);
 			}
 
diff --git a/Engine/Source/Programs/Shared/EpicGames.Perforce/PerforceReflection.cs b/Engine/Source/Programs/Shared/EpicGames.Perforce/PerforceReflection.cs
--- a/Engine/Source/Programs/Shared/EpicGames.Perforce/PerforceReflection.cs
+++ b/Engine/Source/Programs/Shared/EpicGames.Perforce/PerforceReflection.cs
@@ -37,6 +37,11 @@ namespace EpicGames.Perforce
 	/// </summary>
 	record struct NestedRecordInfo(PropertyInfo PropertyInfo, CreateRecordDelegate CreateInstance);
 
+	/// <summary>
+	/// Name and rank of a tagged property within a record
+	/// </summary>
+	record struct TaggedPropertyNameAndRank(Utf8String Tag, int Rank);
+
 	/// <summary>
 	/// Stores cached information about a property with a <see cref="PerforceTagAttribute"/> attribute.
 	/// </summary>
@@ -141,7 +146,7 @@ namespace EpicGames.Perforce
 		/// <summary>
 		/// Map of name to tag info
 		/// </summary>
-		public Dictionary<Utf8String, TaggedPropertyInfo> NameToInfo { get; set; } = new Dictionary<Utf8String, TaggedPropertyInfo>();
+		public Dictionary<TaggedPropertyNameAndRank, TaggedPropertyInfo> NameAndRankToInfo { get; set; } = new Dictionary<TaggedPropertyNameAndRank, TaggedPropertyInfo>();
 
 		/// <summary>
 		/// Bitmask of all the required tags. Formed by bitwise-or'ing the RequiredTagBitMask fields for each required CachedTagInfo.
@@ -468,7 +473,7 @@ namespace EpicGames.Perforce
 
 				// Find all the properties in the record
 				AddRecordProperties(recordType, record, Array.Empty<NestedRecordInfo>());
-				record.NameToInfo = record.Properties.ToDictionary(x => x.Tag, x => x);
+				record.NameAndRankToInfo = record.Properties.ToDictionary(x => new TaggedPropertyNameAndRank(x.Tag, x.ParentRecords.Length), x => x);
 
 				// Try to save the record info, or get the version that's already in the cache
 				if (!s_recordTypeToInfo.TryAdd(recordType, record))
