EpicGames.Perforce: Fix parsing of array fields from Perforce records.

diff --git a/Engine/Source/Programs/Shared/EpicGames.Perforce/IPerforceOutput.cs b/Engine/Source/Programs/Shared/EpicGames.Perforce/IPerforceOutput.cs
--- a/Engine/Source/Programs/Shared/EpicGames.Perforce/IPerforceOutput.cs
+++ b/Engine/Source/Programs/Shared/EpicGames.Perforce/IPerforceOutput.cs
@@ -624,11 +624,15 @@ namespace EpicGames.Perforce
 					}
 				}
 
-				// Try to find the matching field
-				TaggedPropertyInfo? tagInfo;
-				if (recordInfo.NameAndRankToInfo.TryGetValue(new TaggedPropertyNameAndRank(tag, rank), out tagInfo))
+				// Try to find the matching field. Check all rank values down to zero in case we're parsing an array.
+				TaggedPropertyInfo? tagInfo = null;
+				for (; rank >= 0; rank--)
 				{
-					requiredTagsBitMask |= tagInfo.RequiredTagBitMask;
+					if (recordInfo.NameAndRankToInfo.TryGetValue(new TaggedPropertyNameAndRank(tag, rank), out tagInfo))
+					{
+						requiredTagsBitMask |= tagInfo.RequiredTagBitMask;
+						break;
+					}
 				}
 
 				// Find the target object for this tag
