UnrealBuildTool: Move several helper functions to EpicGames.Build

#rnx

diff --git a/Engine/Source/Programs/Shared/EpicGames.Build/Unreal.cs b/Engine/Source/Programs/Shared/EpicGames.Build/Unreal.cs
--- a/Engine/Source/Programs/Shared/EpicGames.Build/Unreal.cs
+++ b/Engine/Source/Programs/Shared/EpicGames.Build/Unreal.cs
@@ -195,6 +195,16 @@ namespace UnrealBuildBase
 		/// </summary>
 		static private bool? bIsEngineInstalled;
 
+		/// <summary>
+		/// Whether we're running with an installed project (ie. a mod kit)
+		/// </summary>
+		static private bool? bIsProjectInstalled;
+
+		/// <summary>
+		/// If we are running with an installed project, specifies the path to it
+		/// </summary>
+		static private FileReference? InstalledProjectFile;
+
 		/// <summary>
 		/// Returns where another platform's Dotnet is located
 		/// </summary>
@@ -206,7 +216,7 @@ namespace UnrealBuildBase
 		}
 
 		/// <summary>
-		/// Returns true of UnrealBuildTool is running on a build machine
+		/// Returns true if the application is running on a build machine
 		/// </summary>
 		/// <returns>True if running on a build machine</returns>
 		static public bool IsBuildMachine()
@@ -219,7 +229,7 @@ namespace UnrealBuildBase
 		}
 
 		/// <summary>
-		/// Returns true if UnrealBuildTool is running using installed Engine components
+		/// Returns true if the application is running using installed Engine components
 		/// </summary>
 		/// <returns>True if running using installed Engine components</returns>
 		static public bool IsEngineInstalled()
@@ -231,6 +241,56 @@ namespace UnrealBuildBase
 			return bIsEngineInstalled.Value;
 		}
 
+		/// <summary>
+		/// Returns true if the application is running using an installed project (ie. a mod kit)
+		/// </summary>
+		/// <returns>True if running using an installed project</returns>
+		public static bool IsProjectInstalled()
+		{
+			if (!bIsProjectInstalled.HasValue)
+			{
+				FileReference InstalledProjectLocationFile = FileReference.Combine(EngineDirectory, "Build", "InstalledProjectBuild.txt");
+				if (FileReference.Exists(InstalledProjectLocationFile))
+				{
+					InstalledProjectFile = FileReference.Combine(RootDirectory, FileReference.ReadAllText(InstalledProjectLocationFile).Trim());
+					bIsProjectInstalled = true;
+				}
+				else
+				{
+					InstalledProjectFile = null;
+					bIsProjectInstalled = false;
+				}
+			}
+			return bIsProjectInstalled.Value;
+		}
+
+		/// <summary>
+		/// Gets the installed project file
+		/// </summary>
+		/// <returns>Location of the installed project file</returns>
+		public static FileReference? GetInstalledProjectFile()
+		{
+			return IsProjectInstalled() ? InstalledProjectFile : null;
+		}
+
+		/// <summary>
+		/// Checks whether the given file is under an installed directory, and should not be overridden
+		/// </summary>
+		/// <param name="File">File to test</param>
+		/// <returns>True if the file is part of the installed distribution, false otherwise</returns>
+		public static bool IsFileInstalled(FileReference File)
+		{
+			if (IsEngineInstalled() && File.IsUnderDirectory(EngineDirectory))
+			{
+				return true;
+			}
+			if (IsProjectInstalled() && File.IsUnderDirectory(InstalledProjectFile!.Directory))
+			{
+				return true;
+			}
+			return false;
+		}
+
 		public static class LocationOverride
 		{
 			/// <summary>
@@ -239,7 +299,6 @@ namespace UnrealBuildBase
 			public static DirectoryReference? RootDirectory = null;
 		}
 
-
 		// A subset of the functionality in DataDrivenPlatformInfo.GetAllPlatformInfos() - finds the DataDrivenPlatformInfo.ini files and records their existence, but does not parse them
 		// (perhaps DataDrivenPlatformInfo.GetAllPlatformInfos() could be modified to use this data to avoid an additional search through the filesystem)
 		private static Lazy<HashSet<string>> IniPresentForPlatform = new Lazy<HashSet<string>>(() =>
diff --git a/Engine/Source/Programs/UnrealBuildTool/Configuration/TargetDescriptor.cs b/Engine/Source/Programs/UnrealBuildTool/Configuration/TargetDescriptor.cs
--- a/Engine/Source/Programs/UnrealBuildTool/Configuration/TargetDescriptor.cs
+++ b/Engine/Source/Programs/UnrealBuildTool/Configuration/TargetDescriptor.cs
@@ -638,9 +638,9 @@ namespace UnrealBuildTool
 				}
 			}
 
-			if (UnrealBuildTool.IsProjectInstalled())
+			if (Unreal.IsProjectInstalled())
 			{
-				ProjectFile = UnrealBuildTool.GetInstalledProjectFile()!;
+				ProjectFile = Unreal.GetInstalledProjectFile()!;
 				return true;
 			}
 
diff --git a/Engine/Source/Programs/UnrealBuildTool/Configuration/UEBuildTarget.cs b/Engine/Source/Programs/UnrealBuildTool/Configuration/UEBuildTarget.cs
--- a/Engine/Source/Programs/UnrealBuildTool/Configuration/UEBuildTarget.cs
+++ b/Engine/Source/Programs/UnrealBuildTool/Configuration/UEBuildTarget.cs
@@ -2489,7 +2489,7 @@ namespace UnrealBuildTool
 			List<UEBuildBinary> OriginalBinaries = Binaries;
 
 			// For installed builds, filter out all the binaries that aren't in mods
-			if (UnrealBuildTool.IsProjectInstalled())
+			if (Unreal.IsProjectInstalled())
 			{
 				List<DirectoryReference> ModDirectories = EnabledPlugins!.Where(x => x.Type == PluginType.Mod).Select(x => x.Directory).ToList();
 
@@ -2626,7 +2626,7 @@ namespace UnrealBuildTool
 			if (!ShouldCompileMonolithic() & !Rules.bFormalBuild && Platform.IsInGroup(UnrealPlatformGroup.Windows))
 			{
 				FileReference DefaultResourceLocation = FileReference.Combine(Unreal.EngineDirectory, "Build", "Windows", "Resources", "Default.rc2");
-				if (!UnrealBuildTool.IsFileInstalled(DefaultResourceLocation))
+				if (!Unreal.IsFileInstalled(DefaultResourceLocation))
 				{
 					CppCompileEnvironment DefaultResourceCompileEnvironment = new CppCompileEnvironment(GlobalCompileEnvironment);
 
@@ -2685,7 +2685,7 @@ namespace UnrealBuildTool
 
 			foreach (KeyValuePair<FileReference, FileReference> Pair in RuntimeDependencyTargetFileToSourceFile)
 			{
-				if (!UnrealBuildTool.IsFileInstalled(Pair.Key))
+				if (!Unreal.IsFileInstalled(Pair.Key))
 				{
 					Makefile.OutputItems.Add(MakefileBuilder.CreateCopyAction(Pair.Value, Pair.Key));
 				}
@@ -2764,7 +2764,7 @@ namespace UnrealBuildTool
 			// Remove any installed build products that don't exist. They may be part of an optional install.
 			if (Unreal.IsEngineInstalled())
 			{
-				BuildProducts.RemoveAll(x => UnrealBuildTool.IsFileInstalled(x.Key) && !FileReference.Exists(x.Key));
+				BuildProducts.RemoveAll(x => Unreal.IsFileInstalled(x.Key) && !FileReference.Exists(x.Key));
 			}
 
 			// Make sure all the checked headers were valid
@@ -3078,7 +3078,7 @@ namespace UnrealBuildTool
 				ProducedItems.Add(Info.ReceiptFile);
 			}
 
-			if (!ProducedItems.Any(x => UnrealBuildTool.IsFileInstalled(x)))
+			if (!ProducedItems.Any(x => Unreal.IsFileInstalled(x)))
 			{
 				TargetMakefile Makefile = MakefileBuilder.Makefile;
 
@@ -3840,7 +3840,7 @@ namespace UnrealBuildTool
 					List<UEBuildModule> GameModules = Binary.FindHotReloadModules();
 					if (GameModules != null && GameModules.Count > 0)
 					{
-						if (!UnrealBuildTool.IsProjectInstalled() || EnabledPlugins!.Where(x => x.Type == PluginType.Mod).Any(x => Binary.OutputFilePaths[0].IsUnderDirectory(x.Directory)))
+						if (!Unreal.IsProjectInstalled() || EnabledPlugins!.Where(x => x.Type == PluginType.Mod).Any(x => Binary.OutputFilePaths[0].IsUnderDirectory(x.Directory)))
 						{
 							HotReloadModuleNames.UnionWith(GameModules.OfType<UEBuildModuleCPP>().Where(x => !x.Rules.bUsePrecompiled).Select(x => x.Name));
 						}
diff --git a/Engine/Source/Programs/UnrealBuildTool/Modes/GenerateProjectFilesMode.cs b/Engine/Source/Programs/UnrealBuildTool/Modes/GenerateProjectFilesMode.cs
--- a/Engine/Source/Programs/UnrealBuildTool/Modes/GenerateProjectFilesMode.cs
+++ b/Engine/Source/Programs/UnrealBuildTool/Modes/GenerateProjectFilesMode.cs
@@ -329,7 +329,7 @@ namespace UnrealBuildTool
 				return true;
 			}
 
-			FileReference? InstalledProjectFile = UnrealBuildTool.GetInstalledProjectFile();
+			FileReference? InstalledProjectFile = Unreal.GetInstalledProjectFile();
 			if (InstalledProjectFile != null)
 			{
 				ProjectFile = InstalledProjectFile;
diff --git a/Engine/Source/Programs/UnrealBuildTool/Modes/WriteMetadataMode.cs b/Engine/Source/Programs/UnrealBuildTool/Modes/WriteMetadataMode.cs
--- a/Engine/Source/Programs/UnrealBuildTool/Modes/WriteMetadataMode.cs
+++ b/Engine/Source/Programs/UnrealBuildTool/Modes/WriteMetadataMode.cs
@@ -156,7 +156,7 @@ namespace UnrealBuildTool
 			foreach (KeyValuePair<FileReference, ModuleManifest> Pair in TargetInfo.FileToManifest)
 			{
 				FileReference ManifestFile = Pair.Key;
-				if (!UnrealBuildTool.IsFileInstalled(ManifestFile))
+				if (!Unreal.IsFileInstalled(ManifestFile))
 				{
 					ModuleManifest Manifest = Pair.Value;
 					Manifest.BuildId = BuildId ?? String.Empty;
@@ -197,7 +197,7 @@ namespace UnrealBuildTool
 			foreach (KeyValuePair<FileReference, LoadOrderManifest> Pair in TargetInfo.FileToLoadOrderManifest)
 			{
 				FileReference ManifestFile = Pair.Key;
-				if (!UnrealBuildTool.IsFileInstalled(ManifestFile))
+				if (!Unreal.IsFileInstalled(ManifestFile))
 				{
 					LoadOrderManifest Manifest = Pair.Value;
 
diff --git a/Engine/Source/Programs/UnrealBuildTool/Platform/Windows/VCToolChain.cs b/Engine/Source/Programs/UnrealBuildTool/Platform/Windows/VCToolChain.cs
--- a/Engine/Source/Programs/UnrealBuildTool/Platform/Windows/VCToolChain.cs
+++ b/Engine/Source/Programs/UnrealBuildTool/Platform/Windows/VCToolChain.cs
@@ -1666,7 +1666,7 @@ namespace UnrealBuildTool
 			if (SourceFile.HasExtension(".natvis") || SourceFile.HasExtension(".natstepfilter"))
 			{
 				FileReference IntermediateFile = FileReference.Combine(IntermediateDirectory, SourceFile.Name);
-				if (!UnrealBuildTool.IsFileInstalled(IntermediateFile))
+				if (!Unreal.IsFileInstalled(IntermediateFile))
 				{
 					FileItem Item = FileItem.GetItemByFileReference(IntermediateFile);
 					Graph.CreateCopyAction(SourceFile, Item);
diff --git a/Engine/Source/Programs/UnrealBuildTool/ProjectFiles/ProjectFileGenerator.cs b/Engine/Source/Programs/UnrealBuildTool/ProjectFiles/ProjectFileGenerator.cs
--- a/Engine/Source/Programs/UnrealBuildTool/ProjectFiles/ProjectFileGenerator.cs
+++ b/Engine/Source/Programs/UnrealBuildTool/ProjectFiles/ProjectFileGenerator.cs
@@ -2911,7 +2911,7 @@ namespace UnrealBuildTool
 					ProjectFile ProjectFile = FindOrAddProject(ProjectFilePath, BaseFolder, IncludeInGeneratedProjects: true, bAlreadyExisted: out bProjectAlreadyExisted);
 					ProjectFile.IsForeignProject = CheckProjectFile != null && !NativeProjects.IsNativeProject(CheckProjectFile, Logger);
 					ProjectFile.IsGeneratedProject = true;
-					ProjectFile.IsStubProject = UnrealBuildTool.IsProjectInstalled();
+					ProjectFile.IsStubProject = Unreal.IsProjectInstalled();
 					ProjectFile.IsHybridContentOnlyProject = CheckProjectFile != null && NativeProjects.IsHybridContentOnlyProject(CheckProjectFile, Logger);
 					if (TargetRulesObject.bBuildInSolutionByDefault.HasValue)
 					{
diff --git a/Engine/Source/Programs/UnrealBuildTool/ProjectFiles/Rider/RiderProjectFileGenerator.cs b/Engine/Source/Programs/UnrealBuildTool/ProjectFiles/Rider/RiderProjectFileGenerator.cs
--- a/Engine/Source/Programs/UnrealBuildTool/ProjectFiles/Rider/RiderProjectFileGenerator.cs
+++ b/Engine/Source/Programs/UnrealBuildTool/ProjectFiles/Rider/RiderProjectFileGenerator.cs
@@ -325,7 +325,7 @@ namespace UnrealBuildTool
 					ProjectFile.IsForeignProject =
 						CheckProjectFile != null && !NativeProjects.IsNativeProject(CheckProjectFile, Logger);
 					ProjectFile.IsGeneratedProject = true;
-					ProjectFile.IsStubProject = UnrealBuildTool.IsProjectInstalled();
+					ProjectFile.IsStubProject = Unreal.IsProjectInstalled();
 					if (TargetRulesObject.bBuildInSolutionByDefault.HasValue)
 					{
 						ProjectFile.ShouldBuildByDefaultForSolutionTargets =
diff --git a/Engine/Source/Programs/UnrealBuildTool/System/DynamicCompilation.cs b/Engine/Source/Programs/UnrealBuildTool/System/DynamicCompilation.cs
--- a/Engine/Source/Programs/UnrealBuildTool/System/DynamicCompilation.cs
+++ b/Engine/Source/Programs/UnrealBuildTool/System/DynamicCompilation.cs
@@ -33,7 +33,7 @@ namespace UnrealBuildTool
 		private static bool RequiresCompilation(IEnumerable<FileReference> SourceFiles, FileReference AssemblyManifestFilePath, FileReference OutputAssemblyPath, ILogger Logger)
 		{
 			// Do not compile the file if it's installed
-			if (UnrealBuildTool.IsFileInstalled(OutputAssemblyPath))
+			if (Unreal.IsFileInstalled(OutputAssemblyPath))
 			{
 				Logger.LogDebug("Skipping {OutputAssemblyPath}: File is installed", OutputAssemblyPath);
 				return false;
diff --git a/Engine/Source/Programs/UnrealBuildTool/System/ExternalExecution.cs b/Engine/Source/Programs/UnrealBuildTool/System/ExternalExecution.cs
--- a/Engine/Source/Programs/UnrealBuildTool/System/ExternalExecution.cs
+++ b/Engine/Source/Programs/UnrealBuildTool/System/ExternalExecution.cs
@@ -626,7 +626,7 @@ namespace UnrealBuildTool
 				}
 
 				// Don't check timestamps for individual binaries if we're using the installed version of UHT. It will always be up to date.
-				if (!UnrealBuildTool.IsFileInstalled(ReceiptFile.Location))
+				if (!Unreal.IsFileInstalled(ReceiptFile.Location))
 				{
 					TargetReceipt? Receipt;
 					if (!TargetReceipt.TryRead(ReceiptPath, out Receipt))
diff --git a/Engine/Source/Programs/UnrealBuildTool/System/RulesCompiler.cs b/Engine/Source/Programs/UnrealBuildTool/System/RulesCompiler.cs
--- a/Engine/Source/Programs/UnrealBuildTool/System/RulesCompiler.cs
+++ b/Engine/Source/Programs/UnrealBuildTool/System/RulesCompiler.cs
@@ -329,7 +329,7 @@ namespace UnrealBuildTool
 					}
 					else
 					{
-						ProjectRulesAssembly = new RulesAssembly(Scope, new List<DirectoryReference> { MainProjectDirectory }, ProjectPlugins, ModuleFiles, TargetFiles, AssemblyFileName, bContainsEngineModules: false, DefaultBuildSettings: null, bReadOnly: UnrealBuildTool.IsProjectInstalled(), bSkipCompile: bSkipCompile, bForceCompile: bForceCompile, Parent: Parent, Logger: Logger);
+						ProjectRulesAssembly = new RulesAssembly(Scope, new List<DirectoryReference> { MainProjectDirectory }, ProjectPlugins, ModuleFiles, TargetFiles, AssemblyFileName, bContainsEngineModules: false, DefaultBuildSettings: null, bReadOnly: Unreal.IsProjectInstalled(), bSkipCompile: bSkipCompile, bForceCompile: bForceCompile, Parent: Parent, Logger: Logger);
 					}
 					return ProjectRulesAssembly;
 				});
diff --git a/Engine/Source/Programs/UnrealBuildTool/UnrealBuildTool.cs b/Engine/Source/Programs/UnrealBuildTool/UnrealBuildTool.cs
--- a/Engine/Source/Programs/UnrealBuildTool/UnrealBuildTool.cs
+++ b/Engine/Source/Programs/UnrealBuildTool/UnrealBuildTool.cs
@@ -48,16 +48,6 @@ namespace UnrealBuildTool
 		/// </summary>
 		public static System.Collections.IDictionary? InitialEnvironment;
 
-		/// <summary>
-		/// Whether we're running with an installed project
-		/// </summary>
-		private static bool? bIsProjectInstalled;
-
-		/// <summary>
-		/// If we are running with an installed project, specifies the path to it
-		/// </summary>
-		static FileReference? InstalledProjectFile;
-
 		/// <summary>
 		/// The full name of the Engine/Source directory
 		/// </summary>
@@ -140,58 +130,23 @@ namespace UnrealBuildTool
 		/// Returns true if UnrealBuildTool is running using an installed project (ie. a mod kit)
 		/// </summary>
 		/// <returns>True if running using an installed project</returns>
-		public static bool IsProjectInstalled()
-		{
-			if (!bIsProjectInstalled.HasValue)
-			{
-				FileReference InstalledProjectLocationFile = FileReference.Combine(Unreal.RootDirectory, "Engine", "Build", "InstalledProjectBuild.txt");
-				if (FileReference.Exists(InstalledProjectLocationFile))
-				{
-					InstalledProjectFile = FileReference.Combine(Unreal.RootDirectory, File.ReadAllText(InstalledProjectLocationFile.FullName).Trim());
-					bIsProjectInstalled = true;
-				}
-				else
-				{
-					InstalledProjectFile = null;
-					bIsProjectInstalled = false;
-				}
-			}
-			return bIsProjectInstalled.Value;
-		}
+		[Obsolete("Deprecated in UE5.5 - use Unreal.IsProjectInstalled")]
+		public static bool IsProjectInstalled() => Unreal.IsProjectInstalled();
 
 		/// <summary>
 		/// Gets the installed project file
 		/// </summary>
 		/// <returns>Location of the installed project file</returns>
-		public static FileReference? GetInstalledProjectFile()
-		{
-			if (IsProjectInstalled())
-			{
-				return InstalledProjectFile;
-			}
-			else
-			{
-				return null;
-			}
-		}
+		[Obsolete("Deprecated in UE5.5 - use Unreal.GetInstalledProjectFile")]
+		public static FileReference? GetInstalledProjectFile() => Unreal.GetInstalledProjectFile();
 
 		/// <summary>
 		/// Checks whether the given file is under an installed directory, and should not be overridden
 		/// </summary>
 		/// <param name="File">File to test</param>
 		/// <returns>True if the file is part of the installed distribution, false otherwise</returns>
-		public static bool IsFileInstalled(FileReference File)
-		{
-			if (Unreal.IsEngineInstalled() && File.IsUnderDirectory(Unreal.EngineDirectory))
-			{
-				return true;
-			}
-			if (IsProjectInstalled() && File.IsUnderDirectory(InstalledProjectFile!.Directory))
-			{
-				return true;
-			}
-			return false;
-		}
+		[Obsolete("Deprecated in UE5.5 - use Unreal.IsFileInstalled")]
+		public static bool IsFileInstalled(FileReference File) => Unreal.IsFileInstalled(File);
 
 		/// <summary>
 		/// Gets the absolute path to the UBT assembly.
