EpicGames.UBA: Maintain reference to wrapped native handles in CacheClientImpl.cs

#rnx

diff --git a/Engine/Source/Programs/Shared/EpicGames.UBA/Impl/CacheClientImpl.cs b/Engine/Source/Programs/Shared/EpicGames.UBA/Impl/CacheClientImpl.cs
--- a/Engine/Source/Programs/Shared/EpicGames.UBA/Impl/CacheClientImpl.cs
+++ b/Engine/Source/Programs/Shared/EpicGames.UBA/Impl/CacheClientImpl.cs
@@ -8,6 +8,7 @@ namespace EpicGames.UBA
 	internal class RootPathsImpl : IRootPaths
 	{
 		IntPtr _handle = IntPtr.Zero;
+		readonly ILogger _logger;
 
 		#region DllImport
 		[DllImport("UbaHost", CharSet = CharSet.Auto)]
@@ -48,7 +49,8 @@ namespace EpicGames.UBA
 
 		public RootPathsImpl(ILogger logger)
 		{
-			_handle = RootPaths_Create(logger.GetHandle());
+			_logger = logger;
+			_handle = RootPaths_Create(_logger.GetHandle());
 		}
 
 		public bool RegisterRoot(string path, bool includeInKey)
@@ -67,6 +69,8 @@ namespace EpicGames.UBA
 	internal class CacheClientImpl : ICacheClient
 	{
 		IntPtr _handle = IntPtr.Zero;
+		readonly ISessionServer _sessionServer;
+
 		public delegate void ExitCallback(IntPtr userData, IntPtr handle);
 
 		#region DllImport
@@ -118,9 +122,10 @@ namespace EpicGames.UBA
 		}
 		#endregion
 
-		public CacheClientImpl(ISessionServer session, bool reportMissReason)
+		public CacheClientImpl(ISessionServer server, bool reportMissReason)
 		{
-			_handle = CacheClient_Create(session.GetHandle(), reportMissReason);
+			_sessionServer = server;
+			_handle = CacheClient_Create(_sessionServer.GetHandle(), reportMissReason);
 		}
 
 		public bool Connect(string host, int port)
