Fix spelling error
#rnx

diff --git a/Engine/Source/Programs/Shared/EpicGames.Perforce/PerforceConnection.cs b/Engine/Source/Programs/Shared/EpicGames.Perforce/PerforceConnection.cs
--- a/Engine/Source/Programs/Shared/EpicGames.Perforce/PerforceConnection.cs
+++ b/Engine/Source/Programs/Shared/EpicGames.Perforce/PerforceConnection.cs
@@ -2288,7 +2288,7 @@ namespace EpicGames.Perforce
 		#region p4 have
 
 		/// <summary>
-		/// Dtermine files currently synced to the client
+		/// Determine files currently synced to the client
 		/// </summary>
 		/// <param name="connection">Connection to the Perforce server</param>
 		/// <param name="fileSpec">Files to query</param>
@@ -2300,7 +2300,7 @@ namespace EpicGames.Perforce
 		}
 
 		/// <summary>
-		/// Dtermine files currently synced to the client
+		/// Determine files currently synced to the client
 		/// </summary>
 		/// <param name="connection">Connection to the Perforce server</param>
 		/// <param name="fileSpec">Files to query</param>
