Added UHT support for FAnsiString and FUtf8String.

diff --git a/Engine/Source/Programs/Shared/EpicGames.UHT/Types/Properties/UhtAnsiStrProperty.cs b/Engine/Source/Programs/Shared/EpicGames.UHT/Types/Properties/UhtAnsiStrProperty.cs
new file mode 100644
--- /dev/null
+++ b/Engine/Source/Programs/Shared/EpicGames.UHT/Types/Properties/UhtAnsiStrProperty.cs
@@ -0,0 +1,138 @@
+// Copyright Epic Games, Inc. All Rights Reserved.
+
+using System.Diagnostics.CodeAnalysis;
+using System.Text;
+using EpicGames.Core;
+using EpicGames.UHT.Tables;
+using EpicGames.UHT.Tokenizer;
+using EpicGames.UHT.Utils;
+
+namespace EpicGames.UHT.Types
+{
+
+	/// <summary>
+	/// FAnsiStrProperty
+	/// </summary>
+	[UnrealHeaderTool]
+	[UhtEngineClass(Name = "AnsiStrProperty", IsProperty = true)]
+	public class UhtAnsiStrProperty : UhtProperty
+	{
+		/// <inheritdoc/>
+		public override string EngineClassName => "AnsiStrProperty";
+
+		/// <inheritdoc/>
+		protected override string CppTypeText => "FAnsiString";
+
+		/// <inheritdoc/>
+		protected override string PGetMacroText => "PROPERTY";
+
+		/// <inheritdoc/>
+		protected override UhtPGetArgumentType PGetTypeArgument => UhtPGetArgumentType.EngineClass;
+
+		/// <summary>
+		/// Construct a new property
+		/// </summary>
+		/// <param name="propertySettings">Property settings</param>
+		public UhtAnsiStrProperty(UhtPropertySettings propertySettings) : base(propertySettings)
+		{
+			// Other caps not supported until engine support catches up
+			PropertyCaps |= UhtPropertyCaps.PassCppArgsByRef; // | UhtPropertyCaps.CanExposeOnSpawn | UhtPropertyCaps.IsParameterSupportedByBlueprint |
+				//UhtPropertyCaps.IsMemberSupportedByBlueprint | UhtPropertyCaps.SupportsRigVM;
+		}
+
+		/// <inheritdoc/>
+		public override StringBuilder AppendNullConstructorArg(StringBuilder builder, bool isInitializer)
+		{
+			builder.Append("TEXT(\"\")");
+			return builder;
+		}
+
+		/// <inheritdoc/>
+		public override StringBuilder AppendMemberDecl(StringBuilder builder, IUhtPropertyMemberContext context, string name, string nameSuffix, int tabs)
+		{
+			return AppendMemberDecl(builder, context, name, nameSuffix, tabs, "FAnsiStrPropertyParams");
+		}
+
+		/// <inheritdoc/>
+		public override StringBuilder AppendMemberDef(StringBuilder builder, IUhtPropertyMemberContext context, string name, string nameSuffix, string? offset, int tabs)
+		{
+			AppendMemberDefStart(builder, context, name, nameSuffix, offset, tabs, "FAnsiStrPropertyParams", "UECodeGen_Private::EPropertyGenFlags::AnsiStr");
+			AppendMemberDefEnd(builder, context, name, nameSuffix);
+			return builder;
+		}
+
+		/// <inheritdoc/>
+		public override bool SanitizeDefaultValue(IUhtTokenReader defaultValueReader, StringBuilder innerDefaultValue)
+		{
+			if (defaultValueReader.TryOptional("FAnsiString"))
+			{
+				defaultValueReader.Require('(');
+				StringView value = defaultValueReader.GetWrappedConstString();
+				defaultValueReader.Require(')');
+				innerDefaultValue.Append(value);
+			}
+			else
+			{
+				StringView value = defaultValueReader.GetWrappedConstString();
+				innerDefaultValue.Append(value);
+			}
+			return true;
+		}
+
+		/// <inheritdoc/>
+		public override bool IsSameType(UhtProperty other)
+		{
+			return other is UhtAnsiStrProperty;
+		}
+
+		/// <inheritdoc/>
+		protected override void ValidateFunctionArgument(UhtFunction function, UhtValidationOptions options)
+		{
+			base.ValidateFunctionArgument(function, options);
+
+			if (function.FunctionFlags.HasAnyFlags(EFunctionFlags.Net))
+			{
+				if (!function.FunctionFlags.HasAnyFlags(EFunctionFlags.NetRequest))
+				{
+					if (RefQualifier != UhtPropertyRefQualifier.ConstRef && !IsStaticArray)
+					{
+						this.LogError("Replicated FAnsiString parameters must be passed by const reference");
+					}
+				}
+			}
+		}
+
+		[UhtPropertyType(Keyword = "FAnsiString")]
+		[SuppressMessage("CodeQuality", "IDE0051:Remove unused private members", Justification = "Attribute accessed method")]
+		[SuppressMessage("Style", "IDE0060:Remove unused parameter", Justification = "Attribute accessed method")]
+		private static UhtProperty? AnsiStrProperty(UhtPropertyResolvePhase resolvePhase, UhtPropertySettings propertySettings, IUhtTokenReader tokenReader, UhtToken matchedToken)
+		{
+			if (!tokenReader.SkipExpectedType(matchedToken.Value, propertySettings.PropertyCategory == UhtPropertyCategory.Member))
+			{
+				return null;
+			}
+			UhtAnsiStrProperty property = new(propertySettings);
+			if (property.PropertyCategory != UhtPropertyCategory.Member)
+			{
+				if (tokenReader.TryOptional('&'))
+				{
+					if (property.PropertyFlags.HasAnyFlags(EPropertyFlags.ConstParm))
+					{
+						property.PropertyFlags &= ~EPropertyFlags.ConstParm;
+
+						// We record here that we encountered a const reference, because we need to remove that information from flags for code generation purposes.
+						property.RefQualifier = UhtPropertyRefQualifier.ConstRef;
+					}
+					else
+					{
+						property.PropertyFlags |= EPropertyFlags.OutParm;
+
+						// And we record here that we encountered a non-const reference here too.
+						property.RefQualifier = UhtPropertyRefQualifier.NonConstRef;
+					}
+				}
+			}
+			return property;
+		}
+	}
+}
diff --git a/Engine/Source/Programs/Shared/EpicGames.UHT/Types/Properties/UhtUtf8StrProperty.cs b/Engine/Source/Programs/Shared/EpicGames.UHT/Types/Properties/UhtUtf8StrProperty.cs
new file mode 100644
--- /dev/null
+++ b/Engine/Source/Programs/Shared/EpicGames.UHT/Types/Properties/UhtUtf8StrProperty.cs
@@ -0,0 +1,138 @@
+// Copyright Epic Games, Inc. All Rights Reserved.
+
+using System.Diagnostics.CodeAnalysis;
+using System.Text;
+using EpicGames.Core;
+using EpicGames.UHT.Tables;
+using EpicGames.UHT.Tokenizer;
+using EpicGames.UHT.Utils;
+
+namespace EpicGames.UHT.Types
+{
+
+	/// <summary>
+	/// FUtf8StrProperty
+	/// </summary>
+	[UnrealHeaderTool]
+	[UhtEngineClass(Name = "Utf8StrProperty", IsProperty = true)]
+	public class UhtUtf8StrProperty : UhtProperty
+	{
+		/// <inheritdoc/>
+		public override string EngineClassName => "Utf8StrProperty";
+
+		/// <inheritdoc/>
+		protected override string CppTypeText => "FUtf8String";
+
+		/// <inheritdoc/>
+		protected override string PGetMacroText => "PROPERTY";
+
+		/// <inheritdoc/>
+		protected override UhtPGetArgumentType PGetTypeArgument => UhtPGetArgumentType.EngineClass;
+
+		/// <summary>
+		/// Construct a new property
+		/// </summary>
+		/// <param name="propertySettings">Property settings</param>
+		public UhtUtf8StrProperty(UhtPropertySettings propertySettings) : base(propertySettings)
+		{
+			// Other caps not supported until engine support catches up
+			PropertyCaps |= UhtPropertyCaps.PassCppArgsByRef; // | UhtPropertyCaps.CanExposeOnSpawn | UhtPropertyCaps.IsParameterSupportedByBlueprint |
+				//UhtPropertyCaps.IsMemberSupportedByBlueprint | UhtPropertyCaps.SupportsRigVM;
+		}
+
+		/// <inheritdoc/>
+		public override StringBuilder AppendNullConstructorArg(StringBuilder builder, bool isInitializer)
+		{
+			builder.Append("TEXT(\"\")");
+			return builder;
+		}
+
+		/// <inheritdoc/>
+		public override StringBuilder AppendMemberDecl(StringBuilder builder, IUhtPropertyMemberContext context, string name, string nameSuffix, int tabs)
+		{
+			return AppendMemberDecl(builder, context, name, nameSuffix, tabs, "FUtf8StrPropertyParams");
+		}
+
+		/// <inheritdoc/>
+		public override StringBuilder AppendMemberDef(StringBuilder builder, IUhtPropertyMemberContext context, string name, string nameSuffix, string? offset, int tabs)
+		{
+			AppendMemberDefStart(builder, context, name, nameSuffix, offset, tabs, "FUtf8StrPropertyParams", "UECodeGen_Private::EPropertyGenFlags::Utf8Str");
+			AppendMemberDefEnd(builder, context, name, nameSuffix);
+			return builder;
+		}
+
+		/// <inheritdoc/>
+		public override bool SanitizeDefaultValue(IUhtTokenReader defaultValueReader, StringBuilder innerDefaultValue)
+		{
+			if (defaultValueReader.TryOptional("FUtf8String"))
+			{
+				defaultValueReader.Require('(');
+				StringView value = defaultValueReader.GetWrappedConstString();
+				defaultValueReader.Require(')');
+				innerDefaultValue.Append(value);
+			}
+			else
+			{
+				StringView value = defaultValueReader.GetWrappedConstString();
+				innerDefaultValue.Append(value);
+			}
+			return true;
+		}
+
+		/// <inheritdoc/>
+		public override bool IsSameType(UhtProperty other)
+		{
+			return other is UhtUtf8StrProperty;
+		}
+
+		/// <inheritdoc/>
+		protected override void ValidateFunctionArgument(UhtFunction function, UhtValidationOptions options)
+		{
+			base.ValidateFunctionArgument(function, options);
+
+			if (function.FunctionFlags.HasAnyFlags(EFunctionFlags.Net))
+			{
+				if (!function.FunctionFlags.HasAnyFlags(EFunctionFlags.NetRequest))
+				{
+					if (RefQualifier != UhtPropertyRefQualifier.ConstRef && !IsStaticArray)
+					{
+						this.LogError("Replicated FUtf8String parameters must be passed by const reference");
+					}
+				}
+			}
+		}
+
+		[UhtPropertyType(Keyword = "FUtf8String")]
+		[SuppressMessage("CodeQuality", "IDE0051:Remove unused private members", Justification = "Attribute accessed method")]
+		[SuppressMessage("Style", "IDE0060:Remove unused parameter", Justification = "Attribute accessed method")]
+		private static UhtProperty? Utf8StrProperty(UhtPropertyResolvePhase resolvePhase, UhtPropertySettings propertySettings, IUhtTokenReader tokenReader, UhtToken matchedToken)
+		{
+			if (!tokenReader.SkipExpectedType(matchedToken.Value, propertySettings.PropertyCategory == UhtPropertyCategory.Member))
+			{
+				return null;
+			}
+			UhtUtf8StrProperty property = new(propertySettings);
+			if (property.PropertyCategory != UhtPropertyCategory.Member)
+			{
+				if (tokenReader.TryOptional('&'))
+				{
+					if (property.PropertyFlags.HasAnyFlags(EPropertyFlags.ConstParm))
+					{
+						property.PropertyFlags &= ~EPropertyFlags.ConstParm;
+
+						// We record here that we encountered a const reference, because we need to remove that information from flags for code generation purposes.
+						property.RefQualifier = UhtPropertyRefQualifier.ConstRef;
+					}
+					else
+					{
+						property.PropertyFlags |= EPropertyFlags.OutParm;
+
+						// And we record here that we encountered a non-const reference here too.
+						property.RefQualifier = UhtPropertyRefQualifier.NonConstRef;
+					}
+				}
+			}
+			return property;
+		}
+	}
+}
