Horde: Fix http auth handler not using correct configuration for server when run locally.

- Caused by trying to use a single HttpAuthMessageHandlerState for DI container, even though multiple clients can now be created.
- A message handler is now owned by each Horde client instance, which has an appropriately configured HTTP message handler for that server.

diff --git a/Engine/Source/Programs/AutomationTool/BuildGraph/Tasks/HordeGetSecretsTask.cs b/Engine/Source/Programs/AutomationTool/BuildGraph/Tasks/HordeGetSecretsTask.cs
--- a/Engine/Source/Programs/AutomationTool/BuildGraph/Tasks/HordeGetSecretsTask.cs
+++ b/Engine/Source/Programs/AutomationTool/BuildGraph/Tasks/HordeGetSecretsTask.cs
@@ -136,6 +136,7 @@ namespace AutomationTool.Tasks
 			if (secretToReplacementInfo.Count > 0)
 			{
 				ServiceCollection serviceCollection = new ServiceCollection();
+				serviceCollection.AddLogging(builder => builder.AddEpicDefault());
 				serviceCollection.AddHorde(options => options.AllowAuthPrompt = !CommandUtils.IsBuildMachine);
 
 				await using (ServiceProvider serviceProvider = serviceCollection.BuildServiceProvider())
@@ -151,9 +152,13 @@ namespace AutomationTool.Tasks
 						{
 							secret = await hordeHttpClient.GetSecretAsync(secretId);
 						}
+						catch (HttpRequestException ex) when (ex.StatusCode == System.Net.HttpStatusCode.NotFound)
+						{
+							throw new AutomationException(ex, $"Secret '{secretId}' was not found on {hordeClient.ServerUrl}");
+						}
 						catch (HttpRequestException ex) when (ex.StatusCode == System.Net.HttpStatusCode.Forbidden)
 						{
-							throw new AutomationException(ex, $"User does not have permissions to read {secretId}");
+							throw new AutomationException(ex, $"User does not have permissions to read '{secretId}' on {hordeClient.ServerUrl}");
 						}
 
 						foreach (ReplacementInfo replacement in replacements)
diff --git a/Engine/Source/Programs/Shared/EpicGames.Horde/HordeClient.cs b/Engine/Source/Programs/Shared/EpicGames.Horde/HordeClient.cs
--- a/Engine/Source/Programs/Shared/EpicGames.Horde/HordeClient.cs
+++ b/Engine/Source/Programs/Shared/EpicGames.Horde/HordeClient.cs
@@ -1,6 +1,7 @@
 ﻿// Copyright Epic Games, Inc. All Rights Reserved.
 
 using System;
+using System.Net;
 using System.Net.Http;
 using System.Net.Http.Headers;
 using System.Net.Http.Json;
@@ -17,8 +18,13 @@ using EpicGames.Horde.Storage.Bundles;
 using Grpc.Core;
 using Grpc.Net.Client;
 using Grpc.Net.Client.Configuration;
+using Microsoft.Extensions.Http;
 using Microsoft.Extensions.Logging;
 using Microsoft.Extensions.Options;
+using Polly;
+using Polly.Extensions.Http;
+using Polly.Retry;
+using Polly.Timeout;
 
 #pragma warning disable CA2234 // Use URIs instead of strings
 
@@ -30,7 +36,6 @@ namespace EpicGames.Horde
 	abstract class HordeClient : IHordeClient
 	{
 		readonly Uri _serverUrl;
-		readonly IHttpClientFactory _httpClientFactory;
 		readonly BundleCache _bundleCache;
 		readonly HordeOptions _hordeOptions;
 		readonly ILoggerFactory _loggerFactory;
@@ -44,10 +49,9 @@ namespace EpicGames.Horde
 		/// <summary>
 		/// Constructor
 		/// </summary>
-		public HordeClient(Uri serverUrl, IHttpClientFactory httpClientFactory, BundleCache bundleCache, IOptionsSnapshot<HordeOptions> hordeOptions, ILoggerFactory loggerFactory)
+		public HordeClient(Uri serverUrl, BundleCache bundleCache, IOptionsSnapshot<HordeOptions> hordeOptions, ILoggerFactory loggerFactory)
 		{
 			_serverUrl = serverUrl;
-			_httpClientFactory = httpClientFactory;
 			_bundleCache = bundleCache;
 			_hordeOptions = hordeOptions.Value;
 			_loggerFactory = loggerFactory;
@@ -55,7 +59,7 @@ namespace EpicGames.Horde
 		}
 
 		/// <inheritdoc/>
-		public async ValueTask DisposeAsync()
+		public virtual async ValueTask DisposeAsync()
 		{
 			if (_grpcChannel != null)
 			{
@@ -64,6 +68,74 @@ namespace EpicGames.Horde
 			}
 		}
 
+		/// <summary>
+		/// Creates a new message handler with default resilience policies
+		/// </summary>
+		protected HttpMessageHandler CreateDefaultHttpMessageHandler()
+		{
+			HttpMessageHandler? httpMessageHandler = null;
+			try
+			{
+#pragma warning disable CA2000 // Call dispose on httpMessageHandler (disposed by child handlers)
+				httpMessageHandler = new SocketsHttpHandler();
+				httpMessageHandler = new PolicyHttpMessageHandler(request => CreateDefaultTimeoutRetryPolicy(request)){ InnerHandler = httpMessageHandler };
+				httpMessageHandler = new PolicyHttpMessageHandler(request => CreateDefaultTransientErrorPolicy(request)){ InnerHandler = httpMessageHandler };
+#pragma warning restore CA2000
+				return httpMessageHandler;
+			}
+			catch
+			{
+				httpMessageHandler?.Dispose();
+				throw;
+			}
+		}
+
+		/// <summary>
+		/// Create a default timeout retry policy
+		/// </summary>
+		protected IAsyncPolicy<HttpResponseMessage> CreateDefaultTimeoutRetryPolicy(HttpRequestMessage request)
+		{
+			// Wait 30 seconds for operations to timeout
+			Task OnTimeoutAsync(Context context, TimeSpan timespan, Task timeoutTask)
+			{
+				_logger.LogWarning(KnownLogEvents.Systemic_Horde_Http, "{Method} {Url} timed out after {Time}s.", request.Method, request.RequestUri, (int)timespan.TotalSeconds);
+				return Task.CompletedTask;
+			}
+
+			AsyncTimeoutPolicy<HttpResponseMessage> timeoutPolicy = Policy.TimeoutAsync<HttpResponseMessage>(30, OnTimeoutAsync);
+
+			// Retry twice after a timeout
+			void OnRetry(Exception ex, TimeSpan timespan)
+			{
+				_logger.LogWarning(KnownLogEvents.Systemic_Horde_Http, ex, "{Method} {Url} retrying after {Time}s.", request.Method, request.RequestUri, timespan.TotalSeconds);
+			}
+
+			TimeSpan[] retryTimes = new[] { TimeSpan.FromSeconds(5.0), TimeSpan.FromSeconds(10.0) };
+			AsyncRetryPolicy retryPolicy = Policy.Handle<TimeoutRejectedException>().WaitAndRetryAsync(retryTimes, OnRetry);
+			return retryPolicy.WrapAsync(timeoutPolicy);
+		}
+
+		/// <summary>
+		/// Create a default timeout retry policy
+		/// </summary>
+		protected IAsyncPolicy<HttpResponseMessage> CreateDefaultTransientErrorPolicy(HttpRequestMessage request)
+		{
+			Task OnTimeoutAsync(DelegateResult<HttpResponseMessage> outcome, TimeSpan timespan, int retryAttempt, Context context)
+			{
+				_logger.LogWarning(KnownLogEvents.Systemic_Horde_Http, "{Method} {Url} failed ({Result}). Delaying for {DelayMs}ms (attempt #{RetryNum}).", request.Method, request.RequestUri, outcome.Result?.StatusCode, timespan.TotalMilliseconds, retryAttempt);
+				return Task.CompletedTask;
+			}
+
+			TimeSpan[] retryTimes = new[] { TimeSpan.FromSeconds(1.0), TimeSpan.FromSeconds(5.0), TimeSpan.FromSeconds(10.0), TimeSpan.FromSeconds(30.0), TimeSpan.FromSeconds(30.0) };
+
+			// Policy for transient errors is the same as HttpPolicyExtensions.HandleTransientHttpError(), but excludes HttpStatusCode.ServiceUnavailable (which is used as a response
+			// when allocating compute resources when none are available). This pathway is handled explicitly on the application side.
+			return Policy<HttpResponseMessage>
+				.Handle<HttpRequestException>()
+				.OrResult(x => (x.StatusCode >= HttpStatusCode.InternalServerError && x.StatusCode != HttpStatusCode.ServiceUnavailable) || x.StatusCode == HttpStatusCode.RequestTimeout)
+				.WaitAndRetryAsync(retryTimes, OnTimeoutAsync);
+		}
+
 		/// <inheritdoc/>
 		public abstract Task<bool> LoginAsync(bool allowLogin, CancellationToken cancellationToken);
 
@@ -90,7 +162,7 @@ namespace EpicGames.Horde
 			if (useInsecureConnection)
 			{
 				_logger.LogInformation("Querying server {BaseUrl} for rpc port", serverUri);
-				using (HttpClient httpClient = _httpClientFactory.CreateClient())
+				using (HttpClient httpClient = CreateUnauthenticatedHttpClient())
 				{
 					httpClient.DefaultRequestHeaders.Add("Accept", "application/json");
 					httpClient.Timeout = TimeSpan.FromSeconds(210); // Need to make sure this doesn't cancel any long running gRPC streaming calls (eg. session update)
@@ -122,7 +194,7 @@ namespace EpicGames.Horde
 			channelOptions.ServiceConfig.MethodConfigs.Add(new MethodConfig
 			{
 				Names = { MethodName.Default },
-				RetryPolicy = new RetryPolicy
+				RetryPolicy = new Grpc.Net.Client.Configuration.RetryPolicy
 				{
 					MaxAttempts = 3,
 					InitialBackoff = TimeSpan.FromSeconds(1),
@@ -168,7 +240,7 @@ namespace EpicGames.Horde
 
 		/// <inheritdoc/>
 		public HordeHttpClient CreateHttpClient()
-			=> new HordeHttpClient(CreateDefaultHttpClient());
+			=> new HordeHttpClient(CreateAuthenticatedHttpClient());
 
 		public IComputeClient CreateComputeClient()
 		{
@@ -182,7 +254,7 @@ namespace EpicGames.Horde
 				sessionId = $"{jobId}-{batchId}-{stepId}";
 			}
 
-			return new ServerComputeClient(CreateDefaultHttpClient(), sessionId, _loggerFactory.CreateLogger<ServerComputeClient>());
+			return new ServerComputeClient(CreateAuthenticatedHttpClient(), sessionId, _loggerFactory.CreateLogger<ServerComputeClient>());
 		}
 
 		/// <inheritdoc/>
@@ -190,16 +262,19 @@ namespace EpicGames.Horde
 		{
 			HttpClient CreateClient()
 			{
-				HttpClient httpClient = CreateDefaultHttpClient();
 				if (accessToken != null)
 				{
+					HttpClient httpClient = CreateUnauthenticatedHttpClient();
 					httpClient.DefaultRequestHeaders.Authorization = new AuthenticationHeaderValue("Bearer", accessToken);
+					return httpClient;
+				}
+				else
+				{
+					return CreateAuthenticatedHttpClient();
 				}
-				return httpClient;
 			}
 
-			HttpClient CreateUploadRedirectClient() => _httpClientFactory.CreateClient(HordeHttpClient.UploadRedirectHttpClientName);
-			HttpStorageBackend httpStorageBackend = new HttpStorageBackend(basePath, CreateClient, CreateUploadRedirectClient, _loggerFactory.CreateLogger<HttpStorageBackend>());
+			HttpStorageBackend httpStorageBackend = new HttpStorageBackend(basePath, CreateClient, CreateUnauthenticatedHttpClient, _loggerFactory.CreateLogger<HttpStorageBackend>());
 			return new BundleStorageClient(httpStorageBackend, _bundleCache, _hordeOptions.Bundle, _loggerFactory.CreateLogger<BundleStorageClient>());
 		}
 
@@ -210,7 +285,12 @@ namespace EpicGames.Horde
 		/// <summary>
 		/// Creates an http client for satisfying requests
 		/// </summary>
-		protected abstract HttpClient CreateDefaultHttpClient();
+		protected abstract HttpClient CreateAuthenticatedHttpClient();
+
+		/// <summary>
+		/// Creates an http client for satisfying requests
+		/// </summary>
+		protected abstract HttpClient CreateUnauthenticatedHttpClient();
 	}
 
 	/// <summary>
@@ -219,16 +299,23 @@ namespace EpicGames.Horde
 	class HordeClientWithStaticCredentials : HordeClient
 	{
 		readonly string? _accessToken;
-		readonly IHttpClientFactory _httpClientFactory;
+		readonly HttpMessageHandler _httpMessageHandler;
 
 		/// <summary>
 		/// Constructor
 		/// </summary>
-		public HordeClientWithStaticCredentials(Uri serverUrl, string? accessToken, IHttpClientFactory httpClientFactory, BundleCache bundleCache, IOptionsSnapshot<HordeOptions> hordeOptions, ILoggerFactory loggerFactory)
-			: base(serverUrl, httpClientFactory, bundleCache, hordeOptions, loggerFactory)
+		public HordeClientWithStaticCredentials(Uri serverUrl, string? accessToken, BundleCache bundleCache, IOptionsSnapshot<HordeOptions> hordeOptions, ILoggerFactory loggerFactory)
+			: base(serverUrl, bundleCache, hordeOptions, loggerFactory)
 		{
 			_accessToken = accessToken;
-			_httpClientFactory = httpClientFactory;
+			_httpMessageHandler = CreateDefaultHttpMessageHandler();
+		}
+
+		/// <inheritdoc/>
+		public override async ValueTask DisposeAsync()
+		{
+			_httpMessageHandler.Dispose();
+			await base.DisposeAsync();
 		}
 
 		/// <inheritdoc/>
@@ -243,9 +330,10 @@ namespace EpicGames.Horde
 		public override Task<string?> GetAccessTokenAsync(bool interactive, CancellationToken cancellationToken)
 			=> Task.FromResult(_accessToken);
 
-		protected override HttpClient CreateDefaultHttpClient()
+		/// <inheritdoc/>
+		protected override HttpClient CreateAuthenticatedHttpClient()
 		{
-			HttpClient httpClient = _httpClientFactory.CreateClient(HordeHttpClient.HttpClientName);
+			HttpClient httpClient = new HttpClient(_httpMessageHandler, false);
 			httpClient.BaseAddress = ServerUrl;
 			if (!String.IsNullOrEmpty(_accessToken))
 			{
@@ -253,6 +341,14 @@ namespace EpicGames.Horde
 			}
 			return httpClient;
 		}
+
+		/// <inheritdoc/>
+		protected override HttpClient CreateUnauthenticatedHttpClient()
+		{
+			HttpClient httpClient = new HttpClient(_httpMessageHandler, false);
+			httpClient.BaseAddress = ServerUrl;
+			return httpClient;
+		}
 	}
 
 	/// <summary>
@@ -260,23 +356,35 @@ namespace EpicGames.Horde
 	/// </summary>
 	class HordeClientWithDynamicCredentials : HordeClient
 	{
-		readonly IHttpClientFactory _httpClientFactory;
-		readonly HordeHttpAuthHandlerState _authHandler;
+		readonly HordeHttpAuthHandlerState _authHandlerState;
+		readonly HttpMessageHandler _baseHttpMessageHandler;
+		readonly HttpMessageHandler _authHttpMessageHandler;
 
 		/// <summary>
 		/// Constructor
 		/// </summary>
-		public HordeClientWithDynamicCredentials(Uri serverUrl, IHttpClientFactory httpClientFactory, HordeHttpAuthHandlerState authHandler, BundleCache bundleCache, IOptionsSnapshot<HordeOptions> hordeOptions, ILoggerFactory loggerFactory)
-			: base(serverUrl, httpClientFactory, bundleCache, hordeOptions, loggerFactory)
+		public HordeClientWithDynamicCredentials(Uri serverUrl, BundleCache bundleCache, IOptionsSnapshot<HordeOptions> hordeOptions, ILoggerFactory loggerFactory)
+			: base(serverUrl, bundleCache, hordeOptions, loggerFactory)
+		{
+			_baseHttpMessageHandler = CreateDefaultHttpMessageHandler();
+
+			_authHandlerState = new HordeHttpAuthHandlerState(_baseHttpMessageHandler, serverUrl, hordeOptions, loggerFactory.CreateLogger<HordeHttpAuthHandlerState>());
+			_authHttpMessageHandler = new HordeHttpAuthHandler(_baseHttpMessageHandler, _authHandlerState, hordeOptions);
+		}
+
+		/// <inheritdoc/>
+		public override async ValueTask DisposeAsync()
 		{
-			_httpClientFactory = httpClientFactory;
-			_authHandler = authHandler;
+			await _authHandlerState.DisposeAsync();
+			_authHttpMessageHandler.Dispose();
+			_baseHttpMessageHandler.Dispose();
+			await base.DisposeAsync();
 		}
 
 		/// <inheritdoc/>
 		public override async Task<bool> LoginAsync(bool allowLogin, CancellationToken cancellationToken)
 		{
-			return await _authHandler.LoginAsync(allowLogin, cancellationToken);
+			return await _authHandlerState.LoginAsync(allowLogin, cancellationToken);
 		}
 
 		/// <inheritdoc/>
@@ -284,7 +392,7 @@ namespace EpicGames.Horde
 		{
 			try
 			{
-				return _authHandler.IsLoggedIn();
+				return _authHandlerState.IsLoggedIn();
 			}
 			catch
 			{
@@ -294,12 +402,20 @@ namespace EpicGames.Horde
 
 		/// <inheritdoc/>
 		public override Task<string?> GetAccessTokenAsync(bool interactive, CancellationToken cancellationToken)
-			=> _authHandler.GetAccessTokenAsync(interactive, cancellationToken);
+			=> _authHandlerState.GetAccessTokenAsync(interactive, cancellationToken);
 
 		/// <inheritdoc/>
-		protected override HttpClient CreateDefaultHttpClient()
+		protected override HttpClient CreateAuthenticatedHttpClient()
 		{
-			HttpClient httpClient = _httpClientFactory.CreateClient(HordeHttpClient.HttpClientName);
+			HttpClient httpClient = new HttpClient(_authHttpMessageHandler, false);
+			httpClient.BaseAddress = ServerUrl;
+			return httpClient;
+		}
+
+		/// <inheritdoc/>
+		protected override HttpClient CreateUnauthenticatedHttpClient()
+		{
+			HttpClient httpClient = new HttpClient(_baseHttpMessageHandler, false);
 			httpClient.BaseAddress = ServerUrl;
 			return httpClient;
 		}
@@ -310,17 +426,13 @@ namespace EpicGames.Horde
 	/// </summary>
 	class HordeClientFactory : IHordeClientFactory
 	{
-		readonly IHttpClientFactory _httpClientFactory;
 		readonly BundleCache _bundleCache;
-		readonly HordeHttpAuthHandlerState _authHandlerState;
 		readonly IOptionsSnapshot<HordeOptions> _hordeOptions;
 		readonly ILoggerFactory _loggerFactory;
 
-		public HordeClientFactory(IHttpClientFactory httpClientFactory, BundleCache bundleCache, HordeHttpAuthHandlerState authHandlerState, IOptionsSnapshot<HordeOptions> hordeOptions, ILoggerFactory loggerFactory)
+		public HordeClientFactory(BundleCache bundleCache, IOptionsSnapshot<HordeOptions> hordeOptions, ILoggerFactory loggerFactory)
 		{
-			_httpClientFactory = httpClientFactory;
 			_bundleCache = bundleCache;
-			_authHandlerState = authHandlerState;
 			_hordeOptions = hordeOptions;
 			_loggerFactory = loggerFactory;
 		}
@@ -339,14 +451,14 @@ namespace EpicGames.Horde
 		public IHordeClient Create()
 		{
 			Uri serverUrl = GetServerUrl();
-			return new HordeClientWithDynamicCredentials(serverUrl, _httpClientFactory, _authHandlerState, _bundleCache, _hordeOptions, _loggerFactory);
+			return new HordeClientWithDynamicCredentials(serverUrl, _bundleCache, _hordeOptions, _loggerFactory);
 		}
 
 		/// <inheritdoc/>
 		public IHordeClient Create(string? accessToken)
 		{
 			Uri serverUrl = GetServerUrl();
-			return new HordeClientWithStaticCredentials(serverUrl, accessToken, _httpClientFactory, _bundleCache, _hordeOptions, _loggerFactory);
+			return new HordeClientWithStaticCredentials(serverUrl, accessToken, _bundleCache, _hordeOptions, _loggerFactory);
 		}
 	}
 }
diff --git a/Engine/Source/Programs/Shared/EpicGames.Horde/HordeExtensions.cs b/Engine/Source/Programs/Shared/EpicGames.Horde/HordeExtensions.cs
--- a/Engine/Source/Programs/Shared/EpicGames.Horde/HordeExtensions.cs
+++ b/Engine/Source/Programs/Shared/EpicGames.Horde/HordeExtensions.cs
@@ -1,8 +1,6 @@
 // Copyright Epic Games, Inc. All Rights Reserved.
 
 using System;
-using System.Net.Http;
-using System.Threading.Tasks;
 using EpicGames.Core;
 using EpicGames.Horde.Storage;
 using EpicGames.Horde.Storage.Backends;
@@ -11,10 +9,6 @@ using EpicGames.Horde.Storage.Clients;
 using Microsoft.Extensions.DependencyInjection;
 using Microsoft.Extensions.Logging;
 using Microsoft.Extensions.Options;
-using Polly;
-using Polly.Extensions.Http;
-using Polly.Retry;
-using Polly.Timeout;
 
 namespace EpicGames.Horde
 {
@@ -29,37 +23,6 @@ namespace EpicGames.Horde
 		/// <param name="serviceCollection">Collection to register services with</param>
 		public static void AddHorde(this IServiceCollection serviceCollection)
 		{
-			// Sets defaults from the environment before calling the user provided configuration method
-			void ConfigureClientFromEnvironment(IServiceProvider serviceProvider, HttpClient httpClient)
-			{
-				httpClient.Timeout = TimeSpan.FromSeconds(240); // Global timeout
-			}
-
-			// Register the HTTP client for handling login requests
-			void ConfigureClientFromOptions(IServiceProvider serviceProvider, HttpClient httpClient)
-			{
-				IOptions<HordeOptions> options = serviceProvider.GetRequiredService<IOptions<HordeOptions>>();
-				if (options.Value.ServerUrl != null)
-				{
-					httpClient.BaseAddress = options.Value.ServerUrl;
-				}
-			}
-
-			serviceCollection.AddSingleton<HordeHttpAuthHandlerState>();
-			serviceCollection.AddTransient<HordeHttpAuthHandler>();
-			serviceCollection.AddHttpClient(HordeHttpAuthHandlerState.HttpClientName, ConfigureClientFromOptions);
-
-			// Register the HTTP client for processing upload redirects
-			serviceCollection.AddHttpClient(HordeHttpClient.UploadRedirectHttpClientName)
-				.AddPolicyHandler((serviceProvider, request) => CreateDefaultTimeoutRetryPolicy(request, serviceProvider.GetRequiredService<ILogger<HttpStorageBackend>>()))
-				.AddPolicyHandler((serviceProvider, request) => CreateDefaultTransientErrorPolicy(request, serviceProvider.GetRequiredService<ILogger<HttpStorageBackend>>()));
-
-			// Create the HTTP client for handling Horde requests
-			serviceCollection.AddHttpClient<HordeHttpClient>(HordeHttpClient.HttpClientName, ConfigureClientFromEnvironment)
-				.AddHttpMessageHandler<HordeHttpAuthHandler>()
-				.AddPolicyHandler((serviceProvider, request) => CreateDefaultTimeoutRetryPolicy(request, serviceProvider.GetRequiredService<ILogger<HttpStorageBackend>>()))
-				.AddPolicyHandler((serviceProvider, request) => CreateDefaultTransientErrorPolicy(request, serviceProvider.GetRequiredService<ILogger<HttpStorageBackend>>()));
-
 			static StorageBackendCache CreateBackendCache(IServiceProvider serviceProvider)
 			{
 				StorageBackendCacheOptions options = serviceProvider.GetRequiredService<IOptions<HordeOptions>>().Value.BackendCache;
@@ -67,6 +30,8 @@ namespace EpicGames.Horde
 				return new StorageBackendCache(cacheDir, options.MaxSize, serviceProvider.GetRequiredService<ILogger<StorageBackendCache>>());
 			}
 
+			serviceCollection.AddLogging();
+
 			serviceCollection.AddSingleton<BundleCache>(sp => new BundleCache(sp.GetRequiredService<IOptions<HordeOptions>>().Value.BundleCache));
 			serviceCollection.AddSingleton<StorageBackendCache>(CreateBackendCache);
 			serviceCollection.AddSingleton<HttpStorageBackendFactory>();
@@ -85,45 +50,5 @@ namespace EpicGames.Horde
 			serviceCollection.Configure(configureHorde);
 			AddHorde(serviceCollection);
 		}
-
-		/// <summary>
-		/// Create a default timeout retry policy
-		/// </summary>
-		public static IAsyncPolicy<HttpResponseMessage> CreateDefaultTimeoutRetryPolicy(HttpRequestMessage request, ILogger logger)
-		{
-			// Wait 30 seconds for operations to timeout
-			Task OnTimeoutAsync(Context context, TimeSpan timespan, Task timeoutTask)
-			{
-				logger.LogWarning(KnownLogEvents.Systemic_Horde_Http, "{Method} {Url} timed out after {Time}s.", request.Method, request.RequestUri, (int)timespan.TotalSeconds);
-				return Task.CompletedTask;
-			}
-
-			AsyncTimeoutPolicy<HttpResponseMessage> timeoutPolicy = Policy.TimeoutAsync<HttpResponseMessage>(30, OnTimeoutAsync);
-
-			// Retry twice after a timeout
-			void OnRetry(Exception ex, TimeSpan timespan)
-			{
-				logger.LogWarning(KnownLogEvents.Systemic_Horde_Http, ex, "{Method} {Url} retrying after {Time}s.", request.Method, request.RequestUri, timespan.TotalSeconds);
-			}
-
-			TimeSpan[] retryTimes = new[] { TimeSpan.FromSeconds(5.0), TimeSpan.FromSeconds(10.0) };
-			AsyncRetryPolicy retryPolicy = Policy.Handle<TimeoutRejectedException>().WaitAndRetryAsync(retryTimes, OnRetry);
-			return retryPolicy.WrapAsync(timeoutPolicy);
-		}
-
-		/// <summary>
-		/// Create a default timeout retry policy
-		/// </summary>
-		public static IAsyncPolicy<HttpResponseMessage> CreateDefaultTransientErrorPolicy(HttpRequestMessage request, ILogger logger)
-		{
-			Task OnTimeoutAsync(DelegateResult<HttpResponseMessage> outcome, TimeSpan timespan, int retryAttempt, Context context)
-			{
-				logger.LogWarning(KnownLogEvents.Systemic_Horde_Http, "{Method} {Url} failed ({Result}). Delaying for {DelayMs}ms (attempt #{RetryNum}).", request.Method, request.RequestUri, outcome.Result?.StatusCode, timespan.TotalMilliseconds, retryAttempt);
-				return Task.CompletedTask;
-			}
-
-			TimeSpan[] retryTimes = new[] { TimeSpan.FromSeconds(1.0), TimeSpan.FromSeconds(5.0), TimeSpan.FromSeconds(10.0), TimeSpan.FromSeconds(30.0), TimeSpan.FromSeconds(30.0) };
-			return HttpPolicyExtensions.HandleTransientHttpError().WaitAndRetryAsync(retryTimes, OnTimeoutAsync);
-		}
 	}
 }
diff --git a/Engine/Source/Programs/Shared/EpicGames.Horde/HordeHttpAuthHandler.cs b/Engine/Source/Programs/Shared/EpicGames.Horde/HordeHttpAuthHandler.cs
--- a/Engine/Source/Programs/Shared/EpicGames.Horde/HordeHttpAuthHandler.cs
+++ b/Engine/Source/Programs/Shared/EpicGames.Horde/HordeHttpAuthHandler.cs
@@ -30,7 +30,8 @@ namespace EpicGames.Horde
 		/// <summary>
 		/// Constructor
 		/// </summary>
-		public HordeHttpAuthHandler(HordeHttpAuthHandlerState authState, IOptions<HordeOptions> options)
+		public HordeHttpAuthHandler(HttpMessageHandler innerHandler, HordeHttpAuthHandlerState authState, IOptions<HordeOptions> options)
+			: base(innerHandler)
 		{
 			_authState = authState;
 			_options = options;
@@ -83,11 +84,6 @@ namespace EpicGames.Horde
 	/// </summary>
 	public sealed class HordeHttpAuthHandlerState : IAsyncDisposable
 	{
-		/// <summary>
-		/// HTTP client name
-		/// </summary>
-		public const string HttpClientName = "HordeHttpAuthState";
-
 		record AuthState(IClock Clock, AuthMethod Method, OidcTokenInfo? TokenInfo, bool Interactive)
 		{
 			public bool IsAuthorized()
@@ -101,7 +97,8 @@ namespace EpicGames.Horde
 		readonly object _lockObject = new object();
 		readonly CancellationTokenSource _cancellationTokenSource = new CancellationTokenSource();
 		Task<AuthState>? _authStateTask = null;
-		readonly IHttpClientFactory _httpClientFactory;
+		readonly HttpMessageHandler _httpMessageHandler;
+		readonly Uri _serverUrl;
 		readonly IOptions<HordeOptions> _options;
 		readonly ILogger _logger;
 		
@@ -114,14 +111,16 @@ namespace EpicGames.Horde
 		/// Constructor
 		/// </summary>
 		public HordeHttpAuthHandlerState(
-			IHttpClientFactory httpClientFactory,
+			HttpMessageHandler httpMessageHandler,
+			Uri serverUrl,
 			IOptions<HordeOptions> options,
-			ILogger<HordeHttpAuthHandler> logger,
+			ILogger<HordeHttpAuthHandlerState> logger,
 			ITokenStore? tokenStore = null,
 			IOidcTokenManager? oidcTokenManager = null,
 			IClock? clock = null)
 		{
-			_httpClientFactory = httpClientFactory;
+			_httpMessageHandler = httpMessageHandler;
+			_serverUrl = serverUrl;
 			_options = options;
 			_logger = logger;
 			_tokenStore = tokenStore;
@@ -273,18 +272,11 @@ namespace EpicGames.Horde
 
 		async Task<AuthState> GetAuthStateInternalAsync(bool interactive, CancellationToken cancellationToken)
 		{
-			Uri serverUrl;
-
 			GetAuthConfigResponse? authConfig;
-			using (HttpClient httpClient = _httpClientFactory.CreateClient(HttpClientName))
+			using (HttpClient httpClient = new HttpClient(_httpMessageHandler, false))
 			{
-				if (httpClient.BaseAddress == null)
-				{
-					throw new Exception("Horde base address is not configured.");
-				}
-
-				serverUrl = httpClient.BaseAddress;
-				_logger.LogDebug("Retrieving auth configuration for {Server}", serverUrl);
+				httpClient.BaseAddress = _serverUrl;
+				_logger.LogDebug("Retrieving auth configuration for {Server}", _serverUrl);
 
 				JsonSerializerOptions jsonOptions = new JsonSerializerOptions();
 				HordeHttpClient.ConfigureJsonSerializer(jsonOptions);
@@ -326,7 +318,7 @@ namespace EpicGames.Horde
 			}
 			if (result == null && interactive)
 			{
-				_logger.LogInformation("Logging in to {Server}...", serverUrl);
+				_logger.LogInformation("Logging in to {Server}...", _serverUrl);
 				result = await oidcTokenManager.LoginAsync(oidcProvider, cancellationToken);
 			}
 
diff --git a/Engine/Source/Programs/UnrealBuildTool/Executors/UnrealBuildAccelerator/UBAAgentCoordinatorHorde.cs b/Engine/Source/Programs/UnrealBuildTool/Executors/UnrealBuildAccelerator/UBAAgentCoordinatorHorde.cs
--- a/Engine/Source/Programs/UnrealBuildTool/Executors/UnrealBuildAccelerator/UBAAgentCoordinatorHorde.cs
+++ b/Engine/Source/Programs/UnrealBuildTool/Executors/UnrealBuildAccelerator/UBAAgentCoordinatorHorde.cs
@@ -79,6 +79,7 @@ namespace UnrealBuildTool
 			_crypto = owner.Crypto;
 
 			ServiceCollection services = new();
+			services.AddLogging(builder => builder.AddEpicDefault());
 			services.AddHorde(options =>
 			{
 				options.ServerUrl = serverUrl;
@@ -565,16 +566,6 @@ namespace UnrealBuildTool
 			XmlConfig.ApplyTo(HordeConfig);
 			additionalArguments?.ApplyTo(HordeConfig);
 
-			if (String.IsNullOrEmpty(HordeConfig.HordeServer))
-			{
-				HordeConfig.HordeServer = Environment.GetEnvironmentVariable(HordeHttpClient.HordeUrlEnvVarName);
-			}
-
-			if (String.IsNullOrEmpty(HordeConfig.HordeToken))
-			{
-				HordeConfig.HordeToken = Environment.GetEnvironmentVariable(HordeHttpClient.HordeTokenEnvVarName);
-			}
-
 			// Sentry is currently unsupported for non-Windows and non-x64
 			if (!OperatingSystem.IsWindows() || RuntimeInformation.ProcessArchitecture != Architecture.X64)
 			{
