Horde: Fix test compile errors.

#rnx

diff --git a/Engine/Source/Programs/Shared/EpicGames.Horde.Tests/HordeHttpTests.cs b/Engine/Source/Programs/Shared/EpicGames.Horde.Tests/HordeHttpTests.cs
--- a/Engine/Source/Programs/Shared/EpicGames.Horde.Tests/HordeHttpTests.cs
+++ b/Engine/Source/Programs/Shared/EpicGames.Horde.Tests/HordeHttpTests.cs
@@ -58,7 +58,7 @@ public class FakeServerWithAuth : HttpMessageHandler
 	{
 		return new StubHttpClientFactory(() => new HttpClient(this) { BaseAddress = new Uri(ServerUrl) });
 	}
-	
+
 	/// <inheritdoc/>
 	protected override Task<HttpResponseMessage> SendAsync(HttpRequestMessage request, CancellationToken cancellationToken)
 	{
@@ -129,14 +129,14 @@ public class HordeHttpAuthHandlerTests
 			builder.AddSimpleConsole(options => { options.SingleLine = true; });
 		});
 
-		ILogger<HordeHttpAuthHandler> logger = loggerFactory.CreateLogger<HordeHttpAuthHandler>();
+		ILogger<HordeHttpAuthHandlerState> logger = loggerFactory.CreateLogger<HordeHttpAuthHandlerState>();
 		OptionsWrapper<HordeOptions> options = new (hordeOptions ?? new HordeOptions());
 		InMemoryTokenStore inMemoryTokenStore = new ();
-		HordeHttpAuthHandlerState state = new (server.GetHttpClientFactory(), options, logger, inMemoryTokenStore, oidcTokenManager);
-		HordeHttpAuthHandler authHandler = new (state, options);
+		HordeHttpAuthHandlerState state = new HordeHttpAuthHandlerState(server, new Uri("http://fake-server"), options, logger, inMemoryTokenStore, oidcTokenManager);
+		HordeHttpAuthHandler authHandler = new HordeHttpAuthHandler(server, state, options);
 
 		authHandler.InnerHandler = server;
-		HttpClient client = new (authHandler);
+		HttpClient client = new(authHandler);
 		return (client, server);
 	}
 	
diff --git a/Engine/Source/Programs/Shared/EpicGames.Horde/HordeClient.cs b/Engine/Source/Programs/Shared/EpicGames.Horde/HordeClient.cs
--- a/Engine/Source/Programs/Shared/EpicGames.Horde/HordeClient.cs
+++ b/Engine/Source/Programs/Shared/EpicGames.Horde/HordeClient.cs
@@ -22,7 +22,6 @@ using Microsoft.Extensions.Http;
 using Microsoft.Extensions.Logging;
 using Microsoft.Extensions.Options;
 using Polly;
-using Polly.Extensions.Http;
 using Polly.Retry;
 using Polly.Timeout;
 
@@ -78,8 +77,8 @@ namespace EpicGames.Horde
 			{
 #pragma warning disable CA2000 // Call dispose on httpMessageHandler (disposed by child handlers)
 				httpMessageHandler = new SocketsHttpHandler();
-				httpMessageHandler = new PolicyHttpMessageHandler(request => CreateDefaultTimeoutRetryPolicy(request)){ InnerHandler = httpMessageHandler };
-				httpMessageHandler = new PolicyHttpMessageHandler(request => CreateDefaultTransientErrorPolicy(request)){ InnerHandler = httpMessageHandler };
+				httpMessageHandler = new PolicyHttpMessageHandler(request => CreateDefaultTimeoutRetryPolicy(request)) { InnerHandler = httpMessageHandler };
+				httpMessageHandler = new PolicyHttpMessageHandler(request => CreateDefaultTransientErrorPolicy(request)) { InnerHandler = httpMessageHandler };
 #pragma warning restore CA2000
 				return httpMessageHandler;
 			}
