EpicGames.Redis: Add support for hash sets that use a class definition to name keys, and allow creating conditions using named key types.

#rnx

diff --git a/Engine/Source/Programs/Shared/EpicGames.Redis/RedisHash.cs b/Engine/Source/Programs/Shared/EpicGames.Redis/RedisHash.cs
--- a/Engine/Source/Programs/Shared/EpicGames.Redis/RedisHash.cs
+++ b/Engine/Source/Programs/Shared/EpicGames.Redis/RedisHash.cs
@@ -3,11 +3,31 @@
 using System;
 using System.Collections.Generic;
 using System.Linq;
+using System.Linq.Expressions;
 using System.Threading.Tasks;
 using StackExchange.Redis;
 
 namespace EpicGames.Redis
 {
+	/// <summary>
+	/// Represents a redis hash key, with members corresponding to the property names of a type
+	/// </summary>
+	/// <typeparam name="T">Type of the hash fields</typeparam>
+	public readonly struct RedisHashKey<T>
+	{
+		/// <summary>
+		/// The key for the list
+		/// </summary>
+		public readonly RedisKey Inner { get; }
+
+		/// <summary>
+		/// Constructor
+		/// </summary>
+		/// <param name="inner">Redis key this type is using</param>
+		public RedisHashKey(RedisKey inner)
+			=> Inner = inner;
+	}
+
 	/// <summary>
 	/// Represents a typed Redis list with a given key
 	/// </summary>
@@ -86,8 +106,82 @@ namespace EpicGames.Redis
 	/// </summary>
 	public static class RedisHashExtensions
 	{
+		#region Conditions
+
+		/// <inheritdoc cref="Condition.HashEqual(RedisKey, RedisValue, RedisValue)"/>
+		public static Condition Equal<TRecord, TValue>(this RedisHashKey<TRecord> key, Expression<Func<TRecord, TValue>> selector, TValue value)
+		{
+			MemberExpression memberExpression = (selector.Body as MemberExpression) ?? throw new InvalidOperationException("Expression must be a property accessor");
+			return Condition.HashEqual(key.Inner, memberExpression.Member.Name, RedisSerializer.Serialize(value));
+		}
+
+		/// <inheritdoc cref="Condition.HashEqual(RedisKey, RedisValue, RedisValue)"/>
+		public static Condition Equal<TName, TValue>(this RedisHashKey<TName, TValue> key, TName name, TValue value)
+			=> Condition.HashEqual(key.Inner, RedisSerializer.Serialize(name), RedisSerializer.Serialize(value));
+
+		/// <inheritdoc cref="Condition.HashExists(RedisKey, RedisValue)"/>
+		public static Condition Exists<TRecord, TValue>(this RedisHashKey<TRecord> key, Expression<Func<TRecord, TValue>> selector)
+		{
+			MemberExpression memberExpression = (selector.Body as MemberExpression) ?? throw new InvalidOperationException("Expression must be a property accessor");
+			return Condition.HashExists(key.Inner, memberExpression.Member.Name);
+		}
+
+		/// <inheritdoc cref="Condition.HashExists(RedisKey, RedisValue)"/>
+		public static Condition Exists<TName, TValue>(this RedisHashKey<TName, TValue> key, TName name)
+			=> Condition.HashExists(key.Inner, RedisSerializer.Serialize(name));
+
+		/// <inheritdoc cref="Condition.HashLengthEqual(RedisKey, Int64)"/>
+		public static Condition LengthEqual<TName, TValue>(this RedisHashKey<TName, TValue> key, long length)
+			=> Condition.HashLengthEqual(key.Inner, length);
+
+		/// <inheritdoc cref="Condition.HashLengthGreaterThan(RedisKey, Int64)"/>
+		public static Condition LengthGreaterThan<TName, TValue>(this RedisHashKey<TName, TValue> key, long length)
+			=> Condition.HashLengthGreaterThan(key.Inner, length);
+
+		/// <inheritdoc cref="Condition.HashLengthLessThan(RedisKey, Int64)"/>
+		public static Condition LengthLessThan<TName, TValue>(this RedisHashKey<TName, TValue> key, long length)
+			=> Condition.HashLengthLessThan(key.Inner, length);
+
+		/// <inheritdoc cref="Condition.HashNotExists(RedisKey, RedisValue)"/>
+		public static Condition NotExists<TRecord, TValue>(this RedisHashKey<TRecord> key, Expression<Func<TRecord, TValue>> selector)
+		{
+			MemberExpression memberExpression = (selector.Body as MemberExpression) ?? throw new InvalidOperationException("Expression must be a property accessor");
+			return Condition.HashNotExists(key.Inner, memberExpression.Member.Name);
+		}
+
+		/// <inheritdoc cref="Condition.HashNotExists(RedisKey, RedisValue)"/>
+		public static Condition NotExists<TName, TValue>(this RedisHashKey<TName, TValue> key, TName name)
+			=> Condition.HashNotExists(key.Inner, RedisSerializer.Serialize(name));
+
+		/// <inheritdoc cref="Condition.HashEqual(RedisKey, RedisValue, RedisValue)"/>
+		public static Condition NotEqual<TRecord, TValue>(this RedisHashKey<TRecord> key, Expression<Func<TRecord, TValue>> selector, TValue value)
+		{
+			MemberExpression memberExpression = (selector.Body as MemberExpression) ?? throw new InvalidOperationException("Expression must be a property accessor");
+			return Condition.HashNotEqual(key.Inner, memberExpression.Member.Name, RedisSerializer.Serialize(value));
+		}
+
+		/// <inheritdoc cref="Condition.HashNotEqual(RedisKey, RedisValue, RedisValue)"/>
+		public static Condition NotEqual<TName, TValue>(this RedisHashKey<TName, TValue> key, TName name, TValue value)
+			=> Condition.HashNotEqual(key.Inner, RedisSerializer.Serialize(name), RedisSerializer.Serialize(value));
+
+		#endregion
+
 		#region HashDecrementAsync
 
+		/// <inheritdoc cref="IDatabaseAsync.HashDecrementAsync(RedisKey, RedisValue, Int64, CommandFlags)"/>
+		public static Task<long> HashDecrementAsync<TRecord>(this IDatabaseAsync target, RedisHashKey<TRecord> key, Expression<Func<TRecord, long>> selector, long value = 1L, CommandFlags flags = CommandFlags.None)
+		{
+			MemberExpression memberExpression = (selector.Body as MemberExpression) ?? throw new InvalidOperationException("Expression must be a property accessor");
+			return target.HashDecrementAsync(key.Inner, memberExpression.Member.Name, value, flags);
+		}
+
+		/// <inheritdoc cref="IDatabaseAsync.HashDecrementAsync(RedisKey, RedisValue, Double, CommandFlags)"/>
+		public static Task<double> HashDecrementAsync<TRecord>(this IDatabaseAsync target, RedisHashKey<TRecord> key, Expression<Func<TRecord, double>> selector, double value = 1.0, CommandFlags flags = CommandFlags.None)
+		{
+			MemberExpression memberExpression = (selector.Body as MemberExpression) ?? throw new InvalidOperationException("Expression must be a property accessor");
+			return target.HashDecrementAsync(key.Inner, memberExpression.Member.Name, value, flags);
+		}
+
 		/// <inheritdoc cref="IDatabaseAsync.HashDecrementAsync(RedisKey, RedisValue, Int64, CommandFlags)"/>
 		public static Task<long> HashDecrementAsync<TName>(this IDatabaseAsync target, RedisHashKey<TName, long> key, TName name, long value = 1L, CommandFlags flags = CommandFlags.None)
 		{
@@ -104,6 +198,13 @@ namespace EpicGames.Redis
 
 		#region HashDeleteAsync
 
+		/// <inheritdoc cref="IDatabaseAsync.HashDeleteAsync(RedisKey, RedisValue, CommandFlags)"/>
+		public static Task<bool> HashDeleteAsync<TRecord>(this IDatabaseAsync target, RedisHashKey<TRecord> key, Expression<Func<TRecord, double>> selector, CommandFlags flags = CommandFlags.None)
+		{
+			MemberExpression memberExpression = (selector.Body as MemberExpression) ?? throw new InvalidOperationException("Expression must be a property accessor");
+			return target.HashDeleteAsync(key.Inner, memberExpression.Member.Name, flags);
+		}
+
 		/// <inheritdoc cref="IDatabaseAsync.HashDeleteAsync(RedisKey, RedisValue, CommandFlags)"/>
 		public static Task<bool> HashDeleteAsync<TName, TValue>(this IDatabaseAsync target, RedisHashKey<TName, TValue> key, TName name, CommandFlags flags = CommandFlags.None)
 		{
@@ -130,6 +231,13 @@ namespace EpicGames.Redis
 
 		#region HashGetAsync
 
+		/// <inheritdoc cref="IDatabaseAsync.HashGetAsync(RedisKey, RedisValue, CommandFlags)"/>
+		public static Task<TValue> HashGetAsync<TRecord, TValue>(this IDatabaseAsync target, RedisHashKey<TRecord> key, Expression<Func<TRecord, TValue>> selector, CommandFlags flags = CommandFlags.None)
+		{
+			MemberExpression memberExpression = (selector.Body as MemberExpression) ?? throw new InvalidOperationException("Expression must be a property accessor");
+			return target.HashGetAsync(key.Inner, memberExpression.Member.Name, flags).DeserializeAsync<TValue>();
+		}
+
 		/// <inheritdoc cref="IDatabaseAsync.HashGetAsync(RedisKey, RedisValue, CommandFlags)"/>
 		public static Task<TValue> HashGetAsync<TName, TValue>(this IDatabaseAsync target, RedisHashKey<TName, TValue> key, TName name, CommandFlags flags = CommandFlags.None)
 		{
@@ -157,6 +265,20 @@ namespace EpicGames.Redis
 
 		#region HashIncrementAsync
 
+		/// <inheritdoc cref="IDatabaseAsync.HashIncrementAsync(RedisKey, RedisValue, Int64, CommandFlags)"/>
+		public static Task<long> HashInccrementAsync<TRecord>(this IDatabaseAsync target, RedisHashKey<TRecord> key, Expression<Func<TRecord, long>> selector, long value = 1L, CommandFlags flags = CommandFlags.None)
+		{
+			MemberExpression memberExpression = (selector.Body as MemberExpression) ?? throw new InvalidOperationException("Expression must be a property accessor");
+			return target.HashIncrementAsync(key.Inner, memberExpression.Member.Name, value, flags);
+		}
+
+		/// <inheritdoc cref="IDatabaseAsync.HashIncrementAsync(RedisKey, RedisValue, Double, CommandFlags)"/>
+		public static Task<double> HashIncrementAsync<TRecord>(this IDatabaseAsync target, RedisHashKey<TRecord> key, Expression<Func<TRecord, double>> selector, double value = 1.0, CommandFlags flags = CommandFlags.None)
+		{
+			MemberExpression memberExpression = (selector.Body as MemberExpression) ?? throw new InvalidOperationException("Expression must be a property accessor");
+			return target.HashIncrementAsync(key.Inner, memberExpression.Member.Name, value, flags);
+		}
+
 		/// <inheritdoc cref="IDatabaseAsync.HashIncrementAsync(RedisKey, RedisValue, Int64, CommandFlags)"/>
 		public static Task<long> HashIncrementAsync<TName>(this IDatabaseAsync target, RedisHashKey<TName, long> key, TName name, long value = 1L, CommandFlags flags = CommandFlags.None)
 		{
@@ -193,6 +315,13 @@ namespace EpicGames.Redis
 
 		#region HashSetAsync
 
+		/// <inheritdoc cref="IDatabaseAsync.HashSetAsync(RedisKey, RedisValue, RedisValue, When, CommandFlags)"/>
+		public static Task HashSetAsync<TRecord, TValue>(this IDatabaseAsync target, RedisHashKey<TRecord> key, Expression<Func<TRecord, TValue>> selector, TValue value, When when = When.Always, CommandFlags flags = CommandFlags.None)
+		{
+			MemberExpression memberExpression = (selector.Body as MemberExpression) ?? throw new InvalidOperationException("Expression must be a property accessor");
+			return target.HashSetAsync(key.Inner, memberExpression.Member.Name, RedisSerializer.Serialize(value), when, flags);
+		}
+
 		/// <inheritdoc cref="IDatabaseAsync.HashSetAsync(RedisKey, RedisValue, RedisValue, When, CommandFlags)"/>
 		public static Task HashSetAsync<TName, TValue>(this IDatabaseAsync target, RedisHashKey<TName, TValue> key, TName name, TValue value, When when = When.Always, CommandFlags flags = CommandFlags.None)
 		{
diff --git a/Engine/Source/Programs/Shared/EpicGames.Redis/RedisList.cs b/Engine/Source/Programs/Shared/EpicGames.Redis/RedisList.cs
--- a/Engine/Source/Programs/Shared/EpicGames.Redis/RedisList.cs
+++ b/Engine/Source/Programs/Shared/EpicGames.Redis/RedisList.cs
@@ -45,6 +45,38 @@ namespace EpicGames.Redis
 	/// </summary>
 	public static class RedisListExtensions
 	{
+		#region Conditions
+
+		/// <inheritdoc cref="Condition.ListIndexEqual(RedisKey, Int64, RedisValue)"/>
+		public static Condition IndexEqual<TElement>(this RedisListKey<TElement> key, long index, TElement value)
+			=> Condition.ListIndexEqual(key.Inner, index, RedisSerializer.Serialize(value));
+
+		/// <inheritdoc cref="Condition.ListIndexExists(RedisKey, Int64)"/>
+		public static Condition IndexExists<TElement>(this RedisListKey<TElement> key, long index)
+			=> Condition.ListIndexExists(key.Inner, index);
+
+		/// <inheritdoc cref="Condition.ListIndexNotEqual(RedisKey, Int64, RedisValue)"/>
+		public static Condition IndexNotEqual<TElement>(this RedisListKey<TElement> key, long index, TElement value)
+			=> Condition.ListIndexNotEqual(key.Inner, index, RedisSerializer.Serialize(value));
+
+		/// <inheritdoc cref="Condition.ListIndexNotExists(RedisKey, Int64)"/>
+		public static Condition IndexNotExists<TElement>(this RedisListKey<TElement> key, long index)
+			=> Condition.ListIndexNotExists(key.Inner, index);
+
+		/// <inheritdoc cref="Condition.ListLengthEqual(RedisKey, Int64)"/>
+		public static Condition LengthEqual<TElement>(this RedisListKey<TElement> key, long length)
+			=> Condition.ListLengthEqual(key.Inner, length);
+
+		/// <inheritdoc cref="Condition.ListLengthGreaterThan(RedisKey, Int64)"/>
+		public static Condition LengthGreaterThan<TElement>(this RedisListKey<TElement> key, long length)
+			=> Condition.ListLengthGreaterThan(key.Inner, length);
+
+		/// <inheritdoc cref="Condition.ListLengthLessThan(RedisKey, Int64)"/>
+		public static Condition LengthLessThan<TElement>(this RedisListKey<TElement> key, long length)
+			=> Condition.ListLengthLessThan(key.Inner, length);
+
+		#endregion
+
 		#region ListGetByIndexAsync
 
 		/// <inheritdoc cref="IDatabaseAsync.ListGetByIndexAsync(RedisKey, Int64, CommandFlags)"/>
diff --git a/Engine/Source/Programs/Shared/EpicGames.Redis/RedisSet.cs b/Engine/Source/Programs/Shared/EpicGames.Redis/RedisSet.cs
--- a/Engine/Source/Programs/Shared/EpicGames.Redis/RedisSet.cs
+++ b/Engine/Source/Programs/Shared/EpicGames.Redis/RedisSet.cs
@@ -44,6 +44,30 @@ namespace EpicGames.Redis
 	/// </summary>
 	public static class RedisSetExtensions
 	{
+		#region Conditions
+
+		/// <inheritdoc cref="Condition.SetContains(RedisKey, RedisValue)"/>
+		public static Condition Contains<TElement>(this RedisSetKey<TElement> key, TElement value)
+			=> Condition.SetContains(key.Inner, RedisSerializer.Serialize(value));
+
+		/// <inheritdoc cref="Condition.SetLengthEqual(RedisKey, Int64)"/>
+		public static Condition LengthEqual<TElement>(this RedisSetKey<TElement> key, long length)
+			=> Condition.SetLengthEqual(key.Inner, length);
+
+		/// <inheritdoc cref="Condition.SetLengthGreaterThan(RedisKey, Int64)"/>
+		public static Condition LengthGreaterThan<TElement>(this RedisSetKey<TElement> key, long length)
+			=> Condition.SetLengthGreaterThan(key.Inner, length);
+
+		/// <inheritdoc cref="Condition.SetLengthLessThan(RedisKey, Int64)"/>
+		public static Condition LengthLessThan<TElement>(this RedisSetKey<TElement> key, long length)
+			=> Condition.SetLengthLessThan(key.Inner, length);
+
+		/// <inheritdoc cref="Condition.SetContains(RedisKey, RedisValue)"/>
+		public static Condition NotContains<TElement>(this RedisSetKey<TElement> key, TElement value)
+			=> Condition.SetNotContains(key.Inner, RedisSerializer.Serialize(value));
+
+		#endregion
+
 		#region SetAddAsync
 
 		/// <inheritdoc cref="IDatabaseAsync.SetAddAsync(RedisKey, RedisValue, CommandFlags)"/>
diff --git a/Engine/Source/Programs/Shared/EpicGames.Redis/RedisSortedSet.cs b/Engine/Source/Programs/Shared/EpicGames.Redis/RedisSortedSet.cs
--- a/Engine/Source/Programs/Shared/EpicGames.Redis/RedisSortedSet.cs
+++ b/Engine/Source/Programs/Shared/EpicGames.Redis/RedisSortedSet.cs
@@ -102,6 +102,46 @@ namespace EpicGames.Redis
 	/// </summary>
 	public static class RedisSortedSetExtensions
 	{
+		#region Conditions
+
+		/// <inheritdoc cref="Condition.SortedSetContains(RedisKey, RedisValue)"/>
+		public static Condition Contains<TElement>(this RedisSortedSetKey<TElement> key, TElement value)
+			=> Condition.SortedSetContains(key.Inner, RedisSerializer.Serialize(value));
+
+		/// <inheritdoc cref="Condition.SortedSetEqual(RedisKey, RedisValue, RedisValue)"/>
+		public static Condition Equal<TElement>(this RedisSortedSetKey<TElement> key, TElement value, RedisValue score)
+			=> Condition.SortedSetEqual(key.Inner, RedisSerializer.Serialize(value), score);
+
+		/// <inheritdoc cref="Condition.SortedSetLengthEqual(RedisKey, Int64)"/>
+		public static Condition LengthEqual<TElement>(this RedisSortedSetKey<TElement> key, long length)
+			=> Condition.SortedSetLengthEqual(key.Inner, length);
+
+		/// <inheritdoc cref="Condition.SortedSetLengthGreaterThan(RedisKey, Int64)"/>
+		public static Condition LengthGreaterThan<TElement>(this RedisSortedSetKey<TElement> key, long length)
+			=> Condition.SortedSetLengthGreaterThan(key.Inner, length);
+
+		/// <inheritdoc cref="Condition.SortedSetLengthLessThan(RedisKey, Int64)"/>
+		public static Condition LengthLessThan<TElement>(this RedisSortedSetKey<TElement> key, long length)
+			=> Condition.SortedSetLengthLessThan(key.Inner, length);
+
+		/// <inheritdoc cref="Condition.SortedSetNotContains(RedisKey, RedisValue)"/>
+		public static Condition NotContains<TElement>(this RedisSortedSetKey<TElement> key, TElement value)
+			=> Condition.SortedSetNotContains(key.Inner, RedisSerializer.Serialize(value));
+
+		/// <inheritdoc cref="Condition.SortedSetNotEqual(RedisKey, RedisValue, RedisValue)"/>
+		public static Condition NotEqual<TElement>(this RedisSortedSetKey<TElement> key, TElement value, RedisValue score)
+			=> Condition.SortedSetNotEqual(key.Inner, RedisSerializer.Serialize(value), score);
+
+		/// <inheritdoc cref="Condition.SortedSetScoreExists(RedisKey, RedisValue)"/>
+		public static Condition ScoreExists<TElement>(this RedisSortedSetKey<TElement> key, RedisValue score)
+			=> Condition.SortedSetScoreExists(key.Inner, score);
+
+		/// <inheritdoc cref="Condition.SortedSetScoreExists(RedisKey, RedisValue, RedisValue)"/>
+		public static Condition ScoreExists<TElement>(this RedisSortedSetKey<TElement> key, RedisValue score, RedisValue count)
+			=> Condition.SortedSetScoreExists(key.Inner, score, count);
+
+		#endregion
+
 		#region SortedSetAddAsync
 
 		/// <inheritdoc cref="IDatabaseAsync.SortedSetAddAsync(RedisKey, RedisValue, Double, CommandFlags)"/>
@@ -200,7 +240,7 @@ namespace EpicGames.Redis
 
 		#endregion
 
-		#region 
+		#region SortedSetRemoveAsync
 
 		/// <inheritdoc cref="IDatabaseAsync.SortedSetRemoveAsync(RedisKey, RedisValue, CommandFlags)"/>
 		public static Task<bool> SortedSetRemoveAsync<TElement>(this IDatabaseAsync target, RedisSortedSetKey<TElement> key, TElement value, CommandFlags flags = CommandFlags.None)
diff --git a/Engine/Source/Programs/Shared/EpicGames.Redis/RedisString.cs b/Engine/Source/Programs/Shared/EpicGames.Redis/RedisString.cs
--- a/Engine/Source/Programs/Shared/EpicGames.Redis/RedisString.cs
+++ b/Engine/Source/Programs/Shared/EpicGames.Redis/RedisString.cs
@@ -42,10 +42,34 @@ namespace EpicGames.Redis
 	}
 
 	/// <summary>
-	/// Extension methods for sets
+	/// Extension methods for strings
 	/// </summary>
 	public static class RedisStringExtensions
 	{
+		#region Conditions
+
+		/// <inheritdoc cref="Condition.StringEqual(RedisKey, RedisValue)"/>
+		public static Condition Equal<TElement>(this RedisStringKey<TElement> key, TElement value)
+			=> Condition.StringEqual(key.Inner, RedisSerializer.Serialize(value));
+
+		/// <inheritdoc cref="Condition.StringLengthEqual(RedisKey, Int64)"/>
+		public static Condition LengthEqual<TElement>(this RedisStringKey<TElement> key, long length)
+			=> Condition.StringLengthEqual(key.Inner, length);
+
+		/// <inheritdoc cref="Condition.StringLengthGreaterThan(RedisKey, Int64)"/>
+		public static Condition LengthGreaterThan<TElement>(this RedisStringKey<TElement> key, long length)
+			=> Condition.StringLengthGreaterThan(key.Inner, length);
+
+		/// <inheritdoc cref="Condition.StringLengthLessThan(RedisKey, Int64)"/>
+		public static Condition LengthLessThan<TElement>(this RedisStringKey<TElement> key, long length)
+			=> Condition.StringLengthLessThan(key.Inner, length);
+
+		/// <inheritdoc cref="Condition.StringNotEqual(RedisKey, RedisValue)"/>
+		public static Condition NotEqual<TElement>(this RedisStringKey<TElement> key, TElement value)
+			=> Condition.StringNotEqual(key.Inner, RedisSerializer.Serialize(value));
+
+		#endregion
+
 		#region StringDecrementAsync
 
 		/// <inheritdoc cref="IDatabaseAsync.StringDecrementAsync(RedisKey, Int64, CommandFlags)"/>
