EpicGames.Redis: Add helper methods to allow converting from HashEntry arrays to dictionaries.

#rnx

diff --git a/Engine/Source/Programs/Shared/EpicGames.Redis/RedisHash.cs b/Engine/Source/Programs/Shared/EpicGames.Redis/RedisHash.cs
--- a/Engine/Source/Programs/Shared/EpicGames.Redis/RedisHash.cs
+++ b/Engine/Source/Programs/Shared/EpicGames.Redis/RedisHash.cs
@@ -106,6 +106,18 @@ namespace EpicGames.Redis
 	/// </summary>
 	public static class RedisHashExtensions
 	{
+		/// <summary>
+		/// Helper method to convert an array of hash entries into a dictionary
+		/// </summary>
+		public static async Task<Dictionary<TName, TValue>> ToDictionaryAsync<TName, TValue>(this Task<HashEntry<TName, TValue>[]> entries) where TName : notnull
+			=> (await entries).ToDictionary(x => x.Name, x => x.Value);
+
+		/// <summary>
+		/// Helper method to convert an array of hash entries into a dictionary
+		/// </summary>
+		public static async Task<Dictionary<TName, TValue>> ToDictionaryAsync<TName, TValue>(this Task<HashEntry<TName, TValue>[]> entries, IEqualityComparer<TName>? comparer) where TName : notnull
+			=> (await entries).ToDictionary(x => x.Name, x => x.Value, comparer);
+
 		#region Conditions
 
 		/// <inheritdoc cref="Condition.HashEqual(RedisKey, RedisValue, RedisValue)"/>
diff --git a/Engine/Source/Programs/Shared/EpicGames.Redis/RedisSortedSet.cs b/Engine/Source/Programs/Shared/EpicGames.Redis/RedisSortedSet.cs
--- a/Engine/Source/Programs/Shared/EpicGames.Redis/RedisSortedSet.cs
+++ b/Engine/Source/Programs/Shared/EpicGames.Redis/RedisSortedSet.cs
@@ -308,5 +308,21 @@ namespace EpicGames.Redis
 		}
 
 		#endregion
+
+		#region SortedSetUpdateAsync
+
+		/// <inheritdoc cref="IDatabaseAsync.SortedSetUpdateAsync(RedisKey, RedisValue, double, SortedSetWhen, CommandFlags)"/>
+		public static Task<bool> SortedSetUpdateAsync<TElement>(this IDatabaseAsync target, RedisSortedSetKey<TElement> key, TElement member, double score, SortedSetWhen when = SortedSetWhen.Always, CommandFlags flags = CommandFlags.None)
+		{
+			return target.SortedSetUpdateAsync(key.Inner, RedisSerializer.Serialize(member), score, when, flags);
+		}
+
+		/// <inheritdoc cref="IDatabaseAsync.SortedSetUpdateAsync(RedisKey, RedisValue, double, SortedSetWhen, CommandFlags)"/>
+		public static Task<long> SortedSetUpdateAsync<TElement>(this IDatabaseAsync target, RedisSortedSetKey<TElement> key, SortedSetEntry<TElement>[] values, SortedSetWhen when = SortedSetWhen.Always, CommandFlags flags = CommandFlags.None)
+		{
+			return target.SortedSetUpdateAsync(key.Inner, values.ConvertAll(x => new SortedSetEntry(RedisSerializer.Serialize(x.Element), x.Score)), when, flags);
+		}
+
+		#endregion
 	}
 }
