Upgrade bundled dotnet to 8.0.300

* Rename dotnet directories for windows and linux to include architecture

diff --git a/Engine/Build/BatchFiles/GetDotnetPath.bat b/Engine/Build/BatchFiles/GetDotnetPath.bat
--- a/Engine/Build/BatchFiles/GetDotnetPath.bat
+++ b/Engine/Build/BatchFiles/GetDotnetPath.bat
@@ -9,16 +9,16 @@ rem if UE_USE_SYSTEM_DOTNET we assume a installed dotnet is present
 if "%UE_USE_SYSTEM_DOTNET%" == "1" goto verify_dotnet
 
 rem if UE_DOTNET_VERSION is already set we assume this script has already run
-if "%UE_DOTNET_VERSION%" == "6.0.302" goto Succeeded
+if "%UE_DOTNET_VERSION%" == "8.0.300" goto Succeeded
 
 rem figure out which architecture to use
-set UE_DOTNET_ARCH=windows
+set UE_DOTNET_ARCH=win-x64
 if "%PROCESSOR_ARCHITECTURE%" == "ARM64" (
 	set UE_DOTNET_ARCH=win-arm64
 )
 
 rem add the dotnet sdk in the repo as the current dotnet sdk
-set UE_DOTNET_VERSION=6.0.302
+set UE_DOTNET_VERSION=8.0.300
 set UE_DOTNET_DIR=%~dp0..\..\Binaries\ThirdParty\DotNet\%UE_DOTNET_VERSION%\%UE_DOTNET_ARCH%
 set PATH=%UE_DOTNET_DIR%;%PATH%
 set DOTNET_ROOT=%UE_DOTNET_DIR%
diff --git a/Engine/Build/BatchFiles/Linux/SetupDotnet.sh b/Engine/Build/BatchFiles/Linux/SetupDotnet.sh
--- a/Engine/Build/BatchFiles/Linux/SetupDotnet.sh
+++ b/Engine/Build/BatchFiles/Linux/SetupDotnet.sh
@@ -35,7 +35,11 @@ if [ $IS_DOTNET_INSTALLED -eq 0 ]; then
 	# If this flag isn't set to 0, dotnet crashes during GenerateProjectFiles.sh on Ubuntu 20.04 
 	export DOTNET_gcServer=0
 
-	export UE_DOTNET_DIR="$CUR_DIR/../../../Binaries/ThirdParty/DotNet/6.0.302/linux"
+	# Select the preferred architecture for the current system
+	ARCH=x64
+	[ $(uname -m) == "arm64" ] && ARCH=arm64 
+
+	export UE_DOTNET_DIR="$CUR_DIR/../../../Binaries/ThirdParty/DotNet/8.0.300/linux-$ARCH"
 	chmod u+x "$UE_DOTNET_DIR/dotnet"
 	export PATH="$UE_DOTNET_DIR:$PATH"
 	export DOTNET_ROOT="$UE_DOTNET_DIR"
diff --git a/Engine/Build/BatchFiles/Mac/SetupDotnet.sh b/Engine/Build/BatchFiles/Mac/SetupDotnet.sh
--- a/Engine/Build/BatchFiles/Mac/SetupDotnet.sh
+++ b/Engine/Build/BatchFiles/Mac/SetupDotnet.sh
@@ -39,7 +39,7 @@ if [ $IS_DOTNET_INSTALLED -eq 0 ]; then
 	ARCH=x64
 	[ $(uname -m) == "arm64" ] && ARCH=arm64 
 	
-	export UE_DOTNET_DIR=$CUR_DIR/../../../Binaries/ThirdParty/DotNet/6.0.302/mac-$ARCH
+	export UE_DOTNET_DIR=$CUR_DIR/../../../Binaries/ThirdParty/DotNet/8.0.300/mac-$ARCH
 	chmod u+x "$UE_DOTNET_DIR/dotnet"
 	echo $UE_DOTNET_DIR
 	export PATH=$UE_DOTNET_DIR:$PATH
diff --git a/Engine/Build/InstalledEngineFilters.xml b/Engine/Build/InstalledEngineFilters.xml
--- a/Engine/Build/InstalledEngineFilters.xml
+++ b/Engine/Build/InstalledEngineFilters.xml
@@ -256,7 +256,7 @@
 	<Property Name="ExtraSignFilesWin64">
 		<!-- CsvTools binaries which are manually updated in P4 -->
 		Engine/Binaries/DotNET/CsvTools/...
-		Engine/Binaries/ThirdParty/DotNet/6.0.302/windows/...
+		Engine/Binaries/ThirdParty/DotNet/8.0.300/win-x64/...
 		Engine/Binaries/ThirdParty/Python3/...
 		Engine/Binaries/ThirdParty/PresentMon/Win64/PresentMon-1.8.0-x64.exe
 		Engine/Binaries/ThirdParty/svn/Win64/...
@@ -300,7 +300,7 @@
 		Templates/TP_InCamVFXBP/...
 
 		<!-- DotNet -->
-		Engine/Binaries/ThirdParty/DotNet/6.0.302/windows/...
+		Engine/Binaries/ThirdParty/DotNet/8.0.300/win-x64/...
 		Engine/Binaries/ThirdParty/VisualStudio/...
 
 		<!-- CsvTools binaries which are manually updated in P4 -->
@@ -374,8 +374,8 @@
 		Engine/Build/Xcode/...
 
 		<!-- DotNet -->
-		Engine/Binaries/ThirdParty/DotNet/6.0.302/mac-arm64/...
-		Engine/Binaries/ThirdParty/DotNet/6.0.302/mac-x64/...
+		Engine/Binaries/ThirdParty/DotNet/8.0.300/mac-arm64/...
+		Engine/Binaries/ThirdParty/DotNet/8.0.300/mac-x64/...
 		Engine/Binaries/ThirdParty/VisualStudio/...
 	</Property>
 
@@ -439,7 +439,8 @@
 
 
 		<!-- DotNet -->
-		Engine/Binaries/ThirdParty/DotNet/6.0.302/linux/...
+		Engine/Binaries/ThirdParty/DotNet/8.0.300/linux-x64/...
+		Engine/Binaries/ThirdParty/DotNet/8.0.300/linux-arm64/...
 		Engine/Binaries/ThirdParty/VisualStudio/...
 
 		<!-- DotNet bundled Requirements -->
@@ -595,8 +596,8 @@
 
 	<!-- IOS -->
 	<Property Name="ExtraSignFilesIOS">
-		Engine/Binaries/ThirdParty/DotNet/6.0.302/mac-arm64/...
-		Engine/Binaries/ThirdParty/DotNet/6.0.302/mac-x64/...
+		Engine/Binaries/ThirdParty/DotNet/8.0.300/mac-arm64/...
+		Engine/Binaries/ThirdParty/DotNet/8.0.300/mac-x64/...
 		Engine/Extras/ThirdPartyNotUE/libimobiledevice/x64/...
 	</Property>
 
diff --git a/Engine/Source/Programs/Shared/EpicGames.Build/Unreal.cs b/Engine/Source/Programs/Shared/EpicGames.Build/Unreal.cs
--- a/Engine/Source/Programs/Shared/EpicGames.Build/Unreal.cs
+++ b/Engine/Source/Programs/Shared/EpicGames.Build/Unreal.cs
@@ -87,36 +87,29 @@ namespace UnrealBuildBase
 			return UnrealBuildToolDllPath;
 		}
 
-		static private string DotnetVersionDirectory = "6.0.302";
+		static private string DotnetVersionDirectory = "8.0.300";
 
 		static private string FindRelativeDotnetDirectory(RuntimePlatform.Type HostPlatform)
 		{
-			string HostDotNetDirectoryName;
+			string platform;
+			string architecture;
+
 			switch (HostPlatform)
 			{
-				case RuntimePlatform.Type.Windows:
-					{
-						HostDotNetDirectoryName = "windows";
-						if (RuntimeInformation.ProcessArchitecture == Architecture.Arm64)
-						{
-							HostDotNetDirectoryName = "win-arm64";
-						}
-						break;
-					}
-				case RuntimePlatform.Type.Mac:
-					{
-						HostDotNetDirectoryName = "mac-x64";
-						if (RuntimeInformation.ProcessArchitecture == Architecture.Arm64)
-						{
-							HostDotNetDirectoryName = "mac-arm64";
-						}
-						break;
-					}
-				case RuntimePlatform.Type.Linux: HostDotNetDirectoryName = "linux"; break;
-				default: throw new Exception("Unknown host platform");
+				case RuntimePlatform.Type.Linux: platform = "linux"; break;
+				case RuntimePlatform.Type.Mac: platform = "mac"; break;
+				case RuntimePlatform.Type.Windows: platform = "win"; break;
+				default: throw new Exception($"Unsupported host platform {HostPlatform}");
+			}
+
+			switch (RuntimeInformation.ProcessArchitecture)
+			{
+				case Architecture.Arm64: architecture = "arm64"; break;
+				case Architecture.X64: architecture = "x64"; break;
+				default: throw new Exception($"Unsupported host architecture {RuntimeInformation.ProcessArchitecture}");
 			}
 
-			return Path.Combine("Binaries", "ThirdParty", "DotNet", DotnetVersionDirectory, HostDotNetDirectoryName);
+			return Path.Combine("Binaries", "ThirdParty", "DotNet", DotnetVersionDirectory, $"{platform}-{architecture}");
 		}
 
 		static private string FindRelativeDotnetDirectory() => FindRelativeDotnetDirectory(RuntimePlatform.Current);
