Change the code generation of UHT for BP implementable events to check for any script implementations of the event.

If there are none, then just call directly into the native C++ implemenation of the function, saving us the cost of copying the function params to the script VM.

#pf 6658815ace4032dd41b0218c

diff --git a/Engine/Source/Programs/Shared/EpicGames.UHT/Exporters/CodeGen/UhtHeaderCodeGenerator.cs b/Engine/Source/Programs/Shared/EpicGames.UHT/Exporters/CodeGen/UhtHeaderCodeGenerator.cs
--- a/Engine/Source/Programs/Shared/EpicGames.UHT/Exporters/CodeGen/UhtHeaderCodeGenerator.cs
+++ b/Engine/Source/Programs/Shared/EpicGames.UHT/Exporters/CodeGen/UhtHeaderCodeGenerator.cs
@@ -277,9 +277,29 @@ namespace EpicGames.UHT.Exporters.CodeGen
 			return builder;
 		}
 
-		protected static StringBuilder AppendEventFunctionPrologue(StringBuilder builder, UhtFunction function, string functionName, int tabs, string endl, bool addEventParameterStruct)
+		// Find the UFUNCTION associated with this event so that we can check
+		// if there is a BP implementation of it
+		protected static StringBuilder AppendFindUFunction(StringBuilder builder, UhtClass classObj, UhtFunction function, int tabs, string endl)
 		{
-			builder.AppendTabs(tabs).Append('{').Append(endl);
+			builder
+				.AppendTabs(tabs)
+				.Append("UFunction* Func = FindFunctionChecked(")
+				.Append("NAME_")
+				.Append(classObj.SourceName)
+				.Append('_')
+				.Append(function.EngineName)
+				.Append(");\r\n");
+
+			return builder;
+		}
+
+		protected static StringBuilder AppendEventFunctionPrologue(StringBuilder builder, UhtFunction function, string functionName, int tabs, string endl, bool addEventParameterStruct, bool bAddFunctionScopeBracket = true)
+		{
+			if (bAddFunctionScopeBracket)
+			{
+				builder.AppendTabs(tabs).Append('{').Append(endl);
+			}
+			
 			if (function.Children.Count == 0)
 			{
 				return builder;
@@ -322,7 +342,7 @@ namespace EpicGames.UHT.Exporters.CodeGen
 			return builder;
 		}
 
-		protected static StringBuilder AppendEventFunctionEpilogue(StringBuilder builder, UhtFunction function, int tabs, string endl)
+		protected static StringBuilder AppendEventFunctionEpilogue(StringBuilder builder, UhtFunction function, int tabs, string endl, bool bAddFunctionScopeBracket = true)
 		{
 			++tabs;
 
@@ -382,7 +402,12 @@ namespace EpicGames.UHT.Exporters.CodeGen
 			}
 
 			--tabs;
-			builder.AppendTabs(tabs).Append('}').Append(endl);
+
+			if (bAddFunctionScopeBracket)
+			{
+				builder.AppendTabs(tabs).Append('}').Append(endl);
+			}
+			
 			return builder;
 		}
 
diff --git a/Engine/Source/Programs/Shared/EpicGames.UHT/Exporters/CodeGen/UhtHeaderCodeGeneratorCppFile.cs b/Engine/Source/Programs/Shared/EpicGames.UHT/Exporters/CodeGen/UhtHeaderCodeGeneratorCppFile.cs
--- a/Engine/Source/Programs/Shared/EpicGames.UHT/Exporters/CodeGen/UhtHeaderCodeGeneratorCppFile.cs
+++ b/Engine/Source/Programs/Shared/EpicGames.UHT/Exporters/CodeGen/UhtHeaderCodeGeneratorCppFile.cs
@@ -1143,26 +1143,115 @@ namespace EpicGames.UHT.Exporters.CodeGen
 			}
 			else
 			{
-				AppendEventFunctionPrologue(builder, function, function.EngineName, 0, "\r\n", false);
+				bool isBlueprintEvent = function.FunctionFlags.HasAnyFlags(EFunctionFlags.BlueprintEvent);
+				bool isNetEvent = function.FunctionFlags.HasAnyFlags(EFunctionFlags.Net);
+				bool isCallInEditor = function.MetaData.ContainsKey(UhtNames.CallInEditor);
+				bool isEditorFunction = function.FunctionFlags.HasAnyFlags(EFunctionFlags.EditorOnly);
+				bool hasValidCppImpl = (function.StrippedFunctionName != function.CppImplName);
+
+				// This is a small optimization that we can do for BlueprintNativeEvents.
+				// If there is a native "_Implementation" of the current function, then we can do a small
+				// optimization where we only call "ProcessEvent" if there is actually a BP script override
+				// of the native implementation. This saves us the cost of unnecessarily copying the function
+				// params to the BPVM if we don't have to. We can only do this optimization
+				// if the implementation function name is not the same as the actual C++ function name, as that
+				// would just call the function recursively. We cannot do this optimization for networked events
+				// because ProcessEvent does some important replication behavior which we do not want to lose.
+				bool doNativeImplOptimization = 
+					isBlueprintEvent && 
+					!isNetEvent && 
+					!isCallInEditor && 
+					!isEditorFunction &&
+					hasValidCppImpl;				
+				
+				if (!doNativeImplOptimization)
+				{
+					AppendEventFunctionPrologue(builder, function, function.EngineName, 0, "\r\n", false);
+
+					AppendFindUFunction(builder, classObj, function, 1, "\r\n");
+				}				
+				else
+				{
+					builder.Append("{\r\n");
+
+					AppendFindUFunction(builder, classObj, function, 1, "\r\n");
+
+					builder
+						.Append("\tif (Func->Script.Num() > 0)\r\n")
+						.Append("\t{\r\n");
+
+					AppendEventFunctionPrologue(builder, function, function.EngineName, /*tabs*/ 1, "\r\n", /*addEventParameterStruct=*/false,/*addEventParameterStruct*/ false);
+				}								
 
 				// Cast away const just in case, because ProcessEvent isn't const
 				builder.Append('\t');
 				if (function.FunctionFlags.HasAnyFlags(EFunctionFlags.Const))
 				{
-					builder.Append("const_cast<").Append(classObj.SourceName).Append("*>(this)->");
+					builder.AppendTabs(1).Append("const_cast<").Append(classObj.SourceName).Append("*>(this)->");
 				}
+
 				builder
-					.Append("ProcessEvent(FindFunctionChecked(")
-					.Append("NAME_")
-					.Append(classObj.SourceName)
-					.Append('_')
-					.Append(function.EngineName)
-					.Append("),")
+					.Append("ProcessEvent(Func,")
 					.Append(function.Children.Count > 0 ? "&Parms" : "NULL")
 					.Append(");\r\n");
 
-				AppendEventFunctionEpilogue(builder, function, 0, "\r\n");
+				// Call into the native implementation of the function if there is one
+				// We don't want to do this all the time, like for BlueprintInternalUseOnly functions
+				// which will not have a "_Implementation" appended to their name
+				if (doNativeImplOptimization)
+				{
+					AppendEventFunctionEpilogue(builder, function, /*tabs*/ 1, "\r\n", /*bAddFunctionScopeBracket=*/true);
+
+					builder
+						.AppendTabs(1)
+						.Append("else\r\n")
+						.AppendTabs(1)
+						.Append("{\r\n")
+						.AppendTabs(2)
+						.Append(function.HasReturnProperty ? "return " : "");
+
+					// Cast away const just in case, because ProcessEvent isn't const
+					if (function.FunctionFlags.HasAnyFlags(EFunctionFlags.Const))
+					{
+						builder.Append("const_cast<").Append(classObj.SourceName).Append("*>(this)->");
+					}
+
+					// Begin the native function call...
+					builder
+						.Append(function.CppImplName)
+						.Append('(');
+
+					// For every param in this function call, pass it to our native C++ function
+					int numParams = function.ParameterProperties.Length;
+					ReadOnlySpan<UhtType> paramSpan = function.ParameterProperties.Span;
+					for (int i = 0; i < numParams; i++)
+					{
+						UhtType parameter = paramSpan[i];
+
+						if (parameter is UhtProperty property)
+						{
+							builder.Append(property.SourceName);
+						}
+
+						// Add a "," between function params as long as it isn't the last one
+						if ((i + 1) != numParams)
+						{
+							builder.Append(", ");
+						}
+					}
+
+					// ...close the function call
+					builder
+						.Append(");\r\n")
+						.AppendTabs(1)
+						.Append("}\r\n}\r\n");
+				}
+				else
+				{
+					AppendEventFunctionEpilogue(builder, function, /*tabs*/ 0, "\r\n", /*bAddFunctionScopeBracket=*/true);
+				}			
 			}
+			
 			return builder;
 		}
 
diff --git a/Engine/Source/Programs/Shared/EpicGames.UHT/Utils/UhtNames.cs b/Engine/Source/Programs/Shared/EpicGames.UHT/Utils/UhtNames.cs
--- a/Engine/Source/Programs/Shared/EpicGames.UHT/Utils/UhtNames.cs
+++ b/Engine/Source/Programs/Shared/EpicGames.UHT/Utils/UhtNames.cs
@@ -18,6 +18,7 @@ namespace EpicGames.UHT.Utils
 		public const string BlueprintGetter = "BlueprintGetter";
 		public const string BlueprintInternalUseOnly = "BlueprintInternalUseOnly";
 		public const string BlueprintSetter = "BlueprintSetter";
+		public const string CallInEditor = "CallInEditor";
 		public const string BlueprintType = "BlueprintType";
 		public const string CannotImplementInterfaceInBlueprint = "CannotImplementInterfaceInBlueprint";
 		public const string Category = "Category";
