[UBA]
* Added support for explicit root id for path roots (this is just c# stub code)

diff --git a/Engine/Source/Programs/Shared/EpicGames.UBA/ICacheClient.cs b/Engine/Source/Programs/Shared/EpicGames.UBA/ICacheClient.cs
--- a/Engine/Source/Programs/Shared/EpicGames.UBA/ICacheClient.cs
+++ b/Engine/Source/Programs/Shared/EpicGames.UBA/ICacheClient.cs
@@ -10,12 +10,18 @@ namespace EpicGames.UBA
 		/// <summary>
 		/// Register roots used to normalize paths in caches
 		/// </summary>
-		public abstract bool RegisterRoot(string path, bool includeInKey);
+		/// <param name="path">Path of root</param>
+		/// <param name="includeInKey">set this to false if you want to ignore all files under this folder (you know they are _always_ the same for all machines)</param>
+		/// <param name="id">Id of root. On windows these numbers need to increase two at the time since double backslash paths are added automatically under the hood</param>
+		/// <returns>True if successful</returns>
+		public abstract bool RegisterRoot(string path, bool includeInKey = true, byte id = 0);
 
 		/// <summary>
 		/// Register system roots used to normalize paths in caches
 		/// </summary>
-		public abstract bool RegisterSystemRoots();
+		/// <param name="startId">Start id for system roots. On windows these take up 10 entries</param>
+		/// <returns>True if successful</returns>
+		public abstract bool RegisterSystemRoots(byte startId = 0);
 
 		/// <summary>
 		/// Create root paths instance
diff --git a/Engine/Source/Programs/Shared/EpicGames.UBA/Impl/CacheClientImpl.cs b/Engine/Source/Programs/Shared/EpicGames.UBA/Impl/CacheClientImpl.cs
--- a/Engine/Source/Programs/Shared/EpicGames.UBA/Impl/CacheClientImpl.cs
+++ b/Engine/Source/Programs/Shared/EpicGames.UBA/Impl/CacheClientImpl.cs
@@ -18,10 +18,10 @@ namespace EpicGames.UBA
 		static extern void RootPaths_Destroy(IntPtr rootPaths);
 
 		[DllImport("UbaHost", CharSet = CharSet.Auto)]
-		static extern bool RootPaths_RegisterRoot(IntPtr rootPaths, string root, bool includeInKey);
+		static extern bool RootPaths_RegisterRoot(IntPtr rootPaths, string root, bool includeInKey, byte id);
 
 		[DllImport("UbaHost", CharSet = CharSet.Auto)]
-		static extern bool RootPaths_RegisterSystemRoots(IntPtr rootPaths);
+		static extern bool RootPaths_RegisterSystemRoots(IntPtr rootPaths, byte id);
 		#endregion
 
 		#region IDisposable
@@ -53,14 +53,14 @@ namespace EpicGames.UBA
 			_handle = RootPaths_Create(_logger.GetHandle());
 		}
 
-		public bool RegisterRoot(string path, bool includeInKey)
+		public bool RegisterRoot(string path, bool includeInKey, byte id)
 		{
-			return RootPaths_RegisterRoot(_handle, path, includeInKey);
+			return RootPaths_RegisterRoot(_handle, path, includeInKey, id);
 		}
 
-		public bool RegisterSystemRoots()
+		public bool RegisterSystemRoots(byte startId)
 		{
-			return RootPaths_RegisterSystemRoots(_handle);
+			return RootPaths_RegisterSystemRoots(_handle, startId);
 		}
 
 		public IntPtr GetHandle() => _handle;
