Horde: Add a helper method to get a storage client for an artifact.

#rnx

diff --git a/Engine/Source/Programs/Shared/EpicGames.Horde/IHordeClient.cs b/Engine/Source/Programs/Shared/EpicGames.Horde/IHordeClient.cs
--- a/Engine/Source/Programs/Shared/EpicGames.Horde/IHordeClient.cs
+++ b/Engine/Source/Programs/Shared/EpicGames.Horde/IHordeClient.cs
@@ -3,6 +3,7 @@
 using System;
 using System.Threading;
 using System.Threading.Tasks;
+using EpicGames.Horde.Artifacts;
 using EpicGames.Horde.Compute;
 using EpicGames.Horde.Logs;
 using EpicGames.Horde.Storage;
@@ -99,6 +100,12 @@ namespace EpicGames.Horde
 		public static IStorageClient CreateStorageClient(this IHordeClient hordeClient, NamespaceId namespaceId, string? accessToken = null)
 			=> hordeClient.CreateStorageClient($"api/v1/storage/{namespaceId}", accessToken);
 
+		/// <summary>
+		/// Creates a storage client for a particular artifact
+		/// </summary>
+		public static IStorageClient CreateStorageClient(this IHordeClient hordeClient, ArtifactId artifactId)
+			=> hordeClient.CreateStorageClient($"api/v2/artifacts/{artifactId}");
+
 		/// <summary>
 		/// Creates a storage client for a particular log
 		/// </summary>
