EpicGames.Redis: Add statically typed SetScan method.

#rnx

diff --git a/Engine/Source/Programs/Shared/EpicGames.Redis/RedisSet.cs b/Engine/Source/Programs/Shared/EpicGames.Redis/RedisSet.cs
--- a/Engine/Source/Programs/Shared/EpicGames.Redis/RedisSet.cs
+++ b/Engine/Source/Programs/Shared/EpicGames.Redis/RedisSet.cs
@@ -1,6 +1,7 @@
 // Copyright Epic Games, Inc. All Rights Reserved.
 
 using System;
+using System.Collections.Generic;
 using System.Threading.Tasks;
 using StackExchange.Redis;
 
@@ -145,5 +146,18 @@ namespace EpicGames.Redis
 		}
 
 		#endregion
+
+		#region SetScanAsync
+
+		/// <inheritdoc cref="IDatabaseAsync.SetScanAsync(RedisKey, RedisValue, int, long, int, CommandFlags)"/>
+		public static async IAsyncEnumerable<TElement> SetScanAsync<TElement>(this IDatabaseAsync target, RedisSetKey<TElement> key, RedisValue pattern = default, int pageSize = 250, long cursor = 0, int pageOffset = 0, CommandFlags flags = CommandFlags.None)
+		{
+			await foreach (RedisValue value in target.SetScanAsync(key.Inner, default, pageSize, cursor, pageOffset, flags))
+			{
+				yield return RedisSerializer.Deserialize<TElement>(value);
+			}
+		}
+
+		#endregion
 	}
 }
