EpicGames.Redis: Fix static analysis warning.

#rnx

diff --git a/Engine/Source/Programs/Shared/EpicGames.Redis/RedisSet.cs b/Engine/Source/Programs/Shared/EpicGames.Redis/RedisSet.cs
--- a/Engine/Source/Programs/Shared/EpicGames.Redis/RedisSet.cs
+++ b/Engine/Source/Programs/Shared/EpicGames.Redis/RedisSet.cs
@@ -152,7 +152,7 @@ namespace EpicGames.Redis
 		/// <inheritdoc cref="IDatabaseAsync.SetScanAsync(RedisKey, RedisValue, int, long, int, CommandFlags)"/>
 		public static async IAsyncEnumerable<TElement> SetScanAsync<TElement>(this IDatabaseAsync target, RedisSetKey<TElement> key, RedisValue pattern = default, int pageSize = 250, long cursor = 0, int pageOffset = 0, CommandFlags flags = CommandFlags.None)
 		{
-			await foreach (RedisValue value in target.SetScanAsync(key.Inner, default, pageSize, cursor, pageOffset, flags))
+			await foreach (RedisValue value in target.SetScanAsync(key.Inner, pattern, pageSize, cursor, pageOffset, flags))
 			{
 				yield return RedisSerializer.Deserialize<TElement>(value);
 			}
