[UBA]
* Reverted change that fixed so cancel messages reached the executor.. it caused issues when agents disconnected

diff --git a/Engine/Source/Programs/Shared/EpicGames.UBA/Impl/ProcessImpl.cs b/Engine/Source/Programs/Shared/EpicGames.UBA/Impl/ProcessImpl.cs
--- a/Engine/Source/Programs/Shared/EpicGames.UBA/Impl/ProcessImpl.cs
+++ b/Engine/Source/Programs/Shared/EpicGames.UBA/Impl/ProcessImpl.cs
@@ -173,7 +173,10 @@ namespace EpicGames.UBA
 
 		internal void RaiseExited()
 		{
-			Exited?.Invoke(this, new ExitedEventArgs(this));
+			if (ExitCode != 99999) // See ProcessCancelExitCode
+			{
+				Exited?.Invoke(this, new ExitedEventArgs(this));
+			}
 		}
 	}
 }
