Horde: Move TestDataMessages into EpicGames.Horde
#rnx

diff --git a/Engine/Source/Programs/Horde/Horde.Server.Tests/Jobs/TestData/TestDataTests.cs b/Engine/Source/Programs/Horde/Horde.Server.Tests/Jobs/TestData/TestDataTests.cs
--- a/Engine/Source/Programs/Horde/Horde.Server.Tests/Jobs/TestData/TestDataTests.cs
+++ b/Engine/Source/Programs/Horde/Horde.Server.Tests/Jobs/TestData/TestDataTests.cs
@@ -9,16 +9,17 @@ using System.Threading.Tasks;
 using EpicGames.Horde;
 using EpicGames.Horde.Jobs;
 using EpicGames.Horde.Jobs.Templates;
+using EpicGames.Horde.Jobs.TestData;
 using EpicGames.Horde.Logs;
 using EpicGames.Horde.Projects;
 using EpicGames.Horde.Streams;
 using Horde.Server.Jobs;
 using Horde.Server.Jobs.Graphs;
-using Horde.Server.Jobs.TestData;
 using Horde.Server.Logs;
 using Horde.Server.Projects;
 using Horde.Server.Server;
 using Horde.Server.Streams;
+using Horde.Server.Jobs.TestData;
 using Horde.Server.Users;
 using Horde.Server.Utilities;
 using Microsoft.AspNetCore.Http;
diff --git a/Engine/Source/Programs/Horde/Horde.Server/Jobs/TestData/ITestData.cs b/Engine/Source/Programs/Horde/Horde.Server/Jobs/TestData/ITestData.cs
--- a/Engine/Source/Programs/Horde/Horde.Server/Jobs/TestData/ITestData.cs
+++ b/Engine/Source/Programs/Horde/Horde.Server/Jobs/TestData/ITestData.cs
@@ -4,34 +4,12 @@ using System;
 using System.Collections.Generic;
 using EpicGames.Horde.Jobs;
 using EpicGames.Horde.Jobs.Templates;
+using EpicGames.Horde.Jobs.TestData;
 using EpicGames.Horde.Streams;
 using MongoDB.Bson;
 
 namespace Horde.Server.Jobs.TestData
 {
-	/// <summary>
-	/// Test outcome
-	/// </summary>
-	public enum TestOutcome
-	{
-		/// <summary>
-		/// The test was successful
-		/// </summary>
-		Success,
-		/// <summary>
-		/// The test failed
-		/// </summary>
-		Failure,
-		/// <summary>
-		/// The test was skipped
-		/// </summary>
-		Skipped,
-		/// <summary>
-		/// The test had an unspecified result
-		/// </summary>
-		Unspecified
-	}
-
 	/// <summary>
 	/// Defines a testing environment based on platforms, configurations, targets, etc.
 	/// </summary>
diff --git a/Engine/Source/Programs/Horde/Horde.Server/Jobs/TestData/TestDataCollection.cs b/Engine/Source/Programs/Horde/Horde.Server/Jobs/TestData/TestDataCollection.cs
--- a/Engine/Source/Programs/Horde/Horde.Server/Jobs/TestData/TestDataCollection.cs
+++ b/Engine/Source/Programs/Horde/Horde.Server/Jobs/TestData/TestDataCollection.cs
@@ -7,6 +7,7 @@ using System.Threading;
 using System.Threading.Tasks;
 using EpicGames.Horde.Jobs;
 using EpicGames.Horde.Jobs.Templates;
+using EpicGames.Horde.Jobs.TestData;
 using EpicGames.Horde.Streams;
 using Horde.Server.Server;
 using Horde.Server.Utilities;
diff --git a/Engine/Source/Programs/Horde/Horde.Server/Jobs/TestData/TestDataController.cs b/Engine/Source/Programs/Horde/Horde.Server/Jobs/TestData/TestDataController.cs
--- a/Engine/Source/Programs/Horde/Horde.Server/Jobs/TestData/TestDataController.cs
+++ b/Engine/Source/Programs/Horde/Horde.Server/Jobs/TestData/TestDataController.cs
@@ -7,6 +7,7 @@ using System.Threading;
 using System.Threading.Tasks;
 using EpicGames.Core;
 using EpicGames.Horde.Jobs;
+using EpicGames.Horde.Jobs.TestData;
 using EpicGames.Horde.Streams;
 using Horde.Server.Server;
 using Horde.Server.Streams;
@@ -15,6 +16,7 @@ using Microsoft.AspNetCore.Authorization;
 using Microsoft.AspNetCore.Mvc;
 using Microsoft.Extensions.Options;
 using MongoDB.Bson;
+using MongoDB.Bson.Serialization;
 
 namespace Horde.Server.Jobs.TestData
 {
@@ -69,7 +71,17 @@ namespace Horde.Server.Jobs.TestData
 			[FromQuery(Name = "configuration")] string[]? configurations = null)
 		{
 			IReadOnlyList<ITestMeta> metaData = await _testDataService.FindTestMetaAsync(projects, platforms, configurations, targets);
-			return metaData.ConvertAll(m => new GetTestMetaResponse(m));
+			return metaData.ConvertAll(m => new GetTestMetaResponse
+			{
+				Id = m.Id.ToString(),
+				Platforms = m.Platforms.Select(p => p).ToList(),
+				Configurations = m.Configurations.Select(p => p).ToList(),
+				BuildTargets = m.BuildTargets.Select(p => p).ToList(),
+				ProjectName = m.ProjectName,
+				RHI = m.RHI,
+				Variation = m.Variation
+
+			});
 		}
 
 		/// <summary>
@@ -84,7 +96,20 @@ namespace Horde.Server.Jobs.TestData
 		{
 			TestRefId[] idValues = Array.ConvertAll(ids, x => TestRefId.Parse(x));
 			IReadOnlyList<ITestDataDetails> details = await _testDataService.FindTestDetailsAsync(idValues);
-			return details.Select(d => new GetTestDataDetailsResponse(d)).ToList();
+			return details.Select(d => new GetTestDataDetailsResponse
+			{
+				Id = d.Id.ToString(),
+				TestDataIds = d.TestDataIds.Select(x => x.ToString()).ToList(),
+				SuiteTests = d.SuiteTests?.Select(x => new GetSuiteTestDataResponse
+				{
+					TestId = x.TestId.ToString(),
+					Outcome = x.Outcome,
+					Duration = x.Duration,
+					UID = x.UID,
+					WarningCount = x.WarningCount,
+					ErrorCount = x.ErrorCount
+				}).ToList()
+			}).ToList();
 		}
 
 		/// <summary>
@@ -101,7 +126,14 @@ namespace Horde.Server.Jobs.TestData
 
 			IReadOnlyList<ITest> testValues = await _testDataService.FindTestsAsync(testIds.Select(x => TestId.Parse(x)).ToArray());
 
-			return testValues.Select(x => new GetTestResponse(x)).ToList();
+			return testValues.Select(x => new GetTestResponse
+			{
+				Id = x.Id.ToString(),
+				Name = x.Name,
+				DisplayName = x.DisplayName,
+				SuiteName = x.SuiteName?.ToString(),
+				Metadata = x.Metadata.Select(m => m.ToString()).ToList(),
+			}).ToList();
 		}
 
 		/// <summary>
@@ -222,7 +254,34 @@ namespace Horde.Server.Jobs.TestData
 
 				List<ITestMeta> streamMetaData = metaData.Where(x => streamMetaIds.Contains(x.Id)).ToList();
 
-				responses.Add(new GetTestStreamResponse(s.StreamId, streamTests, streamSuites, streamMetaData));
+				responses.Add(new GetTestStreamResponse
+				{
+					StreamId = s.StreamId.ToString(),
+					Tests = tests.Select(test => new GetTestResponse
+					{
+						Id = test.Id.ToString(),
+						Name = test.Name,
+						DisplayName = test.DisplayName,
+						SuiteName = test.SuiteName?.ToString(),
+						Metadata = test.Metadata.Select(m => m.ToString()).ToList()
+					}).ToList(),
+					TestSuites = suites.Select(suite => new GetTestSuiteResponse
+					{
+						Id = suite.Id.ToString(),
+						Name = suite.Name,
+						Metadata = suite.Metadata.Select(x => x.ToString()).ToList()
+					}).ToList(),
+					TestMetadata = metaData.Select(meta => new GetTestMetaResponse
+					{
+						Id = meta.Id.ToString(),
+						Platforms = meta.Platforms.Select(p => p).ToList(),
+						Configurations = meta.Configurations.Select(p => p).ToList(),
+						BuildTargets = meta.BuildTargets.Select(p => p).ToList(),
+						ProjectName = meta.ProjectName,
+						RHI = meta.RHI,
+						Variation = meta.Variation
+					}).ToList()
+				});
 			}
 
 			return responses;
@@ -273,9 +332,25 @@ namespace Horde.Server.Jobs.TestData
 			}
 
 			IReadOnlyList<ITestDataRef> dataRefs = await _testDataService.FindTestRefsAsync(queryStreams.ToArray(), metaIds.ConvertAll(x => TestMetaId.Parse(x)).ToArray(), testIds, suiteIds, minCreateTime?.UtcDateTime, maxCreateTime?.UtcDateTime, minChange, maxChange);
-			foreach (ITestDataRef d in dataRefs)
+			foreach (ITestDataRef testData in dataRefs)
 			{
-				responses.Add(new GetTestDataRefResponse(d));
+				responses.Add(new GetTestDataRefResponse
+				{
+					Id = testData.Id.ToString(),
+					StreamId = testData.StreamId.ToString(),
+					JobId = testData.JobId?.ToString(),
+					StepId = testData.StepId?.ToString(),
+					Duration = testData.Duration,
+					BuildChangeList = testData.BuildChangeList,
+					MetaId = testData.Metadata.ToString(),
+					TestId = testData.TestId?.ToString(),
+					Outcome = testData.TestId != null ? testData.Outcome : null,
+					SuiteId = testData.SuiteId?.ToString(),
+					SuiteSkipCount = testData.SuiteSkipCount,
+					SuiteWarningCount = testData.SuiteWarningCount,
+					SuiteErrorCount = testData.SuiteErrorCount,
+					SuiteSuccessCount = testData.SuiteSuccessCount
+				});
 			}
 
 			return responses;
@@ -337,11 +412,22 @@ namespace Horde.Server.Jobs.TestData
 			List<object> results = new List<object>();
 
 			IReadOnlyList<ITestData> documents = await _testDataCollection.FindAsync(streamIdValue, minChange, maxChange, jobId, jobStepId, key, index, count, cancellationToken);
-			foreach (ITestData document in documents)
+			foreach (ITestData testData in documents)
 			{
-				if (await _jobService.AuthorizeAsync(document.JobId, JobAclAction.ViewJob, User, _globalConfig.Value, cancellationToken))
+				if (await _jobService.AuthorizeAsync(testData.JobId, JobAclAction.ViewJob, User, _globalConfig.Value, cancellationToken))
 				{
-					results.Add(PropertyFilter.Apply(new GetTestDataResponse(document), filter));
+					results.Add(PropertyFilter.Apply(new GetTestDataResponse
+					{
+						Id = testData.Id.ToString(),
+						StreamId = testData.StreamId.ToString(),
+						TemplateRefId = testData.TemplateRefId.ToString(),
+						JobId = testData.JobId.ToString(),
+						StepId = testData.StepId.ToString(),
+						Change = testData.Change,
+						Key = testData.Key,
+						Data = BsonSerializer.Deserialize<Dictionary<string, object>>(testData.Data)
+					}
+				, filter));
 				}
 			}
 
@@ -370,7 +456,17 @@ namespace Horde.Server.Jobs.TestData
 				return Forbid();
 			}
 
-			return PropertyFilter.Apply(new GetTestDataResponse(testData), filter);
+			return PropertyFilter.Apply(new GetTestDataResponse
+			{
+				Id = testData.Id.ToString(),
+				StreamId = testData.StreamId.ToString(),
+				TemplateRefId = testData.TemplateRefId.ToString(),
+				JobId = testData.JobId.ToString(),
+				StepId = testData.StepId.ToString(),
+				Change = testData.Change,
+				Key = testData.Key,
+				Data = BsonSerializer.Deserialize<Dictionary<string, object>>(testData.Data)
+			}, filter);
 		}
 	}
 }
diff --git a/Engine/Source/Programs/Horde/Horde.Server/Jobs/TestData/TestDataMessages.cs b/Engine/Source/Programs/Horde/Horde.Server/Jobs/TestData/TestDataMessages.cs
deleted file mode 100644
diff --git a/Engine/Source/Programs/Shared/EpicGames.Horde/Jobs/TestData/TestDataMessages.cs b/Engine/Source/Programs/Shared/EpicGames.Horde/Jobs/TestData/TestDataMessages.cs
new file mode 100644
--- /dev/null
+++ b/Engine/Source/Programs/Shared/EpicGames.Horde/Jobs/TestData/TestDataMessages.cs
@@ -0,0 +1,391 @@
+// Copyright Epic Games, Inc. All Rights Reserved.
+
+using System;
+using System.Collections.Generic;
+using System.ComponentModel.DataAnnotations;
+
+#pragma warning disable CA2227 // Change 'x' to be read-only by removing the property setter
+
+namespace EpicGames.Horde.Jobs.TestData
+{
+	/// <summary>
+	/// Test outcome
+	/// </summary>
+	public enum TestOutcome
+	{
+		/// <summary>
+		/// The test was successful
+		/// </summary>
+		Success,
+		/// <summary>
+		/// The test failed
+		/// </summary>
+		Failure,
+		/// <summary>
+		/// The test was skipped
+		/// </summary>
+		Skipped,
+		/// <summary>
+		/// The test had an unspecified result
+		/// </summary>
+		Unspecified
+	}
+
+	/// <summary>
+	/// Response object describing test data to store
+	/// </summary>
+	public class CreateTestDataRequest
+	{
+		/// <summary>
+		/// The job which produced the data
+		/// </summary>
+		[Required]
+		public JobId JobId { get; set; }
+
+		/// <summary>
+		/// The step that ran
+		/// </summary>
+		[Required]
+		public JobStepId StepId { get; set; }
+
+		/// <summary>
+		/// Key used to identify the particular data
+		/// </summary>
+		public string Key { get; set; } = String.Empty;
+
+		/// <summary>
+		/// The data stored for this test
+		/// </summary>
+		[Required]
+		public Dictionary<string, object> Data { get; set; } = new Dictionary<string, object>();
+	}
+
+	/// <summary>
+	/// Response object describing the created document
+	/// </summary>
+	public class CreateTestDataResponse
+	{
+		/// <summary>
+		/// The id for the new document
+		/// </summary>
+		public string Id { get; set; }
+
+		/// <summary>
+		/// Constructor
+		/// </summary>
+		/// <param name="id">Id of the new document</param>
+		public CreateTestDataResponse(string id)
+		{
+			Id = id;
+		}
+	}
+
+	/// <summary>
+	/// Response object describing test results
+	/// </summary>
+	public class GetTestDataResponse
+	{
+		/// <summary>
+		/// Unique id of the test data
+		/// </summary>
+		public string Id { get; set; } = String.Empty;
+
+		/// <summary>
+		/// Stream that generated the test data
+		/// </summary>
+		public string StreamId { get; set; } = String.Empty;
+
+		/// <summary>
+		/// The template reference id
+		/// </summary>
+		public string TemplateRefId { get; set; } = String.Empty;
+
+		/// <summary>
+		/// The job which produced the data
+		/// </summary>
+		public string JobId { get; set; } = String.Empty;
+
+		/// <summary>
+		/// The step that ran
+		/// </summary>
+		public string StepId { get; set; } = String.Empty;
+
+		/// <summary>
+		/// The changelist number that contained the data
+		/// </summary>
+		public int Change { get; set; } = 0;
+
+		/// <summary>
+		/// Key used to identify the particular data
+		/// </summary>
+		public string Key { get; set; } = String.Empty;
+
+		/// <summary>
+		/// The data stored for this test
+		/// </summary>
+		public Dictionary<string, object> Data { get; set; } = new Dictionary<string, object>();
+	}
+
+	/// <summary>
+	/// A test emvironment running in a stream
+	/// </summary>
+	public class GetTestMetaResponse
+	{
+		/// <summary>
+		/// Meta unique id for environment 
+		/// </summary>
+		public string Id { get; set; } = String.Empty;
+
+		/// <summary>
+		/// The platforms in the environment
+		/// </summary>
+		public List<string> Platforms { get; set; } = new List<string>();
+
+		/// <summary>
+		/// The build configurations being tested
+		/// </summary>
+		public List<string> Configurations { get; set; } = new List<string>();
+
+		/// <summary>
+		/// The build targets being tested
+		/// </summary>
+		public List<string> BuildTargets { get; set; } = new List<string>();
+
+		/// <summary>
+		/// The test project name
+		/// </summary>	
+		public string ProjectName { get; set; } = String.Empty;
+
+		/// <summary>
+		/// The rendering hardware interface being used with the test
+		/// </summary>
+		public string RHI { get; set; } = String.Empty;
+
+		/// <summary>
+		/// The varation of the test meta data, for example address sanitizing
+		/// </summary>
+		public string Variation { get; set; } = String.Empty;
+	}
+
+	/// <summary>
+	/// A test that runs in a stream
+	/// </summary>
+	public class GetTestResponse
+	{
+		/// <summary>
+		/// The id of the test
+		/// </summary>
+		public string Id { get; set; } = String.Empty;
+
+		/// <summary>
+		/// The name of the test 
+		/// </summary>
+		public string Name { get; set; } = String.Empty;
+
+		/// <summary>
+		/// The name of the test 
+		/// </summary>
+		public string? DisplayName { get; set; }
+
+		/// <summary>
+		/// The name of the test suite
+		/// </summary>
+		public string? SuiteName { get; set; }
+
+		/// <summary>
+		/// The meta data the test runs on
+		/// </summary>
+		public List<string> Metadata { get; set; } = new List<string>();
+	}
+
+	/// <summary>
+	/// Get tests request
+	/// </summary>
+	public class GetTestsRequest
+	{
+		/// <summary>
+		/// Test ids to get
+		/// </summary>
+		public List<string> TestIds { get; set; } = new List<string>();
+	}
+
+	/// <summary>
+	/// A test suite that runs in a stream, contain subtests
+	/// </summary>
+	public class GetTestSuiteResponse
+	{
+		/// <summary>
+		/// The id of the suite
+		/// </summary>	
+		public string Id { get; set; } = String.Empty;
+
+		/// <summary>
+		/// The name of the test suite
+		/// </summary>
+		public string Name { get; set; } = String.Empty;
+
+		/// <summary>
+		/// The meta data the test suite runs on
+		/// </summary>
+		public List<string> Metadata { get; set; } = new List<string>();
+	}
+
+	/// <summary>
+	/// Response object describing test results
+	/// </summary>
+	public class GetTestStreamResponse
+	{
+		/// <summary>
+		/// The stream id
+		/// </summary>
+		public string StreamId { get; set; } = String.Empty;
+
+		/// <summary>
+		/// Individual tests which run in the stream
+		/// </summary>
+		public List<GetTestResponse> Tests { get; set; } = new List<GetTestResponse>();
+
+		/// <summary>
+		/// Test suites that run in the stream
+		/// </summary>
+		public List<GetTestSuiteResponse> TestSuites { get; set; } = new List<GetTestSuiteResponse>();
+
+		/// <summary>
+		/// Test suites that run in the stream
+		/// </summary>
+		public List<GetTestMetaResponse> TestMetadata { get; set; } = new List<GetTestMetaResponse>();
+	}
+
+	/// <summary>
+	/// Suite test data
+	/// </summary>
+	public class GetSuiteTestDataResponse
+	{
+		/// <summary>
+		/// The test id
+		/// </summary>
+		public string TestId { get; set; } = String.Empty;
+
+		/// <summary>
+		/// The ourcome of the suite test
+		/// </summary>
+		public TestOutcome Outcome { get; set; }
+
+		/// <summary>
+		/// How long the suite test ran
+		/// </summary>
+		public TimeSpan Duration { get; set; }
+
+		/// <summary>
+		/// Test UID for looking up in test details
+		/// </summary>
+		public string UID { get; set; } =String.Empty;
+
+		/// <summary>
+		/// The number of test warnings generated
+		/// </summary>
+		public int? WarningCount { get; set; }
+
+		/// <summary>
+		/// The number of test errors generated
+		/// </summary>
+		public int? ErrorCount { get; set; }
+	}
+
+	/// <summary>
+	/// Test details
+	/// </summary>
+	public class GetTestDataDetailsResponse
+	{
+		/// <summary>
+		/// The corresponding test ref
+		/// </summary>
+		public string Id { get; set; } = String.Empty;
+
+		/// <summary>
+		/// The test documents for this ref
+		/// </summary>
+		public List<string> TestDataIds { get; set; } = new List<string>();
+
+		/// <summary>
+		/// Suite test data
+		/// </summary>		
+		public List<GetSuiteTestDataResponse>? SuiteTests { get; set; }
+	}
+
+	/// <summary>
+	/// Data ref 
+	/// </summary>
+	public class GetTestDataRefResponse
+	{
+		/// <summary>
+		/// The test ref id
+		/// </summary>
+		public string Id { get; set; } = String.Empty;
+
+		/// <summary>
+		/// The associated stream
+		/// </summary>
+		public string StreamId { get; set; } = String.Empty;
+
+		/// <summary>
+		/// The associated job id
+		/// </summary>
+		public string? JobId { get; set; }
+
+		/// <summary>
+		/// The associated step id
+		/// </summary>
+		public string? StepId { get; set; }
+
+		/// <summary>
+		/// How long the test ran
+		/// </summary>
+		public TimeSpan Duration { get; set; }
+
+		/// <summary>
+		/// The build changelist upon which the test ran, may not correspond to the job changelist
+		/// </summary>
+		public int BuildChangeList { get; set; }
+
+		/// <summary>
+		/// The platform the test ran on 
+		/// </summary>
+		public string MetaId { get; set; } = String.Empty;
+
+		/// <summary>
+		/// The test id in stream
+		/// </summary>
+		public string? TestId { get; set; }
+
+		/// <summary>
+		/// The outcome of the test
+		/// </summary>
+		public TestOutcome? Outcome { get; set; }
+
+		/// <summary>
+		/// The if of the stream test suite
+		/// </summary>
+		public string? SuiteId { get; set; }
+
+		/// <summary>
+		/// Suite tests skipped
+		/// </summary>
+		public int? SuiteSkipCount { get; set; }
+
+		/// <summary>
+		/// Suite test warnings
+		/// </summary>
+		public int? SuiteWarningCount { get; set; }
+
+		/// <summary>
+		/// Suite test errors
+		/// </summary>
+		public int? SuiteErrorCount { get; set; }
+
+		/// <summary>
+		/// Suite test successes
+		/// </summary>
+		public int? SuiteSuccessCount { get; set; }
+	}
+}
