Horde: Move logic for computing chunk boundaries into BuzHash class.

#rnx

diff --git a/Engine/Source/Programs/Shared/EpicGames.Core/BuzHash.cs b/Engine/Source/Programs/Shared/EpicGames.Core/BuzHash.cs
--- a/Engine/Source/Programs/Shared/EpicGames.Core/BuzHash.cs
+++ b/Engine/Source/Programs/Shared/EpicGames.Core/BuzHash.cs
@@ -225,5 +225,58 @@ namespace EpicGames.Core
 		{
 			return (uint)((1L << 32) / targetSize);
 		}
+
+		/// <summary>
+		/// Gets the length of a chunk up to the next content-defined boundary
+		/// </summary>
+		/// <param name="inputData">Input data to find a boundary in</param>
+		/// <param name="minSize">Minimum size for a chunk</param>
+		/// <param name="maxSize">Maximum size for a chunk</param>
+		/// <param name="targetSize">Desired average size for a chunk</param>
+		/// <returns></returns>
+		public static int FindChunkLength(ReadOnlySpan<byte> inputData, int minSize, int maxSize, int targetSize)
+		{
+			// If the target option sizes are fixed, just chunk the data along fixed boundaries
+			if (minSize == targetSize && maxSize == targetSize)
+			{
+				return Math.Min(inputData.Length, maxSize);
+			}
+
+			// Cap the append data span to the maximum amount we can add
+			int maxLength = maxSize;
+			if (maxLength < inputData.Length)
+			{
+				inputData = inputData.Slice(0, maxLength);
+			}
+
+			int windowSize = minSize;
+
+			// Fast path for appending data to the buffer up to the chunk window size
+			int length = Math.Min(windowSize, inputData.Length);
+			uint rollingHash = BuzHash.Add(0, inputData.Slice(0, length));
+
+			// Get the threshold for the rolling hash to split the output
+			uint rollingHashThreshold = (uint)((1L << 32) / targetSize);
+
+			// Step through the rest of the data which is completely contained in appendData.
+			if (length < inputData.Length)
+			{
+				Debug.Assert(length >= windowSize);
+
+				ReadOnlySpan<byte> tailSpan = inputData.Slice(length - windowSize, inputData.Length - windowSize);
+				ReadOnlySpan<byte> headSpan = inputData.Slice(length);
+
+				int count = BuzHash.Update(tailSpan, headSpan, rollingHashThreshold, ref rollingHash);
+				if (count != -1)
+				{
+					length += count;
+					return length;
+				}
+
+				length += headSpan.Length;
+			}
+
+			return length;
+		}
 	}
 }
diff --git a/Engine/Source/Programs/Shared/EpicGames.Horde/Storage/Nodes/ChunkedDataNode.cs b/Engine/Source/Programs/Shared/EpicGames.Horde/Storage/Nodes/ChunkedDataNode.cs
--- a/Engine/Source/Programs/Shared/EpicGames.Horde/Storage/Nodes/ChunkedDataNode.cs
+++ b/Engine/Source/Programs/Shared/EpicGames.Horde/Storage/Nodes/ChunkedDataNode.cs
@@ -343,47 +343,7 @@ namespace EpicGames.Horde.Storage.Nodes
 		/// <returns>The number of bytes to append</returns>
 		public static int GetChunkLength(ReadOnlySpan<byte> inputData, LeafChunkedDataNodeOptions options)
 		{
-			// If the target option sizes are fixed, just chunk the data along fixed boundaries
-			if (options.MinSize == options.TargetSize && options.MaxSize == options.TargetSize)
-			{
-				return Math.Min(inputData.Length, options.MaxSize);
-			}
-
-			// Cap the append data span to the maximum amount we can add
-			int maxLength = options.MaxSize;
-			if (maxLength < inputData.Length)
-			{
-				inputData = inputData.Slice(0, maxLength);
-			}
-
-			int windowSize = options.MinSize;
-
-			// Fast path for appending data to the buffer up to the chunk window size
-			int length = Math.Min(windowSize, inputData.Length);
-			uint rollingHash = BuzHash.Add(0, inputData.Slice(0, length));
-
-			// Get the threshold for the rolling hash to split the output
-			uint rollingHashThreshold = (uint)((1L << 32) / options.TargetSize);
-
-			// Step through the rest of the data which is completely contained in appendData.
-			if (length < inputData.Length)
-			{
-				Debug.Assert(length >= windowSize);
-
-				ReadOnlySpan<byte> tailSpan = inputData.Slice(length - windowSize, inputData.Length - windowSize);
-				ReadOnlySpan<byte> headSpan = inputData.Slice(length);
-
-				int count = BuzHash.Update(tailSpan, headSpan, rollingHashThreshold, ref rollingHash);
-				if (count != -1)
-				{
-					length += count;
-					return length;
-				}
-
-				length += headSpan.Length;
-			}
-
-			return length;
+			return BuzHash.FindChunkLength(inputData, options.MinSize, options.MaxSize, options.TargetSize);
 		}
 	}
 
