Horde: Update the last write time when extracting files from an archive, if present.

#rnx

diff --git a/Engine/Source/Programs/Shared/EpicGames.Horde/Storage/Nodes/DirectoryNode.Extract.cs b/Engine/Source/Programs/Shared/EpicGames.Horde/Storage/Nodes/DirectoryNode.Extract.cs
--- a/Engine/Source/Programs/Shared/EpicGames.Horde/Storage/Nodes/DirectoryNode.Extract.cs
+++ b/Engine/Source/Programs/Shared/EpicGames.Horde/Storage/Nodes/DirectoryNode.Extract.cs
@@ -669,6 +669,12 @@ namespace EpicGames.Horde.Storage.Nodes
 			{
 				file.FileInfo.Refresh();
 				FileEntry.SetPermissions(file.FileInfo!, file.FileEntry.Flags);
+
+				if((file.FileEntry.Flags & FileEntryFlags.HasModTime) != 0)
+				{
+					file.FileInfo.LastWriteTimeUtc = file.FileEntry.ModTime;
+				}
+
 				stats.Update(1, 0);
 			}
 		}
