Horde: Add last modified times to file entries in bundles. Presence of a ModTime field is governed by the HasModTime flag in the FileEntry.

diff --git a/Engine/Source/Programs/Horde/Horde.Server/Replicators/PerforceReplicator.cs b/Engine/Source/Programs/Horde/Horde.Server/Replicators/PerforceReplicator.cs
--- a/Engine/Source/Programs/Horde/Horde.Server/Replicators/PerforceReplicator.cs
+++ b/Engine/Source/Programs/Horde/Horde.Server/Replicators/PerforceReplicator.cs
@@ -664,7 +664,7 @@ namespace Horde.Server.Replicators
 						else if (io.Command == PerforceIoCommand.Close)
 						{
 							FileInfo info = await fileWriter.CloseAsync(io.File, cancellationToken);
-							rootUpdate.AddFile(info.Path.ToString(), info.Flags, info.Length, info.LeafChunkedData, info.Md5);
+							rootUpdate.AddFile(info.Path.ToString(), info.Flags, info.Length, info.LeafChunkedData, customData: info.Md5);
 						}
 						else if (io.Command == PerforceIoCommand.Unlink)
 						{
diff --git a/Engine/Source/Programs/Shared/EpicGames.Horde/Storage/Nodes/DirectoryNode.Update.cs b/Engine/Source/Programs/Shared/EpicGames.Horde/Storage/Nodes/DirectoryNode.Update.cs
--- a/Engine/Source/Programs/Shared/EpicGames.Horde/Storage/Nodes/DirectoryNode.Update.cs
+++ b/Engine/Source/Programs/Shared/EpicGames.Horde/Storage/Nodes/DirectoryNode.Update.cs
@@ -142,25 +142,26 @@ namespace EpicGames.Horde.Storage.Nodes
 	/// </summary>
 	/// <param name="Path">Path to the file</param>
 	/// <param name="Length">Length of the file data</param>
+	/// <param name="ModTime">Last modified time for the file</param>
 	/// <param name="Flags">Flags for the new file entry</param>
 	/// <param name="StreamHash">Hash of the entire stream</param>
 	/// <param name="Nodes">Chunked data for the file</param>
 	/// <param name="CustomData"></param>
-	public record class FileUpdate(string Path, FileEntryFlags Flags, long Length, IoHash StreamHash, List<ChunkedDataNodeRef> Nodes, ReadOnlyMemory<byte> CustomData = default)
+	public record class FileUpdate(string Path, FileEntryFlags Flags, long Length, IoHash StreamHash, List<ChunkedDataNodeRef> Nodes, DateTime ModTime = default, ReadOnlyMemory<byte> CustomData = default)
 	{
 		/// <summary>
 		/// Constructor
 		/// </summary>
-		public FileUpdate(string path, FileEntryFlags flags, long length, ChunkedData chunkedData, ReadOnlyMemory<byte> customData = default)
-			: this(path, flags, length, chunkedData.StreamHash, new List<ChunkedDataNodeRef> { chunkedData.Root }, customData)
+		public FileUpdate(string path, FileEntryFlags flags, long length, ChunkedData chunkedData, DateTime modTime = default, ReadOnlyMemory<byte> customData = default)
+			: this(path, flags, length, chunkedData.StreamHash, new List<ChunkedDataNodeRef> { chunkedData.Root }, modTime, customData)
 		{
 		}
 
 		/// <summary>
 		/// Constructor
 		/// </summary>
-		public FileUpdate(string path, FileEntryFlags flags, long length, LeafChunkedData chunkedData, ReadOnlyMemory<byte> customData = default)
-			: this(path, flags, length, chunkedData.Hash, chunkedData.LeafHandles, customData)
+		public FileUpdate(string path, FileEntryFlags flags, long length, LeafChunkedData chunkedData, DateTime modTime = default, ReadOnlyMemory<byte> customData = default)
+			: this(path, flags, length, chunkedData.Hash, chunkedData.LeafHandles, modTime, customData)
 		{
 		}
 
@@ -245,9 +246,10 @@ namespace EpicGames.Horde.Storage.Nodes
 		/// <param name="path">Path to the file</param>
 		/// <param name="flags">Flags for the new file entry</param>
 		/// <param name="length">Length of the file</param>
+		/// <param name="modTime">Last modified time for the file</param>
 		/// <param name="chunkedData">Chunked data instance</param>
 		/// <param name="customData"></param>
-		public FileUpdate AddFile(string path, FileEntryFlags flags, long length, LeafChunkedData chunkedData, ReadOnlyMemory<byte> customData)
+		public FileUpdate AddFile(string path, FileEntryFlags flags, long length, LeafChunkedData chunkedData, DateTime modTime = default, ReadOnlyMemory<byte> customData = default)
 		{
 			string name = path;
 			for (int idx = path.Length - 1; idx >= 0; idx--)
@@ -259,7 +261,7 @@ namespace EpicGames.Horde.Storage.Nodes
 				}
 			}
 
-			FileUpdate update = new FileUpdate(name, flags, length, chunkedData, customData);
+			FileUpdate update = new FileUpdate(name, flags, length, chunkedData, modTime, customData);
 			AddFile(path, update);
 			return update;
 		}
@@ -438,7 +440,7 @@ namespace EpicGames.Horde.Storage.Nodes
 				if (file != null)
 				{
 					await file.WriteInteriorNodesAsync(writer, new ChunkingOptions().InteriorOptions, cancellationToken);
-					directoryNode.AddFile(new FileEntry(name, file.Flags, file.Length, file.StreamHash, file.Nodes[0], file.CustomData));
+					directoryNode.AddFile(new FileEntry(name, file.Flags, file.Length, file.StreamHash, file.Nodes[0], file.ModTime, file.CustomData));
 				}
 			}
 		}
@@ -495,7 +497,7 @@ namespace EpicGames.Horde.Storage.Nodes
 						FileInfo file = batch[idx];
 
 						FileEntryFlags flags = FileEntry.GetPermissions(file);
-						FileUpdate entry = new FileUpdate(file.Name, flags, file.Length, leafChunkedFiles[idx]);
+						FileUpdate entry = new FileUpdate(file.Name, flags, file.Length, leafChunkedFiles[idx], file.LastWriteTimeUtc);
 						update.AddFile(new FileReference(file).MakeRelativeTo(baseDir), entry);
 					}
 				}
@@ -593,7 +595,7 @@ namespace EpicGames.Horde.Storage.Nodes
 				}
 
 				ChunkedData chunkedFile = await InteriorChunkedDataNode.CreateTreeAsync(leafChunkedFile, options.InteriorOptions, writer, cancellationToken);
-				updates.Add(new FileUpdate(entry.FullName, flags, entry.Length, chunkedFile));
+				updates.Add(new FileUpdate(entry.FullName, flags, entry.Length, chunkedFile, entry.LastWriteTime.UtcDateTime));
 			}
 
 			// Update the tree
diff --git a/Engine/Source/Programs/Shared/EpicGames.Horde/Storage/Nodes/DirectoryNode.cs b/Engine/Source/Programs/Shared/EpicGames.Horde/Storage/Nodes/DirectoryNode.cs
--- a/Engine/Source/Programs/Shared/EpicGames.Horde/Storage/Nodes/DirectoryNode.cs
+++ b/Engine/Source/Programs/Shared/EpicGames.Horde/Storage/Nodes/DirectoryNode.cs
@@ -368,14 +368,19 @@ namespace EpicGames.Horde.Storage.Nodes
 				IoHash streamHash = reader.ReadIoHash();
 				ChunkedDataNodeRef target = new ChunkedDataNodeRef(targetType, length, targetHandle);
 
+				DateTime modTime = default;
+				if ((flags & FileEntryFlags.HasModTime) != 0)
+				{
+					modTime = new DateTime((long)reader.ReadUnsignedVarInt());
+				}
+
 				ReadOnlyMemory<byte> customData = default;
 				if ((flags & FileEntryFlags.HasCustomData) != 0)
 				{
 					customData = reader.ReadVariableLengthBytes();
-					flags &= ~FileEntryFlags.HasCustomData;
 				}
 
-				directoryNode.AddFile(new FileEntry(name, flags, length, streamHash, target, customData));
+				directoryNode.AddFile(new FileEntry(name, flags, length, streamHash, target, modTime, customData));
 			}
 
 			int directoryCount = (int)reader.ReadUnsignedVarInt();
@@ -402,11 +407,25 @@ namespace EpicGames.Horde.Storage.Nodes
 				writer.WriteBlobRef(fileEntry.Target.Handle);
 				writer.WriteUnsignedVarInt((int)fileEntry.Target.Type);
 
-				FileEntryFlags flags = (fileEntry.CustomData.Length > 0) ? (fileEntry.Flags | FileEntryFlags.HasCustomData) : (fileEntry.Flags & ~FileEntryFlags.HasCustomData);
+				FileEntryFlags flags = fileEntry.Flags & ~FileEntryFlags.HasCustomData;
+				if (fileEntry.CustomData.Length > 0)
+				{
+					flags |= FileEntryFlags.HasCustomData;
+				}
+				if (fileEntry.ModTime != default)
+				{
+					flags |= FileEntryFlags.HasModTime;
+				}
 
 				writer.WriteString(fileEntry.Name);
 				writer.WriteUnsignedVarInt((ulong)flags);
 				writer.WriteUnsignedVarInt((ulong)fileEntry.Length);
+
+				if((flags & FileEntryFlags.HasModTime) != 0)
+				{
+					writer.WriteUnsignedVarInt((ulong)fileEntry.ModTime.Ticks);
+				}
+
 				writer.WriteIoHash(fileEntry.StreamHash);
 
 				if ((flags & FileEntryFlags.HasCustomData) != 0)
diff --git a/Engine/Source/Programs/Shared/EpicGames.Horde/Storage/Nodes/FileEntry.cs b/Engine/Source/Programs/Shared/EpicGames.Horde/Storage/Nodes/FileEntry.cs
--- a/Engine/Source/Programs/Shared/EpicGames.Horde/Storage/Nodes/FileEntry.cs
+++ b/Engine/Source/Programs/Shared/EpicGames.Horde/Storage/Nodes/FileEntry.cs
@@ -45,6 +45,11 @@ namespace EpicGames.Horde.Storage.Nodes
 		/// File should be materialized as UTF-16 (but is stored as a UTF-8 source)
 		/// </summary>
 		Utf16 = 64,
+
+		/// <summary>
+		/// Whether the file entry includes a modification time
+		/// </summary>
+		HasModTime = 128,
 	}
 
 	/// <summary>
@@ -77,6 +82,11 @@ namespace EpicGames.Horde.Storage.Nodes
 		/// </summary>
 		public ChunkedDataNodeRef Target { get; }
 
+		/// <summary>
+		/// Last modified time. Only valid if FileEntryFlags.HasModTime is set.
+		/// </summary>
+		public DateTime ModTime { get; }
+
 		/// <summary>
 		/// Custom user data for this file entry
 		/// </summary>
@@ -85,21 +95,32 @@ namespace EpicGames.Horde.Storage.Nodes
 		/// <summary>
 		/// Constructor
 		/// </summary>
-		public FileEntry(string name, FileEntryFlags flags, long length, ChunkedData chunkedData, ReadOnlyMemory<byte> customData = default)
-			: this(name, flags, length, chunkedData.StreamHash, chunkedData.Root, customData)
+		public FileEntry(string name, FileEntryFlags flags, long length, ChunkedData chunkedData, DateTime modTime = default, ReadOnlyMemory<byte> customData = default)
+			: this(name, flags, length, chunkedData.StreamHash, chunkedData.Root, modTime, customData)
 		{
 		}
 
 		/// <summary>
 		/// Constructor
 		/// </summary>
-		public FileEntry(string name, FileEntryFlags flags, long length, IoHash streamHash, ChunkedDataNodeRef target, ReadOnlyMemory<byte> customData)
+		public FileEntry(string name, FileEntryFlags flags, long length, IoHash streamHash, ChunkedDataNodeRef target, DateTime modTime, ReadOnlyMemory<byte> customData)
 		{
+			flags &= ~(FileEntryFlags.HasModTime | FileEntryFlags.HasCustomData);
+			if (modTime != default)
+			{
+				flags |= FileEntryFlags.HasModTime;
+			}
+			if (customData.Length > 0)
+			{
+				flags |= FileEntryFlags.HasCustomData;
+			}
+
 			Name = name;
 			Flags = flags;
 			Length = length;
 			StreamHash = streamHash;
 			Target = target;
+			ModTime = modTime;
 			CustomData = customData;
 		}
 
