Horde: Fix incorrect order of fields serialized to directory nodes.

#rnx

diff --git a/Engine/Source/Programs/Shared/EpicGames.Horde/Storage/Nodes/DirectoryNode.cs b/Engine/Source/Programs/Shared/EpicGames.Horde/Storage/Nodes/DirectoryNode.cs
--- a/Engine/Source/Programs/Shared/EpicGames.Horde/Storage/Nodes/DirectoryNode.cs
+++ b/Engine/Source/Programs/Shared/EpicGames.Horde/Storage/Nodes/DirectoryNode.cs
@@ -420,14 +420,13 @@ namespace EpicGames.Horde.Storage.Nodes
 				writer.WriteString(fileEntry.Name);
 				writer.WriteUnsignedVarInt((ulong)flags);
 				writer.WriteUnsignedVarInt((ulong)fileEntry.Length);
+				writer.WriteIoHash(fileEntry.StreamHash);
 
-				if((flags & FileEntryFlags.HasModTime) != 0)
+				if ((flags & FileEntryFlags.HasModTime) != 0)
 				{
 					writer.WriteUnsignedVarInt((ulong)fileEntry.ModTime.Ticks);
 				}
 
-				writer.WriteIoHash(fileEntry.StreamHash);
-
 				if ((flags & FileEntryFlags.HasCustomData) != 0)
 				{
 					writer.WriteVariableLengthBytes(fileEntry.CustomData.Span);
