Fix issue where a TInstancedStructs property wasn't setting the meta data in the correct place when used with a container such as TArray.

#11720

diff --git a/Engine/Source/Programs/Shared/EpicGames.UHT/Types/Properties/UhtStructProperty.cs b/Engine/Source/Programs/Shared/EpicGames.UHT/Types/Properties/UhtStructProperty.cs
--- a/Engine/Source/Programs/Shared/EpicGames.UHT/Types/Properties/UhtStructProperty.cs
+++ b/Engine/Source/Programs/Shared/EpicGames.UHT/Types/Properties/UhtStructProperty.cs
@@ -541,7 +541,6 @@ namespace EpicGames.UHT.Types
 		{
 			TemplateWrapperName = templateWrapperName;
 			TemplateArgumentStruct = templateArgumentStruct;
-			MetaData.Add("BaseStruct", templateArgumentStruct.PathName);
 		}
 
 		/// <inheritdoc/>
@@ -594,13 +593,15 @@ namespace EpicGames.UHT.Types
 				return null;
 			}
 
-			if (propertySettings.MetaData.ContainsKey("BaseStruct"))
+			UhtPropertySettings rootSettings = propertySettings.RootSettings;
+			if (rootSettings.MetaData.ContainsKey("BaseStruct"))
 			{
 				tokenReader.LogError("BaseStruct metadata is implicitly set from the TInstancedStruct template argument and should not be explicitly specified.");
 				return null;
 			}
 
 			// With TInstancedStruct, BaseStruct is used as a type limiter.
+			rootSettings.MetaData.Add("BaseStruct", baseScriptStruct.PathName);
 			return new UhtTemplateStructProperty(propertySettings, baseScriptStruct.Session.FInstancedStruct, "TInstancedStruct", baseScriptStruct);
 		}
 
@@ -618,19 +619,20 @@ namespace EpicGames.UHT.Types
 				return null;
 			}
 
-			if (propertySettings.MetaData.ContainsKey(RefTypeName))
+			UhtPropertySettings rootSettings = propertySettings.RootSettings;
+			if (rootSettings.MetaData.ContainsKey(RefTypeName))
 			{
 				tokenReader.LogError("{0} metadata is implicitly set from the TStateTreePropertyRef template argument and should not be explicitly specified.", RefTypeName);
 				return null;
 			}
 
-			if (propertySettings.MetaData.ContainsKey(IsRefToArrayName))
+			if (rootSettings.MetaData.ContainsKey(IsRefToArrayName))
 			{
 				tokenReader.LogError("{0} metadata is implicitly set from the TStateTreePropertyRef template argument and should not be explicitly specified.", IsRefToArrayName);
 				return null;
 			}
 
-			if(!tokenReader.SkipExpectedType(matchedToken.Value, propertySettings.PropertyCategory == UhtPropertyCategory.Member))
+			if (!tokenReader.SkipExpectedType(matchedToken.Value, propertySettings.PropertyCategory == UhtPropertyCategory.Member))
 			{
 				return null;
 			}
@@ -664,23 +666,23 @@ namespace EpicGames.UHT.Types
 				}
 
 				// It's a UStruct or UClass
-				instancedStructProperty.MetaData.Add(RefTypeName, foundStruct.PathName);
+				rootSettings.MetaData.Add(RefTypeName, foundStruct.PathName);
 			}
 			else if(foundType is UhtEnum foundEnum)
 			{
 				// It's an enum
-				instancedStructProperty.MetaData.Add(RefTypeName, foundEnum.PathName);
+				rootSettings.MetaData.Add(RefTypeName, foundEnum.PathName);
 			}
 			else
 			{
 				// It's a primitive or unknown type.
-				instancedStructProperty.MetaData.Add(RefTypeName, identifier.ToString());
+				rootSettings.MetaData.Add(RefTypeName, identifier.ToString());
 			}
 
 			if (isRefToArray)
 			{
 				tokenReader.Require(">");
-				instancedStructProperty.MetaData.Add(IsRefToArrayName, true);
+				rootSettings.MetaData.Add(IsRefToArrayName, true);
 			}
 
 			tokenReader.Require(">");
diff --git a/Engine/Source/Programs/Shared/EpicGames.UHT/Types/UhtProperty.cs b/Engine/Source/Programs/Shared/EpicGames.UHT/Types/UhtProperty.cs
--- a/Engine/Source/Programs/Shared/EpicGames.UHT/Types/UhtProperty.cs
+++ b/Engine/Source/Programs/Shared/EpicGames.UHT/Types/UhtProperty.cs
@@ -645,6 +645,27 @@ namespace EpicGames.UHT.Types
 	public class UhtPropertySettings
 	{
 
+		/// <summary>
+		/// For nested properties, this points to the parent
+		/// </summary>
+		public UhtPropertySettings? ParentSettings { get; set; } = null;
+
+		/// <summary>
+		/// The root settings for a nested property
+		/// </summary>
+		public UhtPropertySettings RootSettings
+		{
+			get
+			{
+				UhtPropertySettings settings = this;
+				while (settings.ParentSettings != null)
+				{
+					settings = settings.ParentSettings;
+				}
+				return settings;
+			}
+		}
+
 		/// <summary>
 		/// Source name of the property
 		/// </summary>
@@ -758,6 +779,7 @@ namespace EpicGames.UHT.Types
 		/// <param name="messageSite">Message site used to construct meta data object</param>
 		public void Reset(UhtPropertySettings parentPropertySettings, string sourceName, IUhtMessageSite messageSite)
 		{
+			ParentSettings = parentPropertySettings;
 			SourceName = sourceName;
 			EngineName = sourceName;
 			MetaData = new UhtMetaData(messageSite, parentPropertySettings.Outer.Session.Config);
@@ -788,6 +810,7 @@ namespace EpicGames.UHT.Types
 		/// <param name="disallowPropertyFlags">Property flags that are not allowed</param>
 		public void Reset(UhtType outer, int lineNumber, UhtPropertyCategory propertyCategory, EPropertyFlags disallowPropertyFlags)
 		{
+			ParentSettings = null;
 			SourceName = String.Empty;
 			EngineName = String.Empty;
 			MetaData = new UhtMetaData(outer, outer.Session.Config);
@@ -820,6 +843,7 @@ namespace EpicGames.UHT.Types
 			{
 				throw new UhtIceException("Property must have an outer specified");
 			}
+			ParentSettings = null;
 			SourceName = property.SourceName;
 			EngineName = property.EngineName;
 			MetaData = property.MetaData;
@@ -1039,6 +1063,27 @@ namespace EpicGames.UHT.Types
 		[JsonIgnore]
 		public bool IsEditorOnlyProperty => PropertyFlags.HasAnyFlags(EPropertyFlags.EditorOnly);
 
+		/// <summary>
+		/// Return the root property for the given property.  If the property is outside of a container, 
+		/// return the property itself.  Otherwise, return the property associated with the container.
+		/// </summary>
+		public UhtProperty RootProperty
+		{
+			get
+			{
+				UhtProperty root = this;
+				while (true)
+				{
+					UhtProperty? outer = root.Outer as UhtProperty;
+					if (outer == null)
+					{
+						return root;
+					}
+					root = outer;
+				}
+			}
+		}
+
 		/// <summary>
 		/// Construct a new property
 		/// </summary>
