Horde: Make horde chunk boundary calculation consistent with unsync, add unit tests to validate determinism

diff --git a/Engine/Source/Programs/Shared/EpicGames.Core/BuzHash.cs b/Engine/Source/Programs/Shared/EpicGames.Core/BuzHash.cs
--- a/Engine/Source/Programs/Shared/EpicGames.Core/BuzHash.cs
+++ b/Engine/Source/Programs/Shared/EpicGames.Core/BuzHash.cs
@@ -175,7 +175,7 @@ namespace EpicGames.Core
 		/// <param name="state">The current hash value</param>
 		/// <returns>Number of bytes consumed from next</returns>
 #pragma warning disable CA1045 // Do not pass types by reference
-		public static unsafe int Update(ReadOnlySpan<byte> prev, ReadOnlySpan<byte> next, uint threshold, ref uint state)
+		public static unsafe int Update(ReadOnlySpan<byte> prev, ReadOnlySpan<byte> next, int windowSize, uint threshold, ref uint state)
 #pragma warning restore CA1045 // Do not pass types by reference
 		{
 			Debug.Assert(prev.Length == next.Length);
@@ -187,13 +187,13 @@ namespace EpicGames.Core
 				uint hash = state;
 				for (int idx = 0; idx < prev.Length; idx++)
 				{
-					hash = Rol32(hash, 1) ^ table[nextPtr[idx]];
 					if (hash < threshold)
 					{
 						state = hash;
 						return idx;
 					}
-					hash ^= Rol32(table[prevPtr[idx]], prev.Length - 1);
+					hash ^= Rol32(table[prevPtr[idx]], windowSize - 1);
+					hash = Rol32(hash, 1) ^ table[nextPtr[idx]];
 				}
 				state = hash;
 			}
@@ -256,7 +256,7 @@ namespace EpicGames.Core
 			uint rollingHash = BuzHash.Add(0, inputData.Slice(0, length));
 
 			// Get the threshold for the rolling hash to split the output
-			uint rollingHashThreshold = (uint)((1L << 32) / targetSize);
+			uint rollingHashThreshold = (uint)((1L << 32) / (targetSize-minSize));
 
 			// Step through the rest of the data which is completely contained in appendData.
 			if (length < inputData.Length)
@@ -266,7 +266,7 @@ namespace EpicGames.Core
 				ReadOnlySpan<byte> tailSpan = inputData.Slice(length - windowSize, inputData.Length - windowSize);
 				ReadOnlySpan<byte> headSpan = inputData.Slice(length);
 
-				int count = BuzHash.Update(tailSpan, headSpan, rollingHashThreshold, ref rollingHash);
+				int count = BuzHash.Update(tailSpan, headSpan, minSize, rollingHashThreshold, ref rollingHash);
 				if (count != -1)
 				{
 					length += count;
diff --git a/Engine/Source/Programs/Shared/EpicGames.Horde.Tests/ChunkingTests.cs b/Engine/Source/Programs/Shared/EpicGames.Horde.Tests/ChunkingTests.cs
--- a/Engine/Source/Programs/Shared/EpicGames.Horde.Tests/ChunkingTests.cs
+++ b/Engine/Source/Programs/Shared/EpicGames.Horde.Tests/ChunkingTests.cs
@@ -250,5 +250,135 @@ namespace EpicGames.Horde.Tests
 				await GetReadOrderAsync(childHandle, list);
 			}
 		}
+		private struct ChunkOffsetsAndSizes
+		{
+			public List<int> _offsets;
+			public List<int> _sizes;
+		}
+
+		/// <summary>
+		/// Utility function to test low-level chunking algorithm with different chunking parameters
+		/// </summary>
+		private static ChunkOffsetsAndSizes FindChunks(ReadOnlySpan<byte> data, int minSize, int maxSize, int targetSize)
+		{
+			ChunkOffsetsAndSizes chunks;
+
+			chunks._offsets = new List<int>();
+			chunks._sizes = new List<int>();
+
+			int currentOffset = 0;
+			while (currentOffset != data.Length)
+			{
+				chunks._offsets.Add(currentOffset);
+				int chunkLength = BuzHash.FindChunkLength(data.Slice(currentOffset), minSize, maxSize, targetSize);
+				currentOffset += chunkLength;
+				chunks._sizes.Add(currentOffset);
+			}
+
+			Assert.AreEqual(chunks._offsets.Count, chunks._sizes.Count);
+
+			return chunks;
+		}
+
+		/// <summary>
+		/// Utility function to test low-level chunking algorithm with varying target chunk sizes
+		/// </summary>
+		private static ChunkOffsetsAndSizes FindChunks(ReadOnlySpan<byte> data, int targetSize)
+		{
+			int minSize = targetSize / 2;
+			int maxSize = targetSize * 4;
+			return FindChunks(data, minSize, maxSize, targetSize);
+		}
+
+		/// <summary>
+		/// Verify that chunking algorithm produces chunks at predetermined offsets,
+		/// which is important for generating consistent chunking between different implementations.
+		/// </summary>
+		[TestMethod]
+		public void DeterministicChunkOffsets()
+		{
+			int dataSize = 1 << 20; // 1MB test buffer
+			byte[] data = GenerateChunkingTestBuffer(dataSize, 1234);
+
+			List<int> expectedChunkOffsets = new List<int>
+			{
+				0, 34577, 128471, 195115, 238047, 297334, 358754, 396031,
+				462359, 508658, 601550, 702021, 754650, 790285, 854987, 887998,
+				956848, 1042406
+			};
+
+			int targetSize = 65536;
+			ChunkOffsetsAndSizes chunks = FindChunks(data, targetSize);
+
+			CollectionAssert.AreEqual(expectedChunkOffsets, chunks._offsets);
+		}
+
+		/// <summary>
+		/// Verify that chunking algorithm produces chunks of expected sizes for different target chunk size configurations.
+		/// This helps validating chunking consistency between different implementations of the algorithm.
+		/// </summary>
+		[TestMethod]
+		public void DeterministicChunkSize()
+		{
+			int dataSize = 128 << 20; // 128MB test buffer
+			byte[] data = GenerateChunkingTestBuffer(dataSize, 1234);
+
+			// Chunk size in KB and expected number of generated chunks
+			List<(int, int)> configList = new List<(int, int)>
+			{
+				(8, 16442),
+				(16, 8146),
+				(32, 4089),
+				(64, 2019),
+				(96, 1362),
+				(128, 1012),
+				(160, 811),
+				(192, 681),
+				(256, 503),
+			};
+
+			Parallel.ForEach(configList, config =>
+			{
+				(int targetSizeKB, int expectedNumChunks) = config;
+
+				int targetSize = targetSizeKB * 1024;
+
+				ChunkOffsetsAndSizes chunks = FindChunks(data, targetSize);
+
+				int avgChunkSize = dataSize / chunks._offsets.Count;
+				int absError = Math.Abs(avgChunkSize - targetSize);
+				double absErrorPct = (100.0 * absError) / targetSize;
+
+				Assert.IsTrue(absErrorPct < 5.0, $"Average chunk size {avgChunkSize} is more than 5% different from target {targetSize}");
+				Assert.AreEqual(expectedNumChunks, chunks._offsets.Count);
+			});
+		}
+
+		/// <summary>
+		/// Utility function to quickly generate a deterministic byte sequence
+		/// that can be used to ensure consistency between multiple implementations 
+		/// of the chunking algorithm.
+		/// </summary>
+		static byte[] GenerateChunkingTestBuffer(int count, uint seed)
+		{
+			byte[] data = new byte[count];
+			uint rngState = seed;
+
+			for (int i = 0; i < count; ++i)
+			{
+				data[i] = (byte)XorShift32(ref rngState);
+			}
+
+			return data;
+		}
+		static uint XorShift32(ref uint state)
+		{
+			uint x = state;
+			x ^= x << 13;
+			x ^= x >> 17;
+			x ^= x << 5;
+			state = x;
+			return x;
+		}
 	}
 }
diff --git a/Engine/Source/Programs/Shared/EpicGames.Horde/Storage/Nodes/ChunkedDataWriter.cs b/Engine/Source/Programs/Shared/EpicGames.Horde/Storage/Nodes/ChunkedDataWriter.cs
--- a/Engine/Source/Programs/Shared/EpicGames.Horde/Storage/Nodes/ChunkedDataWriter.cs
+++ b/Engine/Source/Programs/Shared/EpicGames.Horde/Storage/Nodes/ChunkedDataWriter.cs
@@ -271,7 +271,7 @@ namespace EpicGames.Horde.Storage.Nodes
 			}
 
 			// Get the threshold for the rolling hash
-			uint rollingHashThreshold = (uint)((1L << 32) / options.TargetSize);
+			uint rollingHashThreshold = (uint)((1L << 32) / (options.TargetSize - options.MinSize));
 
 			// Step through the part of the data where the tail of the window is in currentData, and the head of the window is in appendData.
 			if (appendLength < appendData.Length && windowSize > appendLength)
@@ -282,7 +282,7 @@ namespace EpicGames.Horde.Storage.Nodes
 				ReadOnlySpan<byte> tailSpan = currentData.Slice(currentData.Length - overlap, overlapLength);
 				ReadOnlySpan<byte> headSpan = appendData.Slice(appendLength, overlapLength);
 
-				int count = BuzHash.Update(tailSpan, headSpan, rollingHashThreshold, ref rollingHash);
+				int count = BuzHash.Update(tailSpan, headSpan, windowSize, rollingHashThreshold, ref rollingHash);
 				if (count != -1)
 				{
 					appendLength += count;
@@ -300,7 +300,7 @@ namespace EpicGames.Horde.Storage.Nodes
 				ReadOnlySpan<byte> tailSpan = appendData.Slice(appendLength - windowSize, appendData.Length - windowSize);
 				ReadOnlySpan<byte> headSpan = appendData.Slice(appendLength);
 
-				int count = BuzHash.Update(tailSpan, headSpan, rollingHashThreshold, ref rollingHash);
+				int count = BuzHash.Update(tailSpan, headSpan, windowSize, rollingHashThreshold, ref rollingHash);
 				if (count != -1)
 				{
 					appendLength += count;
