Horde: Add missing documentation for windowSize parameter

diff --git a/Engine/Source/Programs/Shared/EpicGames.Core/BuzHash.cs b/Engine/Source/Programs/Shared/EpicGames.Core/BuzHash.cs
--- a/Engine/Source/Programs/Shared/EpicGames.Core/BuzHash.cs
+++ b/Engine/Source/Programs/Shared/EpicGames.Core/BuzHash.cs
@@ -171,6 +171,7 @@ namespace EpicGames.Core
 		/// </summary>
 		/// <param name="prev">Previous data added to the hash</param>
 		/// <param name="next">Next data to be added to the hash</param>
+		/// <param name="windowSize">Size of the window over which the rolling hash is computed</param>
 		/// <param name="threshold">Threshold for hash below which to terminate the scan</param>
 		/// <param name="state">The current hash value</param>
 		/// <returns>Number of bytes consumed from next</returns>
