Platform SDKs can select the most appropriate manual SDK where the platform supports side-by-side SDK installs.

 - currently only enabled when AutoSDK is configured (i.e. UE_SDKS_ROOT is set)

diff --git a/Engine/Source/Developer/TargetPlatform/Private/TargetPlatformManagerModule.cpp b/Engine/Source/Developer/TargetPlatform/Private/TargetPlatformManagerModule.cpp
--- a/Engine/Source/Developer/TargetPlatform/Private/TargetPlatformManagerModule.cpp
+++ b/Engine/Source/Developer/TargetPlatform/Private/TargetPlatformManagerModule.cpp
@@ -219,6 +219,12 @@ public:
 					SetupAndValidateAutoSDK(Pair.Value.AutoSDKPath);
 				}
 			}
+
+			FString ManualSDKEnvironmentVarsPath = FPaths::EngineIntermediateDir() / FString(TEXT("ManualSDKEnvVars.txt"));
+			if (IFileManager::Get().FileExists(*ManualSDKEnvironmentVarsPath))
+			{
+				SetupEnvironmentFromManualSDK(ManualSDKEnvironmentVarsPath);
+			}
 		}
 #endif
 
@@ -949,6 +955,16 @@ protected:
 		return true;
 #endif // AUTOSDKS_ENABLED
 	}
+
+	bool SetupEnvironmentFromManualSDK(const FString& EnvVarFileName)
+	{
+		TRACE_CPUPROFILER_EVENT_SCOPE(FTargetPlatformManagerModule::SetupEnvironmentFromManualSDK);
+
+		UE_LOG(LogTargetPlatformManager, Verbose, TEXT("Reading the manifest for auto-selected manual sdks") );
+		bool bResult = SetupEnvironmentFromEnvVarFile(EnvVarFileName);
+		IFileManager::Get().Delete(*EnvVarFileName);
+		return bResult;
+	}
 	
 	bool SetupEnvironmentFromAutoSDK(const FString& AutoSDKPath)
 	{
@@ -1016,6 +1032,21 @@ protected:
 		static const FString SDKEnvironmentVarsFile(TEXT("OutputEnvVars.txt"));
 		FString EnvVarFileName = FPaths::Combine(*TargetSDKRoot, *SDKEnvironmentVarsFile);		
 
+		if (!SetupEnvironmentFromEnvVarFile(EnvVarFileName))
+		{
+			UE_LOG(LogTargetPlatformManager, Warning, TEXT("OutputEnvVars.txt not found for platform: '%s'"), *AutoSDKPath);			
+			return false;
+		}
+
+		UE_LOG(LogTargetPlatformManager, Verbose, TEXT("Platform %s has auto sdk install"), *AutoSDKPath);		
+		return true;
+#else
+		return true;
+#endif
+	}
+
+	bool SetupEnvironmentFromEnvVarFile( const FString& EnvVarFileName )
+	{
 		// If we are using a manual install, then it is valid for there to be no OutputEnvVars file.
 		TUniquePtr<FArchive> EnvVarFile(IFileManager::Get().CreateFileReader(*EnvVarFileName));
 		if (EnvVarFile)
@@ -1126,18 +1157,12 @@ protected:
 
 			FString ModifiedPath = FString::Join(ModifiedPathVars, PathDelimiter);
 			FPlatformMisc::SetEnvironmentVar(TEXT("PATH"), *ModifiedPath);			
+			return true;
 		}
 		else
 		{
-			UE_LOG(LogTargetPlatformManager, Warning, TEXT("OutputEnvVars.txt not found for platform: '%s'"), *AutoSDKPath);			
 			return false;
 		}
-
-		UE_LOG(LogTargetPlatformManager, Verbose, TEXT("Platform %s has auto sdk install"), *AutoSDKPath);		
-		return true;
-#else
-		return true;
-#endif
 	}
 
 	bool SetupSDKStatus()
diff --git a/Engine/Source/Programs/Shared/EpicGames.Build/UEBuildPlatformSDK.cs b/Engine/Source/Programs/Shared/EpicGames.Build/UEBuildPlatformSDK.cs
--- a/Engine/Source/Programs/Shared/EpicGames.Build/UEBuildPlatformSDK.cs
+++ b/Engine/Source/Programs/Shared/EpicGames.Build/UEBuildPlatformSDK.cs
@@ -347,6 +347,15 @@ namespace EpicGames.Core
 			}
 			else
 			{
+				// allow the platform to select the most appropriate manually installed SDK if it hasn't already been set up by AutoSDK
+				if (PlatformSupportsAutoSDKs() && HasAutoSDKSystemEnabled() && !HasParentProcessSetupAutoSDK(out _)) // @todo: PlatformSupportsAutoSDKs & HasAutoSDKSystemEnabled check is not technically needed for manual sdk switching but the editor-side code is heavily tied into this and would require further refactoring first
+				{
+					if (TrySelectBestManualSDK( out string? SelectedSDKVersion ))
+					{
+						Console.WriteLine("Auto-selected best installed sdk for {0} - {1}", GetAutoSDKPlatformName(), SelectedSDKVersion);
+					}
+				}
+
 				CachedManualSDKVersions = new Dictionary<string, string>();
 
 				// if there was no parent, get the SDK version before we run AutoSDK to get the manual version
@@ -651,6 +660,40 @@ namespace EpicGames.Core
 		{
 			return new string[] { };
 		}
+		
+		/// <summary>
+		/// Find the highest valid manually-installed SDK, prioritizing the main version if it's available. Requires that the platform SDK implements GetAllInstalledSDKVersions
+		/// </summary>
+		/// <returns>True if successful</returns>
+		protected virtual string? FindBestInstalledSDKVersion()
+		{
+			string[] InstalledSDKVersions = GetAllInstalledSDKVersions();
+			if (InstalledSDKVersions.Length == 0)
+			{
+				return null;
+			}
+
+			// see if the main version is available
+			string MainVersion = GetMainVersion();
+			if (InstalledSDKVersions.Contains(MainVersion))
+			{
+				return MainVersion;
+			}
+
+			// main version is not available - find highest valid version
+			GetValidVersionRange(out string? MinVersion, out string? MaxVersion);
+			if (TryConvertVersionToInt(MinVersion, out ulong MinVersionInt, null) && TryConvertVersionToInt(MaxVersion, out ulong MaxVersionInt, null))
+			{
+				IGrouping<ulong,string>? ResultGroup = InstalledSDKVersions
+					.GroupBy( X => TryConvertVersionToInt(X, out ulong Value, null) ? Value : 0 )   // group by integer version
+					.Where( X => X.Key >= MinVersionInt && X.Key <= MaxVersionInt)                  // select only valid versions
+					.OrderByDescending( X => X.Key )                                                // sort highest version first
+					.FirstOrDefault();
+				return ResultGroup?.FirstOrDefault();
+			}
+
+			return null;
+		}
 
 		/// <summary>
 		/// Switch to another version of the SDK than what GetInstalledSDKVersion() returns. This will be one of the versions returned from GetAllInstalledSDKVersions()
@@ -663,6 +706,18 @@ namespace EpicGames.Core
 			return false;
 		}
 
+		/// <summary>
+		/// Allows a platform to switch a different version of a manually-installed SDK, if the platform supports side-by-side installations of different versions
+		/// </summary>
+		/// <returns>True if successful</returns>
+		protected virtual bool TryGetEnvironmentForManualSDK(string SelectedSDKVersion, out Dictionary<string, string>? EnvVarValues, out List<string>? PathAdds, out List<string>? PathRemoves)
+		{
+			EnvVarValues = null;
+			PathAdds = null;
+			PathRemoves = null;
+			return false;
+		}
+
 
 		/// <summary>
 		/// For a platform that doesn't use properly named AutoSDK directories, the directory name may not be convertible to an integer,
@@ -1263,6 +1318,8 @@ namespace EpicGames.Core
 		/// </summary>
 		protected const string SDKEnvironmentVarsFile = "OutputEnvVars.txt";
 
+		protected const string ManualSDKEnvironmentVarsFile = "ManualSDKEnvVars.txt";
+
 		protected const string SDKRootEnvVar = "UE_SDKS_ROOT";
 
 		protected const string AutoSetupEnvVar = "AutoSDKSetup";
@@ -1524,6 +1581,100 @@ namespace EpicGames.Core
 			}
 		}
 
+		public static void ClearManualSDKEnvVarCache()
+		{
+			string ManualSDKEnvironmentVarsPath = Path.Combine(Unreal.EngineDirectory.ToString(), "Intermediate", ManualSDKEnvironmentVarsFile);
+			if (File.Exists(ManualSDKEnvironmentVarsPath))
+			{
+				File.Delete(ManualSDKEnvironmentVarsPath);
+			}
+		}
+
+		private bool TrySelectBestManualSDK(out string? SelectedSDKVersion)
+		{
+			// find the best valid manual sdk
+			SelectedSDKVersion = FindBestInstalledSDKVersion();
+			if (SelectedSDKVersion == null)
+			{
+				return false;
+			}
+
+			// query the platform for environment changes
+			if (!TryGetEnvironmentForManualSDK(SelectedSDKVersion, out Dictionary<string, string>? EnvVarValues, out List<string>? PathAdds, out List<string>? PathRemoves))
+			{
+				return false;
+			}
+
+			// apply environment variables
+			if (EnvVarValues != null)
+			{
+				foreach (KeyValuePair<string,string> EnvVarValue in EnvVarValues)
+				{
+					Environment.SetEnvironmentVariable(EnvVarValue.Key, EnvVarValue.Value);
+				}
+			}
+
+			// apply PATH modifications
+			if (PathRemoves != null || PathAdds != null)
+			{
+				string OrigPathVar = Environment.GetEnvironmentVariable("PATH")!;
+				IEnumerable<string> PathVars = OrigPathVar.Split( Path.PathSeparator );
+				if (PathRemoves != null)
+				{
+					PathVars = PathVars.Except(PathRemoves, FileUtils.PlatformPathComparer);
+				}
+				if (PathAdds != null)
+				{
+					PathVars = PathVars.Except(PathAdds, FileUtils.PlatformPathComparer); // remove all of the ADDs so that if this function is executed multiple times, the paths will be guaranteed to be in the same order after each run.
+					PathVars = PathVars.Union(PathAdds, FileUtils.PlatformPathComparer);
+				}
+
+				string NewPathVar = String.Join( Path.PathSeparator, PathVars );
+				Environment.SetEnvironmentVariable("PATH", NewPathVar);
+			}
+
+			// write all environment modifications
+			if ( (EnvVarValues != null && EnvVarValues.Any()) || (PathRemoves != null && PathRemoves.Any()) || (PathAdds != null && PathAdds.Any()) )
+			{
+				string ManualSDKEnvironmentVarsPath = Path.Combine(Unreal.EngineDirectory.ToString(), "Intermediate", ManualSDKEnvironmentVarsFile);
+
+				Directory.CreateDirectory(Path.GetDirectoryName(ManualSDKEnvironmentVarsPath)!);
+				using (StreamWriter Writer = File.AppendText(ManualSDKEnvironmentVarsPath))
+				{
+					// write environment variables
+					if (EnvVarValues != null)
+					{
+						foreach (KeyValuePair<string,string> EnvVarValue in EnvVarValues)
+						{
+							Writer.WriteLine($"{EnvVarValue.Key}={EnvVarValue.Value}");
+							Logger.LogDebug("Setting variable '{Name}' to '{Value}'", EnvVarValue.Key, EnvVarValue.Value);
+						}
+					}
+
+					// write PATH modifications
+					if (PathRemoves != null)
+					{
+						foreach (string PathVar in PathRemoves)
+						{
+							Writer.WriteLine($"strippath={PathVar}");
+							Logger.LogDebug("Removing Path: '{Path}'", PathVar);
+						}
+					}
+					if (PathAdds != null)
+					{
+						foreach (string PathVar in PathAdds)
+						{
+							Writer.WriteLine($"addpath={PathVar}");
+							Logger.LogDebug("Adding Path: '{Path}'", PathVar);
+						}
+					}
+				}
+			}
+
+			return true;
+		}
+
+
 		protected bool SetLastRunAutoSDKScriptVersion(string LastRunScriptVersion)
 		{
 			String PlatformSDKRoot = GetPathToPlatformAutoSDKs();
diff --git a/Engine/Source/Programs/UnrealBuildTool/Configuration/UEBuildPlatform.cs b/Engine/Source/Programs/UnrealBuildTool/Configuration/UEBuildPlatform.cs
--- a/Engine/Source/Programs/UnrealBuildTool/Configuration/UEBuildPlatform.cs
+++ b/Engine/Source/Programs/UnrealBuildTool/Configuration/UEBuildPlatform.cs
@@ -256,10 +256,11 @@ namespace UnrealBuildTool
 		}
 
 		private static string[] UATProjectParams = { "-project=", "-scriptsforproject=" };
-		// Before we setup AutoSDK, we check to see if any projects need to override the Main version so that AutoSDK
-		// will set up an alternate SDK
 		private static void InitializePerPlatformSDKs(string[] Args, bool bArgumentsAreForUBT, ILogger Logger)
 		{
+			// Before we setup AutoSDK, we check to see if any projects need to override the Main version so that AutoSDK
+			// will set up an alternate SDK
+		
 			Dictionary<string, string> PlatformToVersionMap = new();
 
 			IEnumerable<FileReference?> ProjectFiles;
@@ -287,6 +288,10 @@ namespace UnrealBuildTool
 			}
 
 			UEBuildPlatformSDK.InitializePerProjectSDKVersions(ProjectFiles.OfType<FileReference>());
+
+
+			// clear the cache used for auto-switching to the best manually-installed SDK
+			UEBuildPlatformSDK.ClearManualSDKEnvVarCache();
 		}
 
 		/// <summary>
