Horde: Include modification times for file entries in bundles by default, but do not write them unless the configured API version allows it.

diff --git a/Engine/Source/Programs/Shared/EpicGames.Horde/HordeApiVersion.cs b/Engine/Source/Programs/Shared/EpicGames.Horde/HordeApiVersion.cs
--- a/Engine/Source/Programs/Shared/EpicGames.Horde/HordeApiVersion.cs
+++ b/Engine/Source/Programs/Shared/EpicGames.Horde/HordeApiVersion.cs
@@ -31,6 +31,11 @@ namespace EpicGames.Horde
 		/// </summary>
 		AddLengthsToInteriorNodes = 2,
 
+		/// <summary>
+		/// Add support for last modified timestamps to file entries in directory nodes
+		/// </summary>
+		AddFileModTimes = 3,
+
 		/// <summary>
 		/// One past the latest known version number. Add new version numbers above this point.
 		/// </summary>
diff --git a/Engine/Source/Programs/Shared/EpicGames.Horde/Storage/BlobSerializer.cs b/Engine/Source/Programs/Shared/EpicGames.Horde/Storage/BlobSerializer.cs
--- a/Engine/Source/Programs/Shared/EpicGames.Horde/Storage/BlobSerializer.cs
+++ b/Engine/Source/Programs/Shared/EpicGames.Horde/Storage/BlobSerializer.cs
@@ -203,6 +203,7 @@ namespace EpicGames.Horde.Storage
 			{
 				options.Converters.Add(new InteriorChunkedDataNodeConverter(2));
 			}
+			options.Converters.Add(new DirectoryNodeConverter(version));
 			return options;
 		}
 	}
diff --git a/Engine/Source/Programs/Shared/EpicGames.Horde/Storage/Nodes/DirectoryNode.Update.cs b/Engine/Source/Programs/Shared/EpicGames.Horde/Storage/Nodes/DirectoryNode.Update.cs
--- a/Engine/Source/Programs/Shared/EpicGames.Horde/Storage/Nodes/DirectoryNode.Update.cs
+++ b/Engine/Source/Programs/Shared/EpicGames.Horde/Storage/Nodes/DirectoryNode.Update.cs
@@ -440,7 +440,7 @@ namespace EpicGames.Horde.Storage.Nodes
 				if (file != null)
 				{
 					await file.WriteInteriorNodesAsync(writer, new ChunkingOptions().InteriorOptions, cancellationToken);
-					directoryNode.AddFile(new FileEntry(name, file.Flags, file.Length, file.StreamHash, file.Nodes[0], customData: file.CustomData));
+					directoryNode.AddFile(new FileEntry(name, file.Flags, file.Length, file.StreamHash, file.Nodes[0], file.ModTime, file.CustomData));
 				}
 			}
 		}
@@ -497,7 +497,7 @@ namespace EpicGames.Horde.Storage.Nodes
 						FileInfo file = batch[idx];
 
 						FileEntryFlags flags = FileEntry.GetPermissions(file);
-						FileUpdate entry = new FileUpdate(file.Name, flags, file.Length, leafChunkedFiles[idx]);
+						FileUpdate entry = new FileUpdate(file.Name, flags, file.Length, leafChunkedFiles[idx], file.LastWriteTimeUtc);
 						update.AddFile(new FileReference(file).MakeRelativeTo(baseDir), entry);
 					}
 				}
@@ -595,7 +595,7 @@ namespace EpicGames.Horde.Storage.Nodes
 				}
 
 				ChunkedData chunkedFile = await InteriorChunkedDataNode.CreateTreeAsync(leafChunkedFile, options.InteriorOptions, writer, cancellationToken);
-				updates.Add(new FileUpdate(entry.FullName, flags, entry.Length, chunkedFile));
+				updates.Add(new FileUpdate(entry.FullName, flags, entry.Length, chunkedFile, entry.LastWriteTime.UtcDateTime));
 			}
 
 			// Update the tree
diff --git a/Engine/Source/Programs/Shared/EpicGames.Horde/Storage/Nodes/DirectoryNode.cs b/Engine/Source/Programs/Shared/EpicGames.Horde/Storage/Nodes/DirectoryNode.cs
--- a/Engine/Source/Programs/Shared/EpicGames.Horde/Storage/Nodes/DirectoryNode.cs
+++ b/Engine/Source/Programs/Shared/EpicGames.Horde/Storage/Nodes/DirectoryNode.cs
@@ -341,10 +341,20 @@ namespace EpicGames.Horde.Storage.Nodes
 
 	class DirectoryNodeConverter : BlobConverter<DirectoryNode>
 	{
+		readonly HordeApiVersion _writeVersion;
+
 		/// <summary>
-		/// Type of serialized directory node blobs
+		/// Constructor
+		/// </summary>
+		public DirectoryNodeConverter()
+			: this(HordeApiVersion.AddFileModTimes)
+		{ }
+
+		/// <summary>
+		/// Constructor
 		/// </summary>
-		public static BlobType BlobType { get; } = new BlobType(DirectoryNode.BlobTypeGuid, 2);
+		public DirectoryNodeConverter(HordeApiVersion writeVersion)
+			=> _writeVersion = writeVersion;
 
 		/// <inheritdoc/>
 		public override DirectoryNode Read(IBlobReader reader, BlobSerializerOptions options)
@@ -357,7 +367,7 @@ namespace EpicGames.Horde.Storage.Nodes
 				IBlobRef<ChunkedDataNode> targetHandle = reader.ReadBlobRef<ChunkedDataNode>();
 
 				ChunkedDataNodeType targetType = ChunkedDataNodeType.Unknown;
-				if (reader.Version >= 2)
+				if (reader.Version >= 2) // Pre-unification with HordeApiVersion
 				{
 					targetType = (ChunkedDataNodeType)reader.ReadUnsignedVarInt();
 				}
@@ -407,12 +417,12 @@ namespace EpicGames.Horde.Storage.Nodes
 				writer.WriteBlobRef(fileEntry.Target.Handle);
 				writer.WriteUnsignedVarInt((int)fileEntry.Target.Type);
 
-				FileEntryFlags flags = fileEntry.Flags & ~FileEntryFlags.HasCustomData;
+				FileEntryFlags flags = fileEntry.Flags & ~(FileEntryFlags.HasCustomData | FileEntryFlags.HasModTime);
 				if (fileEntry.CustomData.Length > 0)
 				{
 					flags |= FileEntryFlags.HasCustomData;
 				}
-				if (fileEntry.ModTime != default)
+				if (_writeVersion >= HordeApiVersion.AddFileModTimes && fileEntry.ModTime != default)
 				{
 					flags |= FileEntryFlags.HasModTime;
 				}
@@ -441,7 +451,21 @@ namespace EpicGames.Horde.Storage.Nodes
 				writer.WriteString(directoryEntry.Name);
 			}
 
-			return BlobType;
+			return GetBlobType(_writeVersion);
+		}
+
+		static BlobType GetBlobType(HordeApiVersion version)
+		{
+			int blobVersion;
+			if (version >= HordeApiVersion.AddFileModTimes)
+			{
+				blobVersion = (int)HordeApiVersion.AddFileModTimes;
+			}
+			else
+			{
+				blobVersion = 2;
+			}
+			return new BlobType(DirectoryNode.BlobTypeGuid, blobVersion);
 		}
 	}
 
diff --git a/Engine/Source/Programs/UnrealBuildTool/Executors/UnrealBuildAccelerator/UBAAgentCoordinatorHorde.cs b/Engine/Source/Programs/UnrealBuildTool/Executors/UnrealBuildAccelerator/UBAAgentCoordinatorHorde.cs
--- a/Engine/Source/Programs/UnrealBuildTool/Executors/UnrealBuildAccelerator/UBAAgentCoordinatorHorde.cs
+++ b/Engine/Source/Programs/UnrealBuildTool/Executors/UnrealBuildAccelerator/UBAAgentCoordinatorHorde.cs
@@ -169,8 +169,8 @@ namespace UnrealBuildTool
 
 		async Task<BlobLocator> CreateToolAsync(DirectoryReference baseDir, IEnumerable<FileReference> files, CancellationToken cancellationToken)
 		{
-			BlobSerializerOptions serializerOptions = new();
-			serializerOptions.Converters.Add(new InteriorChunkedDataNodeConverter(2)); // Lock to v2 for now. Could change based on protocol version.
+			// TODO: should drive this off API version reported by server
+			BlobSerializerOptions serializerOptions = BlobSerializerOptions.Create(HordeApiVersion.Initial);
 
 			await using IBlobWriter writer = _storage.CreateBlobWriter(serializerOptions: serializerOptions);
 			DirectoryNode sandbox = new();
