Horde: Include rolling hashes in bundle archives to support unsync manifests.

* DirectoryNode/ChunkedDataNodeRef now use HordeApiVersion for versioning purposes, rather than their own custom versioning scheme. These blob types will serialize the most relevant Horde API version number they support into archives.

diff --git a/Engine/Source/Programs/Horde/Horde.Server/Artifacts/ArtifactsController.cs b/Engine/Source/Programs/Horde/Horde.Server/Artifacts/ArtifactsController.cs
--- a/Engine/Source/Programs/Horde/Horde.Server/Artifacts/ArtifactsController.cs
+++ b/Engine/Source/Programs/Horde/Horde.Server/Artifacts/ArtifactsController.cs
@@ -613,7 +613,10 @@ namespace Horde.Server.Artifacts
 				writer.WriteStartObject();
 				writer.WriteString("type", "unsync_manifest");
 				writer.WriteString("hash_strong", "Blake3.160");
-				writer.WriteString("chunking", "Variable");
+				writer.WriteString("chunking", "RollingBuzHash");
+				writer.WriteNumber("chunking_block_size_min", LeafChunkedDataNodeOptions.Default.MinSize);
+				writer.WriteNumber("chunking_block_size_max", LeafChunkedDataNodeOptions.Default.MaxSize);
+				writer.WriteNumber("chunking_block_size_target", LeafChunkedDataNodeOptions.Default.TargetSize);
 
 				writer.WriteStartArray("files");
 				foreach (UnsyncFile file in manifest.Files)
@@ -636,6 +639,7 @@ namespace Horde.Server.Artifacts
 						writer.WriteStartObject();
 						writer.WriteNumber("offset", block.Offset);
 						writer.WriteNumber("size", block.Length);
+						writer.WriteNumber("hash_weak", block.RollingHash);
 						writer.WriteString("hash_strong", block.Blob.Hash.ToString());
 						writer.WriteEndObject();
 					}
@@ -683,6 +687,8 @@ namespace Horde.Server.Artifacts
 
 			// Send the response headers
 			HttpResponse response = HttpContext.Response;
+			response.Headers["x-chunk-content-encoding"] = "identity";
+			response.StatusCode = (int)HttpStatusCode.OK;
 
 			await response.StartAsync(cancellationToken);
 			foreach (GetUnsyncBlockRequest block in request.Files.SelectMany(x => x.Blocks))
diff --git a/Engine/Source/Programs/Horde/Horde.Server/Artifacts/UnsyncCache.cs b/Engine/Source/Programs/Horde/Horde.Server/Artifacts/UnsyncCache.cs
--- a/Engine/Source/Programs/Horde/Horde.Server/Artifacts/UnsyncCache.cs
+++ b/Engine/Source/Programs/Horde/Horde.Server/Artifacts/UnsyncCache.cs
@@ -36,8 +36,9 @@ namespace Horde.Server.Artifacts
 	/// </summary>
 	/// <param name="Offset">Offset within the file</param>
 	/// <param name="Length">Length of the block</param>
+	/// <param name="RollingHash">Rolling hash of the block</param>
 	/// <param name="Blob">Handle to the corresponding blob</param>
-	public record class UnsyncBlock(long Offset, long Length, IBlobRef<LeafChunkedDataNode> Blob);
+	public record class UnsyncBlock(long Offset, long Length, uint RollingHash, IBlobRef<LeafChunkedDataNode> Blob);
 
 	/// <summary>
 	/// Implements a cache for downloading Unsync blobs
@@ -204,7 +205,7 @@ namespace Horde.Server.Artifacts
 		{
 			if (nodeRef.Type == ChunkedDataNodeType.Leaf)
 			{
-				blocks.Add(new UnsyncBlock(offset, nodeRef.Length, nodeRef.GetLeafHandle()));
+				blocks.Add(new UnsyncBlock(offset, nodeRef.Length, nodeRef.RollingHash, nodeRef.GetLeafHandle()));
 			}
 			else
 			{
diff --git a/Engine/Source/Programs/Shared/EpicGames.Core/BuzHash.cs b/Engine/Source/Programs/Shared/EpicGames.Core/BuzHash.cs
--- a/Engine/Source/Programs/Shared/EpicGames.Core/BuzHash.cs
+++ b/Engine/Source/Programs/Shared/EpicGames.Core/BuzHash.cs
@@ -234,12 +234,14 @@ namespace EpicGames.Core
 		/// <param name="minSize">Minimum size for a chunk</param>
 		/// <param name="maxSize">Maximum size for a chunk</param>
 		/// <param name="targetSize">Desired average size for a chunk</param>
+		/// <param name="rollingHash">Receives the rolling hash for the chunk</param>
 		/// <returns></returns>
-		public static int FindChunkLength(ReadOnlySpan<byte> inputData, int minSize, int maxSize, int targetSize)
+		public static int FindChunkLength(ReadOnlySpan<byte> inputData, int minSize, int maxSize, int targetSize, out uint rollingHash)
 		{
 			// If the target option sizes are fixed, just chunk the data along fixed boundaries
 			if (minSize == targetSize && maxSize == targetSize)
 			{
+				rollingHash = 0;
 				return Math.Min(inputData.Length, maxSize);
 			}
 
@@ -254,7 +256,7 @@ namespace EpicGames.Core
 
 			// Fast path for appending data to the buffer up to the chunk window size
 			int length = Math.Min(windowSize, inputData.Length);
-			uint rollingHash = BuzHash.Add(0, inputData.Slice(0, length));
+			rollingHash = BuzHash.Add(0, inputData.Slice(0, length));
 
 			// Get the threshold for the rolling hash to split the output
 			uint rollingHashThreshold = (uint)((1L << 32) / (targetSize-minSize));
diff --git a/Engine/Source/Programs/Shared/EpicGames.Horde.Tests/ChunkingTests.cs b/Engine/Source/Programs/Shared/EpicGames.Horde.Tests/ChunkingTests.cs
--- a/Engine/Source/Programs/Shared/EpicGames.Horde.Tests/ChunkingTests.cs
+++ b/Engine/Source/Programs/Shared/EpicGames.Horde.Tests/ChunkingTests.cs
@@ -174,7 +174,7 @@ namespace EpicGames.Horde.Tests
 			chunkingOptions.InteriorOptions = new InteriorChunkedDataNodeOptions(2, 2, 2);
 
 			BlobSerializerOptions serializerOptions = new BlobSerializerOptions();
-			serializerOptions.Converters.Add(new InteriorChunkedDataNodeConverter(2)); // Does not include length fields in interior nodes
+			serializerOptions.Converters.Add(new InteriorChunkedDataNodeConverter(HordeApiVersion.Initial)); // Does not include length fields in interior nodes
 
 			using KeyValueStorageClient store = KeyValueStorageClient.CreateInMemory();
 
@@ -219,7 +219,7 @@ namespace EpicGames.Horde.Tests
 
 			IBlobRef<LeafChunkedDataNode> leafRef = await blobWriter.WriteBlobAsync(new LeafChunkedDataNode(new byte[] { 1, 2, 3 }));
 			int leafRefIndex = MemoryBlobWriter.GetIndex((IBlobRef)leafRef);
-			ChunkedDataNodeRef leafChunkedRef = new ChunkedDataNodeRef(3, leafRef);
+			ChunkedDataNodeRef leafChunkedRef = new ChunkedDataNodeRef(3, 0, leafRef);
 
 			List<ChunkedDataNodeRef> leafNodeRefs = Enumerable.Repeat(leafChunkedRef, 10000).ToList();
 			ChunkedDataNodeRef root = await InteriorChunkedDataNode.CreateTreeAsync(leafNodeRefs, interiorOptions, blobWriter, CancellationToken.None);
@@ -270,7 +270,7 @@ namespace EpicGames.Horde.Tests
 			while (currentOffset != data.Length)
 			{
 				chunks._offsets.Add(currentOffset);
-				int chunkLength = BuzHash.FindChunkLength(data.Slice(currentOffset), minSize, maxSize, targetSize);
+				int chunkLength = BuzHash.FindChunkLength(data.Slice(currentOffset), minSize, maxSize, targetSize, out _);
 				currentOffset += chunkLength;
 				chunks._sizes.Add(currentOffset);
 			}
diff --git a/Engine/Source/Programs/Shared/EpicGames.Horde/HordeApiVersion.cs b/Engine/Source/Programs/Shared/EpicGames.Horde/HordeApiVersion.cs
--- a/Engine/Source/Programs/Shared/EpicGames.Horde/HordeApiVersion.cs
+++ b/Engine/Source/Programs/Shared/EpicGames.Horde/HordeApiVersion.cs
@@ -36,6 +36,11 @@ namespace EpicGames.Horde
 		/// </summary>
 		AddFileModTimes = 3,
 
+		/// <summary>
+		/// Interior nodes in chunked data now include the rolling hash of any leaf nodes
+		/// </summary>
+		AddRollingHashesForLeafNodes = 4,
+
 		/// <summary>
 		/// One past the latest known version number. Add new version numbers above this point.
 		/// </summary>
diff --git a/Engine/Source/Programs/Shared/EpicGames.Horde/Storage/BlobSerializer.cs b/Engine/Source/Programs/Shared/EpicGames.Horde/Storage/BlobSerializer.cs
--- a/Engine/Source/Programs/Shared/EpicGames.Horde/Storage/BlobSerializer.cs
+++ b/Engine/Source/Programs/Shared/EpicGames.Horde/Storage/BlobSerializer.cs
@@ -199,10 +199,7 @@ namespace EpicGames.Horde.Storage
 		public static BlobSerializerOptions Create(HordeApiVersion version)
 		{
 			BlobSerializerOptions options = new BlobSerializerOptions();
-			if (version < HordeApiVersion.AddLengthsToInteriorNodes)
-			{
-				options.Converters.Add(new InteriorChunkedDataNodeConverter(2));
-			}
+			options.Converters.Add(new InteriorChunkedDataNodeConverter(version));
 			options.Converters.Add(new DirectoryNodeConverter(version));
 			return options;
 		}
diff --git a/Engine/Source/Programs/Shared/EpicGames.Horde/Storage/Nodes/ChunkedDataNode.cs b/Engine/Source/Programs/Shared/EpicGames.Horde/Storage/Nodes/ChunkedDataNode.cs
--- a/Engine/Source/Programs/Shared/EpicGames.Horde/Storage/Nodes/ChunkedDataNode.cs
+++ b/Engine/Source/Programs/Shared/EpicGames.Horde/Storage/Nodes/ChunkedDataNode.cs
@@ -4,6 +4,7 @@ using System;
 using System.Buffers;
 using System.Buffers.Binary;
 using System.Collections.Generic;
+using System.Diagnostics;
 using System.IO;
 using System.Linq;
 using System.Threading;
@@ -149,14 +150,15 @@ namespace EpicGames.Horde.Storage.Nodes
 	/// </summary>
 	/// <param name="Type">Type of the referenced node</param>
 	/// <param name="Length">Length of the data stream within this node</param>
+	/// <param name="RollingHash">Rolling hash for this chunk. Only serialized for leaf node references.</param>
 	/// <param name="Handle">Handle to the target node</param>
-	public record class ChunkedDataNodeRef(ChunkedDataNodeType Type, long Length, IBlobRef<ChunkedDataNode> Handle)
+	public record class ChunkedDataNodeRef(ChunkedDataNodeType Type, long Length, uint RollingHash, IBlobRef<ChunkedDataNode> Handle)
 	{
 		/// <summary>
 		/// Leaf node constructor
 		/// </summary>
-		public ChunkedDataNodeRef(long length, IBlobRef<LeafChunkedDataNode> handle)
-			: this(ChunkedDataNodeType.Leaf, length, handle)
+		public ChunkedDataNodeRef(long length, uint rollingHash, IBlobRef<LeafChunkedDataNode> handle)
+			: this(ChunkedDataNodeType.Leaf, length, rollingHash, handle)
 		{
 		}
 
@@ -164,7 +166,7 @@ namespace EpicGames.Horde.Storage.Nodes
 		/// Interior node constructor
 		/// </summary>
 		public ChunkedDataNodeRef(long length, IBlobRef<InteriorChunkedDataNode> handle)
-			: this(ChunkedDataNodeType.Interior, length, handle)
+			: this(ChunkedDataNodeType.Interior, length, 0, handle)
 		{
 		}
 
@@ -332,7 +334,8 @@ namespace EpicGames.Horde.Storage.Nodes
 					break;
 				}
 
-				int nextLength = GetChunkLength(readBuffer.Memory.Span.Slice(0, size), options);
+				uint rollingHash;
+				int nextLength = GetChunkLength(readBuffer.Memory.Span.Slice(0, size), options, out rollingHash);
 
 				ReadOnlyMemory<byte> nextBlobData = readBuffer.Memory.Slice(0, nextLength);
 				writer.WriteFixedLengthBytes(nextBlobData.Span);
@@ -346,7 +349,7 @@ namespace EpicGames.Horde.Storage.Nodes
 				}
 
 				IBlobRef<LeafChunkedDataNode> blobHandle = await writer.CompleteAsync<LeafChunkedDataNode>(LeafChunkedDataNodeConverter.BlobType, cancellationToken);
-				leafNodeRefs.Add(new ChunkedDataNodeRef(nextLength, blobHandle));
+				leafNodeRefs.Add(new ChunkedDataNodeRef(nextLength, rollingHash, blobHandle));
 
 				readBuffer.Memory.Slice(nextLength, size - nextLength).CopyTo(readBuffer.Memory);
 				size -= nextLength;
@@ -363,10 +366,11 @@ namespace EpicGames.Horde.Storage.Nodes
 		/// </summary>
 		/// <param name="inputData">Data to be appended</param>
 		/// <param name="options">Options for chunking the data</param>
+		/// <param name="rollingHash">Receives the rolling hash at the end of this chunk</param>
 		/// <returns>The number of bytes to append</returns>
-		public static int GetChunkLength(ReadOnlySpan<byte> inputData, LeafChunkedDataNodeOptions options)
+		public static int GetChunkLength(ReadOnlySpan<byte> inputData, LeafChunkedDataNodeOptions options, out uint rollingHash)
 		{
-			return BuzHash.FindChunkLength(inputData, options.MinSize, options.MaxSize, options.TargetSize);
+			return BuzHash.FindChunkLength(inputData, options.MinSize, options.MaxSize, options.TargetSize, out rollingHash);
 		}
 	}
 
@@ -476,7 +480,7 @@ namespace EpicGames.Horde.Storage.Nodes
 				foreach ((InteriorChunkedDataNode interiorNode, long interiorLength) in interiorNodes)
 				{
 					IBlobRef<InteriorChunkedDataNode> interiorHandle = await memoryWriter.WriteBlobAsync(interiorNode, cancellationToken);
-					handleBuffer.Add(new ChunkedDataNodeRef(ChunkedDataNodeType.Interior, interiorLength, interiorHandle));
+					handleBuffer.Add(new ChunkedDataNodeRef(interiorLength, interiorHandle));
 				}
 
 				nodeRefs = handleBuffer;
@@ -563,14 +567,21 @@ namespace EpicGames.Horde.Storage.Nodes
 			while (nodeReader.GetMemory(0).Length > 0)
 			{
 				IBlobRef<ChunkedDataNode> handle = nodeReader.ReadBlobRef<ChunkedDataNode>();
+
+				ChunkedDataNodeType type = ChunkedDataNodeType.Unknown;
 				if (nodeReader.Version >= 2)
 				{
-					_ = nodeReader.ReadUnsignedVarInt(); // Type
+					type = (ChunkedDataNodeType)nodeReader.ReadUnsignedVarInt();
 				}
 				if (nodeReader.Version >= 3)
 				{
 					_ = nodeReader.ReadUnsignedVarInt(); // Length
 				}
+				if (nodeReader.Version >= (int)HordeApiVersion.AddRollingHashesForLeafNodes && type == ChunkedDataNodeType.Leaf)
+				{
+					_ = nodeReader.ReadUInt32(); // Rolling hash
+				}
+
 				await ChunkedDataNode.CopyToStreamAsync(handle, outputStream, cancellationToken);
 			}
 		}
@@ -581,12 +592,12 @@ namespace EpicGames.Horde.Storage.Nodes
 	/// </summary>
 	public class InteriorChunkedDataNodeConverter : BlobConverter<InteriorChunkedDataNode>
 	{
-		readonly int _writeVersion;
+		readonly HordeApiVersion _writeVersion;
 
 		/// <summary>
 		/// Constructor
 		/// </summary>
-		public InteriorChunkedDataNodeConverter() : this(3)
+		public InteriorChunkedDataNodeConverter() : this(HordeApiVersion.Latest)
 		{
 		}
 
@@ -594,7 +605,7 @@ namespace EpicGames.Horde.Storage.Nodes
 		/// Constructor
 		/// </summary>
 		/// <param name="writeVersion">Version number for serialized data</param>
-		public InteriorChunkedDataNodeConverter(int writeVersion)
+		public InteriorChunkedDataNodeConverter(HordeApiVersion writeVersion)
 		{
 			_writeVersion = writeVersion;
 		}
@@ -609,18 +620,24 @@ namespace EpicGames.Horde.Storage.Nodes
 				IBlobRef<ChunkedDataNode> handle = reader.ReadBlobRef<ChunkedDataNode>();
 
 				ChunkedDataNodeType type = ChunkedDataNodeType.Unknown;
-				if (reader.Version >= 2)
+				if (reader.Version >= 2) // Pre sync with HordeApiVersion
 				{
 					type = (ChunkedDataNodeType)reader.ReadUnsignedVarInt();
 				}
 
 				long length = -1;
-				if (reader.Version >= 3)
+				if (reader.Version >= 3) // Pre sync with HordeApiVersion
 				{
 					length = (long)reader.ReadUnsignedVarInt();
 				}
 
-				children.Add(new ChunkedDataNodeRef(type, length, handle));
+				uint rollingHash = 0;
+				if (reader.Version >= (int)HordeApiVersion.AddRollingHashesForLeafNodes && type == ChunkedDataNodeType.Leaf)
+				{
+					rollingHash = reader.ReadUInt32();
+				}
+
+				children.Add(new ChunkedDataNodeRef(type, length, rollingHash, handle));
 			}
 			return new InteriorChunkedDataNode(children);
 		}
@@ -631,16 +648,39 @@ namespace EpicGames.Horde.Storage.Nodes
 			foreach (ChunkedDataNodeRef child in value.Children)
 			{
 				writer.WriteBlobRef(child.Handle);
-				if (_writeVersion >= 2)
+				writer.WriteUnsignedVarInt((int)child.Type);
+
+				if (_writeVersion >= HordeApiVersion.AddLengthsToInteriorNodes)
 				{
-					writer.WriteUnsignedVarInt((int)child.Type);
+					writer.WriteUnsignedVarInt((ulong)child.Length);
 				}
-				if (_writeVersion >= 3)
+				if (_writeVersion >= HordeApiVersion.AddRollingHashesForLeafNodes && child.Type == ChunkedDataNodeType.Leaf)
 				{
-					writer.WriteUnsignedVarInt((ulong)child.Length);
+					writer.WriteUInt32(child.RollingHash);
 				}
 			}
-			return new BlobType(InteriorChunkedDataNode.BlobTypeGuid, _writeVersion);
+			return GetBlobType(_writeVersion);
+		}
+
+		/// <summary>
+		/// Gets the blob type for a particular Horde api version
+		/// </summary>
+		public static BlobType GetBlobType(HordeApiVersion version)
+		{
+			int blobVersion;
+			if (version >= HordeApiVersion.AddRollingHashesForLeafNodes)
+			{
+				blobVersion = (int)HordeApiVersion.AddRollingHashesForLeafNodes;
+			}
+			else if (version >= HordeApiVersion.AddLengthsToInteriorNodes)
+			{
+				blobVersion = 3;
+			}
+			else
+			{
+				blobVersion = 2;
+			}
+			return new BlobType(InteriorChunkedDataNode.BlobTypeGuid, blobVersion);
 		}
 	}
 }
diff --git a/Engine/Source/Programs/Shared/EpicGames.Horde/Storage/Nodes/ChunkedDataWriter.cs b/Engine/Source/Programs/Shared/EpicGames.Horde/Storage/Nodes/ChunkedDataWriter.cs
--- a/Engine/Source/Programs/Shared/EpicGames.Horde/Storage/Nodes/ChunkedDataWriter.cs
+++ b/Engine/Source/Programs/Shared/EpicGames.Horde/Storage/Nodes/ChunkedDataWriter.cs
@@ -339,7 +339,7 @@ namespace EpicGames.Horde.Storage.Nodes
 		{
 			int leafLength = _writer.WrittenMemory.Length;
 			IBlobRef<LeafChunkedDataNode> leafHandle = await _writer.CompleteAsync<LeafChunkedDataNode>(LeafChunkedDataNodeConverter.BlobType, cancellationToken);
-			_leafHandles.Add(new ChunkedDataNodeRef(leafLength, leafHandle));
+			_leafHandles.Add(new ChunkedDataNodeRef(leafLength, _leafHash, leafHandle));
 			ResetLeafState();
 		}
 	}
diff --git a/Engine/Source/Programs/Shared/EpicGames.Horde/Storage/Nodes/DirectoryNode.cs b/Engine/Source/Programs/Shared/EpicGames.Horde/Storage/Nodes/DirectoryNode.cs
--- a/Engine/Source/Programs/Shared/EpicGames.Horde/Storage/Nodes/DirectoryNode.cs
+++ b/Engine/Source/Programs/Shared/EpicGames.Horde/Storage/Nodes/DirectoryNode.cs
@@ -347,7 +347,7 @@ namespace EpicGames.Horde.Storage.Nodes
 		/// Constructor
 		/// </summary>
 		public DirectoryNodeConverter()
-			: this(HordeApiVersion.AddFileModTimes)
+			: this(HordeApiVersion.AddRollingHashesForLeafNodes)
 		{ }
 
 		/// <summary>
@@ -372,11 +372,17 @@ namespace EpicGames.Horde.Storage.Nodes
 					targetType = (ChunkedDataNodeType)reader.ReadUnsignedVarInt();
 				}
 
+				uint rollingHash = 0;
+				if (targetType == ChunkedDataNodeType.Leaf && reader.Version >= (int)HordeApiVersion.AddRollingHashesForLeafNodes)
+				{
+					rollingHash = reader.ReadUInt32();
+				}
+
 				string name = reader.ReadString();
 				FileEntryFlags flags = (FileEntryFlags)reader.ReadUnsignedVarInt();
 				long length = (long)reader.ReadUnsignedVarInt();
 				IoHash streamHash = reader.ReadIoHash();
-				ChunkedDataNodeRef target = new ChunkedDataNodeRef(targetType, length, targetHandle);
+				ChunkedDataNodeRef target = new ChunkedDataNodeRef(targetType, length, rollingHash, targetHandle);
 
 				DateTime modTime = default;
 				if ((flags & FileEntryFlags.HasModTime) != 0)
@@ -417,6 +423,11 @@ namespace EpicGames.Horde.Storage.Nodes
 				writer.WriteBlobRef(fileEntry.Target.Handle);
 				writer.WriteUnsignedVarInt((int)fileEntry.Target.Type);
 
+				if (_writeVersion >= HordeApiVersion.AddRollingHashesForLeafNodes && fileEntry.Target.Type == ChunkedDataNodeType.Leaf)
+				{
+					writer.WriteUInt32(fileEntry.Target.RollingHash);
+				}
+
 				FileEntryFlags flags = fileEntry.Flags & ~(FileEntryFlags.HasCustomData | FileEntryFlags.HasModTime);
 				if (fileEntry.CustomData.Length > 0)
 				{
@@ -457,7 +468,11 @@ namespace EpicGames.Horde.Storage.Nodes
 		static BlobType GetBlobType(HordeApiVersion version)
 		{
 			int blobVersion;
-			if (version >= HordeApiVersion.AddFileModTimes)
+			if (version >= HordeApiVersion.AddRollingHashesForLeafNodes)
+			{
+				blobVersion = (int)HordeApiVersion.AddRollingHashesForLeafNodes;
+			}
+			else if (version >= HordeApiVersion.AddFileModTimes)
 			{
 				blobVersion = (int)HordeApiVersion.AddFileModTimes;
 			}
