Horde: Expect a flat array of blob hashes when requesting data for unsync.

#rnx

diff --git a/Engine/Source/Programs/Horde/Horde.Server/Artifacts/ArtifactsController.cs b/Engine/Source/Programs/Horde/Horde.Server/Artifacts/ArtifactsController.cs
--- a/Engine/Source/Programs/Horde/Horde.Server/Artifacts/ArtifactsController.cs
+++ b/Engine/Source/Programs/Horde/Horde.Server/Artifacts/ArtifactsController.cs
@@ -691,9 +691,9 @@ namespace Horde.Server.Artifacts
 			response.StatusCode = (int)HttpStatusCode.OK;
 
 			await response.StartAsync(cancellationToken);
-			foreach (GetUnsyncBlockRequest block in request.Files.SelectMany(x => x.Blocks))
+			foreach (string block in request.Blocks)
 			{
-				if (block.Hash == null || !IoHash.TryParse(block.Hash, out IoHash hash))
+				if (!IoHash.TryParse(block, out IoHash hash))
 				{
 					throw new InvalidOperationException();
 				}
diff --git a/Engine/Source/Programs/Shared/EpicGames.Horde/Artifacts/UnsyncMessages.cs b/Engine/Source/Programs/Shared/EpicGames.Horde/Artifacts/UnsyncMessages.cs
--- a/Engine/Source/Programs/Shared/EpicGames.Horde/Artifacts/UnsyncMessages.cs
+++ b/Engine/Source/Programs/Shared/EpicGames.Horde/Artifacts/UnsyncMessages.cs
@@ -22,35 +22,7 @@ namespace EpicGames.Horde.Artifacts
 		/// <summary>
 		/// Files to retrieve
 		/// </summary>
-		[JsonPropertyName("files")]
-		public List<GetUnsyncFileRequest> Files { get; set; } = new List<GetUnsyncFileRequest>();
-	}
-
-	/// <summary>
-	/// Requests a set of blobs from a particular file
-	/// </summary>
-	public class GetUnsyncFileRequest
-	{
-		/// <summary>
-		/// Path to the file
-		/// </summary>
-		public string Name { get; set; } = String.Empty;
-
-		/// <summary>
-		/// Blobs to return
-		/// </summary>
-		public List<GetUnsyncBlockRequest> Blocks { get; set; } = new List<GetUnsyncBlockRequest>();
-	}
-
-	/// <summary>
-	/// Requests a block of data
-	/// </summary>
-	public class GetUnsyncBlockRequest
-	{
-		/// <summary>
-		/// Hash of the block
-		/// </summary>
-		[JsonPropertyName("hash_strong")]
-		public string? Hash { get; set; }
+		[JsonPropertyName("blocks")]
+		public List<string> Blocks { get; set; } = new List<string>();
 	}
 }
