Horde: Add backwards compatible path for querying unsync blobs.

#rnx

diff --git a/Engine/Source/Programs/Horde/Horde.Server/Artifacts/ArtifactsController.cs b/Engine/Source/Programs/Horde/Horde.Server/Artifacts/ArtifactsController.cs
--- a/Engine/Source/Programs/Horde/Horde.Server/Artifacts/ArtifactsController.cs
+++ b/Engine/Source/Programs/Horde/Horde.Server/Artifacts/ArtifactsController.cs
@@ -685,6 +685,18 @@ namespace Horde.Server.Artifacts
 			IHttpResponseBodyFeature? responseBodyFeature = HttpContext.Features.Get<IHttpResponseBodyFeature>();
 			responseBodyFeature?.DisableBuffering();
 
+			// Merge the legacy block list with the regular one
+			foreach (GetUnsyncFileRequest file in request.Files)
+			{
+				foreach (GetUnsyncBlockRequest block in file.Blocks)
+				{
+					if (block.Hash != null)
+					{
+						request.Blocks.Add(block.Hash);
+					}
+				}
+			}
+
 			// Send the response headers
 			HttpResponse response = HttpContext.Response;
 			response.Headers["x-chunk-content-encoding"] = "identity";
diff --git a/Engine/Source/Programs/Shared/EpicGames.Horde/Artifacts/UnsyncMessages.cs b/Engine/Source/Programs/Shared/EpicGames.Horde/Artifacts/UnsyncMessages.cs
--- a/Engine/Source/Programs/Shared/EpicGames.Horde/Artifacts/UnsyncMessages.cs
+++ b/Engine/Source/Programs/Shared/EpicGames.Horde/Artifacts/UnsyncMessages.cs
@@ -24,5 +24,39 @@ namespace EpicGames.Horde.Artifacts
 		/// </summary>
 		[JsonPropertyName("blocks")]
 		public List<string> Blocks { get; set; } = new List<string>();
+
+		/// <summary>
+		/// Files to retrieve
+		/// </summary>
+		[JsonPropertyName("files")]
+		public List<GetUnsyncFileRequest> Files { get; set; } = new List<GetUnsyncFileRequest>();
+	}
+
+	/// <summary>
+	/// Requests a set of blobs from a particular file
+	/// </summary>
+	public class GetUnsyncFileRequest
+	{
+		/// <summary>
+		/// Path to the file
+		/// </summary>
+		public string Name { get; set; } = String.Empty;
+
+		/// <summary>
+		/// Blobs to return
+		/// </summary>
+		public List<GetUnsyncBlockRequest> Blocks { get; set; } = new List<GetUnsyncBlockRequest>();
+	}
+
+	/// <summary>
+	/// Requests a block of data
+	/// </summary>
+	public class GetUnsyncBlockRequest
+	{
+		/// <summary>
+		/// Hash of the block
+		/// </summary>
+		[JsonPropertyName("hash_strong")]
+		public string? Hash { get; set; }
 	}
 }
