Horde: Add zstd as a compression format for bundles. Also add some basic compression tests, and fix an isuse with gzip streams being truncated.

diff --git a/Engine/Source/Programs/Shared/EpicGames.Horde.Tests/CompressionTests.cs b/Engine/Source/Programs/Shared/EpicGames.Horde.Tests/CompressionTests.cs
new file mode 100644
--- /dev/null
+++ b/Engine/Source/Programs/Shared/EpicGames.Horde.Tests/CompressionTests.cs
@@ -0,0 +1,67 @@
+﻿// Copyright Epic Games, Inc. All Rights Reserved.
+
+using System;
+using System.Collections.Generic;
+using System.Diagnostics;
+using System.Linq;
+using System.Security.Cryptography;
+using System.Text;
+using System.Threading.Tasks;
+using EpicGames.Core;
+using EpicGames.Horde.Storage.Bundles;
+using Microsoft.VisualStudio.TestTools.UnitTesting;
+
+namespace EpicGames.Horde.Tests
+{
+	[TestClass]
+	public class CompressionTests
+	{
+		[TestMethod]
+		[DataTestMethod]
+		[DataRow(BundleCompressionFormat.LZ4)]
+		[DataRow(BundleCompressionFormat.Gzip)]
+		[DataRow(BundleCompressionFormat.Brotli)]
+		[DataRow(BundleCompressionFormat.Zstd)]
+		public void SimpleData(BundleCompressionFormat format)
+		{
+			byte[] input = new byte[2 * 1024 * 1024];
+			for (int idx = 0; idx < input.Length; idx++)
+			{
+				input[idx] = (byte)idx;
+			}
+
+			RunTests(format, input);
+		}
+
+		[TestMethod]
+		[DataTestMethod]
+		[DataRow(BundleCompressionFormat.LZ4)]
+		[DataRow(BundleCompressionFormat.Gzip)]
+		[DataRow(BundleCompressionFormat.Brotli)]
+		[DataRow(BundleCompressionFormat.Zstd)]
+		public void RandomData(BundleCompressionFormat format)
+		{
+			byte[] data = new byte[2 * 1024 * 1024];
+			new Random(0).NextBytes(data);
+
+			RunTests(format, data);
+		}
+
+		void RunTests(BundleCompressionFormat format, byte[] input)
+		{
+			ArrayMemoryWriter writer = new ArrayMemoryWriter(10 * 1024 * 1024);
+			Stopwatch compressTimer = Stopwatch.StartNew();
+			BundleData.Compress(format, input, writer);
+			compressTimer.Stop();
+
+			byte[] output = new byte[input.Length];
+			Stopwatch decompressTimer = Stopwatch.StartNew();
+			BundleData.Decompress(format, writer.WrittenMemory, output);
+			decompressTimer.Stop();
+
+			Assert.IsTrue(input.SequenceEqual(output));
+
+			Console.WriteLine("Output size: {0:n0} bytes, compression time: {1}ms, decompression time: {2}ms", writer.WrittenMemory.Length, compressTimer.Elapsed.TotalMilliseconds, decompressTimer.Elapsed.TotalMilliseconds);
+		}
+	}
+}
diff --git a/Engine/Source/Programs/Shared/EpicGames.Horde/EpicGames.Horde.csproj b/Engine/Source/Programs/Shared/EpicGames.Horde/EpicGames.Horde.csproj
--- a/Engine/Source/Programs/Shared/EpicGames.Horde/EpicGames.Horde.csproj
+++ b/Engine/Source/Programs/Shared/EpicGames.Horde/EpicGames.Horde.csproj
@@ -59,6 +59,7 @@
       <PrivateAssets>all</PrivateAssets>
       <IncludeAssets>runtime; build; native; contentfiles; analyzers; buildtransitive</IncludeAssets>
     </PackageReference>
+    <PackageReference Include="ZstdSharp.Port" Version="0.8.1" />
   </ItemGroup>
   <ItemGroup>
     <ProjectReference Include="..\EpicGames.Core\EpicGames.Core.csproj" />
diff --git a/Engine/Source/Programs/Shared/EpicGames.Horde/Storage/Bundles/BundleCompression.cs b/Engine/Source/Programs/Shared/EpicGames.Horde/Storage/Bundles/BundleCompression.cs
--- a/Engine/Source/Programs/Shared/EpicGames.Horde/Storage/Bundles/BundleCompression.cs
+++ b/Engine/Source/Programs/Shared/EpicGames.Horde/Storage/Bundles/BundleCompression.cs
@@ -9,6 +9,9 @@ using K4os.Compression.LZ4;
 
 namespace EpicGames.Horde.Storage.Bundles
 {
+	using ZstdCompressor = ZstdSharp.Compressor;
+	using ZstdDecompressor = ZstdSharp.Decompressor;
+
 	/// <summary>
 	/// Indicates the compression format in the bundle
 	/// </summary>
@@ -38,6 +41,11 @@ namespace EpicGames.Horde.Storage.Bundles
 		/// Brotli compression
 		/// </summary>
 		Brotli = 4,
+
+		/// <summary>
+		/// ZStandard compression
+		/// </summary>
+		Zstd = 5,
 	}
 
 	/// <summary>
@@ -103,6 +111,21 @@ namespace EpicGames.Horde.Storage.Bundles
 							throw new InvalidOperationException("Unable to compress data using Brotli");
 						}
 
+						writer.Advance(encodedLength);
+						return encodedLength;
+					}
+				case BundleCompressionFormat.Zstd:
+					{
+						int maxSize = ZstdCompressor.GetCompressBound(input.Length);
+
+						ZstdCompressor compressor = new ZstdCompressor();
+
+						Span<byte> buffer = writer.GetSpan(maxSize);
+						if (!compressor.TryWrap(input.Span, buffer, out int encodedLength))
+						{
+							throw new InvalidOperationException("Unable to compress data using Zstd");
+						}
+
 						writer.Advance(encodedLength);
 						return encodedLength;
 					}
@@ -138,7 +161,18 @@ namespace EpicGames.Horde.Storage.Bundles
 						using ReadOnlyMemoryStream inputStream = new ReadOnlyMemoryStream(input);
 						using GZipStream inflatedStream = new GZipStream(inputStream, CompressionMode.Decompress, true);
 
-						int length = inflatedStream.Read(output.Span);
+						int length = 0;
+						for (; ; )
+						{
+							int count = inflatedStream.Read(output.Span.Slice(length));
+							if (count == 0)
+							{
+								break;
+							}
+
+							length += count;
+						}
+
 						if (length != output.Length)
 						{
 							throw new InvalidDataException($"Decoded data is shorter than expected (expected {output.Length} bytes, got {length} bytes)");
@@ -161,6 +195,15 @@ namespace EpicGames.Horde.Storage.Bundles
 						}
 						break;
 					}
+				case BundleCompressionFormat.Zstd:
+					{
+						ZstdDecompressor decompressor = new ZstdDecompressor();
+						if (!decompressor.TryUnwrap(input.Span, output.Span, out int bytesWritten) || bytesWritten != output.Length)
+						{
+							throw new InvalidOperationException("Unable to decompress data using Zstd");
+						}
+						break;
+					}
 				default:
 					throw new InvalidDataException($"Invalid compression format '{(int)format}'");
 			}
