Horde: Fix static analysis warning.

#rnx

diff --git a/Engine/Source/Programs/Shared/EpicGames.Horde/Storage/Bundles/BundleCompression.cs b/Engine/Source/Programs/Shared/EpicGames.Horde/Storage/Bundles/BundleCompression.cs
--- a/Engine/Source/Programs/Shared/EpicGames.Horde/Storage/Bundles/BundleCompression.cs
+++ b/Engine/Source/Programs/Shared/EpicGames.Horde/Storage/Bundles/BundleCompression.cs
@@ -118,7 +118,7 @@ namespace EpicGames.Horde.Storage.Bundles
 					{
 						int maxSize = ZstdCompressor.GetCompressBound(input.Length);
 
-						ZstdCompressor compressor = new ZstdCompressor();
+						using ZstdCompressor compressor = new ZstdCompressor();
 
 						Span<byte> buffer = writer.GetSpan(maxSize);
 						if (!compressor.TryWrap(input.Span, buffer, out int encodedLength))
@@ -197,7 +197,7 @@ namespace EpicGames.Horde.Storage.Bundles
 					}
 				case BundleCompressionFormat.Zstd:
 					{
-						ZstdDecompressor decompressor = new ZstdDecompressor();
+						using ZstdDecompressor decompressor = new ZstdDecompressor();
 						if (!decompressor.TryUnwrap(input.Span, output.Span, out int bytesWritten) || bytesWritten != output.Length)
 						{
 							throw new InvalidOperationException("Unable to decompress data using Zstd");
