Horde: Remove unnecessary using directives.

#rnx

diff --git a/Engine/Source/Programs/Shared/EpicGames.Horde.Tests/CompressionTests.cs b/Engine/Source/Programs/Shared/EpicGames.Horde.Tests/CompressionTests.cs
--- a/Engine/Source/Programs/Shared/EpicGames.Horde.Tests/CompressionTests.cs
+++ b/Engine/Source/Programs/Shared/EpicGames.Horde.Tests/CompressionTests.cs
@@ -1,12 +1,8 @@
 ﻿// Copyright Epic Games, Inc. All Rights Reserved.
 
 using System;
-using System.Collections.Generic;
 using System.Diagnostics;
 using System.Linq;
-using System.Security.Cryptography;
-using System.Text;
-using System.Threading.Tasks;
 using EpicGames.Core;
 using EpicGames.Horde.Storage.Bundles;
 using Microsoft.VisualStudio.TestTools.UnitTesting;
