Horde: Static analysis fix.

#rnx

diff --git a/Engine/Source/Programs/Shared/EpicGames.Horde.Tests/CompressionTests.cs b/Engine/Source/Programs/Shared/EpicGames.Horde.Tests/CompressionTests.cs
--- a/Engine/Source/Programs/Shared/EpicGames.Horde.Tests/CompressionTests.cs
+++ b/Engine/Source/Programs/Shared/EpicGames.Horde.Tests/CompressionTests.cs
@@ -43,7 +43,7 @@ namespace EpicGames.Horde.Tests
 			RunTests(format, data);
 		}
 
-		void RunTests(BundleCompressionFormat format, byte[] input)
+		static void RunTests(BundleCompressionFormat format, byte[] input)
 		{
 			ArrayMemoryWriter writer = new ArrayMemoryWriter(10 * 1024 * 1024);
 			Stopwatch compressTimer = Stopwatch.StartNew();
