[UBA]
* Added parameter for providing crypto key to cache client... (binaries are not submitted yet so this parameter does nothing)

diff --git a/Engine/Source/Programs/Shared/EpicGames.UBA/ICacheClient.cs b/Engine/Source/Programs/Shared/EpicGames.UBA/ICacheClient.cs
--- a/Engine/Source/Programs/Shared/EpicGames.UBA/ICacheClient.cs
+++ b/Engine/Source/Programs/Shared/EpicGames.UBA/ICacheClient.cs
@@ -78,10 +78,11 @@ namespace EpicGames.UBA
 		/// </summary>
 		/// <param name="session">The session</param>
 		/// <param name="reportMissReason">Output reason for cache miss to log.</param>
+		/// <param name="crypto">Enable crypto by using a 32 character crypto string (representing a 16 byte value)</param>
 		/// <returns>The ICacheClient</returns>
-		public static ICacheClient CreateCacheClient(ISessionServer session, bool reportMissReason)
+		public static ICacheClient CreateCacheClient(ISessionServer session, bool reportMissReason, string crypto = "")
 		{
-			return new CacheClientImpl(session, reportMissReason);
+			return new CacheClientImpl(session, reportMissReason, crypto);
 		}
 	}
 }
diff --git a/Engine/Source/Programs/Shared/EpicGames.UBA/Impl/CacheClientImpl.cs b/Engine/Source/Programs/Shared/EpicGames.UBA/Impl/CacheClientImpl.cs
--- a/Engine/Source/Programs/Shared/EpicGames.UBA/Impl/CacheClientImpl.cs
+++ b/Engine/Source/Programs/Shared/EpicGames.UBA/Impl/CacheClientImpl.cs
@@ -81,7 +81,7 @@ namespace EpicGames.UBA
 		static extern void ProcessStartInfo_Destroy(IntPtr server);
 
 		[DllImport("UbaHost", CharSet = CharSet.Auto)]
-		static extern IntPtr CacheClient_Create(IntPtr session, bool reportMissReason);
+		static extern IntPtr CacheClient_Create(IntPtr session, bool reportMissReason, string crypto);
 
 		[DllImport("UbaHost", CharSet = CharSet.Auto)]
 		static extern bool CacheClient_Connect(IntPtr cacheClient, string host, int port);
@@ -122,10 +122,10 @@ namespace EpicGames.UBA
 		}
 		#endregion
 
-		public CacheClientImpl(ISessionServer server, bool reportMissReason)
+		public CacheClientImpl(ISessionServer server, bool reportMissReason, string crypto)
 		{
 			_sessionServer = server;
-			_handle = CacheClient_Create(_sessionServer.GetHandle(), reportMissReason);
+			_handle = CacheClient_Create(_sessionServer.GetHandle(), reportMissReason, crypto);
 		}
 
 		public bool Connect(string host, int port)
