Horde: Add deployment state and progress to tool summary
#rnx

diff --git a/Engine/Source/Programs/Horde/Horde.Server/Tools/ToolsController.cs b/Engine/Source/Programs/Horde/Horde.Server/Tools/ToolsController.cs
--- a/Engine/Source/Programs/Horde/Horde.Server/Tools/ToolsController.cs
+++ b/Engine/Source/Programs/Horde/Horde.Server/Tools/ToolsController.cs
@@ -210,7 +210,7 @@ namespace Horde.Server.Tools
 		static GetToolSummaryResponse CreateGetToolSummaryResponse(ITool tool)
 		{
 			IToolDeployment? deployment = (tool.Deployments.Count == 0) ? null : tool.Deployments[^1];
-			return new GetToolSummaryResponse(tool.Id, tool.Name, tool.Description, tool.Category, tool.Group, tool.Platforms?.ToList(), deployment?.Version, deployment?.Id, tool.ShowInUgs, tool.ShowInDashboard);
+			return new GetToolSummaryResponse(tool.Id, tool.Name, tool.Description, tool.Category, tool.Group, tool.Platforms?.ToList(), deployment?.Version, deployment?.Id, deployment?.State, deployment?.Progress, tool.ShowInUgs, tool.ShowInDashboard);
 		}
 
 		/// <summary>
diff --git a/Engine/Source/Programs/Shared/EpicGames.Horde/Tools/ToolMessages.cs b/Engine/Source/Programs/Shared/EpicGames.Horde/Tools/ToolMessages.cs
--- a/Engine/Source/Programs/Shared/EpicGames.Horde/Tools/ToolMessages.cs
+++ b/Engine/Source/Programs/Shared/EpicGames.Horde/Tools/ToolMessages.cs
@@ -33,9 +33,11 @@ namespace EpicGames.Horde.Tools
 	/// <param name="Platforms">List of platforms that this tool supports, as NET runtime identifiers.</param>
 	/// <param name="Version">Version number of the current deployment of this tool</param>
 	/// <param name="DeploymentId">Identifier for the current deployment</param>
+	/// <param name="DeploymentState">Current state of the deployment</param>
+	/// <param name="DeploymentProgress">Current progress of the deployment</param>
 	/// <param name="ShowInUgs">Whether to show this tool for download inside UGS</param>
 	/// <param name="ShowInDashboard">Whether to show this tool for download on the dashboard</param>
-	public record class GetToolSummaryResponse(ToolId Id, string Name, string Description, string? Category, string? Group, List<string>? Platforms, string? Version, ToolDeploymentId? DeploymentId, bool ShowInUgs, bool ShowInDashboard);
+	public record class GetToolSummaryResponse(ToolId Id, string Name, string Description, string? Category, string? Group, List<string>? Platforms, string? Version, ToolDeploymentId? DeploymentId, ToolDeploymentState? DeploymentState, double? DeploymentProgress, bool ShowInUgs, bool ShowInDashboard);
 
 	/// <summary>
 	/// Response when querying all tools
