Fix a few issues when compiling with TargetFramework net8.0

#rnx

diff --git a/Engine/Source/Programs/Shared/EpicGames.Horde/Compute/Transports/AesTransport.cs b/Engine/Source/Programs/Shared/EpicGames.Horde/Compute/Transports/AesTransport.cs
--- a/Engine/Source/Programs/Shared/EpicGames.Horde/Compute/Transports/AesTransport.cs
+++ b/Engine/Source/Programs/Shared/EpicGames.Horde/Compute/Transports/AesTransport.cs
@@ -72,7 +72,9 @@ namespace EpicGames.Horde.Compute.Transports
 
 			_inner = inner;
 			_readPipe = new Pipe();
+#pragma warning disable SYSLIB0053 // Type or member is obsolete
 			_aesGcm = new AesGcm(key);
+#pragma warning restore SYSLIB0053 // Type or member is obsolete
 			_readNonce = nonce.Slice(0, NonceLength).ToArray();
 			_writeNonce = nonce.Slice(0, NonceLength).ToArray();
 			_leaveInnerOpen = leaveInnerOpen;
