Update c# metadata to include AssemblyFileVersion & bump to 5.5

#rnx

diff --git a/Engine/Source/Programs/Shared/MetaData.cs b/Engine/Source/Programs/Shared/MetaData.cs
--- a/Engine/Source/Programs/Shared/MetaData.cs
+++ b/Engine/Source/Programs/Shared/MetaData.cs
@@ -14,11 +14,12 @@ using System.Runtime.InteropServices;
 [assembly: AssemblyCulture("")]
 
 [assembly: ComVisible(false)]
-[assembly: NeutralResourcesLanguageAttribute("en-US")]
+[assembly: NeutralResourcesLanguage("en-US")]
 
 #if !SPECIFIC_VERSION
 // Automatically generate a version number based on the time of compilation
-[assembly: AssemblyVersion("5.0.0.0")]
-[assembly: AssemblyInformationalVersion("5.0.0")]
+[assembly: AssemblyVersion("5.5.0.0")]
+[assembly: AssemblyFileVersion("5.5.0.0")]
+[assembly: AssemblyInformationalVersion("5.5.0.0")]
 #endif
 
