UnrealBuildTool: Support parsing import projects from .csproj

#rnx

diff --git a/Engine/Source/Programs/Shared/EpicGames.Core/CsProjectInfo.cs b/Engine/Source/Programs/Shared/EpicGames.Core/CsProjectInfo.cs
--- a/Engine/Source/Programs/Shared/EpicGames.Core/CsProjectInfo.cs
+++ b/Engine/Source/Programs/Shared/EpicGames.Core/CsProjectInfo.cs
@@ -6,6 +6,7 @@ using System.Diagnostics.CodeAnalysis;
 using System.IO;
 using System.Linq;
 using System.Runtime.InteropServices;
+using System.Text.RegularExpressions;
 using System.Xml;
 
 namespace EpicGames.Core
@@ -64,6 +65,8 @@ namespace EpicGames.Core
 		{
 			ProjectPath = inProjectPath;
 			Properties = new Dictionary<string, string>(inProperties);
+			Properties.TryAdd("MSBuildProjectFile", inProjectPath.FullName);
+			Properties.TryAdd("MSBuildThisFileDirectory", inProjectPath.Directory.FullName);
 		}
 
 		/// <summary>
@@ -336,6 +339,12 @@ namespace EpicGames.Core
 			{
 				switch (element.Name)
 				{
+					case "Import":
+						if (EvaluateCondition(element, projectInfo))
+						{
+							ParseImportProject(element, projectInfo);
+						}
+						break;
 					case "PropertyGroup":
 						if (EvaluateCondition(element, projectInfo))
 						{
@@ -414,6 +423,46 @@ namespace EpicGames.Core
 			return true;
 		}
 
+		/// <summary>
+		/// Parses a 'Import' emenebt
+		/// </summary>
+		/// <param name="element">The element</param>
+		/// <param name="projectInfo">Dictionary mapping property names to values</param>
+		static void ParseImportProject(XmlElement element, CsProjectInfo projectInfo)
+		{
+			string importProjectStr = element.GetAttribute("Project");
+			foreach (Match match in Regex.Matches(importProjectStr, @"(\$\(.*?\))").OfType<Match>())
+			{
+				importProjectStr = importProjectStr.Replace(match.Value, projectInfo.Properties.GetValueOrDefault(match.Value.Substring(2, match.Length - 3), String.Empty), StringComparison.Ordinal);
+			}
+			FileReference importProject = Path.IsPathFullyQualified(importProjectStr) ? new FileReference(importProjectStr) : FileReference.Combine(projectInfo.ProjectPath.Directory, importProjectStr);
+			if (!FileReference.Exists(importProject))
+			{
+				return;
+			}
+			XmlDocument document = new XmlDocument();
+			document.Load(importProject.FullName);
+			if (document.DocumentElement?.Name == "Project")
+			{
+				projectInfo.Properties.TryGetValue("MSBuildThisFileDirectory", out string? msBuildThisFileDirectory);
+				projectInfo.Properties.Remove("MSBuildThisFileDirectory");
+				projectInfo.Properties.Add("MSBuildThisFileDirectory", importProject.Directory.FullName);
+				try
+				{
+					ParseNode(document.DocumentElement, projectInfo);
+				}
+				catch (Exception)
+				{
+					// Ignore issues in imported .props files
+				}
+				projectInfo.Properties.Remove("MSBuildThisFileDirectory");
+				if (msBuildThisFileDirectory != null)
+				{
+					projectInfo.Properties.Add("MSBuildThisFileDirectory", msBuildThisFileDirectory);
+				}
+			}
+		}
+
 		/// <summary>
 		/// Parses a 'PropertyGroup' element.
 		/// </summary>
@@ -424,6 +473,28 @@ namespace EpicGames.Core
 			// We need to know the overridden output type and output path for the selected configuration.
 			foreach (XmlElement element in parentElement.ChildNodes.OfType<XmlElement>())
 			{
+				// Common properties used by UnrealEngine.csproj.props, manually handle because the parsing is awful
+				switch (element.Name)
+				{
+					case "IsLinux": projectInfo.Properties[element.Name] = OperatingSystem.IsLinux().ToString(); continue;
+					case "IsMacOS": projectInfo.Properties[element.Name] = OperatingSystem.IsMacOS().ToString(); continue;
+					case "IsWindows": projectInfo.Properties[element.Name] = OperatingSystem.IsWindows().ToString(); continue;
+					case "OSArchitecture": projectInfo.Properties[element.Name] = RuntimeInformation.OSArchitecture.ToString(); continue;
+					case "EngineDirectory": projectInfo.Properties.TryAdd(element.Name, String.Empty); continue;
+					case "IsEngineProject":
+					{
+						projectInfo.Properties[element.Name] = "False";
+						if (projectInfo.Properties.TryGetValue("EngineDirectory", out string? engineDirectory) && !String.IsNullOrEmpty(engineDirectory))
+						{
+							projectInfo.Properties[element.Name] = projectInfo.ProjectPath.IsUnderDirectory(new DirectoryReference(engineDirectory)).ToString();
+						}
+						continue;
+					}
+					case "IsAutomationProject": projectInfo.Properties[element.Name] = projectInfo.ProjectPath.FullName.EndsWith(".automation.csproj", StringComparison.OrdinalIgnoreCase).ToString(); continue;
+					case "HasAssemblyInfo": projectInfo.Properties[element.Name] = FileReference.Exists(FileReference.Combine(projectInfo.ProjectPath.Directory, "Properties", "AssemblyInfo.cs")).ToString(); continue;
+					default: break;
+				}
+
 				if (EvaluateCondition(element, projectInfo))
 				{
 					projectInfo.Properties[element.Name] = ExpandProperties(element.InnerText, projectInfo.Properties);
@@ -655,7 +726,7 @@ namespace EpicGames.Core
 				string? copyTo = GetChildElementString(parentElement, "CopyToOutputDirectory", null);
 				bool shouldCopy = !String.IsNullOrEmpty(copyTo) && (copyTo.Equals("Always", StringComparison.OrdinalIgnoreCase) || copyTo.Equals("PreserveNewest", StringComparison.OrdinalIgnoreCase));
 				FileReference contentFile = FileReference.Combine(baseDirectory, includePath);
-				contents.Add(contentFile, shouldCopy);
+				contents.TryAdd(contentFile, shouldCopy);
 			}
 		}
 
diff --git a/Engine/Source/Programs/UnrealBuildTool/ProjectFiles/ProjectFileGenerator.cs b/Engine/Source/Programs/UnrealBuildTool/ProjectFiles/ProjectFileGenerator.cs
--- a/Engine/Source/Programs/UnrealBuildTool/ProjectFiles/ProjectFileGenerator.cs
+++ b/Engine/Source/Programs/UnrealBuildTool/ProjectFiles/ProjectFileGenerator.cs
@@ -3456,22 +3456,32 @@ namespace UnrealBuildTool
 
 				try
 				{
-
+					// Parse the project and check if UnrealEngine.csproj.props is imported
 					bool configsFound = false;
-					// Parse the project and ensure both Development and Debug configurations are present
-					foreach (string Config in XElement.Load(InProject.ProjectFilePath.FullName).Elements("{http://schemas.microsoft.com/developer/msbuild/2003}PropertyGroup")
-										   .Where(node => node.Attribute("Condition") != null)
-										   .Select(node => node.Attribute("Condition")!.ToString()))
+					if (InProject is VCSharpProjectFile csProject && csProject.IsDotNETCoreProject())
 					{
 						configsFound = true;
+						bFoundDevelopmentConfig = csProject.Configurations.Contains("Development");
+						bFoundDebugConfig = csProject.Configurations.Contains("Debug");
+					}
 
-						if (Config.Contains("Development|"))
-						{
-							bFoundDevelopmentConfig = true;
-						}
-						else if (Config.Contains("Debug|"))
+					// Parse the project and ensure both Development and Debug configurations are present
+					if (!configsFound)
+					{
+						foreach (string Config in XElement.Load(InProject.ProjectFilePath.FullName).Elements("{http://schemas.microsoft.com/developer/msbuild/2003}PropertyGroup")
+										   .Where(node => node.Attribute("Condition") != null)
+										   .Select(node => node.Attribute("Condition")!.ToString()))
 						{
-							bFoundDebugConfig = true;
+							configsFound = true;
+
+							if (Config.Contains("Development|"))
+							{
+								bFoundDevelopmentConfig = true;
+							}
+							else if (Config.Contains("Debug|"))
+							{
+								bFoundDebugConfig = true;
+							}
 						}
 					}
 
diff --git a/Engine/Source/Programs/UnrealBuildTool/ProjectFiles/VisualStudio/VCProject.cs b/Engine/Source/Programs/UnrealBuildTool/ProjectFiles/VisualStudio/VCProject.cs
--- a/Engine/Source/Programs/UnrealBuildTool/ProjectFiles/VisualStudio/VCProject.cs
+++ b/Engine/Source/Programs/UnrealBuildTool/ProjectFiles/VisualStudio/VCProject.cs
@@ -2584,38 +2584,46 @@ namespace UnrealBuildTool
 				}
 				else
 				{
-					bool ConfigurationsFound = false;
-					foreach (XmlElement PropertyGroup in Document.DocumentElement.ChildNodes.OfType<XmlElement>()
-						.Where(element => element.Name == "PropertyGroup"))
+					foreach (string c in GetProjectProperty("Configurations").Split(';'))
 					{
-						XmlNodeList ConfigNodeList = PropertyGroup.GetElementsByTagName("Configurations");
-						// if this property group does not set configurations we do not care about it
-						if (ConfigNodeList.Count == 0)
-						{
-							continue;
-						}
+						Configurations.Add(c);
+					}
+					bool ConfigurationsFound = Configurations.Any();
 
-						if (PropertyGroup.HasAttribute("Condition"))
-						{
-							string Condition = PropertyGroup.GetAttribute("Condition");
-							Logger.LogWarning("Unable to parse configuration from property group with condition '{InitFilePath}': {Condition}. UBT Requires you to set the configuration without conditionals.", InitFilePath, Condition);
-							continue;
-						}
-						string[]? ParsedConfigurations = ConfigNodeList[0]?.FirstChild?.Value?.Split(';');
-						if (ParsedConfigurations != null)
+					if (!ConfigurationsFound)
+					{
+						foreach (XmlElement PropertyGroup in Document.DocumentElement.ChildNodes.OfType<XmlElement>()
+							.Where(element => element.Name == "PropertyGroup"))
 						{
-							foreach (string c in ParsedConfigurations)
+							XmlNodeList ConfigNodeList = PropertyGroup.GetElementsByTagName("Configurations");
+							// if this property group does not set configurations we do not care about it
+							if (ConfigNodeList.Count == 0)
 							{
-								Configurations.Add(c);
+								continue;
 							}
-						}
 
-						// platforms change meaning quite a bit in .net core but typically you do not specify this and its derived from the build instead
-						// for most intents it is just Any CPU from .net framework
-						Platforms.Add("AnyCPU");
+							if (PropertyGroup.HasAttribute("Condition"))
+							{
+								string Condition = PropertyGroup.GetAttribute("Condition");
+								Logger.LogWarning("Unable to parse configuration from property group with condition '{InitFilePath}': {Condition}. UBT Requires you to set the configuration without conditionals.", InitFilePath, Condition);
+								continue;
+							}
+							string[]? ParsedConfigurations = ConfigNodeList[0]?.FirstChild?.Value?.Split(';');
+							if (ParsedConfigurations != null)
+							{
+								foreach (string c in ParsedConfigurations)
+								{
+									Configurations.Add(c);
+								}
+							}
 
-						ConfigurationsFound = true;
-						break;
+							// platforms change meaning quite a bit in .net core but typically you do not specify this and its derived from the build instead
+							// for most intents it is just Any CPU from .net framework
+							Platforms.Add("AnyCPU");
+
+							ConfigurationsFound = true;
+							break;
+						}
 					}
 
 					// dotnet does not require you to specify configurations or platforms, if you do not debug and release are the defaults
@@ -2648,6 +2656,7 @@ namespace UnrealBuildTool
 			Dictionary<string, string> Properties = new Dictionary<string, string>();
 			Properties.Add("Platform", "AnyCPU");
 			Properties.Add("Configuration", InConfiguration.ToString());
+			Properties.Add("EngineDirectory", Unreal.EngineDirectory.FullName);
 			if (CsProjectInfo.TryRead(ProjectFilePath, Properties, out Info))
 			{
 				CachedProjectInfo.Add(InConfiguration, Info);
@@ -2665,6 +2674,17 @@ namespace UnrealBuildTool
 			return Info.IsDotNETCoreProject();
 		}
 
+
+		/// <summary>
+		/// Gets a property from the project
+		/// </summary>
+		public string GetProjectProperty(string property)
+		{
+			CsProjectInfo Info = GetProjectInfo(UnrealTargetConfiguration.Debug)!;
+			Info.Properties.TryGetValue(property, out string? value);
+			return value ?? String.Empty;
+		}
+
 		/// <inheritdoc/>
 		public override MSBuildProjectContext? GetMatchingProjectContext(TargetType SolutionTarget, UnrealTargetConfiguration SolutionConfiguration, UnrealTargetPlatform SolutionPlatform, PlatformProjectGeneratorCollection PlatformProjectGenerators, UnrealArch? Architecture, ILogger Logger)
 		{
