Fix net8.0 analysis warnings in shared projects

#rnx

diff --git a/Engine/Source/Programs/Shared/EpicGames.Horde/Compute/AgentMessageHandler.cs b/Engine/Source/Programs/Shared/EpicGames.Horde/Compute/AgentMessageHandler.cs
--- a/Engine/Source/Programs/Shared/EpicGames.Horde/Compute/AgentMessageHandler.cs
+++ b/Engine/Source/Programs/Shared/EpicGames.Horde/Compute/AgentMessageHandler.cs
@@ -63,11 +63,15 @@ namespace EpicGames.Horde.Compute
 			using CancellationTokenSource cancellationSource = CancellationTokenSource.CreateLinkedTokenSource(cancellationToken);
 			ExceptionDispatchInfo? exceptionInfo = null;
 
-			void PostException(Exception ex)
+			async Task PostException(Exception ex)
 			{
 				// Capture stack from call site
 				Interlocked.CompareExchange(ref exceptionInfo, ExceptionDispatchInfo.Capture(ex), null);
-				cancellationSource.Cancel();
+#if NET8_0_OR_GREATER
+				await cancellationSource.CancelAsync();
+#else
+				await Task.Run(() => cancellationSource.Cancel());
+#endif
 			}
 
 			await RunAsync(socket, 0, 4 * 1024 * 1024, PostException, cancellationSource.Token);
@@ -81,7 +85,7 @@ namespace EpicGames.Horde.Compute
 #pragma warning restore CA1508
 		}
 
-		async Task RunAsync(ComputeSocket socket, int channelId, int bufferSize, Action<Exception> postException, CancellationToken cancellationToken)
+		async Task RunAsync(ComputeSocket socket, int channelId, int bufferSize, Func<Exception, Task> postException, CancellationToken cancellationToken)
 		{
 			List<Task> childTasks = new List<Task>();
 			using AgentMessageChannel channel = socket.CreateAgentMessageChannel(channelId, bufferSize);
@@ -158,7 +162,7 @@ namespace EpicGames.Horde.Compute
 			{
 				_logger.LogInformation(ex, "Compute Channel {ChannelId}: Exception: {Message}", channelId, ex.Message);
 				await channel.SendExceptionAsync(ex, cancellationToken);
-				postException(ex);
+				await postException(ex);
 			}
 			finally
 			{
diff --git a/Engine/Source/Programs/Shared/EpicGames.Horde/Compute/Transports/IdleTimeoutTransport.cs b/Engine/Source/Programs/Shared/EpicGames.Horde/Compute/Transports/IdleTimeoutTransport.cs
--- a/Engine/Source/Programs/Shared/EpicGames.Horde/Compute/Transports/IdleTimeoutTransport.cs
+++ b/Engine/Source/Programs/Shared/EpicGames.Horde/Compute/Transports/IdleTimeoutTransport.cs
@@ -83,7 +83,11 @@ public class IdleTimeoutTransport : ComputeTransport
 			if (remainingTime < TimeSpan.Zero)
 			{
 				logger.LogWarning("Terminating compute transport due to timeout (last tick at {Time})", DateTime.UtcNow - TimeSinceActivity);
+#if NET8_0_OR_GREATER
+				await cts.CancelAsync();
+#else
 				cts.Cancel();
+#endif
 				break;
 			}
 			await Task.Delay(remainingTime + TimeSpan.FromSeconds(0.2), cancellationToken);
diff --git a/Engine/Source/Programs/Shared/EpicGames.Horde/HordeHttpAuthHandler.cs b/Engine/Source/Programs/Shared/EpicGames.Horde/HordeHttpAuthHandler.cs
--- a/Engine/Source/Programs/Shared/EpicGames.Horde/HordeHttpAuthHandler.cs
+++ b/Engine/Source/Programs/Shared/EpicGames.Horde/HordeHttpAuthHandler.cs
@@ -133,7 +133,11 @@ namespace EpicGames.Horde
 		{
 			if (_authStateTask != null && !_authStateTask.IsCompleted)
 			{
+#if NET8_0_OR_GREATER
+				await _cancellationTokenSource.CancelAsync();
+#else
 				_cancellationTokenSource.Cancel();
+#endif
 				try
 				{
 #pragma warning disable VSTHRD003
diff --git a/Engine/Source/Programs/Shared/EpicGames.Horde/Storage/Nodes/DirectoryNode.Extract.cs b/Engine/Source/Programs/Shared/EpicGames.Horde/Storage/Nodes/DirectoryNode.Extract.cs
--- a/Engine/Source/Programs/Shared/EpicGames.Horde/Storage/Nodes/DirectoryNode.Extract.cs
+++ b/Engine/Source/Programs/Shared/EpicGames.Horde/Storage/Nodes/DirectoryNode.Extract.cs
@@ -318,7 +318,11 @@ namespace EpicGames.Horde.Storage.Nodes
 					catch (Exception ex)
 					{
 						logger.LogError(ex, "Error while extracting data: {Message}", ex.Message);
+#if NET8_0_OR_GREATER
+						await cancellationSource.CancelAsync();
+#else
 						cancellationSource.Cancel();
+#endif
 					}
 				}
 
diff --git a/Engine/Source/Programs/Shared/EpicGames.OIDC/WindowsTokenStore.cs b/Engine/Source/Programs/Shared/EpicGames.OIDC/WindowsTokenStore.cs
--- a/Engine/Source/Programs/Shared/EpicGames.OIDC/WindowsTokenStore.cs
+++ b/Engine/Source/Programs/Shared/EpicGames.OIDC/WindowsTokenStore.cs
@@ -4,7 +4,6 @@ using System;
 using System.Collections.Generic;
 using System.ComponentModel;
 using System.IO;
-using System.Linq;
 using System.Runtime.InteropServices;
 using System.Text;
 using System.Text.Json;
diff --git a/Engine/Source/Programs/Shared/EpicGames.UBA/Impl/SessionServerImpl.cs b/Engine/Source/Programs/Shared/EpicGames.UBA/Impl/SessionServerImpl.cs
--- a/Engine/Source/Programs/Shared/EpicGames.UBA/Impl/SessionServerImpl.cs
+++ b/Engine/Source/Programs/Shared/EpicGames.UBA/Impl/SessionServerImpl.cs
@@ -133,7 +133,9 @@ namespace EpicGames.UBA
 					using (BinaryWriter writer = new(environmentMemory, System.Text.Encoding.UTF8, true))
 					{
 						foreach (DictionaryEntry de in Environment.GetEnvironmentVariables())
+						{
 							writer.Write($"{de.Key}={de.Value}");
+						}
 					}
 				}
 			}
