Fixed build errors in Horde and CloudDDC test assemblies due to duplicate assembly file version attributes

diff --git a/Engine/Source/Programs/Shared/EpicGames.BuildGraph.Tests/EpicGames.BuildGraph.Tests.csproj b/Engine/Source/Programs/Shared/EpicGames.BuildGraph.Tests/EpicGames.BuildGraph.Tests.csproj
--- a/Engine/Source/Programs/Shared/EpicGames.BuildGraph.Tests/EpicGames.BuildGraph.Tests.csproj
+++ b/Engine/Source/Programs/Shared/EpicGames.BuildGraph.Tests/EpicGames.BuildGraph.Tests.csproj
@@ -11,6 +11,7 @@
     <GenerateAssemblyInformationalVersionAttribute>false</GenerateAssemblyInformationalVersionAttribute>
     <GenerateAssemblyProductAttribute>false</GenerateAssemblyProductAttribute>
     <GenerateAssemblyVersionAttribute>false</GenerateAssemblyVersionAttribute>
+    <GenerateAssemblyFileVersionAttribute>false</GenerateAssemblyFileVersionAttribute>
     <Configurations>Debug;Release;Development;Analyze</Configurations>
     <IsPackable>false</IsPackable>
     <Nullable>enable</Nullable>
diff --git a/Engine/Source/Programs/Shared/EpicGames.Core.Tests/EpicGames.Core.Tests.csproj b/Engine/Source/Programs/Shared/EpicGames.Core.Tests/EpicGames.Core.Tests.csproj
--- a/Engine/Source/Programs/Shared/EpicGames.Core.Tests/EpicGames.Core.Tests.csproj
+++ b/Engine/Source/Programs/Shared/EpicGames.Core.Tests/EpicGames.Core.Tests.csproj
@@ -11,6 +11,7 @@
     <GenerateAssemblyInformationalVersionAttribute>false</GenerateAssemblyInformationalVersionAttribute>
     <GenerateAssemblyProductAttribute>false</GenerateAssemblyProductAttribute>
     <GenerateAssemblyVersionAttribute>false</GenerateAssemblyVersionAttribute>
+    <GenerateAssemblyFileVersionAttribute>false</GenerateAssemblyFileVersionAttribute>
     <Configurations>Debug;Release;Development;Analyze</Configurations>
     <IsPackable>false</IsPackable>
     <Nullable>enable</Nullable>
diff --git a/Engine/Source/Programs/Shared/EpicGames.Perforce.Fixture/EpicGames.Perforce.Fixture.csproj b/Engine/Source/Programs/Shared/EpicGames.Perforce.Fixture/EpicGames.Perforce.Fixture.csproj
--- a/Engine/Source/Programs/Shared/EpicGames.Perforce.Fixture/EpicGames.Perforce.Fixture.csproj
+++ b/Engine/Source/Programs/Shared/EpicGames.Perforce.Fixture/EpicGames.Perforce.Fixture.csproj
@@ -10,6 +10,7 @@
     <GenerateAssemblyInformationalVersionAttribute>false</GenerateAssemblyInformationalVersionAttribute>
     <GenerateAssemblyProductAttribute>false</GenerateAssemblyProductAttribute>
     <GenerateAssemblyVersionAttribute>false</GenerateAssemblyVersionAttribute>
+    <GenerateAssemblyFileVersionAttribute>false</GenerateAssemblyFileVersionAttribute>
     <Configurations>Debug;Release;Development;Analyze</Configurations>
     <IsPackable>false</IsPackable>
     <Nullable>enable</Nullable>
diff --git a/Engine/Source/Programs/Shared/EpicGames.Perforce.Managed.Tests/EpicGames.Perforce.Managed.Tests.csproj b/Engine/Source/Programs/Shared/EpicGames.Perforce.Managed.Tests/EpicGames.Perforce.Managed.Tests.csproj
--- a/Engine/Source/Programs/Shared/EpicGames.Perforce.Managed.Tests/EpicGames.Perforce.Managed.Tests.csproj
+++ b/Engine/Source/Programs/Shared/EpicGames.Perforce.Managed.Tests/EpicGames.Perforce.Managed.Tests.csproj
@@ -10,6 +10,7 @@
     <GenerateAssemblyInformationalVersionAttribute>false</GenerateAssemblyInformationalVersionAttribute>
     <GenerateAssemblyProductAttribute>false</GenerateAssemblyProductAttribute>
     <GenerateAssemblyVersionAttribute>false</GenerateAssemblyVersionAttribute>
+    <GenerateAssemblyFileVersionAttribute>false</GenerateAssemblyFileVersionAttribute>
     <Configurations>Debug;Release;Development</Configurations>
     <IsPackable>false</IsPackable>
     <Nullable>enable</Nullable>
diff --git a/Engine/Source/Programs/Shared/EpicGames.Perforce.Tests/EpicGames.Perforce.Tests.csproj b/Engine/Source/Programs/Shared/EpicGames.Perforce.Tests/EpicGames.Perforce.Tests.csproj
--- a/Engine/Source/Programs/Shared/EpicGames.Perforce.Tests/EpicGames.Perforce.Tests.csproj
+++ b/Engine/Source/Programs/Shared/EpicGames.Perforce.Tests/EpicGames.Perforce.Tests.csproj
@@ -10,6 +10,7 @@
     <GenerateAssemblyInformationalVersionAttribute>false</GenerateAssemblyInformationalVersionAttribute>
     <GenerateAssemblyProductAttribute>false</GenerateAssemblyProductAttribute>
     <GenerateAssemblyVersionAttribute>false</GenerateAssemblyVersionAttribute>
+    <GenerateAssemblyFileVersionAttribute>false</GenerateAssemblyFileVersionAttribute>
     <Configurations>Debug;Release;Development;Analyze</Configurations>
     <IsPackable>false</IsPackable>
     <Nullable>enable</Nullable>
diff --git a/Engine/Source/Programs/Shared/EpicGames.Redis.Tests/EpicGames.Redis.Tests.csproj b/Engine/Source/Programs/Shared/EpicGames.Redis.Tests/EpicGames.Redis.Tests.csproj
--- a/Engine/Source/Programs/Shared/EpicGames.Redis.Tests/EpicGames.Redis.Tests.csproj
+++ b/Engine/Source/Programs/Shared/EpicGames.Redis.Tests/EpicGames.Redis.Tests.csproj
@@ -11,6 +11,7 @@
     <GenerateAssemblyInformationalVersionAttribute>false</GenerateAssemblyInformationalVersionAttribute>
     <GenerateAssemblyProductAttribute>false</GenerateAssemblyProductAttribute>
     <GenerateAssemblyVersionAttribute>false</GenerateAssemblyVersionAttribute>
+    <GenerateAssemblyFileVersionAttribute>false</GenerateAssemblyFileVersionAttribute>
     <Configurations>Debug;Release;Development;Analyze</Configurations>
     <IsPackable>false</IsPackable>
     <Nullable>enable</Nullable>
diff --git a/Engine/Source/Programs/Shared/EpicGames.Serialization.Tests/EpicGames.Serialization.Tests.csproj b/Engine/Source/Programs/Shared/EpicGames.Serialization.Tests/EpicGames.Serialization.Tests.csproj
--- a/Engine/Source/Programs/Shared/EpicGames.Serialization.Tests/EpicGames.Serialization.Tests.csproj
+++ b/Engine/Source/Programs/Shared/EpicGames.Serialization.Tests/EpicGames.Serialization.Tests.csproj
@@ -11,6 +11,7 @@
     <GenerateAssemblyInformationalVersionAttribute>false</GenerateAssemblyInformationalVersionAttribute>
     <GenerateAssemblyProductAttribute>false</GenerateAssemblyProductAttribute>
     <GenerateAssemblyVersionAttribute>false</GenerateAssemblyVersionAttribute>
+    <GenerateAssemblyFileVersionAttribute>false</GenerateAssemblyFileVersionAttribute>
     <Configurations>Debug;Release;Development;Analyze</Configurations>
     <IsPackable>false</IsPackable>
     <Nullable>enable</Nullable>
diff --git a/Engine/Source/Programs/Shared/EpicGames.Tracing.Tests/EpicGames.Tracing.Tests.csproj b/Engine/Source/Programs/Shared/EpicGames.Tracing.Tests/EpicGames.Tracing.Tests.csproj
--- a/Engine/Source/Programs/Shared/EpicGames.Tracing.Tests/EpicGames.Tracing.Tests.csproj
+++ b/Engine/Source/Programs/Shared/EpicGames.Tracing.Tests/EpicGames.Tracing.Tests.csproj
@@ -10,6 +10,7 @@
     <GenerateAssemblyInformationalVersionAttribute>false</GenerateAssemblyInformationalVersionAttribute>
     <GenerateAssemblyProductAttribute>false</GenerateAssemblyProductAttribute>
     <GenerateAssemblyVersionAttribute>false</GenerateAssemblyVersionAttribute>
+    <GenerateAssemblyFileVersionAttribute>false</GenerateAssemblyFileVersionAttribute>
     <Configurations>Debug;Release;Development;Analyze</Configurations>
     <IsPackable>false</IsPackable>
     <Nullable>enable</Nullable>
