Fix static analysis warning

#rnx

diff --git a/Engine/Source/Programs/Shared/EpicGames.Horde/Compute/AgentMessageHandler.cs b/Engine/Source/Programs/Shared/EpicGames.Horde/Compute/AgentMessageHandler.cs
--- a/Engine/Source/Programs/Shared/EpicGames.Horde/Compute/AgentMessageHandler.cs
+++ b/Engine/Source/Programs/Shared/EpicGames.Horde/Compute/AgentMessageHandler.cs
@@ -70,7 +70,7 @@ namespace EpicGames.Horde.Compute
 #if NET8_0_OR_GREATER
 				await cancellationSource.CancelAsync();
 #else
-				await Task.Run(() => cancellationSource.Cancel());
+				await Task.Run(() => cancellationSource.Cancel(), CancellationToken.None);
 #endif
 			}
 
