Horde: Fix hang when reading from blob pipeline.

#rnx

diff --git a/Engine/Source/Programs/Shared/EpicGames.Horde.Tests/BlobPipelineTests.cs b/Engine/Source/Programs/Shared/EpicGames.Horde.Tests/BlobPipelineTests.cs
new file mode 100644
--- /dev/null
+++ b/Engine/Source/Programs/Shared/EpicGames.Horde.Tests/BlobPipelineTests.cs
@@ -0,0 +1,49 @@
+// Copyright Epic Games, Inc. All Rights Reserved.
+
+using System;
+using System.Collections.Generic;
+using System.Linq;
+using System.Text;
+using System.Threading.Tasks;
+using EpicGames.Core;
+using EpicGames.Horde.Storage;
+using EpicGames.Horde.Storage.Bundles;
+using Microsoft.Extensions.Logging.Abstractions;
+using Microsoft.VisualStudio.TestTools.UnitTesting;
+
+namespace EpicGames.Horde.Tests
+{
+	[TestClass]
+	public class BlobPipelineTests
+	{
+		static BlobType s_blobType = new BlobType(Guid.Empty, 1);
+
+		[TestMethod]
+		public async Task TestAsync()
+		{
+			using BundleStorageClient client = BundleStorageClient.CreateInMemory(NullLogger.Instance);
+
+			List<IBlobRef> handles = new List<IBlobRef>();
+			await using (IBlobWriter writer = client.CreateBlobWriter())
+			{
+				for (int idx = 0; idx < 4000; idx++)
+				{
+					writer.WriteInt32(idx);
+					handles.Add(await writer.CompleteAsync(s_blobType));
+				}
+			}
+
+			await using (BlobPipeline<IoHash> pipeline = new BlobPipeline<IoHash>())
+			{
+				foreach (IBlobRef blobRef in handles)
+				{
+					pipeline.Add(new BlobRequest<IoHash>(blobRef, blobRef.Hash));
+				}
+				pipeline.FinishAdding();
+
+				List<BlobResponse<IoHash>> responses = await pipeline.ReadAllAsync().ToListAsync();
+				Assert.AreEqual(handles.Count, responses.Count);
+			}
+		}
+	}
+}
diff --git a/Engine/Source/Programs/Shared/EpicGames.Horde/Storage/BlobPipeline.cs b/Engine/Source/Programs/Shared/EpicGames.Horde/Storage/BlobPipeline.cs
--- a/Engine/Source/Programs/Shared/EpicGames.Horde/Storage/BlobPipeline.cs
+++ b/Engine/Source/Programs/Shared/EpicGames.Horde/Storage/BlobPipeline.cs
@@ -99,10 +99,7 @@ namespace EpicGames.Horde.Storage
 			_responseBatchChannel = Channel.CreateBounded<BlobResponseBatch<TUserData>>(new BoundedChannelOptions(options.ResponseBufferSize) { FullMode = BoundedChannelFullMode.Wait });
 
 			_tasks.Add(Task.Run(() => BatchLoopAsync(_cancellationTokenSource.Token), _cancellationTokenSource.Token));
-			for (int idx = 0; idx < options.NumFetchTasks; idx++)
-			{
-				_tasks.Add(Task.Run(() => FetchLoopAsync(_cancellationTokenSource.Token), _cancellationTokenSource.Token));
-			}
+			_tasks.Add(Task.Run(() => FetchLoopAsync(_cancellationTokenSource.Token), _cancellationTokenSource.Token));
 		}
 
 		/// <inheritdoc/>
@@ -185,6 +182,10 @@ namespace EpicGames.Horde.Storage
 				}
 				if (Interlocked.CompareExchange(ref _writerCount, writerCount + delta, writerCount) == writerCount)
 				{
+					if (writerCount + delta == 0)
+					{
+						_requestChannel.Writer.TryComplete();
+					}
 					break;
 				}
 			}
@@ -274,6 +275,25 @@ namespace EpicGames.Horde.Storage
 		}
 
 		async Task FetchLoopAsync(CancellationToken cancellationToken)
+		{
+			try
+			{
+				List<Task> tasks = new List<Task>();
+				for (int idx = 0; idx < _options.NumFetchTasks; idx++)
+				{
+					tasks.Add(Task.Run(() => FetchWorkerAsync(cancellationToken), cancellationToken));
+				}
+				await Task.WhenAll(tasks);
+
+				_responseBatchChannel.Writer.TryComplete();
+			}
+			catch(Exception ex)
+			{
+				_responseBatchChannel.Writer.TryComplete(ex);
+			}
+		}
+
+		async Task FetchWorkerAsync(CancellationToken cancellationToken)
 		{
 			while (await _requestBatchChannel.Reader.WaitToReadAsync(cancellationToken))
 			{
