EpicGames.Perforce: Strip all non-numeric characters from the end of a datetime string before trying to convert to a DateTimeOffset. Some time zones include other punctuation characters (eg. "Paris, Madrid").

diff --git a/Engine/Source/Programs/Shared/EpicGames.Perforce/PerforceReflection.cs b/Engine/Source/Programs/Shared/EpicGames.Perforce/PerforceReflection.cs
--- a/Engine/Source/Programs/Shared/EpicGames.Perforce/PerforceReflection.cs
+++ b/Engine/Source/Programs/Shared/EpicGames.Perforce/PerforceReflection.cs
@@ -670,7 +670,7 @@ namespace EpicGames.Perforce
 		static object ParseStringAsNullableDateTimeOffset(Utf8String str)
 		{
 			string text = str.ToString();
-			return DateTimeOffset.Parse(Regex.Replace(text, "[a-zA-Z. ]*$", "")); // Strip timezone name (eg. "EST")
+			return DateTimeOffset.Parse(Regex.Replace(text, "[^0-9]*$", "")); // Strip timezone name (eg. "EST")
 		}
 	}
 }
