EpicGames.Perforce: Do not attempt to use the native client on Linux./ARM64.

diff --git a/Engine/Source/Programs/Shared/EpicGames.Perforce/NativePerforceConnection.cs b/Engine/Source/Programs/Shared/EpicGames.Perforce/NativePerforceConnection.cs
--- a/Engine/Source/Programs/Shared/EpicGames.Perforce/NativePerforceConnection.cs
+++ b/Engine/Source/Programs/Shared/EpicGames.Perforce/NativePerforceConnection.cs
@@ -387,7 +387,11 @@ namespace EpicGames.Perforce
 			{
 				return RuntimeInformation.OSArchitecture != Architecture.Arm64;
 			}
-			return true;
+			else if (RuntimeInformation.IsOSPlatform(OSPlatform.Linux))
+			{
+				return RuntimeInformation.OSArchitecture != Architecture.Arm64;
+			}
+			return false;
 		}
 
 		void GetNextWriteBuffer(NativeWriteBuffer nativeWriteBuffer, int minSize)
