Horde: Speculative fix for intermittent test error (not flushing streams)

#rnx

diff --git a/Engine/Source/Programs/Shared/EpicGames.Horde.Tests/TransportTests.cs b/Engine/Source/Programs/Shared/EpicGames.Horde.Tests/TransportTests.cs
--- a/Engine/Source/Programs/Shared/EpicGames.Horde.Tests/TransportTests.cs
+++ b/Engine/Source/Programs/Shared/EpicGames.Horde.Tests/TransportTests.cs
@@ -60,21 +60,28 @@ namespace EpicGames.Horde.Tests
 
 			{
 				using MemoryStream memoryStream = new MemoryStream();
-				await using StreamTransport streamTransport = new StreamTransport(memoryStream);
-				await using AesTransport aesTransport = new AesTransport(streamTransport, key, nonce);
-				await aesTransport.SendAsync(input, CancellationToken.None);
+				await using (StreamTransport streamTransport = new StreamTransport(memoryStream))
+				{
+					await using (AesTransport aesTransport = new AesTransport(streamTransport, key, nonce))
+					{
+						await aesTransport.SendAsync(input, CancellationToken.None);
+					}
+				}
 				encrypted = memoryStream.ToArray();
 			}
 
 			byte[] output = new byte[input.Length];
 			{
 				using MemoryStream memoryStream = new MemoryStream(encrypted);
-				await using StreamTransport streamTransport = new StreamTransport(memoryStream);
-				await using AesTransport aesTransport = new AesTransport(streamTransport, key, nonce);
-
-				for (int offset = 0; offset < output.Length;)
+				await using (StreamTransport streamTransport = new StreamTransport(memoryStream))
 				{
-					offset += await aesTransport.RecvAsync(output.AsMemory(offset), CancellationToken.None);
+					await using (AesTransport aesTransport = new AesTransport(streamTransport, key, nonce))
+					{
+						for (int offset = 0; offset < output.Length;)
+						{
+							offset += await aesTransport.RecvAsync(output.AsMemory(offset), CancellationToken.None);
+						}
+					}
 				}
 			}
 
