Horde: Fix compile warnings

#rnx

diff --git a/Engine/Source/Programs/Shared/EpicGames.Horde.Tests/BlobPipelineTests.cs b/Engine/Source/Programs/Shared/EpicGames.Horde.Tests/BlobPipelineTests.cs
--- a/Engine/Source/Programs/Shared/EpicGames.Horde.Tests/BlobPipelineTests.cs
+++ b/Engine/Source/Programs/Shared/EpicGames.Horde.Tests/BlobPipelineTests.cs
@@ -3,7 +3,6 @@
 using System;
 using System.Collections.Generic;
 using System.Linq;
-using System.Text;
 using System.Threading.Tasks;
 using EpicGames.Core;
 using EpicGames.Horde.Storage;
@@ -16,7 +15,7 @@ namespace EpicGames.Horde.Tests
 	[TestClass]
 	public class BlobPipelineTests
 	{
-		static BlobType s_blobType = new BlobType(Guid.Empty, 1);
+		static readonly BlobType s_blobType = new BlobType(Guid.Empty, 1);
 
 		[TestMethod]
 		public async Task TestAsync()
