Horde: Add support for auto-assigning cluster ID during compute allocation requests

- Important for load and geo balancing remote execution
- Backwards compatible with current C# interface
- Treat response code 503 same as 429 as we migrate to the latter. 503 is technically correct but also interferes with monitoring as 5xx is usually associated with errors.

diff --git a/Engine/Source/Programs/Shared/EpicGames.Horde/Compute/Clients/AgentComputeClient.cs b/Engine/Source/Programs/Shared/EpicGames.Horde/Compute/Clients/AgentComputeClient.cs
--- a/Engine/Source/Programs/Shared/EpicGames.Horde/Compute/Clients/AgentComputeClient.cs
+++ b/Engine/Source/Programs/Shared/EpicGames.Horde/Compute/Clients/AgentComputeClient.cs
@@ -75,7 +75,7 @@ namespace EpicGames.Horde.Compute.Clients
 		public ValueTask DisposeAsync() => new ValueTask();
 
 		/// <inheritdoc/>
-		public async Task<IComputeLease?> TryAssignWorkerAsync(ClusterId clusterId, Requirements? requirements, string? requestId, ConnectionMetadataRequest? connection, ILogger logger, CancellationToken cancellationToken)
+		public async Task<IComputeLease?> TryAssignWorkerAsync(ClusterId? clusterId, Requirements? requirements, string? requestId, ConnectionMetadataRequest? connection, ILogger logger, CancellationToken cancellationToken)
 		{
 			logger.LogInformation("** CLIENT **");
 			logger.LogInformation("Launching {Path} to handle remote", _hordeAgentAssembly);
diff --git a/Engine/Source/Programs/Shared/EpicGames.Horde/Compute/Clients/LocalComputeClient.cs b/Engine/Source/Programs/Shared/EpicGames.Horde/Compute/Clients/LocalComputeClient.cs
--- a/Engine/Source/Programs/Shared/EpicGames.Horde/Compute/Clients/LocalComputeClient.cs
+++ b/Engine/Source/Programs/Shared/EpicGames.Horde/Compute/Clients/LocalComputeClient.cs
@@ -84,7 +84,7 @@ namespace EpicGames.Horde.Compute.Clients
 		}
 
 		/// <inheritdoc/>
-		public Task<IComputeLease?> TryAssignWorkerAsync(ClusterId clusterId, Requirements? requirements, string? requestId, ConnectionMetadataRequest? connection, ILogger logger, CancellationToken cancellationToken)
+		public Task<IComputeLease?> TryAssignWorkerAsync(ClusterId? clusterId, Requirements? requirements, string? requestId, ConnectionMetadataRequest? connection, ILogger logger, CancellationToken cancellationToken)
 		{
 #pragma warning disable CA2000 // Dispose objects before losing scope
 			RemoteComputeSocket socket = new RemoteComputeSocket(new TcpTransport(_socket), ComputeProtocol.Latest, logger);
diff --git a/Engine/Source/Programs/Shared/EpicGames.Horde/Compute/Clients/ServerComputeClient.cs b/Engine/Source/Programs/Shared/EpicGames.Horde/Compute/Clients/ServerComputeClient.cs
--- a/Engine/Source/Programs/Shared/EpicGames.Horde/Compute/Clients/ServerComputeClient.cs
+++ b/Engine/Source/Programs/Shared/EpicGames.Horde/Compute/Clients/ServerComputeClient.cs
@@ -218,7 +218,7 @@ namespace EpicGames.Horde.Compute.Clients
 		}
 
 		/// <inheritdoc/>
-		public async Task<IComputeLease?> TryAssignWorkerAsync(ClusterId clusterId, Requirements? requirements, string? requestId, ConnectionMetadataRequest? connection, ILogger logger, CancellationToken cancellationToken)
+		public async Task<IComputeLease?> TryAssignWorkerAsync(ClusterId? clusterId, Requirements? requirements, string? requestId, ConnectionMetadataRequest? connection, ILogger logger, CancellationToken cancellationToken)
 		{
 			try
 			{
@@ -245,7 +245,7 @@ namespace EpicGames.Horde.Compute.Clients
 			response.EnsureSuccessStatusCode();
 		}
 
-		async IAsyncEnumerable<LeaseInfo> ConnectAsync(ClusterId clusterId, Requirements? requirements, string? requestId, ConnectionMetadataRequest? connection, ILogger workerLogger, [EnumeratorCancellation] CancellationToken cancellationToken)
+		async IAsyncEnumerable<LeaseInfo> ConnectAsync(ClusterId? clusterId, Requirements? requirements, string? requestId, ConnectionMetadataRequest? connection, ILogger workerLogger, [EnumeratorCancellation] CancellationToken cancellationToken)
 		{
 			_logger.LogDebug("Requesting compute resource");
 
@@ -262,16 +262,17 @@ namespace EpicGames.Horde.Compute.Clients
 			}
 
 			AssignComputeResponse? response;
-			using (HttpResponseMessage httpResponse = await HordeHttpClient.PostAsync(_httpClient, $"api/v2/compute/{clusterId}", request, _cancellationSource.Token))
+			string path = clusterId == null ? "api/v2/compute" : $"api/v2/compute/{clusterId}";
+			using (HttpResponseMessage httpResponse = await HordeHttpClient.PostAsync(_httpClient, path, request, _cancellationSource.Token))
 			{
 				if (httpResponse.StatusCode == HttpStatusCode.NotFound)
 				{
-					throw new NoComputeAgentsFoundException(clusterId, requirements);
+					throw new NoComputeAgentsFoundException(clusterId ?? new ClusterId("null"), requirements);
 				}
 
-				if (httpResponse.StatusCode == HttpStatusCode.ServiceUnavailable)
+				if (httpResponse.StatusCode is HttpStatusCode.ServiceUnavailable or HttpStatusCode.TooManyRequests)
 				{
-					_logger.LogDebug("No compute resource is available.");
+					_logger.LogDebug("No compute resource is available");
 					yield break;
 				}
 
diff --git a/Engine/Source/Programs/Shared/EpicGames.Horde/Compute/IComputeClient.cs b/Engine/Source/Programs/Shared/EpicGames.Horde/Compute/IComputeClient.cs
--- a/Engine/Source/Programs/Shared/EpicGames.Horde/Compute/IComputeClient.cs
+++ b/Engine/Source/Programs/Shared/EpicGames.Horde/Compute/IComputeClient.cs
@@ -16,13 +16,13 @@ namespace EpicGames.Horde.Compute
 		/// <summary>
 		/// Adds a new remote request
 		/// </summary>
-		/// <param name="clusterId">Cluster to execute the request</param>
+		/// <param name="clusterId">Optional cluster ID. If not set, cluster will automatically be resolved by server</param>
 		/// <param name="requirements">Requirements for the agent</param>
 		/// <param name="requestId">Optional ID identifying the request over multiple calls, such as retrying the same request</param>
 		/// <param name="connection">Optional preference of connection details</param>
 		/// <param name="logger">Logger for output from this worker</param>
 		/// <param name="cancellationToken">Cancellation token for the operation</param>
-		public Task<IComputeLease?> TryAssignWorkerAsync(ClusterId clusterId, Requirements? requirements, string? requestId, ConnectionMetadataRequest? connection, ILogger logger, CancellationToken cancellationToken = default);
+		public Task<IComputeLease?> TryAssignWorkerAsync(ClusterId? clusterId, Requirements? requirements, string? requestId, ConnectionMetadataRequest? connection, ILogger logger, CancellationToken cancellationToken = default);
 
 		/// <summary>
 		/// Declare resource needs for current client
