Gauntlet Issue - Remove Summary regex match (use only property)

#rnx

diff --git a/Engine/Source/Programs/Horde/Horde.Server.Tests/Issues/IssueServiceTests.cs b/Engine/Source/Programs/Horde/Horde.Server.Tests/Issues/IssueServiceTests.cs
--- a/Engine/Source/Programs/Horde/Horde.Server.Tests/Issues/IssueServiceTests.cs
+++ b/Engine/Source/Programs/Horde/Horde.Server.Tests/Issues/IssueServiceTests.cs
@@ -2010,7 +2010,7 @@ namespace Horde.Server.Tests.Issues
 				await UpdateCompleteStepAsync(job, 0, 0, JobStepOutcome.Failure);
 				await using (TestJsonLogger logger = await CreateLoggerAsync(job, 0, 1))
 				{
-					logger.LogError(KnownLogEvents.Gauntlet_FatalEvent, "Engine encounter a critical failure\n{Message}\n{Callstack}", logMessage, logAlternateCallstack);
+					logger.LogError(KnownLogEvents.Gauntlet_FatalEvent, "Engine encounter a critical failure\n{Summary}\n{Callstack}", logMessage, logAlternateCallstack);
 				}
 				await UpdateCompleteStepAsync(job, 0, 1, JobStepOutcome.Failure);
 
diff --git a/Engine/Source/Programs/Shared/EpicGames.Horde/Issues/Handlers/GauntletIssueHandler.cs b/Engine/Source/Programs/Shared/EpicGames.Horde/Issues/Handlers/GauntletIssueHandler.cs
--- a/Engine/Source/Programs/Shared/EpicGames.Horde/Issues/Handlers/GauntletIssueHandler.cs
+++ b/Engine/Source/Programs/Shared/EpicGames.Horde/Issues/Handlers/GauntletIssueHandler.cs
@@ -51,11 +51,6 @@ namespace EpicGames.Horde.Issues.Handlers
 		/// </summary>
 		static readonly Utf8String s_summaryLogType = new Utf8String("Summary");
 
-		/// <summary>
-		/// Regex pattern to match assertion, ensure and other crash report summary from UE
-		/// </summary>
-		static readonly Regex s_summaryPattern = new Regex(@"\w[\w ]+:\s+.+\s+\[[^]]+\.[^]]+\](?:\s*\[[^]]+\])?(?:\s*\n\w+.+)?", RegexOptions.Multiline | RegexOptions.ExplicitCapture);
-
 		/// <summary>
 		/// Max Message Length to hash
 		/// </summary>
@@ -133,18 +128,8 @@ namespace EpicGames.Horde.Issues.Handlers
 
 		private static bool TryGetHash(IssueEvent issueEvent, out Md5Hash hash)
 		{
-			string? sanitized = GetSummaryProperty(issueEvent);
 			// Use only the summary if one is found instead of the full callstack
-			if (sanitized == null)
-			{
-				sanitized = issueEvent.Message;
-				// Let's try with regex
-				Match matchSummary = s_summaryPattern.Match(sanitized);
-				if (matchSummary.Success)
-				{
-					sanitized = matchSummary.Groups[0].Value;
-				}
-			}
+			string sanitized = GetSummaryProperty(issueEvent) ?? issueEvent.Message;
 			sanitized = sanitized.Length > MaxMessageLength ? sanitized.Substring(0, MaxMessageLength) : sanitized;
 			sanitized = sanitized.Trim().ToUpperInvariant();
 			sanitized = Regex.Replace(sanitized, @"(?<![A-Z])(?:[A-Z]:|/)[^ :]+[/\\]SYNC[/\\]", "{root}/"); // Redact things that look like workspace roots; may be different between agents
