EpicGames.Build: Prevent scanning DirectoryItems that are inaccessable

#rnx

diff --git a/Engine/Source/Programs/Shared/EpicGames.Build/System/DirectoryItem.cs b/Engine/Source/Programs/Shared/EpicGames.Build/System/DirectoryItem.cs
--- a/Engine/Source/Programs/Shared/EpicGames.Build/System/DirectoryItem.cs
+++ b/Engine/Source/Programs/Shared/EpicGames.Build/System/DirectoryItem.cs
@@ -6,6 +6,7 @@ using System.Collections.Generic;
 using System.Diagnostics.CodeAnalysis;
 using System.IO;
 using System.Linq;
+using System.Security;
 using EpicGames.Core;
 
 namespace UnrealBuildBase
@@ -223,26 +224,32 @@ namespace UnrealBuildBase
 		{
 			if (Directories == null)
 			{
-				Dictionary<string, DirectoryItem> NewDirectories;
+				Dictionary<string, DirectoryItem>? NewDirectories = null;
 				if (Info.Value.Exists)
 				{
-					DirectoryInfo[] Directories = Info.Value.GetDirectories();
-					NewDirectories = new Dictionary<string, DirectoryItem>(Directories.Length, DirectoryReference.Comparer);
-					foreach (DirectoryInfo SubDirectoryInfo in Directories)
+					try
 					{
-						if (NewDirectories.ContainsKey(SubDirectoryInfo.Name))
+						DirectoryInfo[] Directories = Info.Value.GetDirectories();
+						NewDirectories = new Dictionary<string, DirectoryItem>(Directories.Length, DirectoryReference.Comparer);
+						foreach (DirectoryInfo SubDirectoryInfo in Directories)
 						{
-							throw new Exception($"Trying to add {SubDirectoryInfo.FullName} as '{SubDirectoryInfo.Name}' yet exists as {NewDirectories[SubDirectoryInfo.Name].FullName}");
-						}
+							if (NewDirectories.ContainsKey(SubDirectoryInfo.Name))
+							{
+								throw new Exception($"Trying to add {SubDirectoryInfo.FullName} as '{SubDirectoryInfo.Name}' yet exists as {NewDirectories[SubDirectoryInfo.Name].FullName}");
+							}
 
-						NewDirectories.Add(SubDirectoryInfo.Name, DirectoryItem.GetItemByDirectoryInfo(SubDirectoryInfo));
+							NewDirectories.Add(SubDirectoryInfo.Name, DirectoryItem.GetItemByDirectoryInfo(SubDirectoryInfo));
+						}
+					}
+					catch (SecurityException)
+					{
+					}
+					catch (UnauthorizedAccessException)
+					{
 					}
 				}
-				else
-				{
-					NewDirectories = new Dictionary<string, DirectoryItem>(DirectoryReference.Comparer);
-				}
-				Directories = NewDirectories;
+
+				Directories = NewDirectories ?? new Dictionary<string, DirectoryItem>(DirectoryReference.Comparer);
 			}
 		}
 
@@ -289,23 +296,31 @@ namespace UnrealBuildBase
 		{
 			if (Files == null)
 			{
-				Dictionary<string, FileItem> NewFiles;
+				Dictionary<string, FileItem>? NewFiles = null;
 				if (Info.Value.Exists)
 				{
-					FileInfo[] FileInfos = Info.Value.GetFiles();
-					NewFiles = new Dictionary<string, FileItem>(FileInfos.Length, FileReference.Comparer);
-					foreach (FileInfo FileInfo in FileInfos)
+					try
+					{
+						FileInfo[] FileInfos = Info.Value.GetFiles();
+						NewFiles = new Dictionary<string, FileItem>(FileInfos.Length, FileReference.Comparer);
+						foreach (FileInfo FileInfo in FileInfos)
+						{
+							FileItem FileItem = FileItem.GetItemByFileInfo(FileInfo);
+							FileItem.UpdateCachedDirectory(this);
+							NewFiles[FileInfo.Name] = FileItem;
+						}
+						Files = NewFiles;
+						return;
+					}
+					catch (SecurityException)
+					{
+					}
+					catch (UnauthorizedAccessException)
 					{
-						FileItem FileItem = FileItem.GetItemByFileInfo(FileInfo);
-						FileItem.UpdateCachedDirectory(this);
-						NewFiles[FileInfo.Name] = FileItem;
 					}
 				}
-				else
-				{
-					NewFiles = new Dictionary<string, FileItem>(FileReference.Comparer);
-				}
-				Files = NewFiles;
+
+				Files = NewFiles ?? new Dictionary<string, FileItem>(FileReference.Comparer);
 			}
 		}
 
