Horde: Fix static analysis warning.

#rnx

diff --git a/Engine/Source/Programs/Shared/EpicGames.Horde/Storage/Bundles/BundleCompression.cs b/Engine/Source/Programs/Shared/EpicGames.Horde/Storage/Bundles/BundleCompression.cs
--- a/Engine/Source/Programs/Shared/EpicGames.Horde/Storage/Bundles/BundleCompression.cs
+++ b/Engine/Source/Programs/Shared/EpicGames.Horde/Storage/Bundles/BundleCompression.cs
@@ -62,7 +62,7 @@ namespace EpicGames.Horde.Storage.Bundles
 		public static byte[] Compress(BundleCompressionFormat format, ReadOnlyMemory<byte> input)
 		{
 			ArrayMemoryWriter writer = new ArrayMemoryWriter(0);
-			int length = Compress(format, input, writer);
+			Compress(format, input, writer);
 			return writer.WrittenMemory.ToArray();
 		}
 
