Horde: Move AclScopeName into EpicGames.Horde.

#rnx

diff --git a/Engine/Source/Programs/Horde/HordeServer.Tests/Artifacts/ArtifactTests.cs b/Engine/Source/Programs/Horde/HordeServer.Tests/Artifacts/ArtifactTests.cs
--- a/Engine/Source/Programs/Horde/HordeServer.Tests/Artifacts/ArtifactTests.cs
+++ b/Engine/Source/Programs/Horde/HordeServer.Tests/Artifacts/ArtifactTests.cs
@@ -13,6 +13,7 @@ using HordeServer.Server;
 using HordeCommon;
 using Microsoft.Extensions.DependencyInjection;
 using Microsoft.VisualStudio.TestTools.UnitTesting;
+using EpicGames.Horde.Acls;
 
 namespace HordeServer.Tests.Artifacts
 {
diff --git a/Engine/Source/Programs/Horde/HordeServer/Acls/AclScopeName.cs b/Engine/Source/Programs/Horde/HordeServer/Acls/AclScopeName.cs
deleted file mode 100644
diff --git a/Engine/Source/Programs/Horde/HordeServer/Artifacts/ArtifactCollection.cs b/Engine/Source/Programs/Horde/HordeServer/Artifacts/ArtifactCollection.cs
--- a/Engine/Source/Programs/Horde/HordeServer/Artifacts/ArtifactCollection.cs
+++ b/Engine/Source/Programs/Horde/HordeServer/Artifacts/ArtifactCollection.cs
@@ -7,10 +7,10 @@ using System.Runtime.CompilerServices;
 using System.Threading;
 using System.Threading.Tasks;
 using EpicGames.Core;
+using EpicGames.Horde.Acls;
 using EpicGames.Horde.Artifacts;
 using EpicGames.Horde.Storage;
 using EpicGames.Horde.Streams;
-using HordeServer.Acls;
 using HordeServer.Server;
 using HordeServer.Storage;
 using HordeServer.Utilities;
diff --git a/Engine/Source/Programs/Horde/HordeServer/Artifacts/ArtifactsController.cs b/Engine/Source/Programs/Horde/HordeServer/Artifacts/ArtifactsController.cs
--- a/Engine/Source/Programs/Horde/HordeServer/Artifacts/ArtifactsController.cs
+++ b/Engine/Source/Programs/Horde/HordeServer/Artifacts/ArtifactsController.cs
@@ -36,6 +36,7 @@ using Microsoft.AspNetCore.Mvc;
 using Microsoft.AspNetCore.StaticFiles;
 using Microsoft.Extensions.Logging;
 using Microsoft.Extensions.Options;
+using EpicGames.Horde.Acls;
 
 namespace HordeServer.Artifacts
 {
diff --git a/Engine/Source/Programs/Horde/HordeServer/Artifacts/IArtifact.cs b/Engine/Source/Programs/Horde/HordeServer/Artifacts/IArtifact.cs
--- a/Engine/Source/Programs/Horde/HordeServer/Artifacts/IArtifact.cs
+++ b/Engine/Source/Programs/Horde/HordeServer/Artifacts/IArtifact.cs
@@ -2,10 +2,10 @@
 
 using System;
 using System.Collections.Generic;
+using EpicGames.Horde.Acls;
 using EpicGames.Horde.Artifacts;
 using EpicGames.Horde.Storage;
 using EpicGames.Horde.Streams;
-using HordeServer.Acls;
 
 namespace HordeServer.Artifacts
 {
diff --git a/Engine/Source/Programs/Horde/HordeServer/Artifacts/IArtifactCollection.cs b/Engine/Source/Programs/Horde/HordeServer/Artifacts/IArtifactCollection.cs
--- a/Engine/Source/Programs/Horde/HordeServer/Artifacts/IArtifactCollection.cs
+++ b/Engine/Source/Programs/Horde/HordeServer/Artifacts/IArtifactCollection.cs
@@ -4,9 +4,9 @@ using System;
 using System.Collections.Generic;
 using System.Threading;
 using System.Threading.Tasks;
+using EpicGames.Horde.Acls;
 using EpicGames.Horde.Artifacts;
 using EpicGames.Horde.Streams;
-using HordeServer.Acls;
 
 namespace HordeServer.Artifacts
 {
diff --git a/Engine/Source/Programs/Horde/HordeServer/Utilities/HordeControllerBase.cs b/Engine/Source/Programs/Horde/HordeServer/Utilities/HordeControllerBase.cs
--- a/Engine/Source/Programs/Horde/HordeServer/Utilities/HordeControllerBase.cs
+++ b/Engine/Source/Programs/Horde/HordeServer/Utilities/HordeControllerBase.cs
@@ -16,7 +16,6 @@ using EpicGames.Horde.Storage;
 using EpicGames.Horde.Streams;
 using EpicGames.Horde.Telemetry;
 using EpicGames.Horde.Tools;
-using HordeServer.Acls;
 using Microsoft.AspNetCore.Http;
 using Microsoft.AspNetCore.Mvc;
 using Microsoft.Extensions.Logging;
diff --git a/Engine/Source/Programs/Shared/EpicGames.Horde/Acls/AclScopeName.cs b/Engine/Source/Programs/Shared/EpicGames.Horde/Acls/AclScopeName.cs
new file mode 100644
--- /dev/null
+++ b/Engine/Source/Programs/Shared/EpicGames.Horde/Acls/AclScopeName.cs
@@ -0,0 +1,92 @@
+// Copyright Epic Games, Inc. All Rights Reserved.
+
+using System;
+using System.ComponentModel;
+using System.Diagnostics;
+using System.Diagnostics.CodeAnalysis;
+using System.Globalization;
+using System.Text.Json;
+using System.Text.Json.Serialization;
+
+namespace EpicGames.Horde.Acls
+{
+	/// <summary>
+	/// Name of an ACL scope
+	/// </summary>
+	[DebuggerDisplay("{Text}")]
+	[JsonConverter(typeof(AclScopeNameJsonConverter))]
+	[TypeConverter(typeof(AclScopeNameTypeConverter))]
+	public record struct AclScopeName(string Text)
+	{
+		/// <summary>
+		/// The root scope name
+		/// </summary>
+		public static AclScopeName Root { get; } = new AclScopeName("horde");
+
+		/// <summary>
+		/// Append another name to this scope
+		/// </summary>
+		/// <param name="name">Name to append</param>
+		/// <returns>New scope name</returns>
+		public AclScopeName Append(string name) => new AclScopeName($"{Text}/{name}");
+
+		/// <inheritdoc/>
+		public bool Equals(AclScopeName other) => Text.Equals(other.Text, StringComparison.Ordinal);
+
+		/// <inheritdoc/>
+		public override int GetHashCode() => Text.GetHashCode(StringComparison.Ordinal);
+
+		/// <inheritdoc/>
+		public override string ToString() => Text;
+	}
+
+	/// <summary>
+	/// Serializes <see cref="AclScopeName"/> objects to JSON
+	/// </summary>
+	class AclScopeNameJsonConverter : JsonConverter<AclScopeName>
+	{
+		/// <inheritdoc/>
+		public override AclScopeName Read(ref Utf8JsonReader reader, Type typeToConvert, JsonSerializerOptions options) => new AclScopeName(reader.GetString() ?? String.Empty);
+
+		/// <inheritdoc/>
+		public override void Write(Utf8JsonWriter writer, AclScopeName value, JsonSerializerOptions options) => writer.WriteStringValue(value.Text);
+	}
+
+	/// <summary>
+	/// Converts <see cref="AclScopeName"/> objects to strings
+	/// </summary>
+	class AclScopeNameTypeConverter : TypeConverter
+	{
+		/// <inheritdoc/>
+		public override bool CanConvertFrom(ITypeDescriptorContext? context, Type sourceType)
+		{
+			return sourceType == typeof(string) || base.CanConvertFrom(context, sourceType);
+		}
+
+		/// <inheritdoc/>
+		public override object? ConvertFrom(ITypeDescriptorContext? context, CultureInfo? culture, object value)
+		{
+			if (value is string str)
+			{
+				return new AclScopeName(str);
+			}
+			return base.ConvertFrom(context, culture, value);
+		}
+
+		/// <inheritdoc/>
+		public override bool CanConvertTo(ITypeDescriptorContext? context, [NotNullWhen(true)] Type? destinationType)
+		{
+			return destinationType == typeof(string) || base.CanConvertTo(context, destinationType);
+		}
+
+		/// <inheritdoc/>
+		public override object? ConvertTo(ITypeDescriptorContext? context, CultureInfo? culture, object? value, Type destinationType)
+		{
+			if (destinationType == typeof(string))
+			{
+				return value?.ToString();
+			}
+			return base.ConvertTo(context, culture, value, destinationType);
+		}
+	}
+}
