Horde: Missing file.

diff --git a/Engine/Source/Programs/Shared/EpicGames.Horde/Storage/Nodes/DdcRefNode.cs b/Engine/Source/Programs/Shared/EpicGames.Horde/Storage/Nodes/DdcRefNode.cs
--- a/Engine/Source/Programs/Shared/EpicGames.Horde/Storage/Nodes/DdcRefNode.cs
+++ b/Engine/Source/Programs/Shared/EpicGames.Horde/Storage/Nodes/DdcRefNode.cs
@@ -3,9 +3,8 @@
 using System;
 using System.Collections.Generic;
 using EpicGames.Core;
-using EpicGames.Horde.Storage;
 
-namespace Horde.Server.Ddc
+namespace EpicGames.Horde.Storage.Nodes
 {
 	/// <summary>
 	/// A node containing ref data
