Horde: Support generating schemas including plugin config types.

#rnx

diff --git a/Engine/Source/Programs/Horde/HordeServer.Tests/Server/SchemaTests.cs b/Engine/Source/Programs/Horde/HordeServer.Tests/Server/SchemaTests.cs
--- a/Engine/Source/Programs/Horde/HordeServer.Tests/Server/SchemaTests.cs
+++ b/Engine/Source/Programs/Horde/HordeServer.Tests/Server/SchemaTests.cs
@@ -1,6 +1,7 @@
 // Copyright Epic Games, Inc. All Rights Reserved.
 
 using EpicGames.Core;
+using HordeServer.Plugins;
 using HordeServer.Streams;
 using HordeServer.Utilities;
 using Microsoft.VisualStudio.TestTools.UnitTesting;
@@ -13,14 +14,16 @@ namespace HordeServer.Tests.Server
 		[TestMethod]
 		public void StreamSchema()
 		{
-			JsonSchema schema = Schemas.CreateSchema(typeof(StreamConfig));
+			JsonSchemaCache cache = new JsonSchemaCache(PluginCollection.Empty);
+			JsonSchema schema = cache.CreateSchema(typeof(StreamConfig));
 			_ = schema;
 		}
 
 		[TestMethod]
 		public void ProjectSchema()
 		{
-			JsonSchema schema = Schemas.CreateSchema(typeof(StreamConfig));
+			JsonSchemaCache cache = new JsonSchemaCache(PluginCollection.Empty);
+			JsonSchema schema = cache.CreateSchema(typeof(StreamConfig));
 			_ = schema;
 		}
 	}
diff --git a/Engine/Source/Programs/Horde/HordeServer/Commands/Generate/ConfigSchemasCommand.cs b/Engine/Source/Programs/Horde/HordeServer/Commands/Generate/ConfigSchemasCommand.cs
--- a/Engine/Source/Programs/Horde/HordeServer/Commands/Generate/ConfigSchemasCommand.cs
+++ b/Engine/Source/Programs/Horde/HordeServer/Commands/Generate/ConfigSchemasCommand.cs
@@ -1,9 +1,11 @@
 // Copyright Epic Games, Inc. All Rights Reserved.
 
 using System;
+using System.Collections.Generic;
 using System.ComponentModel;
 using System.Threading.Tasks;
 using EpicGames.Core;
+using HordeServer.Plugins;
 using HordeServer.Server;
 using HordeServer.Utilities;
 using Microsoft.Extensions.Logging;
@@ -21,11 +23,13 @@ namespace HordeServer.Commands.Generate
 		{
 			_outputDir ??= DirectoryReference.Combine(ServerApp.AppDir, "Schemas");
 
+			JsonSchemaCache schemaCache = new JsonSchemaCache(new PluginCollection(new Dictionary<string, IPluginStartup>()));
+
 			DirectoryReference.CreateDirectory(_outputDir);
 			foreach (Type schemaType in SchemaController.ConfigSchemas)
 			{
 				FileReference outputFile = FileReference.Combine(_outputDir, $"{schemaType.Name}.json");
-				Schemas.CreateSchema(schemaType).Write(outputFile);
+				schemaCache.CreateSchema(schemaType).Write(outputFile);
 			}
 
 			return Task.FromResult(0);
diff --git a/Engine/Source/Programs/Horde/HordeServer/Plugins/PluginCollection.cs b/Engine/Source/Programs/Horde/HordeServer/Plugins/PluginCollection.cs
--- a/Engine/Source/Programs/Horde/HordeServer/Plugins/PluginCollection.cs
+++ b/Engine/Source/Programs/Horde/HordeServer/Plugins/PluginCollection.cs
@@ -9,6 +9,11 @@ namespace HordeServer.Plugins
 	/// </summary>
 	class PluginCollection : IPluginCollection
 	{
+		/// <summary>
+		/// Static empty plugin collection
+		/// </summary>
+		public static IPluginCollection Empty { get; } = new PluginCollection(new Dictionary<string, IPluginStartup>());
+
 		/// <inheritdoc/>
 		public IReadOnlyDictionary<string, IPluginStartup> EnabledPlugins { get; }
 
diff --git a/Engine/Source/Programs/Horde/HordeServer/Server/SchemaController.cs b/Engine/Source/Programs/Horde/HordeServer/Server/SchemaController.cs
--- a/Engine/Source/Programs/Horde/HordeServer/Server/SchemaController.cs
+++ b/Engine/Source/Programs/Horde/HordeServer/Server/SchemaController.cs
@@ -39,6 +39,16 @@ namespace HordeServer.Server
 		/// </summary>
 		public static Type[] ConfigSchemas { get; } = FindSchemaTypes();
 
+		readonly JsonSchemaCache _schemaCache;
+
+		/// <summary>
+		/// Constructor
+		/// </summary>
+		public SchemaController(JsonSchemaCache schemaCache)
+		{
+			_schemaCache = schemaCache;
+		}
+
 		static Type[] FindSchemaTypes()
 		{
 			List<Type> schemaTypes = new List<Type>();
@@ -90,7 +100,7 @@ namespace HordeServer.Server
 			{
 				if (schemaType.Name.Equals(typeName, StringComparison.OrdinalIgnoreCase))
 				{
-					JsonSchema schema = Schemas.CreateSchema(schemaType);
+					JsonSchema schema = _schemaCache.CreateSchema(schemaType);
 
 					using MemoryStream stream = new MemoryStream();
 					schema.Write(stream);
diff --git a/Engine/Source/Programs/Horde/HordeServer/Startup.cs b/Engine/Source/Programs/Horde/HordeServer/Startup.cs
--- a/Engine/Source/Programs/Horde/HordeServer/Startup.cs
+++ b/Engine/Source/Programs/Horde/HordeServer/Startup.cs
@@ -368,6 +368,7 @@ namespace HordeServer
 		public void ConfigureServices(IServiceCollection services)
 		{
 			services.AddSingleton<IServerInfo>(new ServerInfo());
+			services.AddSingleton<JsonSchemaCache>();
 
 			// Register the plugin collection
 			Dictionary<string, IPluginStartup> plugins = new Dictionary<string, IPluginStartup>();
diff --git a/Engine/Source/Programs/Horde/HordeServer/Utilities/Schemas.cs b/Engine/Source/Programs/Horde/HordeServer/Utilities/Schemas.cs
deleted file mode 100644
diff --git a/Engine/Source/Programs/Shared/EpicGames.Core/JsonSchema.cs b/Engine/Source/Programs/Shared/EpicGames.Core/JsonSchema.cs
--- a/Engine/Source/Programs/Shared/EpicGames.Core/JsonSchema.cs
+++ b/Engine/Source/Programs/Shared/EpicGames.Core/JsonSchema.cs
@@ -283,8 +283,8 @@ namespace EpicGames.Core
 		/// <returns>New schema object</returns>
 		public static JsonSchema FromType(Type type, XmlDocReader? xmlDocReader)
 		{
-			JsonSchemaAttribute? schemaAttribute = type.GetCustomAttribute<JsonSchemaAttribute>();
-			return new JsonSchema(schemaAttribute?.Id, CreateSchemaType(type, new Dictionary<Type, JsonSchemaType>(), xmlDocReader));
+			JsonSchemaFactory factory = new JsonSchemaFactory(xmlDocReader);
+			return factory.CreateSchema(type);
 		}
 
 		/// <summary>
@@ -292,175 +292,8 @@ namespace EpicGames.Core
 		/// </summary>
 		public static JsonSchema CreateSchema(Type type, XmlDocReader? xmlDocReader)
 		{
-			JsonSchemaAttribute? schemaAttribute = type.GetCustomAttribute<JsonSchemaAttribute>();
-			return new JsonSchema(schemaAttribute?.Id, CreateSchemaType(type, new Dictionary<Type, JsonSchemaType>(), xmlDocReader));
-		}
-
-		/// <summary>
-		/// Constructs a schema type for the given property
-		/// </summary>
-		static JsonSchemaType CreateSchemaType(PropertyInfo property, Dictionary<Type, JsonSchemaType> typeCache, XmlDocReader? xmlDocReader)
-		{
-			if (property.GetCustomAttribute<JsonSchemaStringAttribute>() != null)
-			{
-				return new JsonSchemaString();
-			}
-
-			return CreateSchemaType(property.PropertyType, typeCache, xmlDocReader);
-		}
-
-		/// <summary>
-		/// Constructs a schema type from the given type object
-		/// </summary>
-		static JsonSchemaType CreateSchemaType(Type type, Dictionary<Type, JsonSchemaType> typeCache, XmlDocReader? xmlDocReader)
-		{
-			switch (Type.GetTypeCode(type))
-			{
-				case TypeCode.Boolean:
-					return new JsonSchemaBoolean();
-				case TypeCode.Byte:
-				case TypeCode.SByte:
-				case TypeCode.Int16:
-				case TypeCode.UInt16:
-				case TypeCode.Int32:
-				case TypeCode.UInt32:
-				case TypeCode.Int64:
-				case TypeCode.UInt64:
-					if (type.IsEnum)
-					{
-						return CreateEnumSchemaType(type, xmlDocReader);
-					}
-					else
-					{
-						return new JsonSchemaInteger();
-					}
-				case TypeCode.Single:
-				case TypeCode.Double:
-					return new JsonSchemaNumber();
-				case TypeCode.String:
-					return new JsonSchemaString();
-			}
-
-			JsonSchemaTypeAttribute? attribute = type.GetCustomAttribute<JsonSchemaTypeAttribute>();
-			switch (attribute)
-			{
-				case JsonSchemaStringAttribute str:
-					return new JsonSchemaString(str.Format);
-			}
-
-			if (type.IsGenericType && type.GetGenericTypeDefinition() == typeof(Nullable<>))
-			{
-				return CreateSchemaType(type.GetGenericArguments()[0], typeCache, xmlDocReader);
-			}
-			if (type == typeof(DateTime) || type == typeof(DateTimeOffset))
-			{
-				return new JsonSchemaString(JsonSchemaStringFormat.DateTime);
-			}
-			if (type == typeof(TimeSpan))
-			{
-				return new JsonSchemaString();
-			}
-			if (type == typeof(Uri))
-			{
-				return new JsonSchemaString(JsonSchemaStringFormat.Uri);
-			}
-			if (type == typeof(System.Text.Json.Nodes.JsonObject))
-			{
-				return new JsonSchemaObject();
-			}
-
-			Type[] interfaceTypes = type.GetInterfaces();
-			foreach (Type interfaceType in interfaceTypes)
-			{
-				if (interfaceType.IsGenericType)
-				{
-					Type[] arguments = interfaceType.GetGenericArguments();
-					if (interfaceType.GetGenericTypeDefinition() == typeof(IList<>))
-					{
-						return new JsonSchemaArray(CreateSchemaType(arguments[0], typeCache, xmlDocReader));
-					}
-					if (interfaceType.GetGenericTypeDefinition() == typeof(IDictionary<,>))
-					{
-						JsonSchemaObject obj = new JsonSchemaObject();
-						obj.AdditionalProperties = CreateSchemaType(arguments[1], typeCache, xmlDocReader);
-						return obj;
-					}
-				}
-			}
-
-			if (type.IsClass)
-			{
-				if (typeCache.TryGetValue(type, out JsonSchemaType? schemaType))
-				{
-					return schemaType;
-				}
-
-				JsonKnownTypesAttribute? knownTypes = type.GetCustomAttribute<JsonKnownTypesAttribute>(false);
-				if (knownTypes != null)
-				{
-					JsonSchemaOneOf obj = new JsonSchemaOneOf();
-					typeCache[type] = obj;
-					SetOneOfProperties(obj, type, knownTypes.Types, typeCache, xmlDocReader);
-					return obj;
-				}
-				else
-				{
-					JsonSchemaObject obj = new JsonSchemaObject();
-					typeCache[type] = obj;
-					SetObjectProperties(obj, type, typeCache, xmlDocReader);
-					return obj;
-				}
-			}
-
-			throw new Exception($"Unknown type for schema generation: {type}");
-		}
-
-		static JsonSchemaEnum CreateEnumSchemaType(Type type, XmlDocReader? xmlDocReader)
-		{
-			string[] names = Enum.GetNames(type);
-			string[] descriptions = new string[names.Length];
-
-			for (int idx = 0; idx < names.Length; idx++)
-			{
-				descriptions[idx] = xmlDocReader?.GetDescription(type, names[idx]) ?? String.Empty;
-			}
-
-			string? enumDescription = xmlDocReader?.GetDescription(type);
-			return new JsonSchemaEnum(names, descriptions) { Name = type.Name, Description = enumDescription };
-		}
-
-		static void SetOneOfProperties(JsonSchemaOneOf obj, Type type, Type[] knownTypes, Dictionary<Type, JsonSchemaType> typeCache, XmlDocReader? xmlDocReader)
-		{
-			obj.Name = type.Name;
-
-			foreach (Type knownType in knownTypes)
-			{
-				JsonDiscriminatorAttribute? attribute = knownType.GetCustomAttribute<JsonDiscriminatorAttribute>();
-				if (attribute != null)
-				{
-					JsonSchemaObject knownObject = new JsonSchemaObject();
-					knownObject.Properties.Add(new JsonSchemaProperty("type", "Type discriminator", new JsonSchemaEnum(new[] { attribute.Name }, new[] { "Identifier for the derived type" })));
-					SetObjectProperties(knownObject, knownType, typeCache, xmlDocReader);
-					obj.Types.Add(knownObject);
-				}
-			}
-		}
-
-		static void SetObjectProperties(JsonSchemaObject obj, Type type, Dictionary<Type, JsonSchemaType> typeCache, XmlDocReader? xmlDocReader)
-		{
-			obj.Name = type.Name;
-			obj.Description = xmlDocReader?.GetDescription(type);
-
-			PropertyInfo[] properties = type.GetProperties(BindingFlags.Instance | BindingFlags.Public);
-			foreach (PropertyInfo property in properties)
-			{
-				if (property.GetCustomAttribute<JsonIgnoreAttribute>() == null)
-				{
-					string? description = xmlDocReader?.GetDescription(property);
-					JsonSchemaType propertyType = CreateSchemaType(property, typeCache, xmlDocReader);
-					obj.Properties.Add(new JsonSchemaProperty(property.Name, description, propertyType));
-				}
-			}
+			JsonSchemaFactory factory = new JsonSchemaFactory(xmlDocReader);
+			return factory.CreateSchema(type);
 		}
 	}
 }
diff --git a/Engine/Source/Programs/Shared/EpicGames.Core/JsonSchemaAttributes.cs b/Engine/Source/Programs/Shared/EpicGames.Core/JsonSchemaAttributes.cs
--- a/Engine/Source/Programs/Shared/EpicGames.Core/JsonSchemaAttributes.cs
+++ b/Engine/Source/Programs/Shared/EpicGames.Core/JsonSchemaAttributes.cs
@@ -80,6 +80,14 @@ namespace EpicGames.Core
 	{
 	}
 
+	/// <summary>
+	/// Attribute allowing any possible value in a json schema
+	/// </summary>
+	[AttributeUsage(AttributeTargets.Class | AttributeTargets.Struct | AttributeTargets.Property)]
+	public sealed class JsonSchemaAnyAttribute : JsonSchemaTypeAttribute
+	{
+	}
+
 	/// <summary>
 	/// Attribute setting properties for a type to be serialized as a string
 	/// </summary>
diff --git a/Engine/Source/Programs/Shared/EpicGames.Core/JsonSchemaFactory.cs b/Engine/Source/Programs/Shared/EpicGames.Core/JsonSchemaFactory.cs
new file mode 100644
--- /dev/null
+++ b/Engine/Source/Programs/Shared/EpicGames.Core/JsonSchemaFactory.cs
@@ -0,0 +1,224 @@
+// Copyright Epic Games, Inc.All Rights Reserved.
+
+using System;
+using System.Collections.Generic;
+using System.Linq;
+using System.Reflection;
+using System.Text;
+using System.Text.Json.Serialization;
+using System.Threading.Tasks;
+
+namespace EpicGames.Core
+{
+	/// <summary>
+	/// Factory for creating Json schema elements from native properties and types
+	/// </summary>
+	public class JsonSchemaFactory
+	{
+		/// <summary>
+		/// Reader for XML documentation
+		/// </summary>
+		public XmlDocReader? XmlDocReader { get; }
+
+		/// <summary>
+		/// Cache of known types
+		/// </summary>
+		public Dictionary<Type, JsonSchemaType> TypeCache { get; } = new Dictionary<Type, JsonSchemaType>();
+
+		/// <summary>
+		/// Constructor
+		/// </summary>
+		public JsonSchemaFactory(XmlDocReader? xmlDocReader)
+		{
+			XmlDocReader = xmlDocReader;
+		}
+
+		/// <summary>
+		/// Creates a schema object for the given type
+		/// </summary>
+		public JsonSchema CreateSchema(Type type)
+		{
+			JsonSchemaAttribute? schemaAttribute = type.GetCustomAttribute<JsonSchemaAttribute>();
+			return new JsonSchema(schemaAttribute?.Id, CreateSchemaType(type));
+		}
+
+		/// <summary>
+		/// Constructs a schema type for the given property
+		/// </summary>
+		public virtual JsonSchemaType CreateSchemaType(PropertyInfo property)
+		{
+			if (property.GetCustomAttribute<JsonSchemaAnyAttribute>() != null)
+			{
+				return new JsonSchemaAny();
+			}
+			if (property.GetCustomAttribute<JsonSchemaStringAttribute>() != null)
+			{
+				return new JsonSchemaString();
+			}
+
+			if (property.Name == "Plugins")
+			{
+				Console.Write("");
+				}
+
+
+			return CreateSchemaType(property.PropertyType);
+		}
+
+		/// <summary>
+		/// Constructs a schema type from the given type object
+		/// </summary>
+		public virtual JsonSchemaType CreateSchemaType(Type type)
+		{
+			switch (Type.GetTypeCode(type))
+			{
+				case TypeCode.Boolean:
+					return new JsonSchemaBoolean();
+				case TypeCode.Byte:
+				case TypeCode.SByte:
+				case TypeCode.Int16:
+				case TypeCode.UInt16:
+				case TypeCode.Int32:
+				case TypeCode.UInt32:
+				case TypeCode.Int64:
+				case TypeCode.UInt64:
+					if (type.IsEnum)
+					{
+						return CreateEnumSchemaType(type);
+					}
+					else
+					{
+						return new JsonSchemaInteger();
+					}
+				case TypeCode.Single:
+				case TypeCode.Double:
+					return new JsonSchemaNumber();
+				case TypeCode.String:
+					return new JsonSchemaString();
+			}
+
+			JsonSchemaTypeAttribute? attribute = type.GetCustomAttribute<JsonSchemaTypeAttribute>();
+			switch (attribute)
+			{
+				case JsonSchemaAnyAttribute any:
+					return new JsonSchemaAny();
+				case JsonSchemaStringAttribute str:
+					return new JsonSchemaString(str.Format);
+			}
+
+			if (type.IsGenericType && type.GetGenericTypeDefinition() == typeof(Nullable<>))
+			{
+				return CreateSchemaType(type.GetGenericArguments()[0]);
+			}
+			if (type == typeof(DateTime) || type == typeof(DateTimeOffset))
+			{
+				return new JsonSchemaString(JsonSchemaStringFormat.DateTime);
+			}
+			if (type == typeof(TimeSpan))
+			{
+				return new JsonSchemaString();
+			}
+			if (type == typeof(Uri))
+			{
+				return new JsonSchemaString(JsonSchemaStringFormat.Uri);
+			}
+			if (type == typeof(System.Text.Json.Nodes.JsonObject))
+			{
+				return new JsonSchemaObject();
+			}
+
+			if (type.IsClass)
+			{
+				if (TypeCache.TryGetValue(type, out JsonSchemaType? schemaType))
+				{
+					return schemaType;
+				}
+
+				Type[] interfaceTypes = type.GetInterfaces();
+				foreach (Type interfaceType in interfaceTypes)
+				{
+					if (interfaceType.IsGenericType)
+					{
+						Type[] arguments = interfaceType.GetGenericArguments();
+						if (interfaceType.GetGenericTypeDefinition() == typeof(IList<>))
+						{
+							return new JsonSchemaArray(CreateSchemaType(arguments[0]));
+						}
+						if (interfaceType.GetGenericTypeDefinition() == typeof(IDictionary<,>))
+						{
+							JsonSchemaObject obj = new JsonSchemaObject();
+							obj.AdditionalProperties = CreateSchemaType(arguments[1]);
+							return obj;
+						}
+					}
+				}
+
+				JsonKnownTypesAttribute? knownTypes = type.GetCustomAttribute<JsonKnownTypesAttribute>(false);
+				if (knownTypes != null)
+				{
+					JsonSchemaOneOf obj = new JsonSchemaOneOf();
+					TypeCache[type] = obj;
+					SetOneOfProperties(obj, type, knownTypes.Types);
+					return obj;
+				}
+				else
+				{
+					JsonSchemaObject obj = new JsonSchemaObject();
+					TypeCache[type] = obj;
+					SetObjectProperties(obj, type);
+					return obj;
+				}
+			}
+
+			throw new Exception($"Unknown type for schema generation: {type}");
+		}
+
+		protected JsonSchemaEnum CreateEnumSchemaType(Type type)
+		{
+			string[] names = Enum.GetNames(type);
+			string[] descriptions = new string[names.Length];
+
+			for (int idx = 0; idx < names.Length; idx++)
+			{
+				descriptions[idx] = XmlDocReader?.GetDescription(type, names[idx]) ?? String.Empty;
+			}
+
+			string? enumDescription = XmlDocReader?.GetDescription(type);
+			return new JsonSchemaEnum(names, descriptions) { Name = type.Name, Description = enumDescription };
+		}
+
+		protected void SetOneOfProperties(JsonSchemaOneOf obj, Type type, Type[] knownTypes)
+		{
+			obj.Name = type.Name;
+
+			foreach (Type knownType in knownTypes)
+			{
+				JsonDiscriminatorAttribute? attribute = knownType.GetCustomAttribute<JsonDiscriminatorAttribute>();
+				if (attribute != null)
+				{
+					JsonSchemaObject knownObject = new JsonSchemaObject();
+					knownObject.Properties.Add(new JsonSchemaProperty("type", "Type discriminator", new JsonSchemaEnum(new[] { attribute.Name }, new[] { "Identifier for the derived type" })));
+					SetObjectProperties(knownObject, knownType);
+					obj.Types.Add(knownObject);
+				}
+			}
+		}
+
+		protected void SetObjectProperties(JsonSchemaObject obj, Type type)
+		{
+			obj.Name = type.Name;
+			obj.Description = XmlDocReader?.GetDescription(type);
+
+			PropertyInfo[] properties = type.GetProperties(BindingFlags.Instance | BindingFlags.Public);
+			foreach (PropertyInfo property in properties)
+			{
+				if (property.GetCustomAttribute<JsonIgnoreAttribute>() == null)
+				{
+					string? description = XmlDocReader?.GetDescription(property);
+					JsonSchemaType propertyType = CreateSchemaType(property);
+					obj.Properties.Add(new JsonSchemaProperty(property.Name, description, propertyType));
+				}
+			}
+		}
+	}
+}
diff --git a/Engine/Source/Programs/Shared/EpicGames.Core/JsonSchemaType.cs b/Engine/Source/Programs/Shared/EpicGames.Core/JsonSchemaType.cs
--- a/Engine/Source/Programs/Shared/EpicGames.Core/JsonSchemaType.cs
+++ b/Engine/Source/Programs/Shared/EpicGames.Core/JsonSchemaType.cs
@@ -34,6 +34,17 @@ namespace EpicGames.Core
 	{
 	}
 
+	/// <summary>
+	/// Represents any type in a Json schema
+	/// </summary>
+	public class JsonSchemaAny : JsonSchemaPrimitiveType
+	{
+		/// <inheritdoc/>
+		public override void Write(IJsonSchemaWriter writer)
+		{
+		}
+	}
+
 	/// <summary>
 	/// Represents a boolean in a Json schema
 	/// </summary>
