EpicGames.Core: Fix static analysis warning.

#rnx

diff --git a/Engine/Source/Programs/Shared/EpicGames.Core/JsonSchemaFactory.cs b/Engine/Source/Programs/Shared/EpicGames.Core/JsonSchemaFactory.cs
--- a/Engine/Source/Programs/Shared/EpicGames.Core/JsonSchemaFactory.cs
+++ b/Engine/Source/Programs/Shared/EpicGames.Core/JsonSchemaFactory.cs
@@ -45,24 +45,18 @@ namespace EpicGames.Core
 		/// <summary>
 		/// Constructs a schema type for the given property
 		/// </summary>
-		public virtual JsonSchemaType CreateSchemaType(PropertyInfo property)
+		public virtual JsonSchemaType CreateSchemaType(PropertyInfo propertyInfo)
 		{
-			if (property.GetCustomAttribute<JsonSchemaAnyAttribute>() != null)
+			if (propertyInfo.GetCustomAttribute<JsonSchemaAnyAttribute>() != null)
 			{
 				return new JsonSchemaAny();
 			}
-			if (property.GetCustomAttribute<JsonSchemaStringAttribute>() != null)
+			if (propertyInfo.GetCustomAttribute<JsonSchemaStringAttribute>() != null)
 			{
 				return new JsonSchemaString();
 			}
 
-			if (property.Name == "Plugins")
-			{
-				Console.Write("");
-				}
-
-
-			return CreateSchemaType(property.PropertyType);
+			return CreateSchemaType(propertyInfo.PropertyType);
 		}
 
 		/// <summary>
