EpicGames.Core: Fix some static analysis warnings.

#rnx

diff --git a/Engine/Source/Programs/Shared/EpicGames.Core/JsonSchema.cs b/Engine/Source/Programs/Shared/EpicGames.Core/JsonSchema.cs
--- a/Engine/Source/Programs/Shared/EpicGames.Core/JsonSchema.cs
+++ b/Engine/Source/Programs/Shared/EpicGames.Core/JsonSchema.cs
@@ -3,9 +3,7 @@
 using System;
 using System.Collections.Generic;
 using System.IO;
-using System.Reflection;
 using System.Text.Json;
-using System.Text.Json.Serialization;
 
 namespace EpicGames.Core
 {
diff --git a/Engine/Source/Programs/Shared/EpicGames.Core/JsonSchemaFactory.cs b/Engine/Source/Programs/Shared/EpicGames.Core/JsonSchemaFactory.cs
--- a/Engine/Source/Programs/Shared/EpicGames.Core/JsonSchemaFactory.cs
+++ b/Engine/Source/Programs/Shared/EpicGames.Core/JsonSchemaFactory.cs
@@ -2,11 +2,8 @@
 
 using System;
 using System.Collections.Generic;
-using System.Linq;
 using System.Reflection;
-using System.Text;
 using System.Text.Json.Serialization;
-using System.Threading.Tasks;
 
 namespace EpicGames.Core
 {
@@ -94,7 +91,7 @@ namespace EpicGames.Core
 			JsonSchemaTypeAttribute? attribute = type.GetCustomAttribute<JsonSchemaTypeAttribute>();
 			switch (attribute)
 			{
-				case JsonSchemaAnyAttribute any:
+				case JsonSchemaAnyAttribute _:
 					return new JsonSchemaAny();
 				case JsonSchemaStringAttribute str:
 					return new JsonSchemaString(str.Format);
@@ -167,6 +164,9 @@ namespace EpicGames.Core
 			throw new Exception($"Unknown type for schema generation: {type}");
 		}
 
+		/// <summary>
+		/// Create an enum schema type
+		/// </summary>
 		protected JsonSchemaEnum CreateEnumSchemaType(Type type)
 		{
 			string[] names = Enum.GetNames(type);
@@ -181,6 +181,9 @@ namespace EpicGames.Core
 			return new JsonSchemaEnum(names, descriptions) { Name = type.Name, Description = enumDescription };
 		}
 
+		/// <summary>
+		/// Set a one-of schema element to the known types
+		/// </summary>
 		protected void SetOneOfProperties(JsonSchemaOneOf obj, Type type, Type[] knownTypes)
 		{
 			obj.Name = type.Name;
@@ -198,6 +201,9 @@ namespace EpicGames.Core
 			}
 		}
 
+		/// <summary>
+		/// Fill out the properties for a schema object
+		/// </summary>
 		protected void SetObjectProperties(JsonSchemaObject obj, Type type)
 		{
 			obj.Name = type.Name;
