Fix spelling error
#rnx

diff --git a/Engine/Source/Programs/Shared/EpicGames.Core/FileSystemReference.cs b/Engine/Source/Programs/Shared/EpicGames.Core/FileSystemReference.cs
--- a/Engine/Source/Programs/Shared/EpicGames.Core/FileSystemReference.cs
+++ b/Engine/Source/Programs/Shared/EpicGames.Core/FileSystemReference.cs
@@ -161,7 +161,7 @@ namespace EpicGames.Core
 				return false;
 			}
 
-			// Find each occurence of the name within the remaining string, then test whether it's surrounded by directory separators
+			// Find each occurrence of the name within the remaining string, then test whether it's surrounded by directory separators
 			int matchIdx = offset;
 			for (; ; )
 			{
