Horde: Fix tracking of allocated size in MemoryMappedFileCache.

diff --git a/Engine/Source/Programs/Shared/EpicGames.Core/MemoryMappedFileCache.cs b/Engine/Source/Programs/Shared/EpicGames.Core/MemoryMappedFileCache.cs
--- a/Engine/Source/Programs/Shared/EpicGames.Core/MemoryMappedFileCache.cs
+++ b/Engine/Source/Programs/Shared/EpicGames.Core/MemoryMappedFileCache.cs
@@ -19,6 +19,7 @@ namespace EpicGames.Core
 			public FileReference Key { get; }
 			public LinkedListNode<MappedFile> ListNode { get; }
 
+			readonly MemoryMappedFileCache _cache;
 			readonly FileInfo _fileInfo;
 
 			MemoryMappedFile? _memoryMappedFile;
@@ -33,8 +34,10 @@ namespace EpicGames.Core
 
 			public ulong MappedSize => _memoryMappedViewAccessor?.SafeMemoryMappedViewHandle.ByteLength ?? 0UL;
 
-			public MappedFile(FileReference locator, FileInfo fileInfo)
+			public MappedFile(MemoryMappedFileCache cache, FileReference locator, FileInfo fileInfo)
 			{
+				_cache = cache;
+
 				Key = locator;
 				ListNode = new LinkedListNode<MappedFile>(this);
 
@@ -46,7 +49,8 @@ namespace EpicGames.Core
 					_memoryMappedViewAccessor = _memoryMappedFile.CreateViewAccessor(0, 0, MemoryMappedFileAccess.Read);
 					_memoryMappedView = new MemoryMappedView(_memoryMappedViewAccessor);
 
-					_data = _memoryMappedView.GetMemory(0, (int)fileInfo.Length);
+					_data = _memoryMappedView.GetMemory(0, (int)Math.Min(fileInfo.Length, Int32.MaxValue));
+					Interlocked.Add(ref _cache._mappedSize, _data.Length);
 				}
 				catch
 				{
@@ -57,6 +61,7 @@ namespace EpicGames.Core
 
 			public void Dispose()
 			{
+				Interlocked.Add(ref _cache._mappedSize, -_data.Length);
 				_data = ReadOnlyMemory<byte>.Empty;
 
 				if (_memoryMappedView != null)
@@ -171,8 +176,10 @@ namespace EpicGames.Core
 			_maxMappedCount = maxMappedCount;
 		}
 
-		/// <inheritdoc/>
-		public void Dispose()
+		/// <summary>
+		/// Removes any cached files
+		/// </summary>
+		public void Clear()
 		{
 			lock (_lockObject)
 			{
@@ -180,6 +187,10 @@ namespace EpicGames.Core
 			}
 		}
 
+		/// <inheritdoc/>
+		public void Dispose()
+			=> Clear();
+
 		/// <summary>
 		/// Maps data for a file into memory
 		/// </summary>
@@ -248,11 +259,9 @@ namespace EpicGames.Core
 					UnmapFiles(maxSize, _maxMappedCount - 1);
 				}
 
-				mappedFile = new MappedFile(key, fileInfo);
+				mappedFile = new MappedFile(this, key, fileInfo);
 				_pathToMappedFile.Add(key, mappedFile);
 				_mappedFiles.AddFirst(mappedFile.ListNode);
-
-				_mappedSize += (long)mappedFile.MappedSize;
 			}
 			return mappedFile;
 		}
