Horde: Fix a resource leak when deserializing blobs, due to blob data not being disposed.

diff --git a/Engine/Source/Programs/Shared/EpicGames.Horde/Storage/BlobSerializer.cs b/Engine/Source/Programs/Shared/EpicGames.Horde/Storage/BlobSerializer.cs
--- a/Engine/Source/Programs/Shared/EpicGames.Horde/Storage/BlobSerializer.cs
+++ b/Engine/Source/Programs/Shared/EpicGames.Horde/Storage/BlobSerializer.cs
@@ -219,7 +219,7 @@ namespace EpicGames.Horde.Storage
 		/// <param name="cancellationToken">Cancellation token for the operation</param>
 		public static async ValueTask<T> ReadBlobAsync<T>(this IBlobHandle handle, BlobSerializerOptions? options = null, CancellationToken cancellationToken = default)
 		{
-			BlobData data = await handle.ReadBlobDataAsync(cancellationToken);
+			using BlobData data = await handle.ReadBlobDataAsync(cancellationToken);
 			return BlobSerializer.Deserialize<T>(data, options);
 		}
 
@@ -231,7 +231,7 @@ namespace EpicGames.Horde.Storage
 		/// <param name="cancellationToken">Cancellation token for the operation</param>
 		public static async ValueTask<T> ReadBlobAsync<T>(this IBlobRef<T> handle, CancellationToken cancellationToken = default)
 		{
-			BlobData data = await handle.ReadBlobDataAsync(cancellationToken);
+			using BlobData data = await handle.ReadBlobDataAsync(cancellationToken);
 			return BlobSerializer.Deserialize<T>(data, handle.SerializerOptions);
 		}
 
