Horde: Fix some additional resource leaks for blob data reads.

diff --git a/Engine/Source/Programs/Horde/HordeServer/Commands/Generate/ToolDataCommand.cs b/Engine/Source/Programs/Horde/HordeServer/Commands/Generate/ToolDataCommand.cs
--- a/Engine/Source/Programs/Horde/HordeServer/Commands/Generate/ToolDataCommand.cs
+++ b/Engine/Source/Programs/Horde/HordeServer/Commands/Generate/ToolDataCommand.cs
@@ -143,7 +143,7 @@ namespace HordeServer.Commands.Generate
 				IBlobRef? blobRef = await client.TryReadRefAsync(refName, cancellationToken: cancellationToken);
 				if (blobRef != null)
 				{
-					BlobData blobData = await blobRef.ReadBlobDataAsync(cancellationToken);
+					using BlobData blobData = await blobRef.ReadBlobDataAsync(cancellationToken);
 					if (blobData.Type.Guid == DirectoryNode.BlobTypeGuid)
 					{
 						IBlobRef<DirectoryNode> directoryRef = BlobRef.Create<DirectoryNode>(blobRef);
diff --git a/Engine/Source/Programs/Shared/EpicGames.Horde.Tests/BundleStoreTests.cs b/Engine/Source/Programs/Shared/EpicGames.Horde.Tests/BundleStoreTests.cs
--- a/Engine/Source/Programs/Shared/EpicGames.Horde.Tests/BundleStoreTests.cs
+++ b/Engine/Source/Programs/Shared/EpicGames.Horde.Tests/BundleStoreTests.cs
@@ -68,7 +68,7 @@ namespace EpicGames.Horde.Tests
 			IBlobRef<TextNode> handle = await writer.WriteBlobAsync(node);
 			await writer.FlushAsync();
 
-			IReadOnlyMemoryOwner<byte> owner = await memoryStore.ReadBlobAsync(handle.GetLocator().BaseLocator, 0, null);
+			using IReadOnlyMemoryOwner<byte> owner = await memoryStore.ReadBlobAsync(handle.GetLocator().BaseLocator, 0, null);
 			return owner.Memory.ToArray();
 		}
 
diff --git a/Engine/Source/Programs/Shared/EpicGames.Horde/Storage/Bundles/BundleStorageClient.cs b/Engine/Source/Programs/Shared/EpicGames.Horde/Storage/Bundles/BundleStorageClient.cs
--- a/Engine/Source/Programs/Shared/EpicGames.Horde/Storage/Bundles/BundleStorageClient.cs
+++ b/Engine/Source/Programs/Shared/EpicGames.Horde/Storage/Bundles/BundleStorageClient.cs
@@ -103,7 +103,7 @@ namespace EpicGames.Horde.Storage.Bundles
 		/// </summary>
 		public async Task<IEnumerable<BlobLocator>> ReadBundleReferencesAsync(BlobLocator locator, CancellationToken cancellationToken)
 		{
-			IReadOnlyMemoryOwner<byte> data = await Backend.ReadBlobAsync(locator, 0, 32 * 1024 * 1024, cancellationToken);
+			using IReadOnlyMemoryOwner<byte> data = await Backend.ReadBlobAsync(locator, 0, 32 * 1024 * 1024, cancellationToken);
 			return ReadBundleReferences(data.Memory);
 		}
 
