Horde: Static analysis fixes.

#rnx

diff --git a/Engine/Source/Programs/Shared/EpicGames.Horde/Utilities/ConfigDocAttribute.cs b/Engine/Source/Programs/Shared/EpicGames.Horde/Utilities/ConfigDocAttribute.cs
--- a/Engine/Source/Programs/Shared/EpicGames.Horde/Utilities/ConfigDocAttribute.cs
+++ b/Engine/Source/Programs/Shared/EpicGames.Horde/Utilities/ConfigDocAttribute.cs
@@ -1,16 +1,13 @@
 // Copyright Epic Games, Inc. All Rights Reserved.
 
 using System;
-using System.Collections.Generic;
-using System.Linq;
-using System.Text;
-using System.Threading.Tasks;
 
 namespace HordeServer.Utilities
 {
 	/// <summary>
 	/// Attribute indicating that an object should generate a schema doc page
 	/// </summary>
+	[AttributeUsage(AttributeTargets.Class)]
 	public sealed class ConfigDocAttribute : Attribute
 	{
 		/// <summary>
