Horde: Add default constructors to legacy bundle header collections. Default initializer sets everything to null.

#rnx

diff --git a/Engine/Source/Programs/Shared/EpicGames.Horde/Storage/Bundles/V1/BundleHeader.cs b/Engine/Source/Programs/Shared/EpicGames.Horde/Storage/Bundles/V1/BundleHeader.cs
--- a/Engine/Source/Programs/Shared/EpicGames.Horde/Storage/Bundles/V1/BundleHeader.cs
+++ b/Engine/Source/Programs/Shared/EpicGames.Horde/Storage/Bundles/V1/BundleHeader.cs
@@ -424,6 +424,11 @@ namespace EpicGames.Horde.Storage.Bundles.V1
 		/// <inheritdoc/>
 		public int Count => _types.Length;
 
+		/// <summary>
+		/// Constructor
+		/// </summary>
+		public BundleTypeCollection() : this(Array.Empty<BlobType>()) { }
+
 		/// <summary>
 		/// Constructor
 		/// </summary>
@@ -493,6 +498,11 @@ namespace EpicGames.Horde.Storage.Bundles.V1
 		/// <inheritdoc/>
 		public int Count => _imports.Length;
 
+		/// <summary>
+		/// Constructor
+		/// </summary>
+		public BundleImportCollection() : this(Array.Empty<BlobLocator>()) { }
+
 		/// <summary>
 		/// Constructor
 		/// </summary>
@@ -640,6 +650,11 @@ namespace EpicGames.Horde.Storage.Bundles.V1
 		/// <inheritdoc/>
 		public int Count => _packets?.Length ?? 0;
 
+		/// <summary>
+		/// Constructor
+		/// </summary>
+		public BundlePacketCollection() : this(Array.Empty<BundlePacket>()) { }
+	
 		/// <summary>
 		/// Constructor
 		/// </summary>
@@ -798,6 +813,11 @@ namespace EpicGames.Horde.Storage.Bundles.V1
 		/// <inheritdoc/>
 		public int Count => _exports.Length;
 
+		/// <summary>
+		/// Constructor
+		/// </summary>
+		public BundleExportCollection() : this(Array.Empty<BundleExport>()) { }
+
 		/// <summary>
 		/// Constructor
 		/// </summary>
