Horde: Move issue handler priority setting onto the handler itself, so we can populate handlers without needing their attributes. Also remove plain text message from captured log events.

#rnx

diff --git a/Engine/Source/Programs/Horde/Plugins/Build/HordeServer.Build/Issues/IssueService.cs b/Engine/Source/Programs/Horde/Plugins/Build/HordeServer.Build/Issues/IssueService.cs
--- a/Engine/Source/Programs/Horde/Plugins/Build/HordeServer.Build/Issues/IssueService.cs
+++ b/Engine/Source/Programs/Horde/Plugins/Build/HordeServer.Build/Issues/IssueService.cs
@@ -152,7 +152,7 @@ namespace HordeServer.Issues
 			public ILogEventData EventData { get; }
 
 			public IssueEventInternal(ILogEvent logEvent, ILogEventData logEventData)
-				: base(logEvent.LineIndex, GetLogLevelFromSeverity(logEventData.Severity), logEventData.EventId, logEventData.Message, logEventData.Lines)
+				: base(logEvent.LineIndex, GetLogLevelFromSeverity(logEventData.Severity), logEventData.EventId, logEventData.Lines)
 			{
 				Event = logEvent;
 				EventData = logEventData;
@@ -285,7 +285,7 @@ namespace HordeServer.Issues
 					handlerTypes.Add((type, attribute));
 				}
 			}
-			_handlerTypes = handlerTypes.OrderByDescending(x => x.Attribute.Priority).ToArray();
+			_handlerTypes = handlerTypes.ToArray();
 		}
 
 		/// <inheritdoc/>
@@ -699,9 +699,10 @@ namespace HordeServer.Issues
 					handlers.Add((IssueHandler)ActivatorUtilities.CreateInstance(serviceProvider, type));
 				}
 			}
+			handlers.SortBy(x => -x.Priority);
 
 			// Create all the issue definitions by passing each log event to the handlers in order until one attaches it to an issue
-			List<ILogEvent> stepEvents = await log.GetEventsAsync(cancellationToken: cancellationToken);
+			List <ILogEvent> stepEvents = await log.GetEventsAsync(cancellationToken: cancellationToken);
 			foreach (ILogEvent stepEvent in stepEvents)
 			{
 				ILogEventData stepEventData = await stepEvent.GetDataAsync(cancellationToken);
@@ -754,7 +755,7 @@ namespace HordeServer.Issues
 				_logger.LogInformation("Group {Digest}: Type '{FingerprintType}', keys '{FingerprintKeys}', {NumEvents} events", eventGroup.Id.ToString(), eventGroup.Fingerprint.Type, String.Join(", ", eventGroup.Fingerprint.Keys), eventGroup.Events.Count);
 				foreach (IssueEvent eventItem in eventGroup.Events)
 				{
-					_logger.LogDebug("Group {Digest}: [{Line}] {Message}", eventGroup.Id.ToString(), eventItem.LineIndex, eventItem.Message);
+					_logger.LogDebug("Group {Digest}: [{Line}] {Message}", eventGroup.Id.ToString(), eventItem.LineIndex, eventItem.Render());
 				}
 			}
 
diff --git a/Engine/Source/Programs/Shared/EpicGames.Horde/Issues/Handlers/BuildGraphIssueHandler.cs b/Engine/Source/Programs/Shared/EpicGames.Horde/Issues/Handlers/BuildGraphIssueHandler.cs
--- a/Engine/Source/Programs/Shared/EpicGames.Horde/Issues/Handlers/BuildGraphIssueHandler.cs
+++ b/Engine/Source/Programs/Shared/EpicGames.Horde/Issues/Handlers/BuildGraphIssueHandler.cs
@@ -9,11 +9,14 @@ namespace EpicGames.Horde.Issues.Handlers
 	/// <summary>
 	/// Instance of a particular BuildGraph script error
 	/// </summary>
-	[IssueHandler(Priority = 10)]
+	[IssueHandler]
 	public class BuildGraphIssueHandler : IssueHandler
 	{
 		readonly List<IssueEventGroup> _issues = new List<IssueEventGroup>();
 
+		/// <inheritdoc/>
+		public override int Priority => 10;
+
 		/// <summary>
 		/// Determines if the given event id matches
 		/// </summary>
diff --git a/Engine/Source/Programs/Shared/EpicGames.Horde/Issues/Handlers/CompileIssueHandler.cs b/Engine/Source/Programs/Shared/EpicGames.Horde/Issues/Handlers/CompileIssueHandler.cs
--- a/Engine/Source/Programs/Shared/EpicGames.Horde/Issues/Handlers/CompileIssueHandler.cs
+++ b/Engine/Source/Programs/Shared/EpicGames.Horde/Issues/Handlers/CompileIssueHandler.cs
@@ -9,7 +9,7 @@ namespace EpicGames.Horde.Issues.Handlers
 	/// <summary>
 	/// Instance of a particular compile error
 	/// </summary>
-	[IssueHandler(Priority = 10)]
+	[IssueHandler]
 	public class CompileIssueHandler : IssueHandler
 	{
 		/// <summary>
@@ -25,6 +25,9 @@ namespace EpicGames.Horde.Issues.Handlers
 		readonly IssueHandlerContext _context;
 		readonly List<IssueEventGroup> _issues = new List<IssueEventGroup>();
 
+		/// <inheritdoc/>
+		public override int Priority => 10;
+
 		/// <summary>
 		/// Constructor
 		/// </summary>
diff --git a/Engine/Source/Programs/Shared/EpicGames.Horde/Issues/Handlers/ContentIssueHandler.cs b/Engine/Source/Programs/Shared/EpicGames.Horde/Issues/Handlers/ContentIssueHandler.cs
--- a/Engine/Source/Programs/Shared/EpicGames.Horde/Issues/Handlers/ContentIssueHandler.cs
+++ b/Engine/Source/Programs/Shared/EpicGames.Horde/Issues/Handlers/ContentIssueHandler.cs
@@ -9,7 +9,7 @@ namespace EpicGames.Horde.Issues.Handlers
 	/// <summary>
 	/// Instance of a particular compile error
 	/// </summary>
-	[IssueHandler(Priority = 10)]
+	[IssueHandler]
 	public class ContentIssueHandler : IssueHandler
 	{
 		readonly List<IssueEventGroup> _issues = new List<IssueEventGroup>();
@@ -17,6 +17,9 @@ namespace EpicGames.Horde.Issues.Handlers
 		static bool IsMatchingEventId(EventId eventId) => eventId == KnownLogEvents.Engine_AssetLog;
 		static bool IsMaskedEventId(EventId eventId) => eventId == KnownLogEvents.ExitCode;
 
+		/// <inheritdoc/>
+		public override int Priority => 10;
+
 		/// <inheritdoc/>
 		public override bool HandleEvent(IssueEvent issueEvent)
 		{
diff --git a/Engine/Source/Programs/Shared/EpicGames.Horde/Issues/Handlers/CopyrightIssueHandler.cs b/Engine/Source/Programs/Shared/EpicGames.Horde/Issues/Handlers/CopyrightIssueHandler.cs
--- a/Engine/Source/Programs/Shared/EpicGames.Horde/Issues/Handlers/CopyrightIssueHandler.cs
+++ b/Engine/Source/Programs/Shared/EpicGames.Horde/Issues/Handlers/CopyrightIssueHandler.cs
@@ -8,11 +8,14 @@ namespace EpicGames.Horde.Issues.Handlers
 	/// <summary>
 	/// Instance of a particular compile error
 	/// </summary>
-	[IssueHandler(Priority = 10)]
+	[IssueHandler]
 	public class CopyrightIssueHandler : IssueHandler
 	{
 		readonly List<IssueEventGroup> _issues = new List<IssueEventGroup>();
 
+		/// <inheritdoc/>
+		public override int Priority => 10;
+
 		/// <inheritdoc/>
 		public override bool HandleEvent(IssueEvent issueEvent)
 		{
diff --git a/Engine/Source/Programs/Shared/EpicGames.Horde/Issues/Handlers/DefaultIssueHandler.cs b/Engine/Source/Programs/Shared/EpicGames.Horde/Issues/Handlers/DefaultIssueHandler.cs
--- a/Engine/Source/Programs/Shared/EpicGames.Horde/Issues/Handlers/DefaultIssueHandler.cs
+++ b/Engine/Source/Programs/Shared/EpicGames.Horde/Issues/Handlers/DefaultIssueHandler.cs
@@ -7,12 +7,15 @@ namespace EpicGames.Horde.Issues.Handlers
 	/// <summary>
 	/// Default handler for log events not matched by any other handler
 	/// </summary>
-	[IssueHandler(Priority = 0)]
+	[IssueHandler]
 	public class DefaultIssueHandler : IssueHandler
 	{
 		readonly IssueHandlerContext _context;
 		IssueEventGroup? _issue;
 
+		/// <inheritdoc/>
+		public override int Priority => 0;
+
 		/// <summary>
 		/// Constructor
 		/// </summary>
diff --git a/Engine/Source/Programs/Shared/EpicGames.Horde/Issues/Handlers/GauntletIssueHandler.cs b/Engine/Source/Programs/Shared/EpicGames.Horde/Issues/Handlers/GauntletIssueHandler.cs
--- a/Engine/Source/Programs/Shared/EpicGames.Horde/Issues/Handlers/GauntletIssueHandler.cs
+++ b/Engine/Source/Programs/Shared/EpicGames.Horde/Issues/Handlers/GauntletIssueHandler.cs
@@ -13,7 +13,7 @@ namespace EpicGames.Horde.Issues.Handlers
 	/// <summary>
 	/// Instance of a particular Gauntlet error
 	/// </summary>
-	[IssueHandler(Priority = 10)]
+	[IssueHandler]
 	public class GauntletIssueHandler : IssueHandler
 	{
 		/// <summary>
@@ -72,6 +72,9 @@ namespace EpicGames.Horde.Issues.Handlers
 			{ KnownLogEvents.Gauntlet_FatalEvent, FatalPrefix}
 		};
 
+		/// <inheritdoc/>
+		public override int Priority => 10;
+
 		/// <summary>
 		/// Constructor
 		/// </summary>
@@ -130,7 +133,7 @@ namespace EpicGames.Horde.Issues.Handlers
 		private static bool TryGetHash(IssueEvent issueEvent, out Md5Hash hash)
 		{
 			// Use only the summary if one is found instead of the full callstack
-			string sanitized = GetSummaryProperty(issueEvent) ?? issueEvent.Message;
+			string sanitized = GetSummaryProperty(issueEvent) ?? issueEvent.Render();
 			sanitized = sanitized.Length > MaxMessageLength ? sanitized.Substring(0, MaxMessageLength) : sanitized;
 			sanitized = sanitized.Trim().ToUpperInvariant();
 			sanitized = Regex.Replace(sanitized, @"(?<![A-Z])(?:[A-Z]:|/)[^ :]+[/\\]SYNC[/\\]", "{root}/"); // Redact things that look like workspace roots; may be different between agents
diff --git a/Engine/Source/Programs/Shared/EpicGames.Horde/Issues/Handlers/HashedIssueHandler.cs b/Engine/Source/Programs/Shared/EpicGames.Horde/Issues/Handlers/HashedIssueHandler.cs
--- a/Engine/Source/Programs/Shared/EpicGames.Horde/Issues/Handlers/HashedIssueHandler.cs
+++ b/Engine/Source/Programs/Shared/EpicGames.Horde/Issues/Handlers/HashedIssueHandler.cs
@@ -12,7 +12,7 @@ namespace EpicGames.Horde.Issues.Handlers
 	/// <summary>
 	/// Instance of a particular compile error
 	/// </summary>
-	[IssueHandler(Priority = 1)]
+	[IssueHandler]
 	public class HashedIssueHandler : IssueHandler
 	{
 		readonly IssueHandlerContext _context;
@@ -38,6 +38,9 @@ namespace EpicGames.Horde.Issues.Handlers
 			return s_knownGeneralEvents.Contains(eventId) || (eventId.Id >= KnownLogEvents.Systemic.Id && eventId.Id <= KnownLogEvents.Systemic_Max.Id);
 		}
 
+		/// <inheritdoc/>
+		public override int Priority => 1;
+
 		/// <inheritdoc/>
 		public override bool HandleEvent(IssueEvent logEvent)
 		{
@@ -58,7 +61,7 @@ namespace EpicGames.Horde.Issues.Handlers
 
 			foreach (IssueEvent stepEvent in _issueEvents)
 			{
-				string hashSource = stepEvent.Message;
+				string hashSource = stepEvent.Render();
 
 				if (!allGeneral && stepEvent.EventId != null)
 				{
diff --git a/Engine/Source/Programs/Shared/EpicGames.Horde/Issues/Handlers/LocalizationIssueHandler.cs b/Engine/Source/Programs/Shared/EpicGames.Horde/Issues/Handlers/LocalizationIssueHandler.cs
--- a/Engine/Source/Programs/Shared/EpicGames.Horde/Issues/Handlers/LocalizationIssueHandler.cs
+++ b/Engine/Source/Programs/Shared/EpicGames.Horde/Issues/Handlers/LocalizationIssueHandler.cs
@@ -11,11 +11,14 @@ namespace EpicGames.Horde.Issues.Handlers
 	/// <summary>
 	/// Instance of a localization error
 	/// </summary>
-	[IssueHandler(Priority = 10)]
+	[IssueHandler]
 	public class LocalizationIssueHandler : IssueHandler
 	{
 		readonly List<IssueEventGroup> _issues = new List<IssueEventGroup>();
 
+		/// <inheritdoc/>
+		public override int Priority => 10;
+
 		/// <summary>
 		/// Determines if the given event id matches
 		/// </summary>
diff --git a/Engine/Source/Programs/Shared/EpicGames.Horde/Issues/Handlers/PerforceCaseIssueHandler.cs b/Engine/Source/Programs/Shared/EpicGames.Horde/Issues/Handlers/PerforceCaseIssueHandler.cs
--- a/Engine/Source/Programs/Shared/EpicGames.Horde/Issues/Handlers/PerforceCaseIssueHandler.cs
+++ b/Engine/Source/Programs/Shared/EpicGames.Horde/Issues/Handlers/PerforceCaseIssueHandler.cs
@@ -9,11 +9,14 @@ namespace EpicGames.Horde.Issues.Handlers
 	/// <summary>
 	/// Instance of a Perforce case mismatch error
 	/// </summary>
-	[IssueHandler(Priority = 10)]
+	[IssueHandler]
 	public class PerforceCaseIssueHandler : IssueHandler
 	{
 		readonly List<IssueEventGroup> _issues = new List<IssueEventGroup>();
 
+		/// <inheritdoc/>
+		public override int Priority => 10;
+
 		/// <summary>
 		/// Determines if the given event id matches
 		/// </summary>
diff --git a/Engine/Source/Programs/Shared/EpicGames.Horde/Issues/Handlers/SanitizerIssueHandler.cs b/Engine/Source/Programs/Shared/EpicGames.Horde/Issues/Handlers/SanitizerIssueHandler.cs
--- a/Engine/Source/Programs/Shared/EpicGames.Horde/Issues/Handlers/SanitizerIssueHandler.cs
+++ b/Engine/Source/Programs/Shared/EpicGames.Horde/Issues/Handlers/SanitizerIssueHandler.cs
@@ -11,7 +11,7 @@ namespace EpicGames.Horde.Issues.Handlers
 	/// <summary>
 	/// Instance of a specific Thread Sanitizer error
 	/// </summary>
-	[IssueHandler(Priority = 10)]
+	[IssueHandler]
 	public class SanitizerIssueHandler : IssueHandler
 	{
 		static readonly Utf8String s_summaryReasonAnnotation = new Utf8String("SummaryReason");
@@ -19,6 +19,8 @@ namespace EpicGames.Horde.Issues.Handlers
 		static readonly Utf8String s_sanitizerNameAnnotation = new Utf8String("SanitizerName");
 		readonly List<IssueEventGroup> _issues = new List<IssueEventGroup>();
 
+		/// <inheritdoc/>
+		public override int Priority => 10;
 		/// <summary>
 		/// Log value describing the thread sanitizer error summary reason
 		/// </summary>
diff --git a/Engine/Source/Programs/Shared/EpicGames.Horde/Issues/Handlers/ScopedIssueHandler.cs b/Engine/Source/Programs/Shared/EpicGames.Horde/Issues/Handlers/ScopedIssueHandler.cs
--- a/Engine/Source/Programs/Shared/EpicGames.Horde/Issues/Handlers/ScopedIssueHandler.cs
+++ b/Engine/Source/Programs/Shared/EpicGames.Horde/Issues/Handlers/ScopedIssueHandler.cs
@@ -12,7 +12,7 @@ namespace EpicGames.Horde.Issues.Handlers
 	/// <summary>
 	/// Instance of a particular compile error
 	/// </summary>
-	[IssueHandler(Priority = 2, Tag = "Scoped")]
+	[IssueHandler(Tag = "Scoped")]
 	public class ScopedIssueHandler : IssueHandler
 	{
 		const string NodeName = "Node";
@@ -21,6 +21,8 @@ namespace EpicGames.Horde.Issues.Handlers
 		readonly IssueHandlerContext _context;
 		readonly List<IssueEventGroup> _issues = new List<IssueEventGroup>();
 
+		/// <inheritdoc/>
+		public override int Priority => 2;
 		/// <summary>
 		/// Constructor
 		/// </summary>
@@ -62,7 +64,7 @@ namespace EpicGames.Horde.Issues.Handlers
 
 			string fingerprintType = $"Scoped:{scope}";
 
-			string hashSource = logEvent.Message;
+			string hashSource = logEvent.Render();
 
 			if (TryGetHash(hashSource, out Md5Hash hash))
 			{
diff --git a/Engine/Source/Programs/Shared/EpicGames.Horde/Issues/Handlers/ShaderIssueHandler.cs b/Engine/Source/Programs/Shared/EpicGames.Horde/Issues/Handlers/ShaderIssueHandler.cs
--- a/Engine/Source/Programs/Shared/EpicGames.Horde/Issues/Handlers/ShaderIssueHandler.cs
+++ b/Engine/Source/Programs/Shared/EpicGames.Horde/Issues/Handlers/ShaderIssueHandler.cs
@@ -9,7 +9,7 @@ namespace EpicGames.Horde.Issues.Handlers
 	/// <summary>
 	/// Instance of a particular shader compile error
 	/// </summary>
-	[IssueHandler(Priority = 10)]
+	[IssueHandler]
 	public class ShaderIssueHandler : IssueHandler
 	{
 		readonly List<IssueEventGroup> _issues = new List<IssueEventGroup>();
@@ -17,6 +17,9 @@ namespace EpicGames.Horde.Issues.Handlers
 		static bool IsMatchingEventId(EventId id) => id == KnownLogEvents.Engine_ShaderCompiler;
 		static bool IsMaskedEventId(EventId id) => id == KnownLogEvents.Generic || id == KnownLogEvents.ExitCode || id == KnownLogEvents.Systemic_Xge_BuildFailed || id == KnownLogEvents.Engine_Crash || id == KnownLogEvents.AutomationTool_CrashExitCode || id == KnownLogEvents.Engine_AppError;
 
+		/// <inheritdoc/>
+		public override int Priority => 10;
+
 		/// <inheritdoc/>
 		public override bool HandleEvent(IssueEvent logEvent)
 		{
diff --git a/Engine/Source/Programs/Shared/EpicGames.Horde/Issues/Handlers/SymbolIssueHandler.cs b/Engine/Source/Programs/Shared/EpicGames.Horde/Issues/Handlers/SymbolIssueHandler.cs
--- a/Engine/Source/Programs/Shared/EpicGames.Horde/Issues/Handlers/SymbolIssueHandler.cs
+++ b/Engine/Source/Programs/Shared/EpicGames.Horde/Issues/Handlers/SymbolIssueHandler.cs
@@ -9,13 +9,16 @@ namespace EpicGames.Horde.Issues.Handlers
 	/// <summary>
 	/// Instance of a particular compile error
 	/// </summary>
-	[IssueHandler(Priority = 10)]
+	[IssueHandler]
 	public class SymbolIssueHandler : IssueHandler
 	{
 		const string EventIdName = "EventId";
 
 		readonly List<IssueEventGroup> _issues = new List<IssueEventGroup>();
 
+		/// <inheritdoc/>
+		public override int Priority => 10;
+
 		/// <summary>
 		/// Determines if the given event id matches
 		/// </summary>
diff --git a/Engine/Source/Programs/Shared/EpicGames.Horde/Issues/Handlers/SystemicIssueHandler.cs b/Engine/Source/Programs/Shared/EpicGames.Horde/Issues/Handlers/SystemicIssueHandler.cs
--- a/Engine/Source/Programs/Shared/EpicGames.Horde/Issues/Handlers/SystemicIssueHandler.cs
+++ b/Engine/Source/Programs/Shared/EpicGames.Horde/Issues/Handlers/SystemicIssueHandler.cs
@@ -9,7 +9,7 @@ namespace EpicGames.Horde.Issues.Handlers
 	/// <summary>
 	/// Instance of a particular systemic error
 	/// </summary>
-	[IssueHandler(Priority = 10)]
+	[IssueHandler]
 	public class SystemicIssueHandler : IssueHandler
 	{
 		/// <summary>
@@ -17,6 +17,9 @@ namespace EpicGames.Horde.Issues.Handlers
 		/// </summary>
 		static readonly HashSet<EventId> s_knownSystemic = new HashSet<EventId> { KnownLogEvents.Horde, KnownLogEvents.Horde_InvalidPreflight };
 
+		/// <inheritdoc/>
+		public override int Priority => 10;
+
 		/// <summary>
 		/// Determines if the given event id matches
 		/// </summary>
diff --git a/Engine/Source/Programs/Shared/EpicGames.Horde/Issues/Handlers/UnacceptableWordsIssueHandler.cs b/Engine/Source/Programs/Shared/EpicGames.Horde/Issues/Handlers/UnacceptableWordsIssueHandler.cs
--- a/Engine/Source/Programs/Shared/EpicGames.Horde/Issues/Handlers/UnacceptableWordsIssueHandler.cs
+++ b/Engine/Source/Programs/Shared/EpicGames.Horde/Issues/Handlers/UnacceptableWordsIssueHandler.cs
@@ -8,11 +8,14 @@ namespace EpicGames.Horde.Issues.Handlers
 	/// <summary>
 	/// Instance of a Perforce case mismatch error
 	/// </summary>
-	[IssueHandler(Priority = 8)]
+	[IssueHandler]
 	public class UnacceptableWordsIssueHandler : IssueHandler
 	{
 		readonly List<IssueEventGroup> _issues = new List<IssueEventGroup>();
 
+		/// <inheritdoc/>
+		public override int Priority => 8;
+
 		/// <inheritdoc/>
 		public override bool HandleEvent(IssueEvent issueEvent)
 		{
diff --git a/Engine/Source/Programs/Shared/EpicGames.Horde/Issues/IssueEvent.cs b/Engine/Source/Programs/Shared/EpicGames.Horde/Issues/IssueEvent.cs
--- a/Engine/Source/Programs/Shared/EpicGames.Horde/Issues/IssueEvent.cs
+++ b/Engine/Source/Programs/Shared/EpicGames.Horde/Issues/IssueEvent.cs
@@ -1,6 +1,10 @@
 // Copyright Epic Games, Inc. All Rights Reserved.
 
+using System;
 using System.Collections.Generic;
+using System.Linq;
+using System.Text.Json;
+using System.Text.Json.Serialization;
 using EpicGames.Core;
 using Microsoft.Extensions.Logging;
 
@@ -26,11 +30,6 @@ namespace EpicGames.Horde.Issues
 		/// </summary>
 		public EventId? EventId { get; }
 
-		/// <summary>
-		/// The complete rendered message, in plaintext
-		/// </summary>
-		public string Message { get; }
-
 		/// <summary>
 		/// Gets this event data as a BSON document
 		/// </summary>
@@ -39,17 +38,22 @@ namespace EpicGames.Horde.Issues
 		/// <summary>
 		/// Constructor
 		/// </summary>
-		public IssueEvent(int lineIndex, LogLevel severity, EventId? eventId, string message, IReadOnlyList<JsonLogEvent> lines)
+		public IssueEvent(int lineIndex, LogLevel severity, EventId? eventId, IReadOnlyList<JsonLogEvent> lines)
 		{
 			LineIndex = lineIndex;
 			Severity = severity;
 			EventId = eventId;
-			Message = message;
 			Lines = lines;
 		}
 
+		/// <summary>
+		/// Renders the entire message of this event
+		/// </summary>
+		public string Render()
+			=> String.Join("\n", Lines.Select(x => x.GetRenderedMessage().ToString()));
+
 		/// <inheritdoc/>
-		public override string ToString() => $"[{LineIndex}] {Message}";
+		public override string ToString() => $"[{LineIndex}] {Render()}";
 	}
 
 	/// <summary>
@@ -80,6 +84,7 @@ namespace EpicGames.Horde.Issues
 		/// <summary>
 		/// Filter for changes that should be included in this issue
 		/// </summary>
+		[JsonConverter(typeof(ChangeFilterJsonConverter))]
 		public IReadOnlyList<string> ChangeFilter { get; set; }
 
 		/// <summary>
@@ -100,4 +105,15 @@ namespace EpicGames.Horde.Issues
 			ChangeFilter = changeFilter;
 		}
 	}
+
+	class ChangeFilterJsonConverter : JsonConverter<IReadOnlyList<string>>
+	{
+		/// <inheritdoc/>
+		public override IReadOnlyList<string>? Read(ref Utf8JsonReader reader, Type typeToConvert, JsonSerializerOptions options)
+			=> (reader.GetString() ?? String.Empty).Split(';');
+
+		/// <inheritdoc/>
+		public override void Write(Utf8JsonWriter writer, IReadOnlyList<string> value, JsonSerializerOptions options)
+			=> writer.WriteStringValue(String.Join(";", value));
+	}
 }
diff --git a/Engine/Source/Programs/Shared/EpicGames.Horde/Issues/IssueHandler.cs b/Engine/Source/Programs/Shared/EpicGames.Horde/Issues/IssueHandler.cs
--- a/Engine/Source/Programs/Shared/EpicGames.Horde/Issues/IssueHandler.cs
+++ b/Engine/Source/Programs/Shared/EpicGames.Horde/Issues/IssueHandler.cs
@@ -13,17 +13,36 @@ namespace EpicGames.Horde.Issues
 	[AttributeUsage(AttributeTargets.Class)]
 	public sealed class IssueHandlerAttribute : Attribute
 	{
-		/// <summary> 
-		/// Priority of this handler
-		/// </summary>
-		public int Priority { get; set; }
-
 		/// <summary>
 		/// Class of handler which can be explicitly enabled via a workflow
 		/// </summary>
 		public string? Tag { get; set; }
 	}
 
+	/// <summary>
+	/// Interface for issue matchers
+	/// </summary>
+	public abstract class IssueHandler
+	{
+		/// <summary>
+		/// Priority value for this issue handler. Handlers are processed in order of decreasing priority.
+		/// </summary>
+		public abstract int Priority { get; }
+
+		/// <summary>
+		/// Attempts to assign a log event to an issue
+		/// </summary>
+		/// <param name="issueEvent">Events to process</param>
+		/// <returns>Issue definition for this log event</returns>
+		public abstract bool HandleEvent(IssueEvent issueEvent);
+
+		/// <summary>
+		/// Gets all the issues created by this handler
+		/// </summary>
+		/// <returns></returns>
+		public abstract IEnumerable<IssueEventGroup> GetIssues();
+	}
+
 	/// <summary>
 	/// Context object for issue handlers
 	/// </summary>
@@ -60,23 +79,4 @@ namespace EpicGames.Horde.Issues
 			NodeAnnotations = nodeAnnotations;
 		}
 	}
-
-	/// <summary>
-	/// Interface for issue matchers
-	/// </summary>
-	public abstract class IssueHandler
-	{
-		/// <summary>
-		/// Attempts to assign a log event to an issue
-		/// </summary>
-		/// <param name="issueEvent">Events to process</param>
-		/// <returns>Issue definition for this log event</returns>
-		public abstract bool HandleEvent(IssueEvent issueEvent);
-
-		/// <summary>
-		/// Gets all the issues created by this handler
-		/// </summary>
-		/// <returns></returns>
-		public abstract IEnumerable<IssueEventGroup> GetIssues();
-	}
 }
