Horde: Improve Compute Failure messaging with InternalServerErrors

Don't retry on ISE and handle that status code to print the resulting message from the server

#rnx

diff --git a/Engine/Source/Programs/Shared/EpicGames.Horde/Compute/Clients/ServerComputeClient.cs b/Engine/Source/Programs/Shared/EpicGames.Horde/Compute/Clients/ServerComputeClient.cs
--- a/Engine/Source/Programs/Shared/EpicGames.Horde/Compute/Clients/ServerComputeClient.cs
+++ b/Engine/Source/Programs/Shared/EpicGames.Horde/Compute/Clients/ServerComputeClient.cs
@@ -328,6 +328,20 @@ namespace EpicGames.Horde.Compute.Clients
 					}
 				}
 
+				if (httpResponse.StatusCode == HttpStatusCode.InternalServerError)
+				{
+					string? content;
+					try
+					{
+						content = await httpResponse.Content.ReadAsStringAsync(cancellationToken);
+					}
+					catch
+					{
+						content = "None";
+					}
+					throw new ComputeClientException($"InternalServerError requesting compute resources: \"{content}\"");
+				}
+
 				httpResponse.EnsureSuccessStatusCode();
 				response = await httpResponse.Content.ReadFromJsonAsync<AssignComputeResponse>(HordeHttpClient.JsonSerializerOptions, cancellationToken);
 				if (response == null)
diff --git a/Engine/Source/Programs/Shared/EpicGames.Horde/HordeClient.cs b/Engine/Source/Programs/Shared/EpicGames.Horde/HordeClient.cs
--- a/Engine/Source/Programs/Shared/EpicGames.Horde/HordeClient.cs
+++ b/Engine/Source/Programs/Shared/EpicGames.Horde/HordeClient.cs
@@ -131,7 +131,7 @@ namespace EpicGames.Horde
 			// when allocating compute resources when none are available). This pathway is handled explicitly on the application side.
 			return Policy<HttpResponseMessage>
 				.Handle<HttpRequestException>()
-				.OrResult(x => (x.StatusCode >= HttpStatusCode.InternalServerError && x.StatusCode != HttpStatusCode.ServiceUnavailable) || x.StatusCode == HttpStatusCode.RequestTimeout)
+				.OrResult(x => (x.StatusCode > HttpStatusCode.InternalServerError && x.StatusCode != HttpStatusCode.ServiceUnavailable) || x.StatusCode == HttpStatusCode.RequestTimeout)
 				.WaitAndRetryAsync(retryTimes, OnTimeoutAsync);
 		}
 
