Allow p4 changes to be called with the -r (reverse) specifier

diff --git a/Engine/Source/Programs/Shared/EpicGames.Perforce/PerforceConnection.cs b/Engine/Source/Programs/Shared/EpicGames.Perforce/PerforceConnection.cs
--- a/Engine/Source/Programs/Shared/EpicGames.Perforce/PerforceConnection.cs
+++ b/Engine/Source/Programs/Shared/EpicGames.Perforce/PerforceConnection.cs
@@ -1022,6 +1022,10 @@ namespace EpicGames.Perforce
 			{
 				arguments.Add("-l");
 			}
+			if ((options & ChangesOptions.Reverse) != 0)
+			{
+				arguments.Add("-r");
+			}
 			if ((options & ChangesOptions.TruncatedLongOutput) != 0)
 			{
 				arguments.Add("-L");
diff --git a/Engine/Source/Programs/Shared/EpicGames.Perforce/Records/ChangesOptions.cs b/Engine/Source/Programs/Shared/EpicGames.Perforce/Records/ChangesOptions.cs
--- a/Engine/Source/Programs/Shared/EpicGames.Perforce/Records/ChangesOptions.cs
+++ b/Engine/Source/Programs/Shared/EpicGames.Perforce/Records/ChangesOptions.cs
@@ -39,5 +39,10 @@ namespace EpicGames.Perforce
 		/// View restricted changes (requires admin permission)
 		/// </summary>
 		IncludeRestricted = 16,
+
+		/// <summary>
+		/// Gets oldest changes first
+		/// </summary>
+		Reverse = 32,
 	}
 }
