[UBA]
* Fixed static analysis error

diff --git a/Engine/Source/Programs/Shared/EpicGames.UBA/Impl/SessionServerImpl.cs b/Engine/Source/Programs/Shared/EpicGames.UBA/Impl/SessionServerImpl.cs
--- a/Engine/Source/Programs/Shared/EpicGames.UBA/Impl/SessionServerImpl.cs
+++ b/Engine/Source/Programs/Shared/EpicGames.UBA/Impl/SessionServerImpl.cs
@@ -98,7 +98,7 @@ namespace EpicGames.UBA
 		static extern void SessionServer_UpdateProgress(IntPtr server, uint processesTotal, uint processesDone, uint errorCount);
 
 		[DllImport("UbaHost", CharSet = CharSet.Auto)]
-		static extern uint SessionServer_UpdateStatus(IntPtr server, uint statusRow, uint statusColumn, string statusText, byte statusType, string? statusLink);
+		static extern void SessionServer_UpdateStatus(IntPtr server, uint statusRow, uint statusColumn, string statusText, byte statusType, string? statusLink);
 
 		[DllImport("UbaHost", CharSet = CharSet.Auto)]
 		static extern IntPtr SessionServer_RunProcess(IntPtr server, IntPtr info, bool async, bool enableDetour);
