ParentView is optional per perforce documentation
Fixes Unhandled Exception (Missing 'ParentView' tag when parsing 'StreamRecord') for streams that does not have it set

diff --git a/Engine/Source/Programs/Shared/EpicGames.Perforce/Records/StreamRecord.cs b/Engine/Source/Programs/Shared/EpicGames.Perforce/Records/StreamRecord.cs
--- a/Engine/Source/Programs/Shared/EpicGames.Perforce/Records/StreamRecord.cs
+++ b/Engine/Source/Programs/Shared/EpicGames.Perforce/Records/StreamRecord.cs
@@ -67,7 +67,7 @@ namespace EpicGames.Perforce
 		/// <summary>
 		/// Parent view
 		/// </summary>
-		[PerforceTag("ParentView")]
+		[PerforceTag("ParentView", Optional = true)]
 		public string? ParentView { get; set; }
 
 		/// <summary>
