Horde: Support destructing objects in standard EpicGames.Core logger, and attaching additional properties via ILogger.BeginScope().

diff --git a/Engine/Source/Programs/Shared/EpicGames.Core/Log.cs b/Engine/Source/Programs/Shared/EpicGames.Core/Log.cs
--- a/Engine/Source/Programs/Shared/EpicGames.Core/Log.cs
+++ b/Engine/Source/Programs/Shared/EpicGames.Core/Log.cs
@@ -784,12 +784,59 @@ namespace EpicGames.Core
 	}
 
 	/// <summary>
-	/// NullScope which does nothing
+	/// Base implementation for loggers which handles creation of LogEvents from log commands
 	/// </summary>
-	internal sealed class NullScope : IDisposable
+	public sealed class LoggerScopeCollection
 	{
-		/// <inheritdoc/>
-		public void Dispose() { }
+		class Scope : IDisposable
+		{
+			readonly LoggerScopeCollection _outer;
+			readonly IEnumerable<KeyValuePair<string, object>>? _properties;
+
+			public IEnumerable<KeyValuePair<string, object>>? Properties => _properties;
+
+			public Scope(LoggerScopeCollection outer, object? value)
+			{
+				_outer = outer;
+				_properties = value as IEnumerable<KeyValuePair<string, object>>;
+
+				_outer._scopes.Add(this);
+			}
+
+			public void Dispose()
+				=> _outer._scopes.Remove(this);
+		}
+
+		readonly List<Scope> _scopes = new List<Scope>();
+
+		/// <summary>
+		/// Tests whether the collection is empty
+		/// </summary>
+		public bool IsEmpty()
+			=> _scopes.Count == 0;
+
+		/// <summary>
+		/// Creates a scope with the given state
+		/// </summary>
+		public IDisposable BeginScope<TState>(TState state)
+			=> new Scope(this, state);
+
+		/// <summary>
+		/// Enumerates all current properties
+		/// </summary>
+		public IEnumerable<KeyValuePair<string, object>> GetProperties()
+		{
+			foreach (Scope scope in _scopes)
+			{
+				if (scope.Properties != null)
+				{
+					foreach (KeyValuePair<string, object> property in scope.Properties)
+					{
+						yield return property;
+					}
+				}
+			}
+		}
 	}
 
 	/// <summary>
@@ -797,6 +844,8 @@ namespace EpicGames.Core
 	/// </summary>
 	public class CaptureLogger : ILogger
 	{
+		readonly LoggerScopeCollection _scopeCollection = new LoggerScopeCollection();
+
 		/// <summary>
 		/// List of captured events
 		/// </summary>
@@ -821,7 +870,7 @@ namespace EpicGames.Core
 		public List<string> RenderLines() => Events.ConvertAll(x => x.ToString());
 
 		/// <inheritdoc/>
-		public IDisposable BeginScope<TState>(TState state) => new NullScope();
+		public IDisposable BeginScope<TState>(TState state) => _scopeCollection.BeginScope(state);
 
 		/// <inheritdoc/>
 		public bool IsEnabled(LogLevel logLevel) => true;
@@ -829,7 +878,9 @@ namespace EpicGames.Core
 		/// <inheritdoc/>
 		public void Log<TState>(LogLevel logLevel, EventId eventId, TState state, Exception? exception, Func<TState, Exception?, string> formatter)
 		{
-			Events.Add(LogEvent.FromState(logLevel, eventId, state, exception, formatter));
+			LogEvent logEvent = LogEvent.FromState(logLevel, eventId, state, exception, formatter);
+			logEvent.AddProperties(_scopeCollection.GetProperties());
+			Events.Add(logEvent);
 		}
 	}
 
@@ -915,9 +966,7 @@ namespace EpicGames.Core
 			}
 		}
 
-		/// <summary>
-		/// Object used for synchronization
-		/// </summary>
+		private readonly LoggerScopeCollection _scopeCollection = new LoggerScopeCollection();
 		private readonly object _syncObject = new object();
 
 		/// <summary>
@@ -1175,7 +1224,8 @@ namespace EpicGames.Core
 		}
 
 		/// <inheritdoc/>
-		public IDisposable BeginScope<TState>(TState state) => new NullScope();
+		public IDisposable BeginScope<TState>(TState state)
+			=> _scopeCollection.BeginScope(state);
 
 		/// <inheritdoc/>
 		public bool IsEnabled(LogLevel logLevel)
@@ -1223,14 +1273,18 @@ namespace EpicGames.Core
 				Activity? activity = Activity.Current;
 
 				JsonLogEvent jsonLogEvent;
-				if (activity == null)
+				if (activity == null && _scopeCollection.IsEmpty())
 				{
 					jsonLogEvent = JsonLogEvent.FromLoggerState(logLevel, eventId, state, exception, formatter);
 				}
 				else
 				{
 					LogEvent logEvent = LogEvent.FromState(logLevel, eventId, state, exception, formatter);
-					logEvent.Properties = Enumerable.Append(logEvent.Properties ?? Array.Empty<KeyValuePair<string, object>>(), new KeyValuePair<string, object>("Activity", activity));
+					logEvent.AddProperties(_scopeCollection.GetProperties());
+					if (activity != null)
+					{
+						logEvent.AddProperty("Activity", activity);
+					}
 					jsonLogEvent = new JsonLogEvent(logEvent);
 				}
 				_eventChannel.Writer.TryWrite(jsonLogEvent);
diff --git a/Engine/Source/Programs/Shared/EpicGames.Core/LogEvent.cs b/Engine/Source/Programs/Shared/EpicGames.Core/LogEvent.cs
--- a/Engine/Source/Programs/Shared/EpicGames.Core/LogEvent.cs
+++ b/Engine/Source/Programs/Shared/EpicGames.Core/LogEvent.cs
@@ -2,6 +2,7 @@
 
 using System;
 using System.Buffers;
+using System.Collections;
 using System.Collections.Generic;
 using System.Diagnostics.CodeAnalysis;
 using System.Globalization;
@@ -95,6 +96,33 @@ namespace EpicGames.Core
 		/// </summary>
 		public LogException? Exception { get; }
 
+		class MergedPropertyList : IEnumerable<KeyValuePair<string, object>>
+		{
+			readonly HashSet<string> _names = new HashSet<string>();
+			readonly List<KeyValuePair<string, object>> _properties = new List<KeyValuePair<string, object>>();
+
+			public void AddRange(IEnumerable<KeyValuePair<string, object>> properties)
+			{
+				foreach (KeyValuePair<string, object> property in properties)
+				{
+					if (_names.Add(property.Key))
+					{
+						_properties.Add(property);
+					}
+				}
+			}
+
+			/// <inheritdoc/>
+			public IEnumerator<KeyValuePair<string, object>> GetEnumerator()
+				=> _properties.GetEnumerator();
+
+			/// <inheritdoc/>
+			IEnumerator IEnumerable.GetEnumerator()
+				=> _properties.GetEnumerator();
+		}
+
+		static readonly JsonSerializerOptions s_jsonSerializerOptions = new JsonSerializerOptions { PropertyNamingPolicy = JsonNamingPolicy.CamelCase, DefaultIgnoreCondition = JsonIgnoreCondition.WhenWritingNull };
+
 		/// <summary>
 		/// Constructor
 		/// </summary>
@@ -119,6 +147,37 @@ namespace EpicGames.Core
 			Exception = exception;
 		}
 
+		/// <summary>
+		/// Add a new property to this event
+		/// </summary>
+		/// <param name="name">Name of the property to add</param>
+		/// <param name="value">Value for the property</param>
+		public void AddProperty(string name, object value)
+			=> AddProperties(new[] { KeyValuePair.Create(name, value) });
+
+		/// <summary>
+		/// Add new properties to this event
+		/// </summary>
+		/// <param name="properties">Properties to add</param>
+		public void AddProperties(IEnumerable<KeyValuePair<string, object>> properties)
+		{
+			if (Properties == null)
+			{
+				Properties = properties;
+			}
+			else
+			{
+				MergedPropertyList? list = Properties as MergedPropertyList;
+				if (list == null)
+				{
+					list = new MergedPropertyList();
+					list.AddRange(Properties);
+					Properties = list;
+				}
+				list.AddRange(properties);
+			}
+		}
+
 		/// <summary>
 		/// Gets an untyped property with the given name
 		/// </summary>
@@ -371,8 +430,16 @@ namespace EpicGames.Core
 				{
 					if (!name.Equals(MessageTemplate.FormatPropertyName, StringComparison.Ordinal))
 					{
-						writer.WritePropertyName(name);
-						LogValueFormatter.Format(value, writer);
+						if (name.StartsWith("@", StringComparison.Ordinal))
+						{
+							writer.WritePropertyName(name[1..]);
+							JsonSerializer.Serialize(writer, value, value.GetType(), s_jsonSerializerOptions);
+						}
+						else
+						{
+							writer.WritePropertyName(name);
+							LogValueFormatter.Format(value, writer);
+						}
 					}
 				}
 				writer.WriteEndObject();
diff --git a/Engine/Source/Programs/Shared/EpicGames.Core/MessageTemplate.cs b/Engine/Source/Programs/Shared/EpicGames.Core/MessageTemplate.cs
--- a/Engine/Source/Programs/Shared/EpicGames.Core/MessageTemplate.cs
+++ b/Engine/Source/Programs/Shared/EpicGames.Core/MessageTemplate.cs
@@ -3,6 +3,7 @@
 using System;
 using System.Collections.Generic;
 using System.Text;
+using System.Text.Json;
 
 namespace EpicGames.Core
 {
@@ -47,14 +48,25 @@ namespace EpicGames.Core
 					object? value;
 					if (properties != null && TryGetPropertyValue(format.AsSpan(offset, length), properties, out value))
 					{
+						// Append the text up to this argument
 						int startOffset = offset - 1;
 						if (format[startOffset] == '@' || format[startOffset] == '$')
 						{
 							startOffset--;
 						}
-
 						Unescape(format.AsSpan(nextOffset, startOffset - nextOffset), result);
-						result.Append(value?.ToString() ?? "null");
+
+						// Append the argument
+						if (format[offset] == '@')
+						{
+							result.Append(JsonSerializer.Serialize(value, value?.GetType() ?? typeof(object), new JsonSerializerOptions { PropertyNamingPolicy = JsonNamingPolicy.CamelCase }));
+						}
+						else
+						{
+							result.Append(value?.ToString() ?? "null");
+						}
+
+						// Start the next plain-text run after the closing brace
 						nextOffset = offset + length + 1;
 					}
 				}
