Horde: Store the change filter for issue fingerprints as a plain string. This is easier to script and serialize.

#rnx

diff --git a/Engine/Source/Programs/Horde/Plugins/Build/HordeServer.Build/Issues/IIssueCollection.cs b/Engine/Source/Programs/Horde/Plugins/Build/HordeServer.Build/Issues/IIssueCollection.cs
--- a/Engine/Source/Programs/Horde/Plugins/Build/HordeServer.Build/Issues/IIssueCollection.cs
+++ b/Engine/Source/Programs/Horde/Plugins/Build/HordeServer.Build/Issues/IIssueCollection.cs
@@ -51,7 +51,7 @@ namespace HordeServer.Issues
 		IReadOnlySet<IssueMetadata>? IIssueFingerprint.Metadata => (Metadata.Count > 0) ? Metadata : null;
 
 		/// <inheritdoc/>
-		public IReadOnlyList<string> ChangeFilter { get; set; }
+		public string ChangeFilter { get; set; }
 
 		/// <summary>
 		/// Constructor
@@ -59,7 +59,7 @@ namespace HordeServer.Issues
 		/// <param name="type">The type of issue</param>
 		/// <param name="summaryTemplate">Template for the summary string to display for the issue</param>
 		/// <param name="changeFilter">Filter for changes covered by this issue</param>
-		public NewIssueFingerprint(string type, string summaryTemplate, IReadOnlyList<string> changeFilter)
+		public NewIssueFingerprint(string type, string summaryTemplate, string changeFilter)
 		{
 			Type = type;
 			SummaryTemplate = summaryTemplate;
@@ -75,7 +75,7 @@ namespace HordeServer.Issues
 		/// <param name="rejectKeys">Keys which should not match with this issue</param>
 		/// <param name="metadata">Additional metadata added by the issue handler</param>
 		/// <param name="changeFilter">Filter for changes covered by this issue</param>
-		public NewIssueFingerprint(string type, string summaryTemplate, IEnumerable<IssueKey> keys, IEnumerable<IssueKey>? rejectKeys, IEnumerable<IssueMetadata>? metadata, IEnumerable<string> changeFilter)
+		public NewIssueFingerprint(string type, string summaryTemplate, IEnumerable<IssueKey> keys, IEnumerable<IssueKey>? rejectKeys, IEnumerable<IssueMetadata>? metadata, string changeFilter)
 		{
 			Type = type;
 			SummaryTemplate = summaryTemplate;
@@ -90,7 +90,7 @@ namespace HordeServer.Issues
 				Metadata = new HashSet<IssueMetadata>(metadata);
 			}
 
-			ChangeFilter = new List<string>(changeFilter);
+			ChangeFilter = changeFilter;
 		}
 
 		/// <summary>
diff --git a/Engine/Source/Programs/Horde/Plugins/Build/HordeServer.Build/Issues/IIssueFingerprint.cs b/Engine/Source/Programs/Horde/Plugins/Build/HordeServer.Build/Issues/IIssueFingerprint.cs
--- a/Engine/Source/Programs/Horde/Plugins/Build/HordeServer.Build/Issues/IIssueFingerprint.cs
+++ b/Engine/Source/Programs/Horde/Plugins/Build/HordeServer.Build/Issues/IIssueFingerprint.cs
@@ -40,7 +40,7 @@ namespace HordeServer.Issues
 		/// <summary>
 		/// Filter for changes that should be included in this issue
 		/// </summary>
-		public IReadOnlyList<string> ChangeFilter { get; }
+		public string ChangeFilter { get; }
 	}
 
 	/// <summary>
diff --git a/Engine/Source/Programs/Horde/Plugins/Build/HordeServer.Build/Issues/IssueCollection.cs b/Engine/Source/Programs/Horde/Plugins/Build/HordeServer.Build/Issues/IssueCollection.cs
--- a/Engine/Source/Programs/Horde/Plugins/Build/HordeServer.Build/Issues/IssueCollection.cs
+++ b/Engine/Source/Programs/Horde/Plugins/Build/HordeServer.Build/Issues/IssueCollection.cs
@@ -266,10 +266,18 @@ namespace HordeServer.Issues
 
 			IReadOnlySet<IssueMetadata>? IIssueFingerprint.Metadata => Metadata;
 
+			[BsonElement("fl2")]
+			public string? ChangeFilter { get; set; }
+
+			[Obsolete("Use ChangeFilter instead")]
 			[BsonElement("flt"), BsonIgnoreIfNull]
-			public List<string>? ChangeFilter { get; set; }
+			public List<string>? LegacyChangeFilter
+			{
+				get => null;
+				set => ChangeFilter = (value == null)? null : String.Join(";", value);
+			}
 
-			IReadOnlyList<string> IIssueFingerprint.ChangeFilter => ChangeFilter ?? GetLegacyHandlerInfo(Type).ChangeFilter;
+			string IIssueFingerprint.ChangeFilter => ChangeFilter ?? String.Join(";", GetLegacyHandlerInfo(Type).ChangeFilter);
 
 			[BsonConstructor]
 			private IssueFingerprint()
@@ -292,7 +300,7 @@ namespace HordeServer.Issues
 					Metadata = new HashSet<IssueMetadata>(fingerprint.Metadata);
 				}
 
-				ChangeFilter = new List<string>(fingerprint.ChangeFilter);
+				ChangeFilter = fingerprint.ChangeFilter;
 			}
 
 			#region Legacy
diff --git a/Engine/Source/Programs/Horde/Plugins/Build/HordeServer.Build/Issues/IssueService.cs b/Engine/Source/Programs/Horde/Plugins/Build/HordeServer.Build/Issues/IssueService.cs
--- a/Engine/Source/Programs/Horde/Plugins/Build/HordeServer.Build/Issues/IssueService.cs
+++ b/Engine/Source/Programs/Horde/Plugins/Build/HordeServer.Build/Issues/IssueService.cs
@@ -1219,7 +1219,7 @@ namespace HordeServer.Issues
 		async Task<List<NewIssueSpanSuspectData>> FindSuspectsForSpanAsync(StreamConfig streamConfig, IIssueFingerprint fingerprint, int minChange, int maxChange, CancellationToken cancellationToken)
 		{
 			List<NewIssueSpanSuspectData> suspects = new List<NewIssueSpanSuspectData>();
-			if (fingerprint.ChangeFilter.Count > 0)
+			if (!String.IsNullOrEmpty(fingerprint.ChangeFilter))
 			{
 				_logger.LogDebug("Querying for changes in {StreamName} between {MinChange} and {MaxChange}", streamConfig.Name, minChange, maxChange);
 
@@ -1228,7 +1228,7 @@ namespace HordeServer.Issues
 				_logger.LogDebug("Found {NumResults} changes", changes.Count);
 
 				// Get all the parameters used to rank suspects
-				FileFilter filter = new FileFilter(fingerprint.ChangeFilter);
+				FileFilter filter = new FileFilter(fingerprint.ChangeFilter.Split(";", StringSplitOptions.RemoveEmptyEntries | StringSplitOptions.TrimEntries));
 
 				// Build a set of all the files to include as suspects
 				HashSet<string> suspectFiles = new HashSet<string>(StringComparer.OrdinalIgnoreCase);
diff --git a/Engine/Source/Programs/Shared/EpicGames.Horde/Issues/IssueEvent.cs b/Engine/Source/Programs/Shared/EpicGames.Horde/Issues/IssueEvent.cs
--- a/Engine/Source/Programs/Shared/EpicGames.Horde/Issues/IssueEvent.cs
+++ b/Engine/Source/Programs/Shared/EpicGames.Horde/Issues/IssueEvent.cs
@@ -3,8 +3,6 @@
 using System;
 using System.Collections.Generic;
 using System.Linq;
-using System.Text.Json;
-using System.Text.Json.Serialization;
 using EpicGames.Core;
 using Microsoft.Extensions.Logging;
 
@@ -84,8 +82,7 @@ namespace EpicGames.Horde.Issues
 		/// <summary>
 		/// Filter for changes that should be included in this issue
 		/// </summary>
-		[JsonConverter(typeof(ChangeFilterJsonConverter))]
-		public IReadOnlyList<string> ChangeFilter { get; set; }
+		public string ChangeFilter { get; set; }
 
 		/// <summary>
 		/// Individual log events
@@ -98,22 +95,22 @@ namespace EpicGames.Horde.Issues
 		/// <param name="type">The type of issue</param>
 		/// <param name="summaryTemplate">Template for the summary string to display for the issue</param>
 		/// <param name="changeFilter">Filter for changes covered by this issue</param>
-		public IssueEventGroup(string type, string summaryTemplate, IReadOnlyList<string> changeFilter)
+		public IssueEventGroup(string type, string summaryTemplate, string changeFilter)
 		{
 			Type = type;
 			SummaryTemplate = summaryTemplate;
 			ChangeFilter = changeFilter;
 		}
-	}
-
-	class ChangeFilterJsonConverter : JsonConverter<IReadOnlyList<string>>
-	{
-		/// <inheritdoc/>
-		public override IReadOnlyList<string>? Read(ref Utf8JsonReader reader, Type typeToConvert, JsonSerializerOptions options)
-			=> (reader.GetString() ?? String.Empty).Split(';');
 
-		/// <inheritdoc/>
-		public override void Write(Utf8JsonWriter writer, IReadOnlyList<string> value, JsonSerializerOptions options)
-			=> writer.WriteStringValue(String.Join(";", value));
+		/// <summary>
+		/// Constructor
+		/// </summary>
+		/// <param name="type">The type of issue</param>
+		/// <param name="summaryTemplate">Template for the summary string to display for the issue</param>
+		/// <param name="changeFilter">Filter for changes covered by this issue</param>
+		public IssueEventGroup(string type, string summaryTemplate, IReadOnlyList<string> changeFilter)
+			: this(type, summaryTemplate, String.Join(";", changeFilter))
+		{
+		}
 	}
 }
