Horde: Improve error message when a Perforce response does not contain a required tag.

#rnx

diff --git a/Engine/Source/Programs/Shared/EpicGames.Perforce/IPerforceOutput.cs b/Engine/Source/Programs/Shared/EpicGames.Perforce/IPerforceOutput.cs
--- a/Engine/Source/Programs/Shared/EpicGames.Perforce/IPerforceOutput.cs
+++ b/Engine/Source/Programs/Shared/EpicGames.Perforce/IPerforceOutput.cs
@@ -653,7 +653,7 @@ namespace EpicGames.Perforce
 			// Make sure we've got all the required tags we need
 			if (requiredTagsBitMask != recordInfo.RequiredTagsBitMask)
 			{
-				string missingTagNames = String.Join(", ", recordInfo.NameAndRankToInfo.Where(x => (requiredTagsBitMask | x.Value.RequiredTagBitMask) != requiredTagsBitMask).Select(x => x.Key));
+				string missingTagNames = String.Join(", ", recordInfo.NameAndRankToInfo.Where(x => (requiredTagsBitMask | x.Value.RequiredTagBitMask) != requiredTagsBitMask).Select(x => x.Key.Tag));
 				throw new PerforceException("Missing '{0}' tag when parsing '{1}'", missingTagNames, recordInfo.Type.Name);
 			}
 
