EpicGames.Core: Fix conversion of plain text strings to LogEvent objects.

#rnx

diff --git a/Engine/Source/Programs/Shared/EpicGames.Core/LogEvent.cs b/Engine/Source/Programs/Shared/EpicGames.Core/LogEvent.cs
--- a/Engine/Source/Programs/Shared/EpicGames.Core/LogEvent.cs
+++ b/Engine/Source/Programs/Shared/EpicGames.Core/LogEvent.cs
@@ -544,7 +544,7 @@ namespace EpicGames.Core
 			// Try to log the event
 			IEnumerable<KeyValuePair<string, object?>>? values = state as IEnumerable<KeyValuePair<string, object?>>;
 			string? format = values?.FirstOrDefault(x => x.Key.Equals(MessageTemplate.FormatPropertyName, StringComparison.Ordinal)).Value?.ToString();
-			string message = MessageTemplate.Render(format ?? String.Empty, values);
+			string message = (format == null)? formatter(state, exception) : MessageTemplate.Render(format ?? String.Empty, values);
 			return new LogEvent(time, level, eventId, message, format, values, LogException.FromException(exception));
 		}
 
