Horde: Fix issue metadata not being serialized in fingerprint objects.

#rnx

diff --git a/Engine/Source/Programs/Shared/EpicGames.Core/LogEvent.cs b/Engine/Source/Programs/Shared/EpicGames.Core/LogEvent.cs
--- a/Engine/Source/Programs/Shared/EpicGames.Core/LogEvent.cs
+++ b/Engine/Source/Programs/Shared/EpicGames.Core/LogEvent.cs
@@ -544,7 +544,7 @@ namespace EpicGames.Core
 			// Try to log the event
 			IEnumerable<KeyValuePair<string, object?>>? values = state as IEnumerable<KeyValuePair<string, object?>>;
 			string? format = values?.FirstOrDefault(x => x.Key.Equals(MessageTemplate.FormatPropertyName, StringComparison.Ordinal)).Value?.ToString();
-			string message = (format == null)? formatter(state, exception) : MessageTemplate.Render(format ?? String.Empty, values);
+			string message = (format == null)? formatter(state, exception) : MessageTemplate.Render(format, values);
 			return new LogEvent(time, level, eventId, message, format, values, LogException.FromException(exception));
 		}
 
diff --git a/Engine/Source/Programs/Shared/EpicGames.Horde/Issues/IssueFingerprint.cs b/Engine/Source/Programs/Shared/EpicGames.Horde/Issues/IssueFingerprint.cs
--- a/Engine/Source/Programs/Shared/EpicGames.Horde/Issues/IssueFingerprint.cs
+++ b/Engine/Source/Programs/Shared/EpicGames.Horde/Issues/IssueFingerprint.cs
@@ -68,7 +68,7 @@ namespace EpicGames.Horde.Issues
 		/// </summary>
 		[JsonPropertyName("meta")]
 		[JsonIgnore(Condition = JsonIgnoreCondition.WhenWritingNull)]
-		internal HashSet<IssueMetadata>? SerializedMetadata
+		public HashSet<IssueMetadata>? SerializedMetadata
 		{
 			get => (Metadata.Count == 0) ? null : Metadata;
 			set => Metadata = value ?? new HashSet<IssueMetadata>();
diff --git a/Engine/Source/Programs/Shared/EpicGames.Horde/Issues/IssueMetadata.cs b/Engine/Source/Programs/Shared/EpicGames.Horde/Issues/IssueMetadata.cs
--- a/Engine/Source/Programs/Shared/EpicGames.Horde/Issues/IssueMetadata.cs
+++ b/Engine/Source/Programs/Shared/EpicGames.Horde/Issues/IssueMetadata.cs
@@ -2,6 +2,7 @@
 
 using System;
 using System.Collections.Generic;
+using System.Text.Json.Serialization;
 
 namespace EpicGames.Horde.Issues
 {
@@ -23,6 +24,7 @@ namespace EpicGames.Horde.Issues
 		/// <summary>
 		/// Constructor
 		/// </summary>
+		[JsonConstructor]
 		public IssueMetadata(string key, string value)
 		{
 			Key = key;
