Horde: Separate out log event errors and critical failures into a separate fingerprint when we exceed the 25 event limit in the Hashed Issue Handler.

diff --git a/Engine/Source/Programs/Horde/Plugins/Build/HordeServer.Build.Tests/Issues/IssueServiceTests.cs b/Engine/Source/Programs/Horde/Plugins/Build/HordeServer.Build.Tests/Issues/IssueServiceTests.cs
--- a/Engine/Source/Programs/Horde/Plugins/Build/HordeServer.Build.Tests/Issues/IssueServiceTests.cs
+++ b/Engine/Source/Programs/Horde/Plugins/Build/HordeServer.Build.Tests/Issues/IssueServiceTests.cs
@@ -2746,8 +2746,8 @@ namespace HordeServer.Tests.Issues
 			}
 
 			// #5
-			// Scenario: Job is run which updates issue, with a failure 
-			// Expected: Existing issue is updated and becomes an error
+			// Scenario: Job is run with a new failure, when an existing issue is already open with only a warning
+			// Expected: A new issue is created rather than the existing issue being updated since the previous issue only contained a warning
 			{
 				IJob job = CreateJob(_mainStreamId, 225, "Test Build", _graph, TimeSpan.FromHours(hour++));
 				await AddEventAsync(job, 0, 0, LogLevel.Error);
@@ -2755,7 +2755,7 @@ namespace HordeServer.Tests.Issues
 
 				IReadOnlyList<IIssue> issues = await IssueCollection.FindIssuesAsync();
 				Assert.AreEqual(1, issues.Count);
-				Assert.AreEqual(1, issues[0].Id);
+				Assert.AreEqual(2, issues[0].Id);
 				Assert.AreEqual(IssueSeverity.Error, issues[0].Severity);
 			}
 		}
diff --git a/Engine/Source/Programs/Shared/EpicGames.Horde/Issues/Handlers/HashedIssueHandler.cs b/Engine/Source/Programs/Shared/EpicGames.Horde/Issues/Handlers/HashedIssueHandler.cs
--- a/Engine/Source/Programs/Shared/EpicGames.Horde/Issues/Handlers/HashedIssueHandler.cs
+++ b/Engine/Source/Programs/Shared/EpicGames.Horde/Issues/Handlers/HashedIssueHandler.cs
@@ -54,6 +54,7 @@ namespace EpicGames.Horde.Issues.Handlers
 			List<IssueEventGroup> issues = new List<IssueEventGroup>();
 
 			IssueEventGroup? genericFingerprint = null;
+			IssueEventGroup? genericErrorsFingerprint = null;
 			HashSet<Md5Hash> hashes = new HashSet<Md5Hash>();
 
 			// keep hash consistent when only have general, non-unique events
@@ -85,14 +86,28 @@ namespace EpicGames.Horde.Issues.Handlers
 				}
 				else
 				{
-					if (genericFingerprint == null)
+					if (stepEvent.Severity == LogLevel.Error || stepEvent.Severity == LogLevel.Critical)
 					{
-						genericFingerprint = new IssueEventGroup("Hashed", "{Severity} in {Meta:Node}", IssueChangeFilter.All);
-						genericFingerprint.Keys.Add(IssueKey.FromStep(_context.StreamId, _context.TemplateId, _context.NodeName));
-						genericFingerprint.Metadata.Add("Node", _context.NodeName);
-						issues.Add(genericFingerprint);
+						if (genericErrorsFingerprint == null)
+						{
+							genericErrorsFingerprint = new IssueEventGroup("Hashed", "{Severity} in {Meta:Node}", IssueChangeFilter.All);
+							genericErrorsFingerprint.Keys.Add(IssueKey.FromStepAndSeverity(_context.StreamId, _context.TemplateId, _context.NodeName, LogLevel.Error));
+							genericErrorsFingerprint.Metadata.Add("Node", _context.NodeName);
+							issues.Add(genericErrorsFingerprint);
+						}
+						genericErrorsFingerprint.Events.Add(stepEvent);
+					}
+					else
+					{
+						if (genericFingerprint == null)
+						{
+							genericFingerprint = new IssueEventGroup("Hashed", "{Severity} in {Meta:Node}", IssueChangeFilter.All);
+							genericFingerprint.Keys.Add(IssueKey.FromStep(_context.StreamId, _context.TemplateId, _context.NodeName));
+							genericFingerprint.Metadata.Add("Node", _context.NodeName);
+							issues.Add(genericFingerprint);
+						}
+						genericFingerprint.Events.Add(stepEvent);
 					}
-					genericFingerprint.Events.Add(stepEvent);
 				}
 			}
 
diff --git a/Engine/Source/Programs/Shared/EpicGames.Horde/Issues/IssueKey.cs b/Engine/Source/Programs/Shared/EpicGames.Horde/Issues/IssueKey.cs
--- a/Engine/Source/Programs/Shared/EpicGames.Horde/Issues/IssueKey.cs
+++ b/Engine/Source/Programs/Shared/EpicGames.Horde/Issues/IssueKey.cs
@@ -8,6 +8,7 @@ using System.Text.Json.Serialization;
 using EpicGames.Core;
 using EpicGames.Horde.Jobs.Templates;
 using EpicGames.Horde.Streams;
+using Microsoft.Extensions.Logging;
 
 namespace EpicGames.Horde.Issues
 {
@@ -109,6 +110,11 @@ namespace EpicGames.Horde.Issues
 		/// </summary>
 		public static IssueKey FromStep(StreamId streamId, TemplateId templateId, string nodeName) => new IssueKey($"{streamId}:{templateId}:{nodeName}", IssueKeyType.Step);
 
+		/// <summary>
+		/// Creates an issue key for a particular step and severity
+		/// </summary>
+		public static IssueKey FromStepAndSeverity(StreamId streamId, TemplateId templateId, string nodeName, LogLevel severity) => new IssueKey($"{streamId}:{templateId}:{nodeName}:{severity}", IssueKeyType.Step);
+
 		/// <inheritdoc/>
 		public bool Equals(IssueKey? other) => other is not null && other.Name.Equals(Name, StringComparison.OrdinalIgnoreCase) && other.Type == Type && String.Equals(Scope, other.Scope, StringComparison.Ordinal);
 
