Horde: Static analysis fix.

#rnx

diff --git a/Engine/Source/Programs/Shared/EpicGames.Horde.Tests/ServerLogPacketBuilderTests.cs b/Engine/Source/Programs/Shared/EpicGames.Horde.Tests/ServerLogPacketBuilderTests.cs
--- a/Engine/Source/Programs/Shared/EpicGames.Horde.Tests/ServerLogPacketBuilderTests.cs
+++ b/Engine/Source/Programs/Shared/EpicGames.Horde.Tests/ServerLogPacketBuilderTests.cs
@@ -48,9 +48,9 @@ namespace EpicGames.Horde.Tests
 
 			Assert.AreEqual(impl._lines.Count, 2);
 			Assert.AreEqual(impl._lines[0].ToString(), "Hello world");
-			Assert.AreEqual(impl._lines[0].GetProperty("Text").ToString()!, "world");
+			Assert.AreEqual(impl._lines[0].GetProperty("Text")?.ToString(), "world");
 			Assert.AreEqual(impl._lines[1].ToString(), "Hello world");
-			Assert.AreEqual(impl._lines[1].GetProperty("Text").ToString(), "world");
+			Assert.AreEqual(impl._lines[1].GetProperty("Text")?.ToString(), "world");
 		}
 
 		static string[] SplitMultiLineMessage(string message)
@@ -332,7 +332,7 @@ namespace EpicGames.Horde.Tests
 			string longLineB = String.Join("\n", Enumerable.Range(0, 5).Select(x => $"B{x}"));
 			string longLineC = String.Join("\n", Enumerable.Range(0, 5).Select(x => $"C{x}"));
 
-			Dictionary<string, object> properties = new Dictionary<string, object>
+			Dictionary<string, object?> properties = new Dictionary<string, object?>
 			{
 				["LongLineA"] = longLineA,
 				["LongLineB"] = longLineB,
