Horde: Preliminary support for VCS commit ids that aren't an integer.

diff --git a/Engine/Source/Programs/Shared/EpicGames.Horde/Issues/IssueMessages.cs b/Engine/Source/Programs/Shared/EpicGames.Horde/Issues/IssueMessages.cs
--- a/Engine/Source/Programs/Shared/EpicGames.Horde/Issues/IssueMessages.cs
+++ b/Engine/Source/Programs/Shared/EpicGames.Horde/Issues/IssueMessages.cs
@@ -2,6 +2,7 @@
 
 using System;
 using System.Collections.Generic;
+using EpicGames.Horde.Commits;
 using EpicGames.Horde.Jobs;
 using EpicGames.Horde.Jobs.Templates;
 using EpicGames.Horde.Logs;
@@ -41,7 +42,23 @@ namespace EpicGames.Horde.Issues
 		/// <summary>
 		/// The changelist number
 		/// </summary>
-		public int Change { get; set; }
+		[Obsolete("Use Commit instead")]
+		public int Change
+		{
+			get => _change ?? _commitId?.TryGetPerforceChange() ?? -1;
+			set => _change = value;
+		}
+		int? _change;
+
+		/// <summary>
+		/// The commit for this step
+		/// </summary>
+		public CommitIdWithOrder CommitId
+		{
+			get => _commitId ?? CommitIdWithOrder.FromPerforceChange(_change) ?? CommitIdWithOrder.Empty;
+			set => _commitId = value;
+		}
+		CommitIdWithOrder? _commitId;
 
 		/// <summary>
 		/// Severity of the issue in this step
@@ -130,15 +147,47 @@ namespace EpicGames.Horde.Issues
 		/// </summary>
 		public StreamId StreamId { get; set; }
 
+		/// <summary>
+		/// Minimum commit affected by this issue (ie. last successful build)
+		/// </summary>
+		[Obsolete("Use MinCommitId instead")]
+		public int? MinChange
+		{
+			get => _minChange ?? _minCommitId?.GetPerforceChangeOrMinusOne();
+			set => _minChange = value;
+		}
+		int? _minChange;
+
+		/// <summary>
+		/// Maximum commit affected by this issue (ie. next successful build)
+		/// </summary>
+		[Obsolete("Use MaxCommitId instead")]
+		public int? MaxChange
+		{
+			get => _maxChange ?? _maxCommitId?.GetPerforceChangeOrMinusOne();
+			set => _maxChange = value;
+		}
+		int? _maxChange;
+
 		/// <summary>
 		/// Minimum changelist affected by this issue (ie. last successful build)
 		/// </summary>
-		public int? MinChange { get; set; }
+		public CommitIdWithOrder? MinCommitId
+		{
+			get => _minCommitId ?? CommitIdWithOrder.FromPerforceChange(_minChange);
+			set => _minCommitId = value;
+		}
+		CommitIdWithOrder? _minCommitId;
 
 		/// <summary>
 		/// Maximum changelist affected by this issue (ie. next successful build)
 		/// </summary>
-		public int? MaxChange { get; set; }
+		public CommitIdWithOrder? MaxCommitId
+		{
+			get => _maxCommitId ?? CommitIdWithOrder.FromPerforceChange(_maxChange);
+			set => _maxCommitId = value;
+		}
+		CommitIdWithOrder? _maxCommitId;
 
 		/// <summary>
 		/// Map of steps to (event signature id -> trace id)
@@ -320,14 +369,34 @@ namespace EpicGames.Horde.Issues
 		public DateTime? AcknowledgedAt { get; set; }
 
 		/// <summary>
-		/// Changelist that fixed this issue
+		/// Perforce changelist that fixed this issue
 		/// </summary>
-		public int? FixChange { get; set; }
+		public int? FixChange
+		{
+			get => _fixChange ?? _fixCommitId?.GetPerforceChangeOrMinusOne();
+			set => _fixChange = value;
+		}
+		int? _fixChange;
+
+		/// <summary>
+		/// Commit that fixed this issue
+		/// </summary>
+		public CommitId? FixCommitId
+		{
+			get => _fixCommitId ?? CommitId.FromPerforceChange(_fixChange);
+			set => _fixCommitId = value;
+		}
+		CommitId? _fixCommitId;
 
 		/// <summary>
 		/// Whether the issue is marked fixed as a systemic issue
 		/// </summary>
-		public bool FixedSystemic { get; set; }
+		public bool FixSystemic
+		{
+			get => _fixSystemic ?? (_fixChange < 0);
+			set => _fixSystemic = value;
+		}
+		bool? _fixSystemic;
 
 		/// <summary>
 		/// Time at which the issue was resolved
@@ -529,12 +598,32 @@ namespace EpicGames.Horde.Issues
 		/// <summary>
 		/// Changelist that fixed this issue
 		/// </summary>
-		public int? FixChange { get; set; }
+		public int? FixChange
+		{
+			get => _fixChange ?? ((_fixSystemic ?? false) ? -1 : _fixCommitId?.GetPerforceChangeOrMinusOne());
+			set => _fixChange = value;
+		}
+		int? _fixChange;
+
+		/// <summary>
+		/// Changelist that fixed this issue
+		/// </summary>
+		public CommitId? FixCommitId
+		{
+			get => _fixCommitId ?? CommitId.FromPerforceChange(_fixChange);
+			set => _fixCommitId = value;
+		}
+		CommitId? _fixCommitId;
 
 		/// <summary>
 		/// Whether the issue is marked fixed as a systemic issue
 		/// </summary>
-		public bool FixedSystemic { get; set; }
+		public bool FixSystemic
+		{
+			get => _fixSystemic ?? (_fixChange < 0);
+			set => _fixSystemic = value;
+		}
+		bool? _fixSystemic;
 
 		/// <summary>
 		/// Time at which the issue was resolved
@@ -630,7 +719,33 @@ namespace EpicGames.Horde.Issues
 		/// <summary>
 		/// The change at which the issue is claimed fixed. 0 = not fixed, -1 = systemic issue.
 		/// </summary>
-		public int? FixChange { get; set; }
+		[Obsolete("Use FixCommitId and FixSystemic instead")]
+		public int? FixChange
+		{
+			get => _fixChange ?? ((_fixSystemic ?? false) ? -1 : _fixCommitId?.GetPerforceChangeOrMinusOne());
+			set => _fixChange = value;
+		}
+		int? _fixChange;
+
+		/// <summary>
+		/// The change at which the issue is claimed fixed. """" = not fixed.
+		/// </summary>
+		public CommitId? FixCommitId
+		{
+			get => _fixCommitId ?? ((_fixChange < 0)? null : CommitId.FromPerforceChange(_fixChange));
+			set => _fixCommitId = value;
+		}
+		CommitId? _fixCommitId;
+
+		/// <summary>
+		/// Set to mark the issue as fixed systemically
+		/// </summary>
+		public bool FixSystemic
+		{
+			get => _fixSystemic ?? (_fixChange < 0);
+			set => _fixSystemic = value;
+		}
+		bool? _fixSystemic;
 
 		/// <summary>
 		/// Whether the issue should be marked as resolved
diff --git a/Engine/Source/Programs/Shared/EpicGames.Horde/Jobs/JobMessages.cs b/Engine/Source/Programs/Shared/EpicGames.Horde/Jobs/JobMessages.cs
--- a/Engine/Source/Programs/Shared/EpicGames.Horde/Jobs/JobMessages.cs
+++ b/Engine/Source/Programs/Shared/EpicGames.Horde/Jobs/JobMessages.cs
@@ -191,7 +191,23 @@ namespace EpicGames.Horde.Jobs
 		/// <summary>
 		/// The changelist number to build. Can be null for latest.
 		/// </summary>
-		public int? Change { get; set; }
+		[Obsolete("Use CommitId instead")]
+		public int? Change
+		{
+			get => _change ?? _commitId?.GetPerforceChangeOrMinusOne();
+			set => _change = value;
+		}
+		int? _change;
+
+		/// <summary>
+		/// The changelist number to build. Can be null for latest.
+		/// </summary>
+		public CommitId? CommitId
+		{
+			get => _commitId ?? CommitId.FromPerforceChange(_change);
+			set => _commitId = value;
+		}
+		CommitId? _commitId;
 
 		/// <summary>
 		/// Parameters to use when selecting the change to execute at.
@@ -211,7 +227,23 @@ namespace EpicGames.Horde.Jobs
 		/// <summary>
 		/// The preflight changelist number
 		/// </summary>
-		public int? PreflightChange { get; set; }
+		[Obsolete("Use PreflightCommitId instead")]
+		public int? PreflightChange
+		{
+			get => _preflightChange ?? _preflightCommitId?.GetPerforceChangeOrMinusOne();
+			set => _preflightChange = value;
+		}
+		int? _preflightChange;
+
+		/// <summary>
+		/// The preflight commit
+		/// </summary>
+		public CommitId? PreflightCommitId
+		{
+			get => _preflightCommitId ?? CommitId.FromPerforceChange(_preflightChange);
+			set => _preflightCommitId = value;
+		}
+		CommitId? _preflightCommitId;
 
 		/// <summary>
 		/// Job options
@@ -393,22 +425,86 @@ namespace EpicGames.Horde.Jobs
 		/// <summary>
 		/// The changelist number to build
 		/// </summary>
-		public int Change { get; set; }
+		[Obsolete("Use CommitId instead")]
+		public int Change
+		{
+			get => _change ?? _commitId?.GetPerforceChangeOrMinusOne() ?? 0;
+			set => _change = value;
+		}
+		int? _change;
+
+		/// <summary>
+		/// The commit to build
+		/// </summary>
+		public CommitIdWithOrder CommitId
+		{
+			get => _commitId ?? CommitIdWithOrder.FromPerforceChange(_change) ?? CommitIdWithOrder.Empty;
+			set => _commitId = value;
+		}
+		CommitIdWithOrder? _commitId;
 
 		/// <summary>
 		/// The code changelist
 		/// </summary>
-		public int? CodeChange { get; set; }
+		[Obsolete("Use CodeCommitId instead")]
+		public int? CodeChange
+		{
+			get => _codeChange ?? _codeCommitId?.GetPerforceChangeOrMinusOne();
+			set => _codeChange = value;
+		}
+		int? _codeChange;
+
+		/// <summary>
+		/// The code commit to build
+		/// </summary>
+		public CommitIdWithOrder? CodeCommitId
+		{
+			get => _codeCommitId ?? CommitIdWithOrder.FromPerforceChange(_codeChange);
+			set => _codeCommitId = value;
+		}
+		CommitIdWithOrder? _codeCommitId;
 
 		/// <summary>
 		/// The preflight changelist number
 		/// </summary>
-		public int? PreflightChange { get; set; }
+		[Obsolete("Use PreflightCommitId instead")]
+		public int? PreflightChange
+		{
+			get => _preflightChange ?? _preflightCommitId?.GetPerforceChangeOrMinusOne();
+			set => _preflightChange = value;
+		}
+		int? _preflightChange;
+
+		/// <summary>
+		/// The preflight commit
+		/// </summary>
+		public CommitId? PreflightCommitId
+		{
+			get => _preflightCommitId ?? Horde.Commits.CommitId.FromPerforceChange(_preflightChange);
+			set => _preflightCommitId = value;
+		}
+		CommitId? _preflightCommitId;
 
 		/// <summary>
 		/// The cloned preflight changelist number
 		/// </summary>
-		public int? ClonedPreflightChange { get; set; }
+		[Obsolete("Use ClonedPreflightCommitId instead")]
+		public int? ClonedPreflightChange
+		{
+			get => _clonedPreflightChange ?? _clonedPreflightCommitId?.GetPerforceChangeOrMinusOne();
+			set => _clonedPreflightChange = value;
+		}
+		int? _clonedPreflightChange;
+
+		/// <summary>
+		/// The cloned preflight commit
+		/// </summary>
+		public CommitId? ClonedPreflightCommitId
+		{
+			get => _clonedPreflightCommitId ?? Horde.Commits.CommitId.FromPerforceChange(_clonedPreflightChange);
+			set => _clonedPreflightCommitId = value;
+		}
+		CommitId? _clonedPreflightCommitId;
 
 		/// <summary>
 		/// Description of the preflight
@@ -1302,7 +1398,23 @@ namespace EpicGames.Horde.Jobs
 		/// <summary>
 		/// The change number being built
 		/// </summary>
-		public int Change { get; set; }
+		[Obsolete("Use CommitId instead")]
+		public int Change
+		{
+			get => _change ?? _commitId?.GetPerforceChangeOrMinusOne() ?? 0;
+			set => _change = value;
+		}
+		int? _change;
+
+		/// <summary>
+		/// The commit being built
+		/// </summary>
+		public CommitIdWithOrder CommitId
+		{
+			get => _commitId ?? CommitIdWithOrder.FromPerforceChange(_change) ?? CommitIdWithOrder.Empty;
+			set => _commitId = value;
+		}
+		CommitIdWithOrder? _commitId;
 
 		/// <summary>
 		/// The step log id
diff --git a/Engine/Source/Programs/Shared/EpicGames.Horde/Jobs/TestData/TestDataMessages.cs b/Engine/Source/Programs/Shared/EpicGames.Horde/Jobs/TestData/TestDataMessages.cs
--- a/Engine/Source/Programs/Shared/EpicGames.Horde/Jobs/TestData/TestDataMessages.cs
+++ b/Engine/Source/Programs/Shared/EpicGames.Horde/Jobs/TestData/TestDataMessages.cs
@@ -3,6 +3,7 @@
 using System;
 using System.Collections.Generic;
 using System.ComponentModel.DataAnnotations;
+using EpicGames.Horde.Commits;
 
 #pragma warning disable CA2227 // Change 'x' to be read-only by removing the property setter
 
@@ -113,7 +114,23 @@ namespace EpicGames.Horde.Jobs.TestData
 		/// <summary>
 		/// The changelist number that contained the data
 		/// </summary>
-		public int Change { get; set; } = 0;
+		[Obsolete("Use CommitId instead")]
+		public int Change
+		{
+			get => _change ?? _commitId?.GetPerforceChangeOrMinusOne() ?? 0;
+			set => _change = value;
+		}
+		int? _change;
+
+		/// <summary>
+		/// The changelist number that contained the data
+		/// </summary>
+		public CommitIdWithOrder CommitId
+		{
+			get => _commitId ?? CommitIdWithOrder.FromPerforceChange(_change) ?? CommitIdWithOrder.Empty;
+			set => _commitId = value;
+		}
+		CommitIdWithOrder? _commitId;
 
 		/// <summary>
 		/// Key used to identify the particular data
@@ -346,7 +363,26 @@ namespace EpicGames.Horde.Jobs.TestData
 		/// <summary>
 		/// The build changelist upon which the test ran, may not correspond to the job changelist
 		/// </summary>
-		public int BuildChangeList { get; set; }
+		[Obsolete("Use BuildCommitId instead")]
+		public int BuildChangeList
+		{
+			get => _buildChangeList ?? _buildCommitId?.GetPerforceChangeOrMinusOne() ?? 0;
+			set => _buildChangeList = value;
+		}
+		int? _buildChangeList;
+
+#pragma warning disable CS0618 // Type or member is obsolete
+		/// <summary>
+		/// The build changelist upon which the test ran, may not correspond to the job changelist
+		/// </summary>
+		public CommitId BuildCommitId
+		{
+			get => _buildCommitId ?? CommitId.FromPerforceChange(_buildChangeList) ?? CommitId.Empty;
+			set => _buildCommitId = value;
+		}
+#pragma warning restore CS0618 // Type or member is obsolete
+
+		CommitId? _buildCommitId;
 
 		/// <summary>
 		/// The platform the test ran on 
diff --git a/Engine/Source/Programs/UnrealGameSync/UnrealGameSyncShared/ArchiveInfo.cs b/Engine/Source/Programs/UnrealGameSync/UnrealGameSyncShared/ArchiveInfo.cs
--- a/Engine/Source/Programs/UnrealGameSync/UnrealGameSyncShared/ArchiveInfo.cs
+++ b/Engine/Source/Programs/UnrealGameSync/UnrealGameSyncShared/ArchiveInfo.cs
@@ -340,7 +340,7 @@ namespace UnrealGameSync
 						foreach (GetArtifactResponse response in group)
 						{
 							HordeArchive archive = new HordeArchive(hordeClient, response.Id);
-							channel.ChangeNumberToArchive[response.Change] = archive;
+							channel.ChangeNumberToArchive[response.CommitId!.GetPerforceChange()] = archive;
 						}
 						channels.Add(channel);
 					}
