Horde: Rename IBlobRef -> IHashedBlobRef and IBlobHandle -> IBlobRef. Also add an unhashed IBlobRef<> interface.

#rnx

diff --git a/Engine/Source/Programs/AutomationTool/BuildGraph/Tasks/CreateArtifactTask.cs b/Engine/Source/Programs/AutomationTool/BuildGraph/Tasks/CreateArtifactTask.cs
--- a/Engine/Source/Programs/AutomationTool/BuildGraph/Tasks/CreateArtifactTask.cs
+++ b/Engine/Source/Programs/AutomationTool/BuildGraph/Tasks/CreateArtifactTask.cs
@@ -137,7 +137,7 @@ namespace AutomationTool.Tasks
 					int totalCount = files.Count;
 					long totalSize = files.Sum(x => x.Length);
 
-					IBlobRef<DirectoryNode> outputNodeRef = await writer.WriteFilesAsync(baseDir.ToDirectoryInfo(), files, progress: new UpdateStatsLogger(totalCount, totalSize, Logger));
+					IHashedBlobRef<DirectoryNode> outputNodeRef = await writer.WriteFilesAsync(baseDir.ToDirectoryInfo(), files, progress: new UpdateStatsLogger(totalCount, totalSize, Logger));
 					await writer.FlushAsync();
 
 					await client.WriteRefAsync(response.RefName, outputNodeRef);
diff --git a/Engine/Source/Programs/AutomationTool/BuildGraph/Tasks/DeployToolTask.cs b/Engine/Source/Programs/AutomationTool/BuildGraph/Tasks/DeployToolTask.cs
--- a/Engine/Source/Programs/AutomationTool/BuildGraph/Tasks/DeployToolTask.cs
+++ b/Engine/Source/Programs/AutomationTool/BuildGraph/Tasks/DeployToolTask.cs
@@ -159,7 +159,7 @@ namespace AutomationTool.Tasks
 
 			BlobSerializerOptions serializerOptions = BlobSerializerOptions.Create(infoResponse.ApiVersion);
 
-			IBlobRef handle;
+			IHashedBlobRef handle;
 
 			using IStorageClient storageClient = hordeClient.CreateStorageClient(toolId);
 			await using (IBlobWriter blobWriter = storageClient.CreateBlobWriter(serializerOptions: serializerOptions))
@@ -195,7 +195,7 @@ namespace AutomationTool.Tasks
 				createPaused = true;
 			}
 
-			BlobRefValue locator = handle.GetRefValue();
+			HashedBlobRefValue locator = handle.GetRefValue();
 			ToolDeploymentId deploymentId = await hordeHttpClient.CreateToolDeploymentAsync(toolId, _parameters.Version, duration, createPaused, locator);
 			Logger.LogInformation("Created {ToolId} deployment {DeploymentId}", toolId, deploymentId);
 		}
diff --git a/Engine/Source/Programs/Horde/Drivers/JobDriver/Execution/JobExecutor.cs b/Engine/Source/Programs/Horde/Drivers/JobDriver/Execution/JobExecutor.cs
--- a/Engine/Source/Programs/Horde/Drivers/JobDriver/Execution/JobExecutor.cs
+++ b/Engine/Source/Programs/Horde/Drivers/JobDriver/Execution/JobExecutor.cs
@@ -826,12 +826,12 @@ namespace JobDriver.Execution
 
 				Stopwatch timer = Stopwatch.StartNew();
 
-				IBlobRef<DirectoryNode> rootNodeRef;
+				IHashedBlobRef<DirectoryNode> rootNodeRef;
 				await using (IBlobWriter blobWriter = storage.CreateBlobWriter(artifact.RefName))
 				{
 					DirectoryNode buildGraphNode = new DirectoryNode();
 					await buildGraphNode.AddFilesAsync(workspaceDir, buildGraphFiles, blobWriter, cancellationToken: cancellationToken);
-					IBlobRef<DirectoryNode> outputNodeRef = await blobWriter.WriteBlobAsync(buildGraphNode, cancellationToken);
+					IHashedBlobRef<DirectoryNode> outputNodeRef = await blobWriter.WriteBlobAsync(buildGraphNode, cancellationToken);
 
 					DirectoryNode rootNode = new DirectoryNode();
 					rootNode.AddDirectory(new DirectoryEntry(BuildGraphTempStorageDir, buildGraphNode.Length, outputNodeRef));
@@ -1141,7 +1141,7 @@ namespace JobDriver.Execution
 
 				using IStorageClient storage = CreateStorageClient(new NamespaceId(artifact.NamespaceId), artifact.Token);
 
-				IBlobRef<DirectoryNode> rootRef;
+				IHashedBlobRef<DirectoryNode> rootRef;
 				await using (IBlobWriter blobWriter = storage.CreateBlobWriter(new RefName(artifact.RefName)))
 				{
 					try
@@ -1346,7 +1346,7 @@ namespace JobDriver.Execution
 				// Upload the data
 				Stopwatch timer = Stopwatch.StartNew();
 
-				IBlobRef<DirectoryNode> outputNodeRef;
+				IHashedBlobRef<DirectoryNode> outputNodeRef;
 				await using (IBlobWriter blobWriter = storage.CreateBlobWriter(artifact.RefName))
 				{
 					DirectoryNode outputNode = new DirectoryNode();
@@ -1426,7 +1426,7 @@ namespace JobDriver.Execution
 					// Upload the data
 					Stopwatch timer = Stopwatch.StartNew();
 
-					IBlobRef<DirectoryNode> outputNodeRef;
+					IHashedBlobRef<DirectoryNode> outputNodeRef;
 					await using (IBlobWriter blobWriter = storage.CreateBlobWriter(artifact.RefName))
 					{
 						DirectoryNode outputNode = new DirectoryNode();
diff --git a/Engine/Source/Programs/Horde/Drivers/JobDriver/Utility/TempStorage.cs b/Engine/Source/Programs/Horde/Drivers/JobDriver/Utility/TempStorage.cs
--- a/Engine/Source/Programs/Horde/Drivers/JobDriver/Utility/TempStorage.cs
+++ b/Engine/Source/Programs/Horde/Drivers/JobDriver/Utility/TempStorage.cs
@@ -718,7 +718,7 @@ namespace JobDriver.Utility
 			DirectoryNode rootNode = new DirectoryNode();
 			await rootNode.AddFilesAsync(workspaceDir, archiveFiles, writer, progress: new UpdateStatsLogger(logger), cancellationToken: cancellationToken);
 
-			IBlobRef<DirectoryNode> rootNodeRef = await writer.WriteBlobAsync(rootNode, cancellationToken: cancellationToken);
+			IHashedBlobRef<DirectoryNode> rootNodeRef = await writer.WriteBlobAsync(rootNode, cancellationToken: cancellationToken);
 			return new DirectoryEntry(blockDirectoryName, rootNode.Length, rootNodeRef);
 		}
 
diff --git a/Engine/Source/Programs/Horde/HordeCmd/Commands/Artifacts/ArtifactDownload.cs b/Engine/Source/Programs/Horde/HordeCmd/Commands/Artifacts/ArtifactDownload.cs
--- a/Engine/Source/Programs/Horde/HordeCmd/Commands/Artifacts/ArtifactDownload.cs
+++ b/Engine/Source/Programs/Horde/HordeCmd/Commands/Artifacts/ArtifactDownload.cs
@@ -51,7 +51,7 @@ namespace Horde.Commands.Artifacts
 			}
 
 			using IStorageClient store = _hordeClient.CreateStorageClient(artifact.Id);
-			IBlobHandle handle = await store.ReadRefAsync(new RefName("default"));
+			IBlobRef handle = await store.ReadRefAsync(new RefName("default"));
 
 			Stopwatch timer = Stopwatch.StartNew();
 
diff --git a/Engine/Source/Programs/Horde/HordeCmd/Commands/Bundles/BundleCreate.cs b/Engine/Source/Programs/Horde/HordeCmd/Commands/Bundles/BundleCreate.cs
--- a/Engine/Source/Programs/Horde/HordeCmd/Commands/Bundles/BundleCreate.cs
+++ b/Engine/Source/Programs/Horde/HordeCmd/Commands/Bundles/BundleCreate.cs
@@ -97,7 +97,7 @@ namespace Horde.Commands.Bundles
 				List<FileInfo> fileInfos = files.ConvertAll(x => x.ToFileInfo());
 				UpdateStatsLogger updateStatsLogger = new UpdateStatsLogger(files.Count, fileInfos.Sum(x => x.Length), logger);
 
-				IBlobRef<DirectoryNode> nodeRef = await writer.WriteFilesAsync(baseDir.ToDirectoryInfo(), fileInfos, options, updateStatsLogger, CancellationToken.None);
+				IHashedBlobRef<DirectoryNode> nodeRef = await writer.WriteFilesAsync(baseDir.ToDirectoryInfo(), fileInfos, options, updateStatsLogger, CancellationToken.None);
 
 				await writer.FlushAsync();
 
diff --git a/Engine/Source/Programs/Horde/HordeCmd/Commands/Bundles/BundleExtract.cs b/Engine/Source/Programs/Horde/HordeCmd/Commands/Bundles/BundleExtract.cs
--- a/Engine/Source/Programs/Horde/HordeCmd/Commands/Bundles/BundleExtract.cs
+++ b/Engine/Source/Programs/Horde/HordeCmd/Commands/Bundles/BundleExtract.cs
@@ -55,19 +55,19 @@ namespace Horde.Commands.Bundles
 			if (File != null)
 			{
 				using IStorageClient store = BundleStorageClient.CreateFromDirectory(File.Directory, BundleCache, logger);
-				IBlobHandle handle = store.CreateBlobHandle(await FileStorageBackend.ReadRefAsync(File));
+				IBlobRef handle = store.CreateBlobRef(await FileStorageBackend.ReadRefAsync(File));
 				await ExecuteInternalAsync(store, handle, logger);
 			}
 			else if (Ref != null)
 			{
 				using IStorageClient store = CreateStorageClient();
-				IBlobHandle handle = await store.ReadRefAsync(new RefName(Ref));
+				IBlobRef handle = await store.ReadRefAsync(new RefName(Ref));
 				await ExecuteInternalAsync(store, handle, logger);
 			}
 			else if (Node != null)
 			{
 				using IStorageClient store = CreateStorageClient();
-				IBlobHandle handle = store.CreateBlobHandle(new BlobLocator(Node));
+				IBlobRef handle = store.CreateBlobRef(new BlobLocator(Node));
 				await ExecuteInternalAsync(store, handle, logger);
 			}
 			else
@@ -78,7 +78,7 @@ namespace Horde.Commands.Bundles
 			return 0;
 		}
 
-		protected async Task ExecuteInternalAsync(IStorageClient store, IBlobHandle handle, ILogger logger)
+		protected async Task ExecuteInternalAsync(IStorageClient store, IBlobRef handle, ILogger logger)
 		{
 			Stopwatch timer = Stopwatch.StartNew();
 
diff --git a/Engine/Source/Programs/Horde/HordeCmd/Commands/Compute/ComputeRun.cs b/Engine/Source/Programs/Horde/HordeCmd/Commands/Compute/ComputeRun.cs
--- a/Engine/Source/Programs/Horde/HordeCmd/Commands/Compute/ComputeRun.cs
+++ b/Engine/Source/Programs/Horde/HordeCmd/Commands/Compute/ComputeRun.cs
@@ -160,7 +160,7 @@ namespace Horde.Commands.Compute
 		{
 			await using IBlobWriter writer = storage.CreateBlobWriter();
 
-			IBlobRef<DirectoryNode> sandbox = await writer.WriteFilesAsync(taskFile.Directory, cancellationToken: cancellationToken);
+			IHashedBlobRef<DirectoryNode> sandbox = await writer.WriteFilesAsync(taskFile.Directory, cancellationToken: cancellationToken);
 			await writer.FlushAsync(cancellationToken);
 
 			return sandbox.GetLocator();
diff --git a/Engine/Source/Programs/Horde/HordeCmd/Commands/Tools/ToolUpload.cs b/Engine/Source/Programs/Horde/HordeCmd/Commands/Tools/ToolUpload.cs
--- a/Engine/Source/Programs/Horde/HordeCmd/Commands/Tools/ToolUpload.cs
+++ b/Engine/Source/Programs/Horde/HordeCmd/Commands/Tools/ToolUpload.cs
@@ -40,7 +40,7 @@ namespace Horde.Commands
 		{
 			using IStorageClient storageClient = _httpStorageClientFactory.CreateClientWithPath($"api/v1/tools/{ToolId}");
 
-			IBlobRef<DirectoryNode> target;
+			IHashedBlobRef<DirectoryNode> target;
 			await using (IBlobWriter writer = storageClient.CreateBlobWriter())
 			{
 				target = await writer.WriteFilesAsync(InputDir);
diff --git a/Engine/Source/Programs/Horde/HordeCmd/Commands/Workspace/WorkspaceSync.cs b/Engine/Source/Programs/Horde/HordeCmd/Commands/Workspace/WorkspaceSync.cs
--- a/Engine/Source/Programs/Horde/HordeCmd/Commands/Workspace/WorkspaceSync.cs
+++ b/Engine/Source/Programs/Horde/HordeCmd/Commands/Workspace/WorkspaceSync.cs
@@ -51,19 +51,19 @@ namespace Horde.Commands.Workspace
 			if (File != null)
 			{
 				using IStorageClient store = BundleStorageClient.CreateFromDirectory(File.Directory, BundleCache, logger);
-				IBlobHandle handle = store.CreateBlobHandle(await FileStorageBackend.ReadRefAsync(File));
+				IBlobRef handle = store.CreateBlobRef(await FileStorageBackend.ReadRefAsync(File));
 				return await ExecuteInternalAsync(store, handle, logger);
 			}
 			else if (Ref != null)
 			{
 				using IStorageClient store = CreateStorageClient();
-				IBlobHandle handle = await store.ReadRefAsync(new RefName(Ref));
+				IBlobRef handle = await store.ReadRefAsync(new RefName(Ref));
 				return await ExecuteInternalAsync(store, handle, logger);
 			}
 			else if (Node != null)
 			{
 				using IStorageClient store = CreateStorageClient();
-				IBlobHandle handle = store.CreateBlobHandle(new BlobLocator(Node));
+				IBlobRef handle = store.CreateBlobRef(new BlobLocator(Node));
 				return await ExecuteInternalAsync(store, handle, logger);
 			}
 			else
@@ -72,7 +72,7 @@ namespace Horde.Commands.Workspace
 			}
 		}
 
-		async Task<int> ExecuteInternalAsync(IStorageClient store, IBlobHandle handle, ILogger logger)
+		async Task<int> ExecuteInternalAsync(IStorageClient store, IBlobRef handle, ILogger logger)
 		{
 			RootDir ??= DirectoryReference.GetCurrentDirectory();
 			CancellationToken cancellationToken = CancellationToken.None;
diff --git a/Engine/Source/Programs/Horde/Plugins/Build/HordeServer.Build.Tests/Artifacts/UnsyncCacheTests.cs b/Engine/Source/Programs/Horde/Plugins/Build/HordeServer.Build.Tests/Artifacts/UnsyncCacheTests.cs
--- a/Engine/Source/Programs/Horde/Plugins/Build/HordeServer.Build.Tests/Artifacts/UnsyncCacheTests.cs
+++ b/Engine/Source/Programs/Horde/Plugins/Build/HordeServer.Build.Tests/Artifacts/UnsyncCacheTests.cs
@@ -24,7 +24,7 @@ namespace HordeServer.Tests.Artifacts
 
 			Dictionary<IoHash, ReadOnlyMemory<byte>> chunks = new Dictionary<IoHash, ReadOnlyMemory<byte>>();
 
-			IBlobRef<DirectoryNode> directoryRef;
+			IHashedBlobRef<DirectoryNode> directoryRef;
 			await using (IBlobWriter writer = storageClient.CreateBlobWriter())
 			{
 				using ChunkedDataWriter chunkedWriter = new ChunkedDataWriter(writer, new ChunkingOptions());
@@ -59,7 +59,7 @@ namespace HordeServer.Tests.Artifacts
 			int offset = 0;
 			foreach (UnsyncBlock block in file.Blocks)
 			{
-				IBlobRef? blobRef = await cache.ReadBlobRefAsync(artifact.Object, block.Blob.Hash);
+				IHashedBlobRef? blobRef = await cache.ReadBlobRefAsync(artifact.Object, block.Blob.Hash);
 				Assert.IsNotNull(blobRef);
 
 				using BlobData blobData = await blobRef.ReadBlobDataAsync();
diff --git a/Engine/Source/Programs/Horde/Plugins/Build/HordeServer.Build.Tests/Issues/IssueServiceTests.cs b/Engine/Source/Programs/Horde/Plugins/Build/HordeServer.Build.Tests/Issues/IssueServiceTests.cs
--- a/Engine/Source/Programs/Horde/Plugins/Build/HordeServer.Build.Tests/Issues/IssueServiceTests.cs
+++ b/Engine/Source/Programs/Horde/Plugins/Build/HordeServer.Build.Tests/Issues/IssueServiceTests.cs
@@ -67,7 +67,7 @@ namespace HordeServer.Tests.Issues
 
 				await using (IBlobWriter writer = _storageClient.CreateBlobWriter())
 				{
-					IBlobRef<LogNode> handle = await _builder.FlushAsync(writer, true, CancellationToken.None);
+					IHashedBlobRef<LogNode> handle = await _builder.FlushAsync(writer, true, CancellationToken.None);
 					await _storageClient.WriteRefAsync(_log.RefName, handle);
 				}
 
@@ -356,7 +356,7 @@ namespace HordeServer.Tests.Issues
 					LogBuilder builder = new LogBuilder(LogFormat.Json, NullLogger.Instance);
 					builder.WriteData(data);
 
-					IBlobRef<LogNode> handle = await builder.FlushAsync(writer, true, CancellationToken.None);
+					IHashedBlobRef<LogNode> handle = await builder.FlushAsync(writer, true, CancellationToken.None);
 					await storageClient.WriteRefAsync(log!.RefName, handle);
 				}
 			}
diff --git a/Engine/Source/Programs/Horde/Plugins/Build/HordeServer.Build/Artifacts/ArtifactsController.cs b/Engine/Source/Programs/Horde/Plugins/Build/HordeServer.Build/Artifacts/ArtifactsController.cs
--- a/Engine/Source/Programs/Horde/Plugins/Build/HordeServer.Build/Artifacts/ArtifactsController.cs
+++ b/Engine/Source/Programs/Horde/Plugins/Build/HordeServer.Build/Artifacts/ArtifactsController.cs
@@ -171,7 +171,7 @@ namespace HordeServer.Artifacts
 					break;
 				}
 
-				BlobRefValue? refValue = await storageBackend.TryReadRefAsync(prevArtifact.RefName, cancellationToken: cancellationToken);
+				HashedBlobRefValue? refValue = await storageBackend.TryReadRefAsync(prevArtifact.RefName, cancellationToken: cancellationToken);
 				if (refValue != null)
 				{
 					return prevArtifact.RefName;
@@ -574,7 +574,7 @@ namespace HordeServer.Artifacts
 			IStorageClient storageClient = _storageService.CreateClient(artifact.NamespaceId);
 			try
 			{
-				IBlobRef<DirectoryNode> directory = await storageClient.ReadRefAsync<DirectoryNode>(artifact.RefName, DateTime.UtcNow.AddHours(1.0), cancellationToken: cancellationToken);
+				IHashedBlobRef<DirectoryNode> directory = await storageClient.ReadRefAsync<DirectoryNode>(artifact.RefName, DateTime.UtcNow.AddHours(1.0), cancellationToken: cancellationToken);
 
 				Stream stream = directory.AsZipStream(filter).WrapOwnership(storageClient);
 				return new FileStreamResult(stream, "application/zip") { FileDownloadName = $"{artifact.RefName}.zip" };
@@ -724,7 +724,7 @@ namespace HordeServer.Artifacts
 						return BadRequest($"Invalid IoHash value: {block}");
 					}
 
-					IBlobRef? blobRef = await _unsyncCache.ReadBlobRefAsync(artifact, hash, cancellationToken);
+					IHashedBlobRef? blobRef = await _unsyncCache.ReadBlobRefAsync(artifact, hash, cancellationToken);
 					if (blobRef == null)
 					{
 						return NotFound($"Hash '{hash}' is not part of artifact {artifact.Id}");
diff --git a/Engine/Source/Programs/Horde/Plugins/Build/HordeServer.Build/Artifacts/UnsyncCache.cs b/Engine/Source/Programs/Horde/Plugins/Build/HordeServer.Build/Artifacts/UnsyncCache.cs
--- a/Engine/Source/Programs/Horde/Plugins/Build/HordeServer.Build/Artifacts/UnsyncCache.cs
+++ b/Engine/Source/Programs/Horde/Plugins/Build/HordeServer.Build/Artifacts/UnsyncCache.cs
@@ -36,7 +36,7 @@ namespace HordeServer.Artifacts
 	/// <param name="Length">Length of the block</param>
 	/// <param name="RollingHash">Rolling hash of the block</param>
 	/// <param name="Blob">Handle to the corresponding blob</param>
-	public record class UnsyncBlock(long Offset, long Length, uint RollingHash, IBlobRef<LeafChunkedDataNode> Blob);
+	public record class UnsyncBlock(long Offset, long Length, uint RollingHash, IHashedBlobRef<LeafChunkedDataNode> Blob);
 
 	/// <summary>
 	/// Implements a cache for downloading Unsync blobs
@@ -48,9 +48,9 @@ namespace HordeServer.Artifacts
 			public IStorageClient StorageClient { get; }
 			public UnsyncManifest Manifest { get; }
 			public ReadOnlyMemory<byte> ManifestData { get; }
-			public FrozenDictionary<IoHash, IBlobRef<LeafChunkedDataNode>> Blobs { get; }
+			public FrozenDictionary<IoHash, IHashedBlobRef<LeafChunkedDataNode>> Blobs { get; }
 
-			public ArtifactInfo(IStorageClient storageClient, UnsyncManifest manifest, ReadOnlyMemory<byte> manifestData, FrozenDictionary<IoHash, IBlobRef<LeafChunkedDataNode>> blobs)
+			public ArtifactInfo(IStorageClient storageClient, UnsyncManifest manifest, ReadOnlyMemory<byte> manifestData, FrozenDictionary<IoHash, IHashedBlobRef<LeafChunkedDataNode>> blobs)
 			{
 				StorageClient = storageClient;
 				Manifest = manifest;
@@ -191,7 +191,7 @@ namespace HordeServer.Artifacts
 		/// <param name="blobHash"></param>
 		/// <param name="cancellationToken"></param>
 		/// <returns></returns>
-		public async ValueTask<IBlobRef?> ReadBlobRefAsync(IArtifact artifact, IoHash blobHash, CancellationToken cancellationToken = default)
+		public async ValueTask<IHashedBlobRef?> ReadBlobRefAsync(IArtifact artifact, IoHash blobHash, CancellationToken cancellationToken = default)
 		{
 			ArtifactInfo? artifactInfo = await GetArtifactInfoAsync(artifact, cancellationToken);
 			if (artifactInfo == null)
@@ -199,7 +199,7 @@ namespace HordeServer.Artifacts
 				return null;
 			}
 
-			IBlobRef<LeafChunkedDataNode>? blobRef;
+			IHashedBlobRef<LeafChunkedDataNode>? blobRef;
 			if (!artifactInfo.Blobs.TryGetValue(blobHash, out blobRef))
 			{
 				return null;
@@ -216,7 +216,7 @@ namespace HordeServer.Artifacts
 				Stopwatch timer = Stopwatch.StartNew();
 				storageClient = _storageService.CreateClient(artifact.NamespaceId);
 
-				IBlobRef<DirectoryNode>? target = await storageClient.TryReadRefAsync<DirectoryNode>(artifact.RefName, cancellationToken: cancellationToken);
+				IHashedBlobRef<DirectoryNode>? target = await storageClient.TryReadRefAsync<DirectoryNode>(artifact.RefName, cancellationToken: cancellationToken);
 				if (target == null)
 				{
 					return null;
@@ -225,7 +225,7 @@ namespace HordeServer.Artifacts
 				List<UnsyncFile> files = new List<UnsyncFile>();
 				await FindFilesAsync(new Utf8StringBuilder(), target, files, cancellationToken);
 
-				Dictionary<IoHash, IBlobRef<LeafChunkedDataNode>> blocks = new Dictionary<IoHash, IBlobRef<LeafChunkedDataNode>>();
+				Dictionary<IoHash, IHashedBlobRef<LeafChunkedDataNode>> blocks = new Dictionary<IoHash, IHashedBlobRef<LeafChunkedDataNode>>();
 				foreach (UnsyncBlock block in files.SelectMany(x => x.Blocks))
 				{
 					blocks[block.Blob.Hash] = block.Blob;
@@ -246,7 +246,7 @@ namespace HordeServer.Artifacts
 			}
 		}
 
-		static async Task FindFilesAsync(Utf8StringBuilder path, IBlobRef<DirectoryNode> directoryNodeRef, List<UnsyncFile> files, CancellationToken cancellationToken)
+		static async Task FindFilesAsync(Utf8StringBuilder path, IHashedBlobRef<DirectoryNode> directoryNodeRef, List<UnsyncFile> files, CancellationToken cancellationToken)
 		{
 			DirectoryNode directoryNode = await directoryNodeRef.ReadBlobAsync(cancellationToken);
 
diff --git a/Engine/Source/Programs/Horde/Plugins/Build/HordeServer.Build/Replicators/PerforceReplicator.cs b/Engine/Source/Programs/Horde/Plugins/Build/HordeServer.Build/Replicators/PerforceReplicator.cs
--- a/Engine/Source/Programs/Horde/Plugins/Build/HordeServer.Build/Replicators/PerforceReplicator.cs
+++ b/Engine/Source/Programs/Horde/Plugins/Build/HordeServer.Build/Replicators/PerforceReplicator.cs
@@ -66,12 +66,12 @@ namespace HordeServer.Replicators
 		{
 			public int Change { get; }
 			public int ParentChange { get; }
-			public IBlobRef<CommitNode>? ParentHandle { get; }
+			public IHashedBlobRef<CommitNode>? ParentHandle { get; }
 			public long CopiedSize { get; set; }
-			public IBlobRef<DirectoryNode>? Contents { get; set; }
+			public IHashedBlobRef<DirectoryNode>? Contents { get; set; }
 			public List<string> Paths { get; }
 
-			public StateNode(int number, int parentNumber, IBlobRef<CommitNode>? parentHandle, long copiedSize, IBlobRef<DirectoryNode>? contents, List<string>? paths = null)
+			public StateNode(int number, int parentNumber, IHashedBlobRef<CommitNode>? parentHandle, long copiedSize, IHashedBlobRef<DirectoryNode>? contents, List<string>? paths = null)
 			{
 				Change = number;
 				ParentChange = parentNumber;
@@ -91,7 +91,7 @@ namespace HordeServer.Replicators
 				int change = (int)reader.ReadUnsignedVarInt();
 				int parentChange = (int)reader.ReadUnsignedVarInt();
 
-				IBlobRef<CommitNode>? parentHandle = null;
+				IHashedBlobRef<CommitNode>? parentHandle = null;
 				if (reader.ReadBoolean())
 				{
 					parentHandle = reader.ReadBlobRef<CommitNode>();
@@ -103,7 +103,7 @@ namespace HordeServer.Replicators
 					copiedSize = (long)reader.ReadUnsignedVarInt();
 				}
 
-				IBlobRef<DirectoryNode>? contents = null;
+				IHashedBlobRef<DirectoryNode>? contents = null;
 				if (reader.ReadBoolean())
 				{
 					contents = reader.ReadBlobRef<DirectoryNode>();
@@ -455,7 +455,7 @@ namespace HordeServer.Replicators
 					RedirectNode<CommitNode>? lastCommit = await store.TryReadRefTargetAsync<RedirectNode<CommitNode>>(refName, options: blobOptions, cancellationToken: cancellationToken);
 
 					CommitNode? parent = null;
-					IBlobRef<CommitNode>? parentHandle = lastCommit?.Target;
+					IHashedBlobRef<CommitNode>? parentHandle = lastCommit?.Target;
 					while (parentHandle != null)
 					{
 						CommitNode parentBlob = await parentHandle.ReadBlobAsync(cancellationToken);
@@ -736,7 +736,7 @@ namespace HordeServer.Replicators
 				await root.UpdateAsync(rootUpdate, directoryWriter, cancellationToken);
 				stateNode.Contents = await directoryWriter.WriteBlobAsync(root, cancellationToken);
 
-				IBlobRef<StateNode> stateNodeRef = await directoryWriter.WriteBlobAsync(stateNode, cancellationToken);
+				IHashedBlobRef<StateNode> stateNodeRef = await directoryWriter.WriteBlobAsync(stateNode, cancellationToken);
 				await directoryWriter.FlushAsync(cancellationToken);
 
 				await store.WriteRefAsync(incRefName, stateNodeRef, cancellationToken: cancellationToken);
@@ -758,11 +758,11 @@ namespace HordeServer.Replicators
 			DirectoryNodeRef rootRef = new DirectoryNodeRef(root.Length, stateNode.Contents!);
 
 			ChangeRecord changeRecord = await perforce.GetChangeAsync(GetChangeOptions.None, change, cancellationToken);
-			CommitNode commitNode = new CommitNode(change, stateNode.ParentHandle, changeRecord.User ?? "Unknown", null, null, null, changeRecord.Description ?? String.Empty, changeRecord.Date, rootRef, new Dictionary<Guid, IBlobRef>());
-			IBlobRef<CommitNode> commitNodeRef = await commitWriter.WriteBlobAsync(commitNode, cancellationToken);
+			CommitNode commitNode = new CommitNode(change, stateNode.ParentHandle, changeRecord.User ?? "Unknown", null, null, null, changeRecord.Description ?? String.Empty, changeRecord.Date, rootRef, new Dictionary<Guid, IHashedBlobRef>());
+			IHashedBlobRef<CommitNode> commitNodeRef = await commitWriter.WriteBlobAsync(commitNode, cancellationToken);
 
 			RedirectNode<CommitNode> redirectNode = new RedirectNode<CommitNode>(commitNodeRef);
-			IBlobRef<RedirectNode<CommitNode>> redirectNodeRef = await commitWriter.WriteBlobAsync(redirectNode, cancellationToken);
+			IHashedBlobRef<RedirectNode<CommitNode>> redirectNodeRef = await commitWriter.WriteBlobAsync(redirectNode, cancellationToken);
 
 			await commitWriter.FlushAsync(cancellationToken);
 			await store.WriteRefAsync(refName, redirectNodeRef, options.RefOptions, cancellationToken: cancellationToken);
diff --git a/Engine/Source/Programs/Horde/Plugins/Compute/HordeServer.Compute.Tests/Logs/LogIndexingTests.cs b/Engine/Source/Programs/Horde/Plugins/Compute/HordeServer.Compute.Tests/Logs/LogIndexingTests.cs
--- a/Engine/Source/Programs/Horde/Plugins/Compute/HordeServer.Compute.Tests/Logs/LogIndexingTests.cs
+++ b/Engine/Source/Programs/Horde/Plugins/Compute/HordeServer.Compute.Tests/Logs/LogIndexingTests.cs
@@ -43,7 +43,7 @@ namespace HordeServer.Tests.Logs
 
 		public async Task<ILog> FlushAsync(bool complete = true)
 		{
-			IBlobRef<LogNode> handle = await _builder.FlushAsync(_blobWriter, complete, CancellationToken.None);
+			IHashedBlobRef<LogNode> handle = await _builder.FlushAsync(_blobWriter, complete, CancellationToken.None);
 			await _storageClient.WriteRefAsync(_logFile.RefName, handle);
 			_logFile = await _logFile.UpdateLineCountAsync(_lineCount, complete, CancellationToken.None);
 			return _logFile;
diff --git a/Engine/Source/Programs/Horde/Plugins/Compute/HordeServer.Compute/Logs/LogRpcService.cs b/Engine/Source/Programs/Horde/Plugins/Compute/HordeServer.Compute/Logs/LogRpcService.cs
--- a/Engine/Source/Programs/Horde/Plugins/Compute/HordeServer.Compute/Logs/LogRpcService.cs
+++ b/Engine/Source/Programs/Horde/Plugins/Compute/HordeServer.Compute/Logs/LogRpcService.cs
@@ -61,7 +61,7 @@ namespace HordeServer.Logs
 				hash = IoHash.Zero;
 			}
 
-			IBlobRef target = store.CreateBlobRef(hash, new BlobLocator(request.TargetLocator));
+			IHashedBlobRef target = store.CreateBlobRef(hash, new BlobLocator(request.TargetLocator));
 			await store.WriteRefAsync(new RefName(request.LogId), target);
 
 			await log.UpdateLineCountAsync(request.LineCount, request.Complete, context.CancellationToken);
diff --git a/Engine/Source/Programs/Horde/Plugins/Ddc/HordeServer.Ddc/Ddc/BlobService.cs b/Engine/Source/Programs/Horde/Plugins/Ddc/HordeServer.Ddc/Ddc/BlobService.cs
--- a/Engine/Source/Programs/Horde/Plugins/Ddc/HordeServer.Ddc/Ddc/BlobService.cs
+++ b/Engine/Source/Programs/Horde/Plugins/Ddc/HordeServer.Ddc/Ddc/BlobService.cs
@@ -176,7 +176,7 @@ namespace HordeServer.Ddc
 		{
 			using IStorageClient storageClient = _storageClientFactory.CreateClient(ns);
 
-			IBlobRef blobRef;
+			IHashedBlobRef blobRef;
 			await using (IBlobWriter writer = storageClient.CreateBlobWriter())
 			{
 				int contentLength = (int)content.Length;
diff --git a/Engine/Source/Programs/Horde/Plugins/Ddc/HordeServer.Ddc/Ddc/RefService.cs b/Engine/Source/Programs/Horde/Plugins/Ddc/HordeServer.Ddc/Ddc/RefService.cs
--- a/Engine/Source/Programs/Horde/Plugins/Ddc/HordeServer.Ddc/Ddc/RefService.cs
+++ b/Engine/Source/Programs/Horde/Plugins/Ddc/HordeServer.Ddc/Ddc/RefService.cs
@@ -43,7 +43,7 @@ namespace HordeServer.Ddc
 				throw new BlobNotFoundException(ns, blobHash);
 			}
 
-			IBlobHandle blobHandle = blobAlias.Target;
+			IBlobRef blobHandle = blobAlias.Target;
 			using BlobData blobContents = await blobHandle.ReadBlobDataAsync(cancellationToken);
 			CbObject payload = new CbObject(blobContents.Data);
 
@@ -73,15 +73,15 @@ namespace HordeServer.Ddc
 				// TODO: We resolved all these blobs above... Need to just have GetReferencedBlobs just return the appropriate handles directly.
 				RefName refName = GetRefName(bucket, key);
 
-				IBlobRef<DdcRefNode> refNodeRef;
+				IHashedBlobRef<DdcRefNode> refNodeRef;
 				await using (IBlobWriter writer = storageClient.CreateBlobWriter(refName))
 				{
 					DdcRefNode refNode = new DdcRefNode(blobHash.AsIoHash());
-					refNode.References.Add(BlobRef.Create(blobHash.AsIoHash(), blobHandle));
+					refNode.References.Add(HashedBlobRef.Create(blobHash.AsIoHash(), blobHandle));
 					foreach (BlobId referencedBlob in referencedBlobs)
 					{
 						BlobAlias? alias = await storageClient.FindAliasAsync(BlobService.GetAlias(referencedBlob), cancellationToken);
-						refNode.References.Add(BlobRef.Create(referencedBlob.AsIoHash(), alias!.Target));
+						refNode.References.Add(HashedBlobRef.Create(referencedBlob.AsIoHash(), alias!.Target));
 					}
 					refNodeRef = await writer.WriteBlobAsync(refNode, cancellationToken);
 				}
diff --git a/Engine/Source/Programs/Horde/Plugins/Storage/HordeServer.Storage.Tests/Storage/BlobStoreTests.cs b/Engine/Source/Programs/Horde/Plugins/Storage/HordeServer.Storage.Tests/Storage/BlobStoreTests.cs
--- a/Engine/Source/Programs/Horde/Plugins/Storage/HordeServer.Storage.Tests/Storage/BlobStoreTests.cs
+++ b/Engine/Source/Programs/Horde/Plugins/Storage/HordeServer.Storage.Tests/Storage/BlobStoreTests.cs
@@ -44,51 +44,51 @@ namespace HordeServer.Tests.Storage
 		class Blob
 		{
 			public ReadOnlyMemory<byte> Data { get; set; }
-			public List<BlobRefValue> References { get; set; }
+			public List<HashedBlobRefValue> References { get; set; }
 
-			public Blob() : this(ReadOnlyMemory<byte>.Empty, new List<BlobRefValue>())
+			public Blob() : this(ReadOnlyMemory<byte>.Empty, new List<HashedBlobRefValue>())
 			{
 			}
 
-			public Blob(ReadOnlyMemory<byte> data, IEnumerable<BlobRefValue> references)
+			public Blob(ReadOnlyMemory<byte> data, IEnumerable<HashedBlobRefValue> references)
 			{
 				Data = data;
-				References = new List<BlobRefValue>(references);
+				References = new List<HashedBlobRefValue>(references);
 			}
 		}
 
-		static async ValueTask<Blob> ReadBlobAsync(IStorageClient store, BlobRefValue locator)
+		static async ValueTask<Blob> ReadBlobAsync(IStorageClient store, HashedBlobRefValue locator)
 		{
 			using (BlobData blobData = await store.CreateBlobRef(locator.Hash, locator.Locator).ReadBlobDataAsync())
 			{
 				BlobReader reader = new BlobReader(blobData, null);
 				ReadOnlyMemory<byte> data = reader.ReadVariableLengthBytes();
 
-				List<IBlobRef> handles = new List<IBlobRef>();
+				List<IHashedBlobRef> handles = new List<IHashedBlobRef>();
 				while (reader.RemainingMemory.Length > 0)
 				{
 					handles.Add(reader.ReadBlobRef<object>());
 				}
 
-				List<BlobRefValue> locators = handles.ConvertAll(x => x.GetRefValue());
+				List<HashedBlobRefValue> locators = handles.ConvertAll(x => x.GetRefValue());
 				return new Blob(data, locators);
 			}
 		}
 
-		static async ValueTask<BlobRefValue> WriteBlobAsync(IStorageClient store, Blob blob)
+		static async ValueTask<HashedBlobRefValue> WriteBlobAsync(IStorageClient store, Blob blob)
 		{
 			await using IBlobWriter writer = store.CreateBlobWriter();
 			writer.WriteVariableLengthBytes(blob.Data.Span);
 
-			foreach (BlobRefValue reference in blob.References)
+			foreach (HashedBlobRefValue reference in blob.References)
 			{
 				writer.WriteBlobRef(store.CreateBlobRef(reference));
 			}
 
-			IBlobRef blobRef = await writer.CompleteAsync(s_blobType);
+			IHashedBlobRef blobRef = await writer.CompleteAsync(s_blobType);
 			await blobRef.FlushAsync();
 
-			return new BlobRefValue(blobRef.Hash, blobRef.GetLocator());
+			return new HashedBlobRefValue(blobRef.Hash, blobRef.GetLocator());
 		}
 
 		[TestMethod]
@@ -98,8 +98,8 @@ namespace HordeServer.Tests.Storage
 
 			byte[] input = CreateTestData(256, 0);
 
-			Blob inputBlob = new Blob(input, Array.Empty<BlobRefValue>());
-			BlobRefValue locator = await WriteBlobAsync(store, inputBlob);
+			Blob inputBlob = new Blob(input, Array.Empty<HashedBlobRefValue>());
+			HashedBlobRefValue locator = await WriteBlobAsync(store, inputBlob);
 			Blob outputBlob = await ReadBlobAsync(store, locator);
 
 			Assert.IsTrue(outputBlob.Data.Span.SequenceEqual(input));
@@ -112,19 +112,19 @@ namespace HordeServer.Tests.Storage
 			using IStorageClient store = CreateStorageClient();
 
 			byte[] input1 = CreateTestData(256, 1);
-			BlobRefValue locator1 = await WriteBlobAsync(store, new Blob(input1, Array.Empty<BlobRefValue>()));
+			HashedBlobRefValue locator1 = await WriteBlobAsync(store, new Blob(input1, Array.Empty<HashedBlobRefValue>()));
 			Blob blob1 = await ReadBlobAsync(store, locator1);
 			Assert.IsTrue(blob1.Data.Span.SequenceEqual(input1));
-			Assert.IsTrue(blob1.References.SequenceEqual(Array.Empty<BlobRefValue>()));
+			Assert.IsTrue(blob1.References.SequenceEqual(Array.Empty<HashedBlobRefValue>()));
 
 			byte[] input2 = CreateTestData(256, 2);
-			BlobRefValue locator2 = await WriteBlobAsync(store, new Blob(input2, new BlobRefValue[] { locator1 }));
+			HashedBlobRefValue locator2 = await WriteBlobAsync(store, new Blob(input2, new HashedBlobRefValue[] { locator1 }));
 			Blob blob2 = await ReadBlobAsync(store, locator2);
 			Assert.IsTrue(blob2.Data.Span.SequenceEqual(input2));
 			Assert.IsTrue(blob2.References.SequenceEqual(new[] { locator1 }));
 
 			byte[] input3 = CreateTestData(256, 3);
-			BlobRefValue locator3 = await WriteBlobAsync(store, new Blob(input3, new BlobRefValue[] { locator1, locator2, locator1 }));
+			HashedBlobRefValue locator3 = await WriteBlobAsync(store, new Blob(input3, new HashedBlobRefValue[] { locator1, locator2, locator1 }));
 			Blob blob3 = await ReadBlobAsync(store, locator3);
 			Assert.IsTrue(blob3.Data.Span.SequenceEqual(input3));
 			Assert.IsTrue(blob3.References.SequenceEqual(new[] { locator1, locator2, locator1 }));
@@ -133,7 +133,7 @@ namespace HordeServer.Tests.Storage
 			{
 				RefName refName = new RefName("hello");
 				await store.WriteRefAsync(refName, store.CreateBlobRef(locator3));
-				IBlobRef refTarget = await store.ReadRefAsync(refName);
+				IHashedBlobRef refTarget = await store.ReadRefAsync(refName);
 				Assert.AreEqual(locator3.Locator, refTarget.GetLocator());
 			}
 		}
@@ -143,8 +143,8 @@ namespace HordeServer.Tests.Storage
 		{
 			using IStorageClient store = CreateStorageClient();
 
-			Blob blob1 = new Blob(new byte[] { 1, 2, 3 }, Array.Empty<BlobRefValue>());
-			BlobRefValue target = await WriteBlobAsync(store, blob1);
+			Blob blob1 = new Blob(new byte[] { 1, 2, 3 }, Array.Empty<HashedBlobRefValue>());
+			HashedBlobRefValue target = await WriteBlobAsync(store, blob1);
 
 			await store.WriteRefAsync("test-ref-1", store.CreateBlobRef(target));
 			await store.WriteRefAsync("test-ref-2", store.CreateBlobRef(target), new RefOptions { Lifetime = TimeSpan.FromMinutes(30.0), Extend = true });
@@ -173,9 +173,9 @@ namespace HordeServer.Tests.Storage
 			Assert.AreEqual(default, await TryReadRefValueAsync(store, "test-ref-3"));
 		}
 
-		static async Task<BlobRefValue?> TryReadRefValueAsync(IStorageClient store, RefName name)
+		static async Task<HashedBlobRefValue?> TryReadRefValueAsync(IStorageClient store, RefName name)
 		{
-			IBlobRef? handle = await store.TryReadRefAsync(name);
+			IHashedBlobRef? handle = await store.TryReadRefAsync(name);
 			if (handle == null)
 			{
 				return default;
diff --git a/Engine/Source/Programs/Horde/Plugins/Storage/HordeServer.Storage.Tests/Storage/GcServiceTests.cs b/Engine/Source/Programs/Horde/Plugins/Storage/HordeServer.Storage.Tests/Storage/GcServiceTests.cs
--- a/Engine/Source/Programs/Horde/Plugins/Storage/HordeServer.Storage.Tests/Storage/GcServiceTests.cs
+++ b/Engine/Source/Programs/Horde/Plugins/Storage/HordeServer.Storage.Tests/Storage/GcServiceTests.cs
@@ -34,15 +34,15 @@ namespace HordeServer.Tests.Storage
 			using IStorageClient store = StorageService.CreateClient(new NamespaceId("default"));
 
 			Random random = new Random(0);
-			IBlobRef[] blobs = await CreateTestDataAsync(store, 30, 50, 30, 5, random);
+			IHashedBlobRef[] blobs = await CreateTestDataAsync(store, 30, 50, 30, 5, random);
 
-			HashSet<IBlobHandle> roots = new HashSet<IBlobHandle>();
+			HashSet<IBlobRef> roots = new HashSet<IBlobRef>();
 			for (int idx = 0; idx < 10; idx++)
 			{
 				int blobIdx = (int)(random.NextDouble() * blobs.Length);
 				if (roots.Add(blobs[blobIdx]))
 				{
-					IBlobRef handle = blobs[blobIdx];
+					IHashedBlobRef handle = blobs[blobIdx];
 					await store.WriteRefAsync(new RefName($"ref-{idx}"), handle);
 				}
 			}
@@ -60,16 +60,16 @@ namespace HordeServer.Tests.Storage
 			Assert.IsTrue(remaining.All(x => nodePaths.Contains(x)));
 		}
 
-		static async Task<HashSet<BlobLocator>> FindNodesAsync(IStorageClient store, IEnumerable<IBlobHandle> roots)
+		static async Task<HashSet<BlobLocator>> FindNodesAsync(IStorageClient store, IEnumerable<IBlobRef> roots)
 		{
 			HashSet<BlobLocator> nodes = new HashSet<BlobLocator>();
 			await FindNodesAsync(store, roots, nodes);
 			return nodes;
 		}
 
-		static async Task FindNodesAsync(IStorageClient store, IEnumerable<IBlobHandle> roots, HashSet<BlobLocator> nodes)
+		static async Task FindNodesAsync(IStorageClient store, IEnumerable<IBlobRef> roots, HashSet<BlobLocator> nodes)
 		{
-			foreach (IBlobHandle root in roots)
+			foreach (IBlobRef root in roots)
 			{
 				BlobLocator locator = root.GetLocator();
 				if (nodes.Add(locator))
@@ -80,7 +80,7 @@ namespace HordeServer.Tests.Storage
 			}
 		}
 
-		static async ValueTask<IBlobRef[]> CreateTestDataAsync(IStorageClient store, int numRoots, int numInterior, int numLeaves, int avgChildren, Random random)
+		static async ValueTask<IHashedBlobRef[]> CreateTestDataAsync(IStorageClient store, int numRoots, int numInterior, int numLeaves, int avgChildren, Random random)
 		{
 			int firstRoot = 0;
 			int firstInterior = firstRoot + numRoots;
@@ -109,14 +109,14 @@ namespace HordeServer.Tests.Storage
 
 			BlobType blobType = new BlobType(Guid.Parse("{AFDF76A7-4DEE-5333-F5B5-37B8451251CA}"), 0);
 
-			IBlobRef[] handles = new IBlobRef[children.Length];
+			IHashedBlobRef[] handles = new IHashedBlobRef[children.Length];
 			for (int idx = numNodes - 1; idx >= 0; idx--)
 			{
-				IBlobRef handle;
+				IHashedBlobRef handle;
 				await using (IBlobWriter writer = store.CreateBlobWriter("gctest"))
 				{
-					List<IBlobRef> imports = children[idx].ConvertAll(x => handles[x]);
-					foreach (IBlobRef import in imports)
+					List<IHashedBlobRef> imports = children[idx].ConvertAll(x => handles[x]);
+					foreach (IHashedBlobRef import in imports)
 					{
 						writer.WriteBlobRef(import);
 					}
diff --git a/Engine/Source/Programs/Horde/Plugins/Storage/HordeServer.Storage.Tests/Storage/StorageServiceTests.cs b/Engine/Source/Programs/Horde/Plugins/Storage/HordeServer.Storage.Tests/Storage/StorageServiceTests.cs
--- a/Engine/Source/Programs/Horde/Plugins/Storage/HordeServer.Storage.Tests/Storage/StorageServiceTests.cs
+++ b/Engine/Source/Programs/Horde/Plugins/Storage/HordeServer.Storage.Tests/Storage/StorageServiceTests.cs
@@ -29,9 +29,9 @@ namespace HordeServer.Tests.Storage
 			BlobType type2 = new BlobType(Guid.Parse("{6CB3A005-4787-26BA-3E79-D286CB7137D1}"), 0);
 			byte[] data2 = new byte[] { 4, 5, 6 };
 
-			IBlobRef handle1a;
-			IBlobRef handle1b;
-			IBlobRef handle2;
+			IHashedBlobRef handle1a;
+			IHashedBlobRef handle1b;
+			IHashedBlobRef handle2;
 			await using (IBlobWriter writer = client.CreateBlobWriter())
 			{
 				writer.WriteFixedLengthBytes(data1);
diff --git a/Engine/Source/Programs/Horde/Plugins/Storage/HordeServer.Storage/Storage/StorageController.cs b/Engine/Source/Programs/Horde/Plugins/Storage/HordeServer.Storage/Storage/StorageController.cs
--- a/Engine/Source/Programs/Horde/Plugins/Storage/HordeServer.Storage/Storage/StorageController.cs
+++ b/Engine/Source/Programs/Horde/Plugins/Storage/HordeServer.Storage/Storage/StorageController.cs
@@ -238,7 +238,7 @@ namespace HordeServer.Storage
 			}
 #pragma warning restore CS0618 // Type or member is obsolete
 
-			await backend.WriteRefAsync(refName, new BlobRefValue(request.Hash, request.Target), request.Options, cancellationToken);
+			await backend.WriteRefAsync(refName, new HashedBlobRefValue(request.Hash, request.Target), request.Options, cancellationToken);
 			return Ok();
 		}
 
@@ -281,7 +281,7 @@ namespace HordeServer.Storage
 				}
 			}
 
-			IBlobRef? target = await client.TryReadRefAsync(refName, cacheTime, cancellationToken: cancellationToken);
+			IHashedBlobRef? target = await client.TryReadRefAsync(refName, cacheTime, cancellationToken: cancellationToken);
 			if (target == null)
 			{
 				return new NotFoundResult();
@@ -397,7 +397,7 @@ namespace HordeServer.Storage
 
 			object content;
 
-			using BlobData blobData = await storageClient.CreateBlobHandle(locator).ReadBlobDataAsync(cancellationToken);
+			using BlobData blobData = await storageClient.CreateBlobRef(locator).ReadBlobDataAsync(cancellationToken);
 			if (data)
 			{
 				ReadOnlyMemoryStream stream = new ReadOnlyMemoryStream(blobData.Data.ToArray());
@@ -442,7 +442,7 @@ namespace HordeServer.Storage
 				if (commitNode.Metadata.Count > 0)
 				{
 					metadata = new Dictionary<Guid, object>();
-					foreach ((Guid blobGuid, IBlobRef handle) in commitNode.Metadata)
+					foreach ((Guid blobGuid, IHashedBlobRef handle) in commitNode.Metadata)
 					{
 						metadata.Add(blobGuid, GetNodeHandleLink(namespaceId, handle));
 					}
@@ -474,7 +474,7 @@ namespace HordeServer.Storage
 			return new { type = typeName, guid = blobData.Type.Guid, data = $"{GetNodeLink(namespaceId, locator)}&data=true", content = content };
 		}
 
-		static object? GetCbNodeObject(NamespaceId namespaceId, CbField field, IEnumerator<IBlobHandle> imports)
+		static object? GetCbNodeObject(NamespaceId namespaceId, CbField field, IEnumerator<IBlobRef> imports)
 		{
 			if (field.IsAttachment())
 			{
@@ -516,9 +516,9 @@ namespace HordeServer.Storage
 		static object? GetNodeObject(NamespaceId namespaceId, DirectoryNodeRef? nodeRef) => (nodeRef == null) ? null : new { nodeRef.Length, nodeRef.Handle.Hash, link = GetNodeLink(namespaceId, nodeRef.Handle.GetLocator()) };
 
 		[return: NotNullIfNotNull("handle")]
-		static object? GetNodeHandleLink(NamespaceId namespaceId, IBlobRef? handle) => (handle == null) ? null : new { handle.Hash, link = GetNodeLink(namespaceId, handle.GetLocator()) };
+		static object? GetNodeHandleLink(NamespaceId namespaceId, IHashedBlobRef? handle) => (handle == null) ? null : new { handle.Hash, link = GetNodeLink(namespaceId, handle.GetLocator()) };
 
-		static string GetNodeLink(NamespaceId namespaceId, IBlobHandle handle) => GetNodeLink(namespaceId, handle.GetLocator());
+		static string GetNodeLink(NamespaceId namespaceId, IBlobRef handle) => GetNodeLink(namespaceId, handle.GetLocator());
 
 		static string GetNodeLink(NamespaceId namespaceId, BlobLocator locator) => $"/api/v1/storage/{namespaceId}/nodes/{locator.BaseLocator}?{locator.Fragment}";
 	}
diff --git a/Engine/Source/Programs/Horde/Plugins/Storage/HordeServer.Storage/Storage/StorageService.cs b/Engine/Source/Programs/Horde/Plugins/Storage/HordeServer.Storage/Storage/StorageService.cs
--- a/Engine/Source/Programs/Horde/Plugins/Storage/HordeServer.Storage/Storage/StorageService.cs
+++ b/Engine/Source/Programs/Horde/Plugins/Storage/HordeServer.Storage/Storage/StorageService.cs
@@ -150,18 +150,18 @@ namespace HordeServer.Storage
 			#region Refs
 
 			/// <inheritdoc/>
-			public async Task<BlobRefValue?> TryReadRefAsync(RefName name, RefCacheTime cacheTime, CancellationToken cancellationToken)
+			public async Task<HashedBlobRefValue?> TryReadRefAsync(RefName name, RefCacheTime cacheTime, CancellationToken cancellationToken)
 			{
 				RefInfo? refInfo = await _outer.TryReadRefAsync(NamespaceId, name, cacheTime, cancellationToken);
 				if (refInfo == null)
 				{
 					return null;
 				}
-				return new BlobRefValue(refInfo.Hash, refInfo.Target);
+				return new HashedBlobRefValue(refInfo.Hash, refInfo.Target);
 			}
 
 			/// <inheritdoc/>
-			public Task WriteRefAsync(RefName name, BlobRefValue value, RefOptions? options = null, CancellationToken cancellationToken = default)
+			public Task WriteRefAsync(RefName name, HashedBlobRefValue value, RefOptions? options = null, CancellationToken cancellationToken = default)
 				=> _outer.WriteRefAsync(NamespaceId, name, value, options, cancellationToken);
 
 			/// <inheritdoc/>
@@ -983,7 +983,7 @@ namespace HordeServer.Storage
 		}
 
 		/// <inheritdoc/>
-		async Task WriteRefAsync(NamespaceId namespaceId, RefName name, BlobRefValue value, RefOptions? options = null, CancellationToken cancellationToken = default)
+		async Task WriteRefAsync(NamespaceId namespaceId, RefName name, HashedBlobRefValue value, RefOptions? options = null, CancellationToken cancellationToken = default)
 		{
 			string path = value.Locator.BaseLocator.ToString();
 
diff --git a/Engine/Source/Programs/Horde/Plugins/Tools/HordeServer.Tools/Commands/Generate/ToolDataCommand.cs b/Engine/Source/Programs/Horde/Plugins/Tools/HordeServer.Tools/Commands/Generate/ToolDataCommand.cs
--- a/Engine/Source/Programs/Horde/Plugins/Tools/HordeServer.Tools/Commands/Generate/ToolDataCommand.cs
+++ b/Engine/Source/Programs/Horde/Plugins/Tools/HordeServer.Tools/Commands/Generate/ToolDataCommand.cs
@@ -74,7 +74,7 @@ namespace HordeServer.Commands.Generate
 			{
 				using (IStorageClient client = BundleStorageClient.CreateFromDirectory(bundleDir, bundleCache, logger))
 				{
-					IBlobRef<DirectoryNode> dirNodeRef;
+					IHashedBlobRef<DirectoryNode> dirNodeRef;
 					await using (DedupeBlobWriter writer = client.CreateDedupeBlobWriter(refName))
 					{
 						logger.LogInformation("Populating cache with existing refs...");
@@ -142,13 +142,13 @@ namespace HordeServer.Commands.Generate
 		{
 			foreach (RefName refName in FileStorageBackend.EnumerateRefs(searchDir))
 			{
-				IBlobRef? blobRef = await client.TryReadRefAsync(refName, cancellationToken: cancellationToken);
+				IHashedBlobRef? blobRef = await client.TryReadRefAsync(refName, cancellationToken: cancellationToken);
 				if (blobRef != null)
 				{
 					using BlobData blobData = await blobRef.ReadBlobDataAsync(cancellationToken);
 					if (blobData.Type.Guid == DirectoryNode.BlobTypeGuid)
 					{
-						IBlobRef<DirectoryNode> directoryRef = BlobRef.Create<DirectoryNode>(blobRef);
+						IHashedBlobRef<DirectoryNode> directoryRef = HashedBlobRef.Create<DirectoryNode>(blobRef);
 						await writer.AddToCacheAsync(directoryRef, cancellationToken);
 					}
 				}
diff --git a/Engine/Source/Programs/Horde/Plugins/Tools/HordeServer.Tools/Tools/ITool.cs b/Engine/Source/Programs/Horde/Plugins/Tools/HordeServer.Tools/Tools/ITool.cs
--- a/Engine/Source/Programs/Horde/Plugins/Tools/HordeServer.Tools/Tools/ITool.cs
+++ b/Engine/Source/Programs/Horde/Plugins/Tools/HordeServer.Tools/Tools/ITool.cs
@@ -88,7 +88,7 @@ namespace HordeServer.Tools
 		/// <param name="target">Path to the root node containing the tool data</param>
 		/// <param name="cancellationToken">Cancellation token for the operation</param>
 		/// <returns>Updated tool document, or null if it does not exist</returns>
-		Task<ITool?> CreateDeploymentAsync(ToolDeploymentConfig options, BlobRefValue target, CancellationToken cancellationToken = default);
+		Task<ITool?> CreateDeploymentAsync(ToolDeploymentConfig options, HashedBlobRefValue target, CancellationToken cancellationToken = default);
 
 		/// <summary>
 		/// Gets the storage backend for a particular tool
diff --git a/Engine/Source/Programs/Horde/Plugins/Tools/HordeServer.Tools/Tools/ToolCollection.cs b/Engine/Source/Programs/Horde/Plugins/Tools/HordeServer.Tools/Tools/ToolCollection.cs
--- a/Engine/Source/Programs/Horde/Plugins/Tools/HordeServer.Tools/Tools/ToolCollection.cs
+++ b/Engine/Source/Programs/Horde/Plugins/Tools/HordeServer.Tools/Tools/ToolCollection.cs
@@ -122,7 +122,7 @@ namespace HordeServer.Tools
 				return _collection.CreateToolObject(document);
 			}
 
-			public async Task<ITool?> CreateDeploymentAsync(ToolDeploymentConfig options, BlobRefValue target, CancellationToken cancellationToken = default)
+			public async Task<ITool?> CreateDeploymentAsync(ToolDeploymentConfig options, HashedBlobRefValue target, CancellationToken cancellationToken = default)
 			{
 				ToolDocument? document = await _collection.CreateDeploymentAsync(_document, _config, options, target, cancellationToken);
 				return _collection.CreateToolObject(document);
@@ -328,7 +328,7 @@ namespace HordeServer.Tools
 
 			using IStorageClient client = _storageService.CreateClient(toolConfig.NamespaceId);
 
-			IBlobRef<DirectoryNode> nodeRef;
+			IHashedBlobRef<DirectoryNode> nodeRef;
 			await using (IBlobWriter writer = client.CreateBlobWriter($"{tool.Id}/{deploymentId}"))
 			{
 				DirectoryNode directoryNode = new DirectoryNode();
@@ -339,7 +339,7 @@ namespace HordeServer.Tools
 			return await CreateDeploymentInternalAsync(tool, toolConfig, deploymentId, options, client, nodeRef, cancellationToken);
 		}
 
-		async Task<ToolDocument?> CreateDeploymentAsync(ToolDocument tool, ToolConfig toolConfig, ToolDeploymentConfig options, BlobRefValue target, CancellationToken cancellationToken)
+		async Task<ToolDocument?> CreateDeploymentAsync(ToolDocument tool, ToolConfig toolConfig, ToolDeploymentConfig options, HashedBlobRefValue target, CancellationToken cancellationToken)
 		{
 			ToolDeploymentId deploymentId = new ToolDeploymentId(BinaryIdUtils.CreateNew());
 
@@ -347,7 +347,7 @@ namespace HordeServer.Tools
 			return await CreateDeploymentInternalAsync(tool, toolConfig, deploymentId, options, client, client.CreateBlobRef(target), cancellationToken);
 		}
 
-		async Task<ToolDocument?> CreateDeploymentInternalAsync(ToolDocument tool, ToolConfig toolConfig, ToolDeploymentId deploymentId, ToolDeploymentConfig options, IStorageClient storageClient, IBlobRef content, CancellationToken cancellationToken)
+		async Task<ToolDocument?> CreateDeploymentInternalAsync(ToolDocument tool, ToolConfig toolConfig, ToolDeploymentId deploymentId, ToolDeploymentConfig options, IStorageClient storageClient, IHashedBlobRef content, CancellationToken cancellationToken)
 		{
 			if (toolConfig is BundledToolConfig)
 			{
@@ -499,7 +499,7 @@ namespace HordeServer.Tools
 			IStorageClient client = CreateStorageClient(tool);
 			try
 			{
-				IBlobRef<DirectoryNode> nodeRef = await client.ReadRefAsync<DirectoryNode>(deployment.RefName, DateTime.UtcNow - TimeSpan.FromDays(2.0), cancellationToken: cancellationToken);
+				IHashedBlobRef<DirectoryNode> nodeRef = await client.ReadRefAsync<DirectoryNode>(deployment.RefName, DateTime.UtcNow - TimeSpan.FromDays(2.0), cancellationToken: cancellationToken);
 				return nodeRef.AsZipStream().WrapOwnership(client);
 			}
 			catch
diff --git a/Engine/Source/Programs/Horde/Plugins/Tools/HordeServer.Tools/Tools/ToolsController.cs b/Engine/Source/Programs/Horde/Plugins/Tools/HordeServer.Tools/Tools/ToolsController.cs
--- a/Engine/Source/Programs/Horde/Plugins/Tools/HordeServer.Tools/Tools/ToolsController.cs
+++ b/Engine/Source/Programs/Horde/Plugins/Tools/HordeServer.Tools/Tools/ToolsController.cs
@@ -320,7 +320,7 @@ namespace HordeServer.Tools
 			IStorageClient client = tool.CreateStorageClient();
 			try
 			{
-				IBlobRef<DirectoryNode> nodeRef = await client.ReadRefAsync<DirectoryNode>(deployment.RefName, DateTime.UtcNow - TimeSpan.FromDays(2.0), cancellationToken: cancellationToken);
+				IHashedBlobRef<DirectoryNode> nodeRef = await client.ReadRefAsync<DirectoryNode>(deployment.RefName, DateTime.UtcNow - TimeSpan.FromDays(2.0), cancellationToken: cancellationToken);
 
 				// If we weren't specifically asked for a zip, see if this download is a single file. If it is, allow downloading it directory.
 				if (action != GetToolAction.Zip)
@@ -356,7 +356,7 @@ namespace HordeServer.Tools
 		static async Task<GetToolDeploymentResponse> GetDeploymentInfoResponseAsync(ITool tool, IToolDeployment deployment, CancellationToken cancellationToken)
 		{
 			using IStorageClient client = tool.CreateStorageClient();
-			IBlobHandle rootHandle = await client.ReadRefAsync(deployment.RefName, cancellationToken: cancellationToken);
+			IBlobRef rootHandle = await client.ReadRefAsync(deployment.RefName, cancellationToken: cancellationToken);
 
 			return new GetToolDeploymentResponse(deployment.Id, deployment.Version, deployment.State, deployment.Progress, deployment.StartedAt, deployment.Duration, deployment.RefName, rootHandle.GetLocator());
 		}
diff --git a/Engine/Source/Programs/Horde/Samples/RemoteClient/RemoteClientApp.cs b/Engine/Source/Programs/Horde/Samples/RemoteClient/RemoteClientApp.cs
--- a/Engine/Source/Programs/Horde/Samples/RemoteClient/RemoteClientApp.cs
+++ b/Engine/Source/Programs/Horde/Samples/RemoteClient/RemoteClientApp.cs
@@ -114,7 +114,7 @@ namespace RemoteClient
 
 				await using (IBlobWriter writer = storage.CreateBlobWriter())
 				{
-					IBlobRef<DirectoryNode> sandbox = await writer.WriteFilesAsync(uploadDir);
+					IHashedBlobRef<DirectoryNode> sandbox = await writer.WriteFilesAsync(uploadDir);
 					await writer.FlushAsync();
 					await channel.UploadFilesAsync("", sandbox.GetLocator(), storage.Backend);
 				}
diff --git a/Engine/Source/Programs/Shared/EpicGames.Horde.Tests/BlobDataTests.cs b/Engine/Source/Programs/Shared/EpicGames.Horde.Tests/BlobDataTests.cs
--- a/Engine/Source/Programs/Shared/EpicGames.Horde.Tests/BlobDataTests.cs
+++ b/Engine/Source/Programs/Shared/EpicGames.Horde.Tests/BlobDataTests.cs
@@ -13,11 +13,11 @@ namespace EpicGames.Horde.Tests
 	[TestClass]
 	public class BlobDataTests
 	{
-		class DummyHandle : IBlobHandle
+		class DummyHandle : IBlobRef
 		{
 			readonly BlobLocator _locator;
 
-			public IBlobHandle Innermost => this;
+			public IBlobRef Innermost => this;
 
 			public DummyHandle(string locator) => _locator = new BlobLocator(new Utf8String(locator));
 			public ValueTask FlushAsync(CancellationToken cancellationToken = default) => default;
@@ -36,7 +36,7 @@ namespace EpicGames.Horde.Tests
 		[TestMethod]
 		public void HeaderSerialization()
 		{
-			List<IBlobHandle> refs = new List<IBlobHandle>();
+			List<IBlobRef> refs = new List<IBlobRef>();
 			refs.Add(new DummyHandle("hello"));
 			refs.Add(new DummyHandle("world"));
 
diff --git a/Engine/Source/Programs/Shared/EpicGames.Horde.Tests/BlobPipelineTests.cs b/Engine/Source/Programs/Shared/EpicGames.Horde.Tests/BlobPipelineTests.cs
--- a/Engine/Source/Programs/Shared/EpicGames.Horde.Tests/BlobPipelineTests.cs
+++ b/Engine/Source/Programs/Shared/EpicGames.Horde.Tests/BlobPipelineTests.cs
@@ -22,7 +22,7 @@ namespace EpicGames.Horde.Tests
 		{
 			using BundleStorageClient client = BundleStorageClient.CreateInMemory(NullLogger.Instance);
 
-			List<IBlobRef> handles = new List<IBlobRef>();
+			List<IHashedBlobRef> handles = new List<IHashedBlobRef>();
 			await using (IBlobWriter writer = client.CreateBlobWriter())
 			{
 				for (int idx = 0; idx < 4000; idx++)
@@ -34,7 +34,7 @@ namespace EpicGames.Horde.Tests
 
 			await using (BlobPipeline<IoHash> pipeline = new BlobPipeline<IoHash>())
 			{
-				foreach (IBlobRef blobRef in handles)
+				foreach (IHashedBlobRef blobRef in handles)
 				{
 					pipeline.Add(new BlobRequest<IoHash>(blobRef, blobRef.Hash));
 				}
diff --git a/Engine/Source/Programs/Shared/EpicGames.Horde.Tests/BundleStoreTests.cs b/Engine/Source/Programs/Shared/EpicGames.Horde.Tests/BundleStoreTests.cs
--- a/Engine/Source/Programs/Shared/EpicGames.Horde.Tests/BundleStoreTests.cs
+++ b/Engine/Source/Programs/Shared/EpicGames.Horde.Tests/BundleStoreTests.cs
@@ -65,7 +65,7 @@ namespace EpicGames.Horde.Tests
 			await using IBlobWriter writer = store.CreateBlobWriter();
 
 			TextNode node = new TextNode("Hello world");
-			IBlobRef<TextNode> handle = await writer.WriteBlobAsync(node);
+			IHashedBlobRef<TextNode> handle = await writer.WriteBlobAsync(node);
 			await writer.FlushAsync();
 
 			using IReadOnlyMemoryOwner<byte> owner = await memoryStore.ReadBlobAsync(handle.GetLocator().BaseLocator, 0, null);
@@ -129,9 +129,9 @@ namespace EpicGames.Horde.Tests
 			public static BlobType BlobType { get; } = new BlobType("{F63606D4-4061-5DBB-446F-55A69F22654F}", 1);
 
 			public ReadOnlySequence<byte> Data { get; }
-			public IReadOnlyList<IBlobRef<SimpleNode>> Refs { get; }
+			public IReadOnlyList<IHashedBlobRef<SimpleNode>> Refs { get; }
 
-			public SimpleNode(ReadOnlySequence<byte> data, IReadOnlyList<IBlobRef<SimpleNode>> refs)
+			public SimpleNode(ReadOnlySequence<byte> data, IReadOnlyList<IHashedBlobRef<SimpleNode>> refs)
 			{
 				Data = data;
 				Refs = refs;
@@ -143,7 +143,7 @@ namespace EpicGames.Horde.Tests
 			public override SimpleNode Read(IBlobReader reader, BlobSerializerOptions options)
 			{
 				ReadOnlyMemory<byte> data = reader.ReadVariableLengthBytes();
-				IReadOnlyList<IBlobRef<SimpleNode>> refs = reader.ReadVariableLengthArray(() => reader.ReadBlobRef<SimpleNode>());
+				IReadOnlyList<IHashedBlobRef<SimpleNode>> refs = reader.ReadVariableLengthArray(() => reader.ReadBlobRef<SimpleNode>());
 				return new SimpleNode(new ReadOnlySequence<byte>(data), refs);
 			}
 
@@ -161,13 +161,13 @@ namespace EpicGames.Horde.Tests
 			{
 				await using IBlobWriter writer = store.CreateBlobWriter("test");
 
-				SimpleNode node1 = new SimpleNode(new ReadOnlySequence<byte>(new byte[] { 1 }), Array.Empty<IBlobRef<SimpleNode>>());
+				SimpleNode node1 = new SimpleNode(new ReadOnlySequence<byte>(new byte[] { 1 }), Array.Empty<IHashedBlobRef<SimpleNode>>());
 				SimpleNode node2 = new SimpleNode(new ReadOnlySequence<byte>(new byte[] { 2 }), new[] { await writer.WriteBlobAsync(node1) });
 				SimpleNode node3 = new SimpleNode(new ReadOnlySequence<byte>(new byte[] { 3 }), new[] { await writer.WriteBlobAsync(node2) });
-				SimpleNode node4 = new SimpleNode(new ReadOnlySequence<byte>(new byte[] { 4 }), Array.Empty<IBlobRef<SimpleNode>>());
+				SimpleNode node4 = new SimpleNode(new ReadOnlySequence<byte>(new byte[] { 4 }), Array.Empty<IHashedBlobRef<SimpleNode>>());
 
 				SimpleNode root = new SimpleNode(new ReadOnlySequence<byte>(new byte[] { 5 }), new[] { await writer.WriteBlobAsync(node4), await writer.WriteBlobAsync(node3) });
-				IBlobRef<SimpleNode> rootRef = await writer.WriteBlobAsync(root);
+				IHashedBlobRef<SimpleNode> rootRef = await writer.WriteBlobAsync(root);
 				await writer.FlushAsync();
 
 				await store.WriteRefAsync(new RefName("test"), rootRef);
@@ -188,31 +188,31 @@ namespace EpicGames.Horde.Tests
 			SimpleNode node5 = root;
 			byte[] data5 = node5.Data.ToArray();
 			Assert.IsTrue(data5.SequenceEqual(new byte[] { 5 }));
-			IReadOnlyList<IBlobRef<SimpleNode>> refs5 = node5.Refs;
+			IReadOnlyList<IHashedBlobRef<SimpleNode>> refs5 = node5.Refs;
 			Assert.AreEqual(2, refs5.Count);
 
 			SimpleNode node4 = await refs5[0].ReadBlobAsync();
 			byte[] data4 = node4.Data.ToArray();
 			Assert.IsTrue(data4.SequenceEqual(new byte[] { 4 }));
-			IReadOnlyList<IBlobRef<SimpleNode>> refs4 = node4.Refs;
+			IReadOnlyList<IHashedBlobRef<SimpleNode>> refs4 = node4.Refs;
 			Assert.AreEqual(0, refs4.Count);
 
 			SimpleNode node3 = await refs5[1].ReadBlobAsync();
 			byte[] data3 = node3.Data.ToArray();
 			Assert.IsTrue(data3.SequenceEqual(new byte[] { 3 }));
-			IReadOnlyList<IBlobRef<SimpleNode>> refs3 = node3.Refs;
+			IReadOnlyList<IHashedBlobRef<SimpleNode>> refs3 = node3.Refs;
 			Assert.AreEqual(1, refs3.Count);
 
 			SimpleNode node2 = await refs3[0].ReadBlobAsync();
 			byte[] data2 = node2.Data.ToArray();
 			Assert.IsTrue(data2.SequenceEqual(new byte[] { 2 }));
-			IReadOnlyList<IBlobRef<SimpleNode>> refs2 = node2.Refs;
+			IReadOnlyList<IHashedBlobRef<SimpleNode>> refs2 = node2.Refs;
 			Assert.AreEqual(1, refs2.Count);
 
 			SimpleNode node1 = await refs2[0].ReadBlobAsync();
 			byte[] data1 = node1.Data.ToArray();
 			Assert.IsTrue(data1.SequenceEqual(new byte[] { 1 }));
-			IReadOnlyList<IBlobRef<SimpleNode>> refs1 = node1.Refs;
+			IReadOnlyList<IHashedBlobRef<SimpleNode>> refs1 = node1.Refs;
 			Assert.AreEqual(0, refs1.Count);
 		}
 
@@ -223,10 +223,10 @@ namespace EpicGames.Horde.Tests
 
 			RefName refName = new RefName("test");
 
-			IBlobRef<SimpleNode> inputRef;
+			IHashedBlobRef<SimpleNode> inputRef;
 			await using (IBlobWriter writer = store.CreateBlobWriter(refName))
 			{
-				inputRef = await writer.WriteBlobAsync(new SimpleNode(new ReadOnlySequence<byte>(new byte[] { (byte)123 }), Array.Empty<IBlobRef<SimpleNode>>()));
+				inputRef = await writer.WriteBlobAsync(new SimpleNode(new ReadOnlySequence<byte>(new byte[] { (byte)123 }), Array.Empty<IHashedBlobRef<SimpleNode>>()));
 			}
 			await store.WriteRefAsync(refName, inputRef);
 
@@ -245,15 +245,15 @@ namespace EpicGames.Horde.Tests
 				await using (IBlobWriter writer = store.CreateBlobWriter(new RefName("test")))
 				{
 					DirectoryNode world = new DirectoryNode();
-					IBlobRef<DirectoryNode> worldRef = await writer.WriteBlobAsync(world);
+					IHashedBlobRef<DirectoryNode> worldRef = await writer.WriteBlobAsync(world);
 
 					DirectoryNode hello = new DirectoryNode();
 					hello.AddDirectory(new DirectoryEntry("world", 0, worldRef));
-					IBlobRef<DirectoryNode> helloRef = await writer.WriteBlobAsync(hello);
+					IHashedBlobRef<DirectoryNode> helloRef = await writer.WriteBlobAsync(hello);
 
 					DirectoryNode root = new DirectoryNode(DirectoryFlags.None);
 					root.AddDirectory(new DirectoryEntry("hello", 0, helloRef));
-					IBlobRef<DirectoryNode> rootRef = await writer.WriteBlobAsync(root);
+					IHashedBlobRef<DirectoryNode> rootRef = await writer.WriteBlobAsync(root);
 
 					await writer.FlushAsync();
 
@@ -301,7 +301,7 @@ namespace EpicGames.Horde.Tests
 				DirectoryNode root = new DirectoryNode();
 				await root.UpdateAsync(fileUpdates, writer);
 
-				IBlobRef<DirectoryNode> rootRef = await writer.WriteBlobAsync(root);
+				IHashedBlobRef<DirectoryNode> rootRef = await writer.WriteBlobAsync(root);
 				await store.WriteRefAsync(new RefName("test"), rootRef);
 
 				await CheckFileTreeAsync(root);
@@ -402,7 +402,7 @@ namespace EpicGames.Horde.Tests
 				root = new DirectoryNode(DirectoryFlags.None);
 				root.AddFile("test", FileEntryFlags.None, fileWriter.Length, chunkedData);
 
-				IBlobRef<DirectoryNode> rootRef = await writer.WriteBlobAsync(root);
+				IHashedBlobRef<DirectoryNode> rootRef = await writer.WriteBlobAsync(root);
 				await store.WriteRefAsync(new RefName("test"), rootRef);
 
 				await CheckLargeFileTreeAsync(root, data);
@@ -424,12 +424,12 @@ namespace EpicGames.Horde.Tests
 			}
 		}
 
-		static async Task GetUniqueBlobsAsync(IBlobHandle handle, Dictionary<BlobLocator, long> locatorToSize)
+		static async Task GetUniqueBlobsAsync(IBlobRef handle, Dictionary<BlobLocator, long> locatorToSize)
 		{
 			using BlobData data = await handle.ReadBlobDataAsync();
 			locatorToSize[handle.GetLocator()] = data.Data.Length;
 
-			foreach (IBlobHandle import in data.Imports)
+			foreach (IBlobRef import in data.Imports)
 			{
 				await GetUniqueBlobsAsync(import, locatorToSize);
 			}
diff --git a/Engine/Source/Programs/Shared/EpicGames.Horde.Tests/BundleTests.cs b/Engine/Source/Programs/Shared/EpicGames.Horde.Tests/BundleTests.cs
--- a/Engine/Source/Programs/Shared/EpicGames.Horde.Tests/BundleTests.cs
+++ b/Engine/Source/Programs/Shared/EpicGames.Horde.Tests/BundleTests.cs
@@ -177,15 +177,15 @@ namespace EpicGames.Horde.Tests
 		{
 			using BundleStorageClient store = BundleStorageClient.CreateInMemory(NullLogger.Instance);
 
-			IBlobRef<DirectoryNode> rootRef;
+			IHashedBlobRef<DirectoryNode> rootRef;
 			await using (IBlobWriter writer = store.CreateBlobWriter())
 			{
 				DirectoryNode world = new DirectoryNode();
-				IBlobRef<DirectoryNode> worldRef = await writer.WriteBlobAsync(world);
+				IHashedBlobRef<DirectoryNode> worldRef = await writer.WriteBlobAsync(world);
 
 				DirectoryNode hello = new DirectoryNode();
 				hello.AddDirectory(new DirectoryEntry("world", 0, worldRef));
-				IBlobRef<DirectoryNode> helloRef = await writer.WriteBlobAsync(hello);
+				IHashedBlobRef<DirectoryNode> helloRef = await writer.WriteBlobAsync(hello);
 
 				DirectoryNode root = new DirectoryNode();
 				root.AddDirectory(new DirectoryEntry("hello", 0, helloRef));
@@ -243,12 +243,12 @@ namespace EpicGames.Horde.Tests
 				for (int idx = 1; idx <= 3; idx++)
 				{
 					DirectoryNode node = new DirectoryNode();
-					IBlobRef<DirectoryNode> nodeRef = await writer.WriteBlobAsync(node);
+					IHashedBlobRef<DirectoryNode> nodeRef = await writer.WriteBlobAsync(node);
 					root.AddDirectory(new DirectoryEntry($"node{idx}", 0, nodeRef));
 				}
 
 				RefName refName = new RefName("ref");
-				IBlobRef<DirectoryNode> rootRef = await writer.WriteBlobAsync(root);
+				IHashedBlobRef<DirectoryNode> rootRef = await writer.WriteBlobAsync(root);
 				await storage.WriteRefAsync(refName, rootRef);
 			}
 
@@ -270,7 +270,7 @@ namespace EpicGames.Horde.Tests
 			{
 				await using (IBlobWriter writer = storage.CreateBlobWriter())
 				{
-					IBlobRef<DirectoryNode> rootRef = await writer.WriteBlobAsync(new DirectoryNode());
+					IHashedBlobRef<DirectoryNode> rootRef = await writer.WriteBlobAsync(new DirectoryNode());
 					for (int idx = 4; idx >= 1; idx--)
 					{
 						DirectoryNode next = new DirectoryNode();
diff --git a/Engine/Source/Programs/Shared/EpicGames.Horde.Tests/ChunkingTests.cs b/Engine/Source/Programs/Shared/EpicGames.Horde.Tests/ChunkingTests.cs
--- a/Engine/Source/Programs/Shared/EpicGames.Horde.Tests/ChunkingTests.cs
+++ b/Engine/Source/Programs/Shared/EpicGames.Horde.Tests/ChunkingTests.cs
@@ -68,7 +68,7 @@ namespace EpicGames.Horde.Tests
 				DirectoryNode directory = new DirectoryNode();
 				directory.AddFile("test.foo", FileEntryFlags.None, 0, chunkedData);
 
-				IBlobRef handle = await writer.WriteBlobAsync(directory);
+				IHashedBlobRef handle = await writer.WriteBlobAsync(directory);
 				await store.WriteRefAsync(RefName, handle);
 			}
 		}
@@ -180,7 +180,7 @@ namespace EpicGames.Horde.Tests
 
 			byte[] data = Encoding.UTF8.GetBytes("hello world");
 
-			IBlobRef<DirectoryNode> handle;
+			IHashedBlobRef<DirectoryNode> handle;
 			await using (IBlobWriter writer = store.CreateBlobWriter(options: serializerOptions))
 			{
 				using ChunkedDataWriter chunkedWriter = new ChunkedDataWriter(writer, chunkingOptions);
@@ -217,20 +217,20 @@ namespace EpicGames.Horde.Tests
 
 			await using MemoryBlobWriter blobWriter = new MemoryBlobWriter(new BlobSerializerOptions());
 
-			IBlobRef<LeafChunkedDataNode> leafRef = await blobWriter.WriteBlobAsync(new LeafChunkedDataNode(new byte[] { 1, 2, 3 }));
-			int leafRefIndex = MemoryBlobWriter.GetIndex((IBlobRef)leafRef);
+			IHashedBlobRef<LeafChunkedDataNode> leafRef = await blobWriter.WriteBlobAsync(new LeafChunkedDataNode(new byte[] { 1, 2, 3 }));
+			int leafRefIndex = MemoryBlobWriter.GetIndex((IHashedBlobRef)leafRef);
 			ChunkedDataNodeRef leafChunkedRef = new ChunkedDataNodeRef(3, 0, leafRef);
 
 			List<ChunkedDataNodeRef> leafNodeRefs = Enumerable.Repeat(leafChunkedRef, 10000).ToList();
 			ChunkedDataNodeRef root = await InteriorChunkedDataNode.CreateTreeAsync(leafNodeRefs, interiorOptions, blobWriter, CancellationToken.None);
 
-			List<IBlobHandle> list = new List<IBlobHandle>();
+			List<IBlobRef> list = new List<IBlobRef>();
 			await GetReadOrderAsync(root.Handle, list);
 
 			int prevIndex = Int32.MaxValue;
 			for (int idx = 0; idx < list.Count; idx++)
 			{
-				int index = MemoryBlobWriter.GetIndex((IBlobRef)list[idx].Innermost);
+				int index = MemoryBlobWriter.GetIndex((IHashedBlobRef)list[idx].Innermost);
 				if (index != leafRefIndex)
 				{
 					Console.WriteLine("{0}", index);
@@ -240,12 +240,12 @@ namespace EpicGames.Horde.Tests
 			}
 		}
 
-		static async Task GetReadOrderAsync(IBlobHandle handle, List<IBlobHandle> list)
+		static async Task GetReadOrderAsync(IBlobRef handle, List<IBlobRef> list)
 		{
 			list.Add(handle);
 
 			using BlobData blobData = await handle.ReadBlobDataAsync(CancellationToken.None);
-			foreach (IBlobHandle childHandle in blobData.Imports)
+			foreach (IBlobRef childHandle in blobData.Imports)
 			{
 				await GetReadOrderAsync(childHandle, list);
 			}
diff --git a/Engine/Source/Programs/Shared/EpicGames.Horde.Tests/ComputeSocketTests.cs b/Engine/Source/Programs/Shared/EpicGames.Horde.Tests/ComputeSocketTests.cs
--- a/Engine/Source/Programs/Shared/EpicGames.Horde.Tests/ComputeSocketTests.cs
+++ b/Engine/Source/Programs/Shared/EpicGames.Horde.Tests/ComputeSocketTests.cs
@@ -193,12 +193,12 @@ namespace EpicGames.Horde.Tests
 						DirectoryNode directory = new DirectoryNode();
 						directory.AddFile("hello.txt", FileEntryFlags.None, data.Length, chunkedData);
 
-						IBlobRef<DirectoryNode> directoryRef = await blobWriter.WriteBlobAsync(directory, cancellationToken: cancellationToken);
+						IHashedBlobRef<DirectoryNode> directoryRef = await blobWriter.WriteBlobAsync(directory, cancellationToken: cancellationToken);
 
 						DirectoryNode root = new DirectoryNode();
 						root.AddDirectory(new DirectoryEntry("subdir", directory.Length, directoryRef));
 
-						IBlobRef<DirectoryNode> handle = await blobWriter.WriteBlobAsync(root, cancellationToken: cancellationToken);
+						IHashedBlobRef<DirectoryNode> handle = await blobWriter.WriteBlobAsync(root, cancellationToken: cancellationToken);
 						await blobWriter.FlushAsync(cancellationToken);
 
 						await channel.UploadFilesAsync("", handle.GetLocator(), storage.Backend, cancellationToken);
diff --git a/Engine/Source/Programs/Shared/EpicGames.Horde.Tests/LogTests.cs b/Engine/Source/Programs/Shared/EpicGames.Horde.Tests/LogTests.cs
--- a/Engine/Source/Programs/Shared/EpicGames.Horde.Tests/LogTests.cs
+++ b/Engine/Source/Programs/Shared/EpicGames.Horde.Tests/LogTests.cs
@@ -69,7 +69,7 @@ namespace EpicGames.Horde.Tests
 			}
 
 			// Flush it to storage, and read the finished log node
-			IBlobRef<LogNode> logRef;
+			IHashedBlobRef<LogNode> logRef;
 			await using (IBlobWriter writer = store.CreateBlobWriter())
 			{
 				logRef = await builder.FlushAsync(writer, true, CancellationToken.None);
@@ -118,7 +118,7 @@ namespace EpicGames.Horde.Tests
 				builder.WriteData(Encoding.UTF8.GetBytes(lines[lineIdx]));
 			}
 
-			IBlobRef<LogNode> rootNodeRef;
+			IHashedBlobRef<LogNode> rootNodeRef;
 			await using (IBlobWriter writer = store.CreateBlobWriter())
 			{
 				rootNodeRef = await builder.FlushAsync(writer, true, CancellationToken.None);
diff --git a/Engine/Source/Programs/Shared/EpicGames.Horde.Tests/ServerLogPacketBuilderTests.cs b/Engine/Source/Programs/Shared/EpicGames.Horde.Tests/ServerLogPacketBuilderTests.cs
--- a/Engine/Source/Programs/Shared/EpicGames.Horde.Tests/ServerLogPacketBuilderTests.cs
+++ b/Engine/Source/Programs/Shared/EpicGames.Horde.Tests/ServerLogPacketBuilderTests.cs
@@ -285,7 +285,7 @@ namespace EpicGames.Horde.Tests
 
 			// Read the log
 			using IStorageClient storageClient = hordeClient.CreateStorageClient(logId);
-			LogNode file = await storageClient.CreateBlobHandle(hordeClient.LogRpc.Logs[logId]).ReadBlobAsync<LogNode>();
+			LogNode file = await storageClient.CreateBlobRef(hordeClient.LogRpc.Logs[logId]).ReadBlobAsync<LogNode>();
 
 			// Check the index text
 			List<Utf8String> extractedIndexText = new List<Utf8String>();
diff --git a/Engine/Source/Programs/Shared/EpicGames.Horde.Tests/StorageClientTests.cs b/Engine/Source/Programs/Shared/EpicGames.Horde.Tests/StorageClientTests.cs
--- a/Engine/Source/Programs/Shared/EpicGames.Horde.Tests/StorageClientTests.cs
+++ b/Engine/Source/Programs/Shared/EpicGames.Horde.Tests/StorageClientTests.cs
@@ -22,15 +22,15 @@ namespace EpicGames.Horde.Tests
 
 			public int Value { get; }
 			public byte[] Padding { get; set; } = Array.Empty<byte>();
-			public IBlobRef<TestNode>[] Refs { get; }
+			public IHashedBlobRef<TestNode>[] Refs { get; }
 
-			public TestNode(int value, params IBlobRef<TestNode>[] refs)
+			public TestNode(int value, params IHashedBlobRef<TestNode>[] refs)
 			{
 				Value = value;
 				Refs = refs;
 			}
 
-			public TestNode(int value, byte[] padding, IBlobRef<TestNode>[] refs)
+			public TestNode(int value, byte[] padding, IHashedBlobRef<TestNode>[] refs)
 			{
 				Value = value;
 				Padding = padding;
@@ -44,7 +44,7 @@ namespace EpicGames.Horde.Tests
 			{
 				int value = reader.ReadInt32();
 				byte[] padding = reader.ReadVariableLengthBytes().ToArray();
-				IBlobRef<TestNode>[] refs = reader.ReadVariableLengthArray(() => reader.ReadBlobRef<TestNode>());
+				IHashedBlobRef<TestNode>[] refs = reader.ReadVariableLengthArray(() => reader.ReadBlobRef<TestNode>());
 
 				return new TestNode(value, padding, refs);
 			}
@@ -75,7 +75,7 @@ namespace EpicGames.Horde.Tests
 
 		static async Task TestBasicAsync(IStorageClient store)
 		{
-			IBlobRef<TestNode> nodeRef;
+			IHashedBlobRef<TestNode> nodeRef;
 			await using (IBlobWriter writer = store.CreateBlobWriter())
 			{
 				nodeRef = await writer.WriteBlobAsync(new TestNode(123));
@@ -93,10 +93,10 @@ namespace EpicGames.Horde.Tests
 
 			using BundleStorageClient store = BundleStorageClient.CreateInMemory(NullLogger.Instance);
 
-			IBlobRef<TestNode> nodeRef2;
+			IHashedBlobRef<TestNode> nodeRef2;
 			await using (IBlobWriter writer = store.CreateBlobWriter())
 			{
-				IBlobRef<TestNode> nodeRef1 = await writer.WriteBlobAsync(new TestNode(123));
+				IHashedBlobRef<TestNode> nodeRef1 = await writer.WriteBlobAsync(new TestNode(123));
 				nodeRef2 = await writer.WriteBlobAsync(new TestNode(456, nodeRef1));
 			}
 			await store.WriteRefAsync("hello", nodeRef2);
@@ -119,10 +119,10 @@ namespace EpicGames.Horde.Tests
 			using BundleStorageClient store = BundleStorageClient.CreateInMemory(bundleOptions, NullLogger.Instance);
 
 			await using IBlobWriter writer = store.CreateBlobWriter();
-			IBlobRef<TestNode> nodeRef1 = await writer.WriteBlobAsync(new TestNode(123) { Padding = new byte[1024] });
+			IHashedBlobRef<TestNode> nodeRef1 = await writer.WriteBlobAsync(new TestNode(123) { Padding = new byte[1024] });
 			await writer.FlushAsync();
-			IBlobRef<TestNode> nodeRef2 = await writer.WriteBlobAsync(new TestNode(456, nodeRef1) { Padding = new byte[1024] });
-			IBlobRef<TestNode> nodeRef3 = await writer.WriteBlobAsync(new TestNode(789, nodeRef2));
+			IHashedBlobRef<TestNode> nodeRef2 = await writer.WriteBlobAsync(new TestNode(456, nodeRef1) { Padding = new byte[1024] });
+			IHashedBlobRef<TestNode> nodeRef3 = await writer.WriteBlobAsync(new TestNode(789, nodeRef2));
 
 			// nodeRef1 is in a flushed bundle
 			ExportHandle export1 = (ExportHandle)nodeRef1.Innermost;
diff --git a/Engine/Source/Programs/Shared/EpicGames.Horde/Compute/AgentMessageHandler.cs b/Engine/Source/Programs/Shared/EpicGames.Horde/Compute/AgentMessageHandler.cs
--- a/Engine/Source/Programs/Shared/EpicGames.Horde/Compute/AgentMessageHandler.cs
+++ b/Engine/Source/Programs/Shared/EpicGames.Horde/Compute/AgentMessageHandler.cs
@@ -193,7 +193,7 @@ namespace EpicGames.Horde.Compute
 			BundleOptions bundleOptions = ComputeProtocolUtilities.GetBundleOptions(channel.Protocol);
 			using BundleStorageClient store = new BundleStorageClient(innerStore, cache, bundleOptions, _logger);
 
-			IBlobHandle handle = store.CreateBlobHandle(locator);
+			IBlobRef handle = store.CreateBlobRef(locator);
 			DirectoryNode directoryNode = await handle.ReadBlobAsync<DirectoryNode>(options, cancellationToken);
 
 			DirectoryReference outputDir = DirectoryReference.Combine(_sandboxDir, path);
diff --git a/Engine/Source/Programs/Shared/EpicGames.Horde/Compute/AgentStorageClient.cs b/Engine/Source/Programs/Shared/EpicGames.Horde/Compute/AgentStorageClient.cs
--- a/Engine/Source/Programs/Shared/EpicGames.Horde/Compute/AgentStorageClient.cs
+++ b/Engine/Source/Programs/Shared/EpicGames.Horde/Compute/AgentStorageClient.cs
@@ -100,10 +100,10 @@ namespace EpicGames.Horde.Compute
 		public Task<bool> DeleteRefAsync(RefName name, CancellationToken cancellationToken = default) => throw new NotSupportedException();
 
 		/// <inheritdoc/>
-		public Task<BlobRefValue?> TryReadRefAsync(RefName name, RefCacheTime cacheTime = default, CancellationToken cancellationToken = default) => throw new NotSupportedException();
+		public Task<HashedBlobRefValue?> TryReadRefAsync(RefName name, RefCacheTime cacheTime = default, CancellationToken cancellationToken = default) => throw new NotSupportedException();
 
 		/// <inheritdoc/>
-		public Task WriteRefAsync(RefName name, BlobRefValue value, RefOptions? options = null, CancellationToken cancellationToken = default) => throw new NotSupportedException();
+		public Task WriteRefAsync(RefName name, HashedBlobRefValue value, RefOptions? options = null, CancellationToken cancellationToken = default) => throw new NotSupportedException();
 
 		#endregion
 
diff --git a/Engine/Source/Programs/Shared/EpicGames.Horde/HordeHttpClient.cs b/Engine/Source/Programs/Shared/EpicGames.Horde/HordeHttpClient.cs
--- a/Engine/Source/Programs/Shared/EpicGames.Horde/HordeHttpClient.cs
+++ b/Engine/Source/Programs/Shared/EpicGames.Horde/HordeHttpClient.cs
@@ -436,7 +436,7 @@ namespace EpicGames.Horde
 		/// <param name="createPaused">Whether to create the deployment, but do not start rolling it out yet</param>
 		/// <param name="target">Location of a directory node describing the deployment</param>
 		/// <param name="cancellationToken">Cancellation token for the operation</param>
-		public async Task<ToolDeploymentId> CreateToolDeploymentAsync(ToolId id, string? version, double? duration, bool? createPaused, BlobRefValue target, CancellationToken cancellationToken = default)
+		public async Task<ToolDeploymentId> CreateToolDeploymentAsync(ToolId id, string? version, double? duration, bool? createPaused, HashedBlobRefValue target, CancellationToken cancellationToken = default)
 		{
 			CreateToolDeploymentRequest request = new CreateToolDeploymentRequest(version ?? String.Empty, duration, createPaused, target);
 			CreateToolDeploymentResponse response = await PostAsync<CreateToolDeploymentResponse, CreateToolDeploymentRequest>(_httpClient, $"api/v2/tools/{id}/deployments", request, cancellationToken);
diff --git a/Engine/Source/Programs/Shared/EpicGames.Horde/Logs/LogChunkNode.cs b/Engine/Source/Programs/Shared/EpicGames.Horde/Logs/LogChunkNode.cs
--- a/Engine/Source/Programs/Shared/EpicGames.Horde/Logs/LogChunkNode.cs
+++ b/Engine/Source/Programs/Shared/EpicGames.Horde/Logs/LogChunkNode.cs
@@ -220,7 +220,7 @@ namespace EpicGames.Horde.Logs
 		/// <summary>
 		/// Handle to the target chunk
 		/// </summary>
-		public IBlobRef<LogChunkNode> Target { get; }
+		public IHashedBlobRef<LogChunkNode> Target { get; }
 
 		/// <summary>
 		/// Constructor
@@ -230,7 +230,7 @@ namespace EpicGames.Horde.Logs
 		/// <param name="offset">Offset within the log file</param>
 		/// <param name="length">Length of the chunk</param>
 		/// <param name="target">Referenced log text</param>
-		public LogChunkRef(int lineIndex, int lineCount, long offset, int length, IBlobRef<LogChunkNode> target)
+		public LogChunkRef(int lineIndex, int lineCount, long offset, int length, IHashedBlobRef<LogChunkNode> target)
 		{
 			LineIndex = lineIndex;
 			LineCount = lineCount;
diff --git a/Engine/Source/Programs/Shared/EpicGames.Horde/Logs/LogIndexNode.cs b/Engine/Source/Programs/Shared/EpicGames.Horde/Logs/LogIndexNode.cs
--- a/Engine/Source/Programs/Shared/EpicGames.Horde/Logs/LogIndexNode.cs
+++ b/Engine/Source/Programs/Shared/EpicGames.Horde/Logs/LogIndexNode.cs
@@ -102,7 +102,7 @@ namespace EpicGames.Horde.Logs
 			for (int idx = 0; idx < appendPlainTextChunks.Count; idx++)
 			{
 				LogChunkNode newChunk = appendPlainTextChunks[idx];
-				IBlobRef<LogChunkNode> newChunkRef = await writer.WriteBlobAsync(newChunk, cancellationToken);
+				IHashedBlobRef<LogChunkNode> newChunkRef = await writer.WriteBlobAsync(newChunk, cancellationToken);
 				newChunks[_plainTextChunkRefs.Length + idx] = new LogChunkRef(lineIndex, newChunk.LineCount, offset, newChunk.Length, newChunkRef);
 				lineIndex += newChunk.LineCount;
 				offset += newChunk.Length;
diff --git a/Engine/Source/Programs/Shared/EpicGames.Horde/Logs/LogNode.cs b/Engine/Source/Programs/Shared/EpicGames.Horde/Logs/LogNode.cs
--- a/Engine/Source/Programs/Shared/EpicGames.Horde/Logs/LogNode.cs
+++ b/Engine/Source/Programs/Shared/EpicGames.Horde/Logs/LogNode.cs
@@ -58,7 +58,7 @@ namespace EpicGames.Horde.Logs
 		/// <summary>
 		/// Index for this log
 		/// </summary>
-		public IBlobRef<LogIndexNode> IndexRef { get; }
+		public IHashedBlobRef<LogIndexNode> IndexRef { get; }
 
 		/// <summary>
 		/// Whether this log is complete
@@ -68,7 +68,7 @@ namespace EpicGames.Horde.Logs
 		/// <summary>
 		/// Deserializing constructor
 		/// </summary>
-		public LogNode(LogFormat format, int lineCount, long length, IReadOnlyList<LogChunkRef> textChunkRefs, IBlobRef<LogIndexNode> indexRef, bool complete)
+		public LogNode(LogFormat format, int lineCount, long length, IReadOnlyList<LogChunkRef> textChunkRefs, IHashedBlobRef<LogIndexNode> indexRef, bool complete)
 		{
 			Format = format;
 			LineCount = lineCount;
@@ -95,7 +95,7 @@ namespace EpicGames.Horde.Logs
 			LogFormat format = (LogFormat)reader.ReadUInt8();
 			int lineCount = (int)reader.ReadUnsignedVarInt();
 			long length = (long)reader.ReadUnsignedVarInt();
-			IBlobRef<LogIndexNode> indexRef = reader.ReadBlobRef<LogIndexNode>();
+			IHashedBlobRef<LogIndexNode> indexRef = reader.ReadBlobRef<LogIndexNode>();
 			List<LogChunkRef> textChunkRefs = reader.ReadList(() => new LogChunkRef(reader));
 			bool complete = reader.ReadBoolean();
 
@@ -264,7 +264,7 @@ namespace EpicGames.Horde.Logs
 		/// <param name="writer">Writer for the output nodes</param>
 		/// <param name="complete">Whether the log is complete</param>
 		/// <param name="cancellationToken">Cancellation token for the operation</param>
-		public async Task<IBlobRef<LogNode>> FlushAsync(IBlobWriter writer, bool complete, CancellationToken cancellationToken)
+		public async Task<IHashedBlobRef<LogNode>> FlushAsync(IBlobWriter writer, bool complete, CancellationToken cancellationToken)
 		{
 			// Capture the new data that needs to be written
 			IReadOnlyList<LogChunkNode> writeTextChunks;
@@ -281,21 +281,21 @@ namespace EpicGames.Horde.Logs
 
 			// Flush any complete chunks to storage
 			LogIndexNode newIndex = await _index.AppendAsync(writer, writeIndexTextChunks, cancellationToken);
-			IBlobRef<LogIndexNode> newIndexRef = await writer.WriteBlobAsync(newIndex, cancellationToken);
+			IHashedBlobRef<LogIndexNode> newIndexRef = await writer.WriteBlobAsync(newIndex, cancellationToken);
 
 			List<LogChunkRef> newJsonChunkRefs = new List<LogChunkRef>(_root?.TextChunkRefs ?? Array.Empty<LogChunkRef>());
 			int lineCount = _root?.LineCount ?? 0;
 			long length = _root?.Length ?? 0;
 			foreach (LogChunkNode writeTextChunk in writeTextChunks)
 			{
-				IBlobRef<LogChunkNode> writeTextChunkRef = await writer.WriteBlobAsync(writeTextChunk, cancellationToken);
+				IHashedBlobRef<LogChunkNode> writeTextChunkRef = await writer.WriteBlobAsync(writeTextChunk, cancellationToken);
 				newJsonChunkRefs.Add(new LogChunkRef(lineCount, writeTextChunk.LineCount, length, writeTextChunk.Length, writeTextChunkRef));
 				lineCount += writeTextChunk.LineCount;
 				length += writeTextChunk.Length;
 			}
 
 			LogNode newRoot = new LogNode(_format, lineCount, length, newJsonChunkRefs, newIndexRef, complete);
-			IBlobRef<LogNode> newRootRef = await writer.WriteBlobAsync(newRoot, cancellationToken);
+			IHashedBlobRef<LogNode> newRootRef = await writer.WriteBlobAsync(newRoot, cancellationToken);
 
 			await writer.FlushAsync(cancellationToken);
 
diff --git a/Engine/Source/Programs/Shared/EpicGames.Horde/Logs/ServerLogger.cs b/Engine/Source/Programs/Shared/EpicGames.Horde/Logs/ServerLogger.cs
--- a/Engine/Source/Programs/Shared/EpicGames.Horde/Logs/ServerLogger.cs
+++ b/Engine/Source/Programs/Shared/EpicGames.Horde/Logs/ServerLogger.cs
@@ -352,7 +352,7 @@ namespace EpicGames.Horde.Logs
 
 			if (flush || _bufferLength > FlushLength)
 			{
-				IBlobRef<LogNode> target = await _builder.FlushAsync(_writer, flush, cancellationToken);
+				IHashedBlobRef<LogNode> target = await _builder.FlushAsync(_writer, flush, cancellationToken);
 				await UpdateLogAsync(target, _builder.LineCount, flush, cancellationToken);
 				_bufferLength = 0;
 			}
@@ -362,7 +362,7 @@ namespace EpicGames.Horde.Logs
 
 		#region RPC calls
 
-		async Task UpdateLogAsync(IBlobRef target, int lineCount, bool complete, CancellationToken cancellationToken)
+		async Task UpdateLogAsync(IHashedBlobRef target, int lineCount, bool complete, CancellationToken cancellationToken)
 		{
 			_internalLogger.LogInformation("Updating log {LogId} to line {LineCount}, target {Locator}", _logId, lineCount, target.GetLocator());
 
diff --git a/Engine/Source/Programs/Shared/EpicGames.Horde/Storage/Backends/FileStorageBackend.cs b/Engine/Source/Programs/Shared/EpicGames.Horde/Storage/Backends/FileStorageBackend.cs
--- a/Engine/Source/Programs/Shared/EpicGames.Horde/Storage/Backends/FileStorageBackend.cs
+++ b/Engine/Source/Programs/Shared/EpicGames.Horde/Storage/Backends/FileStorageBackend.cs
@@ -139,7 +139,7 @@ namespace EpicGames.Horde.Storage.Backends
 		}
 
 		/// <inheritdoc/>
-		public async Task<BlobRefValue?> TryReadRefAsync(RefName name, RefCacheTime cacheTime = default, CancellationToken cancellationToken = default)
+		public async Task<HashedBlobRefValue?> TryReadRefAsync(RefName name, RefCacheTime cacheTime = default, CancellationToken cancellationToken = default)
 		{
 			FileReference file = GetRefFile(name);
 			if (!FileReference.Exists(file))
@@ -153,11 +153,11 @@ namespace EpicGames.Horde.Storage.Backends
 			IoHash hash = IoHash.Parse(lines[0].Trim());
 			BlobLocator locator = new BlobLocator(lines[1].Trim());
 
-			return new BlobRefValue(hash, locator);
+			return new HashedBlobRefValue(hash, locator);
 		}
 
 		/// <inheritdoc/>
-		public async Task WriteRefAsync(RefName name, BlobRefValue value, RefOptions? options = null, CancellationToken cancellationToken = default)
+		public async Task WriteRefAsync(RefName name, HashedBlobRefValue value, RefOptions? options = null, CancellationToken cancellationToken = default)
 		{
 			FileReference file = GetRefFile(name);
 			DirectoryReference.CreateDirectory(file.Directory);
diff --git a/Engine/Source/Programs/Shared/EpicGames.Horde/Storage/Backends/HttpStorageBackend.cs b/Engine/Source/Programs/Shared/EpicGames.Horde/Storage/Backends/HttpStorageBackend.cs
--- a/Engine/Source/Programs/Shared/EpicGames.Horde/Storage/Backends/HttpStorageBackend.cs
+++ b/Engine/Source/Programs/Shared/EpicGames.Horde/Storage/Backends/HttpStorageBackend.cs
@@ -308,7 +308,7 @@ namespace EpicGames.Horde.Storage.Backends
 		}
 
 		/// <inheritdoc/>
-		public async Task<BlobRefValue?> TryReadRefAsync(RefName name, RefCacheTime cacheTime = default, CancellationToken cancellationToken = default)
+		public async Task<HashedBlobRefValue?> TryReadRefAsync(RefName name, RefCacheTime cacheTime = default, CancellationToken cancellationToken = default)
 		{
 			using (HttpClient httpClient = _createClient())
 			{
@@ -337,7 +337,7 @@ namespace EpicGames.Horde.Storage.Backends
 							ReadRefResponse? data = await response.Content.ReadFromJsonAsync<ReadRefResponse>(cancellationToken: cancellationToken);
 							_logger.LogDebug("Read ref {RefName} -> {Hash} / {Locator}", name, data!.Hash, data!.Target);
 
-							return new BlobRefValue(data.Hash, data.Target);
+							return new HashedBlobRefValue(data.Hash, data.Target);
 						}
 					}
 				}
@@ -345,7 +345,7 @@ namespace EpicGames.Horde.Storage.Backends
 		}
 
 		/// <inheritdoc/>
-		public async Task WriteRefAsync(RefName name, BlobRefValue value, RefOptions? options = null, CancellationToken cancellationToken = default)
+		public async Task WriteRefAsync(RefName name, HashedBlobRefValue value, RefOptions? options = null, CancellationToken cancellationToken = default)
 		{
 			_logger.LogDebug("Writing ref {RefName} -> {Hash} / {Locator}", name, value.Hash, value.Locator);
 			using (HttpClient httpClient = _createClient())
diff --git a/Engine/Source/Programs/Shared/EpicGames.Horde/Storage/Backends/MemoryStorageBackend.cs b/Engine/Source/Programs/Shared/EpicGames.Horde/Storage/Backends/MemoryStorageBackend.cs
--- a/Engine/Source/Programs/Shared/EpicGames.Horde/Storage/Backends/MemoryStorageBackend.cs
+++ b/Engine/Source/Programs/Shared/EpicGames.Horde/Storage/Backends/MemoryStorageBackend.cs
@@ -18,7 +18,7 @@ namespace EpicGames.Horde.Storage.Backends
 		record class AliasListNode(BlobLocator Locator, int Rank, ReadOnlyMemory<byte> Data, AliasListNode? Next);
 
 		readonly ConcurrentDictionary<BlobLocator, ReadOnlyMemory<byte>> _blobs = new ConcurrentDictionary<BlobLocator, ReadOnlyMemory<byte>>();
-		readonly ConcurrentDictionary<RefName, BlobRefValue> _refs = new ConcurrentDictionary<RefName, BlobRefValue>();
+		readonly ConcurrentDictionary<RefName, HashedBlobRefValue> _refs = new ConcurrentDictionary<RefName, HashedBlobRefValue>();
 		readonly ConcurrentDictionary<string, AliasListNode?> _aliases = new ConcurrentDictionary<string, AliasListNode?>(StringComparer.Ordinal);
 
 		/// <summary>
@@ -29,7 +29,7 @@ namespace EpicGames.Horde.Storage.Backends
 		/// <summary>
 		/// Accessor for all refs stored by the client
 		/// </summary>
-		public IReadOnlyDictionary<RefName, BlobRefValue> Refs => _refs;
+		public IReadOnlyDictionary<RefName, HashedBlobRefValue> Refs => _refs;
 
 		/// <inheritdoc/>
 		public bool SupportsRedirects => false;
@@ -167,21 +167,21 @@ namespace EpicGames.Horde.Storage.Backends
 		public Task<bool> DeleteRefAsync(RefName name, CancellationToken cancellationToken) => Task.FromResult(_refs.TryRemove(name, out _));
 
 		/// <inheritdoc/>
-		public Task<BlobRefValue?> TryReadRefAsync(RefName name, RefCacheTime cacheTime = default, CancellationToken cancellationToken = default)
+		public Task<HashedBlobRefValue?> TryReadRefAsync(RefName name, RefCacheTime cacheTime = default, CancellationToken cancellationToken = default)
 		{
-			BlobRefValue? value;
+			HashedBlobRefValue? value;
 			if (_refs.TryGetValue(name, out value))
 			{
-				return Task.FromResult<BlobRefValue?>(value);
+				return Task.FromResult<HashedBlobRefValue?>(value);
 			}
 			else
 			{
-				return Task.FromResult<BlobRefValue?>(null);
+				return Task.FromResult<HashedBlobRefValue?>(null);
 			}
 		}
 
 		/// <inheritdoc/>
-		public Task WriteRefAsync(RefName name, BlobRefValue value, RefOptions? options = null, CancellationToken cancellationToken = default)
+		public Task WriteRefAsync(RefName name, HashedBlobRefValue value, RefOptions? options = null, CancellationToken cancellationToken = default)
 		{
 			_refs[name] = value;
 			return Task.CompletedTask;
diff --git a/Engine/Source/Programs/Shared/EpicGames.Horde/Storage/BlobAlias.cs b/Engine/Source/Programs/Shared/EpicGames.Horde/Storage/BlobAlias.cs
--- a/Engine/Source/Programs/Shared/EpicGames.Horde/Storage/BlobAlias.cs
+++ b/Engine/Source/Programs/Shared/EpicGames.Horde/Storage/BlobAlias.cs
@@ -10,7 +10,7 @@ namespace EpicGames.Horde.Storage
 	/// <param name="Target">Handle to the target blob for the alias</param>
 	/// <param name="Rank">Rank for the alias</param>
 	/// <param name="Data">Data stored inline with the alias</param>
-	public record class BlobAlias(IBlobHandle Target, int Rank, ReadOnlyMemory<byte> Data);
+	public record class BlobAlias(IBlobRef Target, int Rank, ReadOnlyMemory<byte> Data);
 
 	/// <summary>
 	/// Data for an alias in the storage system. An alias is a named weak reference to a node.
diff --git a/Engine/Source/Programs/Shared/EpicGames.Horde/Storage/BlobData.cs b/Engine/Source/Programs/Shared/EpicGames.Horde/Storage/BlobData.cs
--- a/Engine/Source/Programs/Shared/EpicGames.Horde/Storage/BlobData.cs
+++ b/Engine/Source/Programs/Shared/EpicGames.Horde/Storage/BlobData.cs
@@ -24,12 +24,12 @@ namespace EpicGames.Horde.Storage
 		/// <summary>
 		/// Handles to referenced blobs
 		/// </summary>
-		public IReadOnlyList<IBlobHandle> Imports { get; }
+		public IReadOnlyList<IBlobRef> Imports { get; }
 
 		/// <summary>
 		/// Constructor
 		/// </summary>
-		public BlobData(BlobType type, ReadOnlyMemory<byte> data, IReadOnlyList<IBlobHandle> imports)
+		public BlobData(BlobType type, ReadOnlyMemory<byte> data, IReadOnlyList<IBlobRef> imports)
 		{
 			Type = type;
 			Data = data;
@@ -62,7 +62,7 @@ namespace EpicGames.Horde.Storage
 		/// <summary>
 		/// Constructor
 		/// </summary>
-		public BlobDataWithOwner(BlobType type, ReadOnlyMemory<byte> data, IReadOnlyList<IBlobHandle> imports, IDisposable owner)
+		public BlobDataWithOwner(BlobType type, ReadOnlyMemory<byte> data, IReadOnlyList<IBlobRef> imports, IDisposable owner)
 			: base(type, data, imports)
 		{
 			_owner = owner;
diff --git a/Engine/Source/Programs/Shared/EpicGames.Horde/Storage/BlobPipeline.cs b/Engine/Source/Programs/Shared/EpicGames.Horde/Storage/BlobPipeline.cs
--- a/Engine/Source/Programs/Shared/EpicGames.Horde/Storage/BlobPipeline.cs
+++ b/Engine/Source/Programs/Shared/EpicGames.Horde/Storage/BlobPipeline.cs
@@ -16,7 +16,7 @@ namespace EpicGames.Horde.Storage
 	/// <summary>
 	/// Request to read a blob from storage
 	/// </summary>
-	public record class BlobRequest<TUserData>(IBlobHandle Handle, TUserData UserData);
+	public record class BlobRequest<TUserData>(IBlobRef Handle, TUserData UserData);
 
 	/// <summary>
 	/// Response from reading a blob from storage
diff --git a/Engine/Source/Programs/Shared/EpicGames.Horde/Storage/BlobSerializer.cs b/Engine/Source/Programs/Shared/EpicGames.Horde/Storage/BlobSerializer.cs
--- a/Engine/Source/Programs/Shared/EpicGames.Horde/Storage/BlobSerializer.cs
+++ b/Engine/Source/Programs/Shared/EpicGames.Horde/Storage/BlobSerializer.cs
@@ -217,7 +217,7 @@ namespace EpicGames.Horde.Storage
 		/// <param name="handle">Handle to the blob to deserialize</param>
 		/// <param name="options">Options to control serialization</param>
 		/// <param name="cancellationToken">Cancellation token for the operation</param>
-		public static async ValueTask<T> ReadBlobAsync<T>(this IBlobHandle handle, BlobSerializerOptions? options = null, CancellationToken cancellationToken = default)
+		public static async ValueTask<T> ReadBlobAsync<T>(this IBlobRef handle, BlobSerializerOptions? options = null, CancellationToken cancellationToken = default)
 		{
 			using BlobData data = await handle.ReadBlobDataAsync(cancellationToken);
 			return BlobSerializer.Deserialize<T>(data, options);
@@ -229,7 +229,7 @@ namespace EpicGames.Horde.Storage
 		/// <typeparam name="T">Return type for the deserialized object</typeparam>
 		/// <param name="handle">Handle to the blob to deserialize</param>
 		/// <param name="cancellationToken">Cancellation token for the operation</param>
-		public static async ValueTask<T> ReadBlobAsync<T>(this IBlobRef<T> handle, CancellationToken cancellationToken = default)
+		public static async ValueTask<T> ReadBlobAsync<T>(this IHashedBlobRef<T> handle, CancellationToken cancellationToken = default)
 		{
 			using BlobData data = await handle.ReadBlobDataAsync(cancellationToken);
 			return BlobSerializer.Deserialize<T>(data, handle.SerializerOptions);
@@ -242,7 +242,7 @@ namespace EpicGames.Horde.Storage
 		/// <param name="value">The object to serialize</param>
 		/// <param name="cancellationToken">Cancellation token for the operation</param>
 		/// <returns>Handle to the serialized blob</returns>
-		public static async ValueTask<IBlobRef<T>> WriteBlobAsync<T>(this IBlobWriter writer, T value, CancellationToken cancellationToken = default)
+		public static async ValueTask<IHashedBlobRef<T>> WriteBlobAsync<T>(this IBlobWriter writer, T value, CancellationToken cancellationToken = default)
 		{
 			BlobType blobType = BlobSerializer.Serialize<T>(writer, value, writer.Options);
 			return await writer.CompleteAsync<T>(blobType, cancellationToken);
@@ -257,14 +257,14 @@ namespace EpicGames.Horde.Storage
 		/// <param name="options">Options to control serialization</param>
 		/// <param name="cancellationToken">Cancellation token for the operation</param>
 		/// <returns>Node for the given ref, or null if it does not exist</returns>
-		public static async Task<IBlobRef<TNode>?> TryReadRefAsync<TNode>(this IStorageClient store, RefName name, DateTime cacheTime = default, BlobSerializerOptions? options = null, CancellationToken cancellationToken = default) where TNode : class
+		public static async Task<IHashedBlobRef<TNode>?> TryReadRefAsync<TNode>(this IStorageClient store, RefName name, DateTime cacheTime = default, BlobSerializerOptions? options = null, CancellationToken cancellationToken = default) where TNode : class
 		{
-			IBlobRef? refTarget = await store.TryReadRefAsync(name, cacheTime, cancellationToken);
+			IHashedBlobRef? refTarget = await store.TryReadRefAsync(name, cacheTime, cancellationToken);
 			if (refTarget == null)
 			{
 				return null;
 			}
-			return BlobRef.Create<TNode>(refTarget.Hash, refTarget, options ?? BlobSerializerOptions.Default);
+			return HashedBlobRef.Create<TNode>(refTarget.Hash, refTarget, options ?? BlobSerializerOptions.Default);
 		}
 
 		/// <summary>
@@ -276,9 +276,9 @@ namespace EpicGames.Horde.Storage
 		/// <param name="options">Options to control serialization</param>
 		/// <param name="cancellationToken">Cancellation token for the operation</param>
 		/// <returns>The blob instance</returns>
-		public static async Task<IBlobRef<TNode>> ReadRefAsync<TNode>(this IStorageClient store, RefName name, DateTime cacheTime = default, BlobSerializerOptions? options = null, CancellationToken cancellationToken = default) where TNode : class
+		public static async Task<IHashedBlobRef<TNode>> ReadRefAsync<TNode>(this IStorageClient store, RefName name, DateTime cacheTime = default, BlobSerializerOptions? options = null, CancellationToken cancellationToken = default) where TNode : class
 		{
-			IBlobRef<TNode>? refValue = await store.TryReadRefAsync<TNode>(name, cacheTime, options, cancellationToken);
+			IHashedBlobRef<TNode>? refValue = await store.TryReadRefAsync<TNode>(name, cacheTime, options, cancellationToken);
 			if (refValue == null)
 			{
 				throw new RefNameNotFoundException(name);
@@ -297,7 +297,7 @@ namespace EpicGames.Horde.Storage
 		/// <returns>Node for the given ref, or null if it does not exist</returns>
 		public static async Task<TNode?> TryReadRefTargetAsync<TNode>(this IStorageClient store, RefName name, DateTime cacheTime = default, BlobSerializerOptions? options = null, CancellationToken cancellationToken = default) where TNode : class
 		{
-			IBlobRef<TNode>? refTarget = await store.TryReadRefAsync<TNode>(name, cacheTime, options, cancellationToken);
+			IHashedBlobRef<TNode>? refTarget = await store.TryReadRefAsync<TNode>(name, cacheTime, options, cancellationToken);
 			if (refTarget == null)
 			{
 				return null;
@@ -316,7 +316,7 @@ namespace EpicGames.Horde.Storage
 		/// <returns>The blob instance</returns>
 		public static async Task<TNode> ReadRefTargetAsync<TNode>(this IStorageClient store, RefName name, DateTime cacheTime = default, BlobSerializerOptions? options = null, CancellationToken cancellationToken = default) where TNode : class
 		{
-			IBlobRef<TNode> blobRef = await ReadRefAsync<TNode>(store, name, cacheTime, options, cancellationToken);
+			IHashedBlobRef<TNode> blobRef = await ReadRefAsync<TNode>(store, name, cacheTime, options, cancellationToken);
 			return await blobRef.ReadBlobAsync(cancellationToken);
 		}
 	}
diff --git a/Engine/Source/Programs/Shared/EpicGames.Horde/Storage/Bundles/BundleStorageClient.cs b/Engine/Source/Programs/Shared/EpicGames.Horde/Storage/Bundles/BundleStorageClient.cs
--- a/Engine/Source/Programs/Shared/EpicGames.Horde/Storage/Bundles/BundleStorageClient.cs
+++ b/Engine/Source/Programs/Shared/EpicGames.Horde/Storage/Bundles/BundleStorageClient.cs
@@ -158,7 +158,7 @@ namespace EpicGames.Horde.Storage.Bundles
 		#region Nodes
 
 		/// <inheritdoc/>
-		public IBlobHandle CreateBlobHandle(BlobLocator locator)
+		public IBlobRef CreateBlobRef(BlobLocator locator)
 		{
 			if (!locator.TryUnwrap(out BlobLocator baseLocator, out Utf8String fragment))
 			{
@@ -184,15 +184,21 @@ namespace EpicGames.Horde.Storage.Bundles
 		}
 
 		/// <inheritdoc/>
-		public IBlobRef CreateBlobRef(IoHash hash, BlobLocator locator)
+		public IBlobRef<T> CreateBlobRef<T>(BlobLocator locator, BlobSerializerOptions? options)
 		{
-			return BlobRef.Create(hash, CreateBlobHandle(locator));
+			return BlobRef.Create<T>(CreateBlobRef(locator), options);
 		}
 
 		/// <inheritdoc/>
-		public IBlobRef<T> CreateBlobRef<T>(IoHash hash, BlobLocator locator, BlobSerializerOptions? options)
+		public IHashedBlobRef CreateBlobRef(IoHash hash, BlobLocator locator)
 		{
-			return BlobRef.Create<T>(hash, CreateBlobHandle(locator), options);
+			return HashedBlobRef.Create(hash, CreateBlobRef(locator));
+		}
+
+		/// <inheritdoc/>
+		public IHashedBlobRef<T> CreateBlobRef<T>(IoHash hash, BlobLocator locator, BlobSerializerOptions? options)
+		{
+			return HashedBlobRef.Create<T>(hash, CreateBlobRef(locator), options);
 		}
 
 		/// <inheritdoc/>
@@ -218,14 +224,14 @@ namespace EpicGames.Horde.Storage.Bundles
 		#region Aliases
 
 		/// <inheritdoc/>
-		public async Task AddAliasAsync(string name, IBlobHandle handle, int rank, ReadOnlyMemory<byte> data, CancellationToken cancellationToken)
+		public async Task AddAliasAsync(string name, IBlobRef handle, int rank, ReadOnlyMemory<byte> data, CancellationToken cancellationToken)
 		{
 			await handle.FlushAsync(cancellationToken);
 			await _backend.AddAliasAsync(name, handle.GetLocator(), rank, data, cancellationToken);
 		}
 
 		/// <inheritdoc/>
-		public async Task RemoveAliasAsync(string name, IBlobHandle handle, CancellationToken cancellationToken)
+		public async Task RemoveAliasAsync(string name, IBlobRef handle, CancellationToken cancellationToken)
 		{
 			await handle.FlushAsync(cancellationToken);
 			await _backend.RemoveAliasAsync(name, handle.GetLocator(), cancellationToken);
@@ -240,7 +246,7 @@ namespace EpicGames.Horde.Storage.Bundles
 			for (int idx = 0; idx < aliases.Length; idx++)
 			{
 				BlobAliasLocator alias = aliases[idx];
-				result[idx] = new BlobAlias(CreateBlobHandle(alias.Target), alias.Rank, alias.Data);
+				result[idx] = new BlobAlias(CreateBlobRef(alias.Target), alias.Rank, alias.Data);
 			}
 			return result;
 		}
@@ -254,9 +260,9 @@ namespace EpicGames.Horde.Storage.Bundles
 			=> _backend.DeleteRefAsync(name, cancellationToken);
 
 		/// <inheritdoc/>
-		public async Task<IBlobRef?> TryReadRefAsync(RefName name, RefCacheTime cacheTime = default, CancellationToken cancellationToken = default)
+		public async Task<IHashedBlobRef?> TryReadRefAsync(RefName name, RefCacheTime cacheTime = default, CancellationToken cancellationToken = default)
 		{
-			BlobRefValue? target = await _backend.TryReadRefAsync(name, cacheTime, cancellationToken);
+			HashedBlobRefValue? target = await _backend.TryReadRefAsync(name, cacheTime, cancellationToken);
 			if (target == null)
 			{
 				return null;
@@ -265,10 +271,10 @@ namespace EpicGames.Horde.Storage.Bundles
 		}
 
 		/// <inheritdoc/>
-		public async Task WriteRefAsync(RefName name, IBlobRef target, RefOptions? options = null, CancellationToken cancellationToken = default)
+		public async Task WriteRefAsync(RefName name, IHashedBlobRef target, RefOptions? options = null, CancellationToken cancellationToken = default)
 		{
 			await target.FlushAsync(cancellationToken);
-			await _backend.WriteRefAsync(name, new BlobRefValue(target.Hash, target.GetLocator()), options, cancellationToken);
+			await _backend.WriteRefAsync(name, new HashedBlobRefValue(target.Hash, target.GetLocator()), options, cancellationToken);
 		}
 
 		#endregion
diff --git a/Engine/Source/Programs/Shared/EpicGames.Horde/Storage/Bundles/V1/BundleReader.cs b/Engine/Source/Programs/Shared/EpicGames.Horde/Storage/Bundles/V1/BundleReader.cs
--- a/Engine/Source/Programs/Shared/EpicGames.Horde/Storage/Bundles/V1/BundleReader.cs
+++ b/Engine/Source/Programs/Shared/EpicGames.Horde/Storage/Bundles/V1/BundleReader.cs
@@ -546,7 +546,7 @@ namespace EpicGames.Horde.Storage.Bundles.V1
 
 			Dictionary<BlobLocator, BundleHandle> locatorToBundleHandle = new Dictionary<BlobLocator, BundleHandle>();
 
-			List<IBlobHandle> imports = new List<IBlobHandle>(export.References.Count);
+			List<IBlobRef> imports = new List<IBlobRef>(export.References.Count);
 			foreach (BundleExportRef reference in export.References)
 			{
 				BlobLocator importBlob;
diff --git a/Engine/Source/Programs/Shared/EpicGames.Horde/Storage/Bundles/V1/BundleWriter.cs b/Engine/Source/Programs/Shared/EpicGames.Horde/Storage/Bundles/V1/BundleWriter.cs
--- a/Engine/Source/Programs/Shared/EpicGames.Horde/Storage/Bundles/V1/BundleWriter.cs
+++ b/Engine/Source/Programs/Shared/EpicGames.Horde/Storage/Bundles/V1/BundleWriter.cs
@@ -16,9 +16,9 @@ using Microsoft.Extensions.Logging;
 namespace EpicGames.Horde.Storage.Bundles.V1
 {
 	/// <summary>
-	/// Implementation of <see cref="IBlobHandle"/> for nodes which can be read from storage
+	/// Implementation of <see cref="IBlobRef"/> for nodes which can be read from storage
 	/// </summary>
-	sealed class FlushedNodeHandle : IBlobHandle
+	sealed class FlushedNodeHandle : IBlobRef
 	{
 		readonly BundleReader _reader;
 
@@ -28,7 +28,7 @@ namespace EpicGames.Horde.Storage.Bundles.V1
 		/// <inheritdoc/>
 		public BundleHandle Outer { get; }
 
-		public IBlobHandle Innermost => this;
+		public IBlobRef Innermost => this;
 
 		/// <summary>
 		/// Constructor
@@ -93,7 +93,7 @@ namespace EpicGames.Horde.Storage.Bundles.V1
 	public sealed class BundleWriter : BlobWriter
 	{
 		// Information about a unique output node. Note that multiple node refs may de-duplicate to the same output node.
-		internal class PendingNode : IBlobRef
+		internal class PendingNode : IHashedBlobRef
 		{
 			readonly BundleReader _reader;
 
@@ -107,16 +107,16 @@ namespace EpicGames.Horde.Storage.Bundles.V1
 			public readonly int Packet;
 			public readonly int Offset;
 			public readonly int Length;
-			public readonly IBlobHandle[] Imports;
+			public readonly IBlobRef[] Imports;
 			public readonly AliasInfo[] Aliases;
 
-			public IBlobHandle Innermost => this;
+			public IBlobRef Innermost => this;
 			public PendingBundle? PendingBundle => _pendingBundle;
 			public FlushedNodeHandle? FlushedNodeHandle => _flushedHandle;
 
 			public BundleHandle Outer => (_flushedHandle != null) ? _flushedHandle.Outer : throw new NotSupportedException();
 
-			public PendingNode(BundleReader reader, BlobType blobType, IoHash hash, int packet, int offset, int length, IReadOnlyList<IBlobHandle> imports, IReadOnlyList<AliasInfo> aliases, PendingBundle pendingBundle)
+			public PendingNode(BundleReader reader, BlobType blobType, IoHash hash, int packet, int offset, int length, IReadOnlyList<IBlobRef> imports, IReadOnlyList<AliasInfo> aliases, PendingBundle pendingBundle)
 			{
 				_reader = reader;
 
@@ -173,7 +173,7 @@ namespace EpicGames.Horde.Storage.Bundles.V1
 					return;
 				}
 
-				foreach (IBlobHandle import in Imports)
+				foreach (IBlobRef import in Imports)
 				{
 					await import.FlushAsync(cancellationToken);
 				}
@@ -321,7 +321,7 @@ namespace EpicGames.Horde.Storage.Bundles.V1
 			}
 
 			// Finish a node write
-			public PendingNode WriteNode(BlobType blobType, int size, IReadOnlyList<IBlobHandle> refs, IReadOnlyList<AliasInfo> aliases)
+			public PendingNode WriteNode(BlobType blobType, int size, IReadOnlyList<IBlobRef> refs, IReadOnlyList<AliasInfo> aliases)
 			{
 				IoHash hash = IoHash.Compute(GetBuffer(size, size).Span);
 
@@ -505,10 +505,10 @@ namespace EpicGames.Horde.Storage.Bundles.V1
 				Dictionary<BlobType, int> typeToIndex = new Dictionary<BlobType, int>();
 
 				// Map of node handle to reference
-				Dictionary<IBlobHandle, BundleExportRef> nodeHandleToExportRef = new Dictionary<IBlobHandle, BundleExportRef>();
+				Dictionary<IBlobRef, BundleExportRef> nodeHandleToExportRef = new Dictionary<IBlobRef, BundleExportRef>();
 				for (int exportIdx = 0; exportIdx < _queue.Count; exportIdx++)
 				{
-					IBlobHandle handle = _queue[exportIdx];
+					IBlobRef handle = _queue[exportIdx];
 					nodeHandleToExportRef[handle] = new BundleExportRef(-1, exportIdx);
 				}
 
@@ -519,7 +519,7 @@ namespace EpicGames.Horde.Storage.Bundles.V1
 					int typeIdx = FindOrAddItemIndex(nodeInfo.BlobType, types, typeToIndex);
 
 					List<BundleExportRef> exportRefs = new List<BundleExportRef>();
-					foreach (IBlobHandle import in nodeInfo.Imports)
+					foreach (IBlobRef import in nodeInfo.Imports)
 					{
 						BundleExportRef exportRef;
 						if (!nodeHandleToExportRef.TryGetValue(import, out exportRef))
@@ -771,7 +771,7 @@ namespace EpicGames.Horde.Storage.Bundles.V1
 		/// <param name="aliases">Aliases for the node</param>
 		/// <param name="cancellationToken">Cancellation token for the operation</param>
 		/// <returns>Handle to the written node</returns>
-		public override async ValueTask<IBlobRef> WriteBlobAsync(BlobType type, int size, IReadOnlyList<IBlobHandle> imports, IReadOnlyList<AliasInfo> aliases, CancellationToken cancellationToken = default)
+		public override async ValueTask<IHashedBlobRef> WriteBlobAsync(BlobType type, int size, IReadOnlyList<IBlobRef> imports, IReadOnlyList<AliasInfo> aliases, CancellationToken cancellationToken = default)
 		{
 			PendingBundle currentBundle = GetCurrentBundle();
 
@@ -780,7 +780,7 @@ namespace EpicGames.Horde.Storage.Bundles.V1
 			TraceLogger?.LogInformation("Added new node for {NodeKey} in bundle {BundleId}", pendingNode, currentBundle.BundleId);
 
 			// Add dependencies on all bundles containing a dependent node
-			foreach (IBlobHandle import in imports)
+			foreach (IBlobRef import in imports)
 			{
 				PendingNode? pendingReference = import.Innermost as PendingNode;
 				if (pendingReference?.PendingBundle != null)
diff --git a/Engine/Source/Programs/Shared/EpicGames.Horde/Storage/Bundles/V2/BundleWriter.cs b/Engine/Source/Programs/Shared/EpicGames.Horde/Storage/Bundles/V2/BundleWriter.cs
--- a/Engine/Source/Programs/Shared/EpicGames.Horde/Storage/Bundles/V2/BundleWriter.cs
+++ b/Engine/Source/Programs/Shared/EpicGames.Horde/Storage/Bundles/V2/BundleWriter.cs
@@ -125,7 +125,7 @@ namespace EpicGames.Horde.Storage.Bundles.V2
 				return _packetWriter.GetExport(exportIdx);
 			}
 
-			public ExportHandle CompleteExport(BlobType type, int size, IReadOnlyList<IBlobHandle> imports, IReadOnlyList<AliasInfo> aliases)
+			public ExportHandle CompleteExport(BlobType type, int size, IReadOnlyList<IBlobRef> imports, IReadOnlyList<AliasInfo> aliases)
 			{
 				RuntimeAssert(_packetWriter != null);
 				RuntimeAssert(_packetHandle != null);
@@ -325,7 +325,7 @@ namespace EpicGames.Horde.Storage.Bundles.V2
 			=> _currentBundle.GetOutputBuffer(usedSize, desiredSize);
 
 		/// <inheritdoc/>
-		public override async ValueTask<IBlobRef> WriteBlobAsync(BlobType type, int size, IReadOnlyList<IBlobHandle> references, IReadOnlyList<AliasInfo> aliases, CancellationToken cancellationToken = default)
+		public override async ValueTask<IHashedBlobRef> WriteBlobAsync(BlobType type, int size, IReadOnlyList<IBlobRef> references, IReadOnlyList<AliasInfo> aliases, CancellationToken cancellationToken = default)
 		{
 			ExportHandle exportHandle = _currentBundle.CompleteExport(type, size, references, aliases);
 			if (_currentBundle.Length > _bundleOptions.MaxBlobSize)
diff --git a/Engine/Source/Programs/Shared/EpicGames.Horde/Storage/Bundles/V2/ExportHandle.cs b/Engine/Source/Programs/Shared/EpicGames.Horde/Storage/Bundles/V2/ExportHandle.cs
--- a/Engine/Source/Programs/Shared/EpicGames.Horde/Storage/Bundles/V2/ExportHandle.cs
+++ b/Engine/Source/Programs/Shared/EpicGames.Horde/Storage/Bundles/V2/ExportHandle.cs
@@ -12,7 +12,7 @@ namespace EpicGames.Horde.Storage.Bundles.V2
 	/// <summary>
 	/// Handle class for blobs exported from a bundle
 	/// </summary>
-	public class ExportHandle : IBlobRef
+	public class ExportHandle : IHashedBlobRef
 	{
 		static readonly Utf8String s_fragmentPrefix = new Utf8String("exp=");
 
@@ -31,7 +31,7 @@ namespace EpicGames.Horde.Storage.Bundles.V2
 		readonly int _exportIdx;
 
 		/// <inheritdoc/>
-		public IBlobHandle Innermost => this;
+		public IBlobRef Innermost => this;
 
 		/// <inheritdoc/>
 		public IoHash Hash => _hash;
diff --git a/Engine/Source/Programs/Shared/EpicGames.Horde/Storage/Bundles/V2/PacketHandle.cs b/Engine/Source/Programs/Shared/EpicGames.Horde/Storage/Bundles/V2/PacketHandle.cs
--- a/Engine/Source/Programs/Shared/EpicGames.Horde/Storage/Bundles/V2/PacketHandle.cs
+++ b/Engine/Source/Programs/Shared/EpicGames.Horde/Storage/Bundles/V2/PacketHandle.cs
@@ -22,7 +22,7 @@ namespace EpicGames.Horde.Storage.Bundles.V2
 		/// </summary>
 		public abstract BundleHandle Bundle { get; }
 
-		/// <inheritdoc cref="IBlobHandle.FlushAsync(CancellationToken)"/>
+		/// <inheritdoc cref="IBlobRef.FlushAsync(CancellationToken)"/>
 		public abstract ValueTask FlushAsync(CancellationToken cancellationToken = default);
 
 		/// <summary>
diff --git a/Engine/Source/Programs/Shared/EpicGames.Horde/Storage/Bundles/V2/PacketReader.cs b/Engine/Source/Programs/Shared/EpicGames.Horde/Storage/Bundles/V2/PacketReader.cs
--- a/Engine/Source/Programs/Shared/EpicGames.Horde/Storage/Bundles/V2/PacketReader.cs
+++ b/Engine/Source/Programs/Shared/EpicGames.Horde/Storage/Bundles/V2/PacketReader.cs
@@ -86,7 +86,7 @@ namespace EpicGames.Horde.Storage.Bundles.V2
 		/// </summary>
 		public BlobData Read()
 		{
-			IBlobHandle[] imports = new IBlobHandle[_decodedPacket.GetImportCount()];
+			IBlobRef[] imports = new IBlobRef[_decodedPacket.GetImportCount()];
 			for (int idx = 0; idx < imports.Length; idx++)
 			{
 				imports[idx] = GetImportedBlobHandle(idx);
@@ -105,7 +105,7 @@ namespace EpicGames.Horde.Storage.Bundles.V2
 
 			BlobType type = _decodedPacket.GetType(exportHeader.TypeIdx);
 
-			IBlobHandle[] imports = new IBlobHandle[exportHeader.Imports.Length];
+			IBlobRef[] imports = new IBlobRef[exportHeader.Imports.Length];
 			for (int idx = 0; idx < exportHeader.Imports.Length; idx++)
 			{
 				imports[idx] = GetImportedBlobHandle(exportHeader.Imports[idx]);
@@ -134,9 +134,9 @@ namespace EpicGames.Horde.Storage.Bundles.V2
 		/// <summary>
 		/// Gets an import handle for the packet
 		/// </summary>
-		IBlobHandle GetImportedBlobHandle(int blobIdx)
+		IBlobRef GetImportedBlobHandle(int blobIdx)
 		{
-			IBlobHandle? blobHandle = _cachedImportHandles[blobIdx] as IBlobHandle;
+			IBlobRef? blobHandle = _cachedImportHandles[blobIdx] as IBlobRef;
 			if (blobHandle is null)
 			{
 				PacketImport blobImportInfo = _decodedPacket.GetImport(blobIdx);
diff --git a/Engine/Source/Programs/Shared/EpicGames.Horde/Storage/Bundles/V2/PacketWriter.cs b/Engine/Source/Programs/Shared/EpicGames.Horde/Storage/Bundles/V2/PacketWriter.cs
--- a/Engine/Source/Programs/Shared/EpicGames.Horde/Storage/Bundles/V2/PacketWriter.cs
+++ b/Engine/Source/Programs/Shared/EpicGames.Horde/Storage/Bundles/V2/PacketWriter.cs
@@ -99,7 +99,7 @@ namespace EpicGames.Horde.Storage.Bundles.V2
 		/// <param name="size"></param>
 		/// <param name="type"></param>
 		/// <param name="imports"></param>
-		public int CompleteExport(int size, BlobType type, IReadOnlyList<IBlobHandle> imports)
+		public int CompleteExport(int size, BlobType type, IReadOnlyList<IBlobRef> imports)
 		{
 			int[] importIndices = new int[imports.Count];
 			for (int idx = 0; idx < imports.Count; idx++)
@@ -159,11 +159,11 @@ namespace EpicGames.Horde.Storage.Bundles.V2
 
 				BlobType type = _types[header.TypeIdx];
 
-				IBlobHandle[] imports = new IBlobHandle[header.Imports.Length];
+				IBlobRef[] imports = new IBlobRef[header.Imports.Length];
 				for (int idx = 0; idx < header.Imports.Length; idx++)
 				{
 					int importIdx = header.Imports[idx];
-					imports[idx] = (IBlobHandle)_importHandles[importIdx + PacketImport.Bias];
+					imports[idx] = (IBlobRef)_importHandles[importIdx + PacketImport.Bias];
 				}
 
 				IReadOnlyMemoryOwner<byte> body = ReadOnlyMemoryOwner.Create(export.GetPayload(), _bufferHandle.AddRef());
@@ -205,7 +205,7 @@ namespace EpicGames.Horde.Storage.Bundles.V2
 		/// </summary>
 		/// <param name="handle">Handle to add</param>
 		/// <returns>Index of the import</returns>
-		public int FindOrAddImport(IBlobHandle handle) => FindOrAddImportInternal(handle);
+		public int FindOrAddImport(IBlobRef handle) => FindOrAddImportInternal(handle);
 
 		int FindOrAddImportInternal(object handle)
 		{
@@ -242,7 +242,7 @@ namespace EpicGames.Horde.Storage.Bundles.V2
 				_bundleImportIndices.Add(bundleImportIdx);
 				return bundleImportIdx;
 			}
-			else if (handle is IBlobHandle blobHandle)
+			else if (handle is IBlobRef blobHandle)
 			{
 				BlobLocator locator;
 				if (!blobHandle.TryGetLocator(out locator))
@@ -260,8 +260,8 @@ namespace EpicGames.Horde.Storage.Bundles.V2
 		/// <summary>
 		/// Gets the import assigned to a particular index
 		/// </summary>
-		public IBlobHandle GetImport(int importIdx)
-			=> (_importHandles[importIdx + PacketImport.Bias] as IBlobHandle) ?? throw new InvalidOperationException("Import is not a blob handle");
+		public IBlobRef GetImport(int importIdx)
+			=> (_importHandles[importIdx + PacketImport.Bias] as IBlobRef) ?? throw new InvalidOperationException("Import is not a blob handle");
 
 		/// <summary>
 		/// Gets the import assigned to a particular index
diff --git a/Engine/Source/Programs/Shared/EpicGames.Horde/Storage/Clients/KeyValueStorageClient.cs b/Engine/Source/Programs/Shared/EpicGames.Horde/Storage/Clients/KeyValueStorageClient.cs
--- a/Engine/Source/Programs/Shared/EpicGames.Horde/Storage/Clients/KeyValueStorageClient.cs
+++ b/Engine/Source/Programs/Shared/EpicGames.Horde/Storage/Clients/KeyValueStorageClient.cs
@@ -16,13 +16,13 @@ namespace EpicGames.Horde.Storage.Clients
 	/// </summary>
 	public sealed class KeyValueStorageClient : IStorageClient
 	{
-		class Handle : IBlobHandle
+		class Handle : IBlobRef
 		{
 			readonly KeyValueStorageClient _keyValueStorageClient;
 			readonly BlobLocator _locator;
 
 			/// <inheritdoc/>
-			public IBlobHandle Innermost => this;
+			public IBlobRef Innermost => this;
 
 			/// <summary>
 			/// Constructor
@@ -98,14 +98,14 @@ namespace EpicGames.Horde.Storage.Clients
 			}
 
 			/// <inheritdoc/>
-			public override async ValueTask<IBlobRef> WriteBlobAsync(BlobType type, int size, IReadOnlyList<IBlobHandle> imports, IReadOnlyList<AliasInfo> aliases, CancellationToken cancellationToken = default)
+			public override async ValueTask<IHashedBlobRef> WriteBlobAsync(BlobType type, int size, IReadOnlyList<IBlobRef> imports, IReadOnlyList<AliasInfo> aliases, CancellationToken cancellationToken = default)
 			{
 				ReadOnlyMemory<byte> data = _data.AsMemory(_offset, size);
 				_offset += size;
 
 				using ReadOnlyMemoryStream stream = new ReadOnlyMemoryStream(data);
 
-				IBlobRef handle = await _outer.WriteBlobAsync(type, stream, imports, _basePath, cancellationToken);
+				IHashedBlobRef handle = await _outer.WriteBlobAsync(type, stream, imports, _basePath, cancellationToken);
 				foreach (AliasInfo aliasInfo in aliases)
 				{
 					await _outer.AddAliasAsync(aliasInfo.Name, handle, aliasInfo.Rank, aliasInfo.Data, cancellationToken);
@@ -140,21 +140,27 @@ namespace EpicGames.Horde.Storage.Clients
 		#region Blobs
 
 		/// <inheritdoc/>
-		public IBlobHandle CreateBlobHandle(BlobLocator locator)
+		public IBlobRef CreateBlobRef(BlobLocator locator)
 		{
 			return new Handle(this, locator);
 		}
 
 		/// <inheritdoc/>
-		public IBlobRef CreateBlobRef(IoHash hash, BlobLocator locator)
+		public IBlobRef<T> CreateBlobRef<T>(BlobLocator locator, BlobSerializerOptions? options)
 		{
-			return BlobRef.Create(hash, CreateBlobHandle(locator));
+			return BlobRef.Create<T>(CreateBlobRef(locator), options);
 		}
 
 		/// <inheritdoc/>
-		public IBlobRef<T> CreateBlobRef<T>(IoHash hash, BlobLocator locator, BlobSerializerOptions? options)
+		public IHashedBlobRef CreateBlobRef(IoHash hash, BlobLocator locator)
 		{
-			return BlobRef.Create<T>(hash, CreateBlobHandle(locator), options);
+			return HashedBlobRef.Create(hash, CreateBlobRef(locator));
+		}
+
+		/// <inheritdoc/>
+		public IHashedBlobRef<T> CreateBlobRef<T>(IoHash hash, BlobLocator locator, BlobSerializerOptions? options)
+		{
+			return HashedBlobRef.Create<T>(hash, CreateBlobRef(locator), options);
 		}
 
 		/// <inheritdoc/>
@@ -166,11 +172,11 @@ namespace EpicGames.Horde.Storage.Clients
 			IReadOnlyMemoryOwner<byte> owner = await _backend.ReadBlobAsync(locator, cancellationToken);
 			EncodedBlobData encodedData = new EncodedBlobData(owner.Memory.ToArray());
 
-			return new BlobDataWithOwner(encodedData.Type, encodedData.Payload, encodedData.Imports.Select(x => CreateBlobHandle(x)).ToArray(), owner);
+			return new BlobDataWithOwner(encodedData.Type, encodedData.Payload, encodedData.Imports.Select(x => CreateBlobRef(x)).ToArray(), owner);
 		}
 
 		/// <inheritdoc/>
-		public async ValueTask<IBlobRef> WriteBlobAsync(BlobType type, Stream stream, IReadOnlyList<IBlobHandle> imports, string? basePath = null, CancellationToken cancellationToken = default)
+		public async ValueTask<IHashedBlobRef> WriteBlobAsync(BlobType type, Stream stream, IReadOnlyList<IBlobRef> imports, string? basePath = null, CancellationToken cancellationToken = default)
 		{
 			BlobLocator[] locators = imports.ConvertAll(x => x.GetLocator()).ToArray();
 			byte[] payload = await stream.ReadAllBytesAsync(cancellationToken);
@@ -180,7 +186,7 @@ namespace EpicGames.Horde.Storage.Clients
 			BlobLocator locator = await _backend.WriteBlobAsync(encodedStream, locators, basePath, cancellationToken);
 
 			IoHash hash = IoHash.Compute(payload);
-			return BlobRef.Create(hash, CreateBlobHandle(locator));
+			return HashedBlobRef.Create(hash, CreateBlobRef(locator));
 		}
 
 		#endregion
@@ -188,14 +194,14 @@ namespace EpicGames.Horde.Storage.Clients
 		#region Aliases
 
 		/// <inheritdoc/>
-		public async Task AddAliasAsync(string name, IBlobHandle handle, int rank = 0, ReadOnlyMemory<byte> data = default, CancellationToken cancellationToken = default)
+		public async Task AddAliasAsync(string name, IBlobRef handle, int rank = 0, ReadOnlyMemory<byte> data = default, CancellationToken cancellationToken = default)
 		{
 			await handle.FlushAsync(cancellationToken);
 			await _backend.AddAliasAsync(name, handle.GetLocator(), rank, data, cancellationToken);
 		}
 
 		/// <inheritdoc/>
-		public async Task RemoveAliasAsync(string name, IBlobHandle handle, CancellationToken cancellationToken = default)
+		public async Task RemoveAliasAsync(string name, IBlobRef handle, CancellationToken cancellationToken = default)
 		{
 			await handle.FlushAsync(cancellationToken);
 			await _backend.RemoveAliasAsync(name, handle.GetLocator(), cancellationToken);
@@ -205,7 +211,7 @@ namespace EpicGames.Horde.Storage.Clients
 		public async Task<BlobAlias[]> FindAliasesAsync(string name, int? maxResults = null, CancellationToken cancellationToken = default)
 		{
 			BlobAliasLocator[] locators = await _backend.FindAliasesAsync(name, maxResults, cancellationToken);
-			return locators.Select(x => new BlobAlias(CreateBlobHandle(x.Target), x.Rank, x.Data)).ToArray();
+			return locators.Select(x => new BlobAlias(CreateBlobRef(x.Target), x.Rank, x.Data)).ToArray();
 		}
 
 		#endregion
@@ -213,9 +219,9 @@ namespace EpicGames.Horde.Storage.Clients
 		#region Refs
 
 		/// <inheritdoc/>
-		public async Task<IBlobRef?> TryReadRefAsync(RefName name, RefCacheTime cacheTime = default, CancellationToken cancellationToken = default)
+		public async Task<IHashedBlobRef?> TryReadRefAsync(RefName name, RefCacheTime cacheTime = default, CancellationToken cancellationToken = default)
 		{
-			BlobRefValue? value = await _backend.TryReadRefAsync(name, cacheTime, cancellationToken);
+			HashedBlobRefValue? value = await _backend.TryReadRefAsync(name, cacheTime, cancellationToken);
 			if (value == null)
 			{
 				return null;
@@ -224,10 +230,10 @@ namespace EpicGames.Horde.Storage.Clients
 		}
 
 		/// <inheritdoc/>
-		public async Task WriteRefAsync(RefName name, IBlobRef value, RefOptions? options = null, CancellationToken cancellationToken = default)
+		public async Task WriteRefAsync(RefName name, IHashedBlobRef value, RefOptions? options = null, CancellationToken cancellationToken = default)
 		{
 			await value.FlushAsync(cancellationToken);
-			await _backend.WriteRefAsync(name, new BlobRefValue(value.Hash, value.GetLocator()), options, cancellationToken);
+			await _backend.WriteRefAsync(name, new HashedBlobRefValue(value.Hash, value.GetLocator()), options, cancellationToken);
 		}
 
 		/// <inheritdoc/>
diff --git a/Engine/Source/Programs/Shared/EpicGames.Horde/Storage/DedupeBlobWriter.cs b/Engine/Source/Programs/Shared/EpicGames.Horde/Storage/DedupeBlobWriter.cs
--- a/Engine/Source/Programs/Shared/EpicGames.Horde/Storage/DedupeBlobWriter.cs
+++ b/Engine/Source/Programs/Shared/EpicGames.Horde/Storage/DedupeBlobWriter.cs
@@ -21,16 +21,16 @@ namespace EpicGames.Horde.Storage
 		{
 			readonly int _maxKeys;
 			readonly Queue<BlobKey> _blobKeys = new Queue<BlobKey>();
-			readonly Dictionary<BlobKey, IBlobRef> _blobKeyToHandle = new Dictionary<BlobKey, IBlobRef>();
+			readonly Dictionary<BlobKey, IHashedBlobRef> _blobKeyToHandle = new Dictionary<BlobKey, IHashedBlobRef>();
 
 			public DedupeCache(int maxKeys)
 			{
 				_maxKeys = maxKeys;
 				_blobKeys = new Queue<BlobKey>(maxKeys);
-				_blobKeyToHandle = new Dictionary<BlobKey, IBlobRef>(maxKeys);
+				_blobKeyToHandle = new Dictionary<BlobKey, IHashedBlobRef>(maxKeys);
 			}
 
-			internal void Add(BlobKey key, IBlobRef handle)
+			internal void Add(BlobKey key, IHashedBlobRef handle)
 			{
 				BlobKey? prevKey;
 				if (_blobKeys.Count == _maxKeys && _blobKeys.TryDequeue(out prevKey))
@@ -40,16 +40,16 @@ namespace EpicGames.Horde.Storage
 				_blobKeyToHandle.TryAdd(key, handle);
 			}
 
-			internal bool TryGetValue(BlobKey key, [NotNullWhen(true)] out IBlobRef? handle) => _blobKeyToHandle.TryGetValue(key, out handle);
+			internal bool TryGetValue(BlobKey key, [NotNullWhen(true)] out IHashedBlobRef? handle) => _blobKeyToHandle.TryGetValue(key, out handle);
 		}
 
-		class WrappedHandle : IBlobRef
+		class WrappedHandle : IHashedBlobRef
 		{
 			public object _lockObject = new object();
-			public IBlobRef? _inner;
+			public IHashedBlobRef? _inner;
 
 			/// <inheritdoc/>
-			public IBlobHandle Innermost
+			public IBlobRef Innermost
 				=> _inner!.Innermost;
 
 			/// <inheritdoc/>
@@ -125,14 +125,14 @@ namespace EpicGames.Horde.Storage
 		/// </summary>
 		/// <param name="type">Type of the blob</param>
 		/// <param name="blobRef">Reference to the blob data</param>
-		public void AddToCache(BlobType type, IBlobRef blobRef)
+		public void AddToCache(BlobType type, IHashedBlobRef blobRef)
 		{
 			_cache.Add(new BlobKey(blobRef.Hash, type), blobRef);
 			Interlocked.Increment(ref _numCacheAdds);
 		}
 
 		/// <inheritdoc/>
-		public override async ValueTask<IBlobRef> WriteBlobAsync(BlobType type, int size, IReadOnlyList<IBlobHandle> imports, IReadOnlyList<AliasInfo> aliases, CancellationToken cancellationToken = default)
+		public override async ValueTask<IHashedBlobRef> WriteBlobAsync(BlobType type, int size, IReadOnlyList<IBlobRef> imports, IReadOnlyList<AliasInfo> aliases, CancellationToken cancellationToken = default)
 		{
 			ReadOnlyMemory<byte> data = _inner.GetOutputBuffer(size, size).Slice(0, size);
 			IoHash hash = IoHash.Compute(data.Span);
@@ -141,7 +141,7 @@ namespace EpicGames.Horde.Storage
 			WrappedHandle? wrappedHandle;
 			lock (_cache)
 			{
-				IBlobRef? handle;
+				IHashedBlobRef? handle;
 				if (_cache.TryGetValue(key, out handle))
 				{
 					_numHits++;
@@ -216,7 +216,7 @@ namespace EpicGames.Horde.Storage
 		/// <param name="dedupeWriter">Dedupe writer to operate on</param>
 		/// <param name="directoryNodeRef">Reference to the directory to add</param>
 		/// <param name="cancellationToken">Cancellation token for the operation</param>
-		public static async Task AddToCacheAsync(this DedupeBlobWriter dedupeWriter, IBlobRef<DirectoryNode> directoryNodeRef, CancellationToken cancellationToken = default)
+		public static async Task AddToCacheAsync(this DedupeBlobWriter dedupeWriter, IHashedBlobRef<DirectoryNode> directoryNodeRef, CancellationToken cancellationToken = default)
 		{
 			using BlobData blobData = await directoryNodeRef.ReadBlobDataAsync(cancellationToken);
 			dedupeWriter.AddToCache(blobData.Type, directoryNodeRef);
diff --git a/Engine/Source/Programs/Shared/EpicGames.Horde/Storage/IBlobHandle.cs b/Engine/Source/Programs/Shared/EpicGames.Horde/Storage/IBlobHandle.cs
deleted file mode 100644
diff --git a/Engine/Source/Programs/Shared/EpicGames.Horde/Storage/IBlobReader.cs b/Engine/Source/Programs/Shared/EpicGames.Horde/Storage/IBlobReader.cs
--- a/Engine/Source/Programs/Shared/EpicGames.Horde/Storage/IBlobReader.cs
+++ b/Engine/Source/Programs/Shared/EpicGames.Horde/Storage/IBlobReader.cs
@@ -26,17 +26,17 @@ namespace EpicGames.Horde.Storage
 		/// <summary>
 		/// Locations of all referenced nodes. These handles do not have valid hashes.
 		/// </summary>
-		IReadOnlyList<IBlobHandle> Imports { get; }
+		IReadOnlyList<IBlobRef> Imports { get; }
 
 		/// <summary>
 		/// Gets the next serialized blob handle
 		/// </summary>
-		IBlobRef ReadBlobRef();
+		IHashedBlobRef ReadBlobRef();
 
 		/// <summary>
 		/// Gets the next serialized blob handle
 		/// </summary>
-		IBlobRef<T> ReadBlobRef<T>();
+		IHashedBlobRef<T> ReadBlobRef<T>();
 	}
 
 	/// <summary>
@@ -72,7 +72,7 @@ namespace EpicGames.Horde.Storage
 		/// <summary>
 		/// Locations of all referenced nodes.
 		/// </summary>
-		public IReadOnlyList<IBlobHandle> Imports => _blobData.Imports;
+		public IReadOnlyList<IBlobRef> Imports => _blobData.Imports;
 
 		readonly BlobData _blobData;
 		readonly BlobSerializerOptions _options;
@@ -91,21 +91,21 @@ namespace EpicGames.Horde.Storage
 		/// <summary>
 		/// Gets the next serialized blob reference
 		/// </summary>
-		public IBlobRef ReadBlobRef()
+		public IHashedBlobRef ReadBlobRef()
 		{
-			IBlobHandle import = Imports[_importIdx++];
+			IBlobRef import = Imports[_importIdx++];
 			IoHash hash = this.ReadIoHash();
-			return BlobRef.Create(hash, import);
+			return HashedBlobRef.Create(hash, import);
 		}
 
 		/// <summary>
 		/// Gets the next serialized blob reference
 		/// </summary>
-		public IBlobRef<T> ReadBlobRef<T>()
+		public IHashedBlobRef<T> ReadBlobRef<T>()
 		{
-			IBlobHandle import = Imports[_importIdx++];
+			IBlobRef import = Imports[_importIdx++];
 			IoHash hash = this.ReadIoHash();
-			return BlobRef.Create<T>(hash, import, _options);
+			return HashedBlobRef.Create<T>(hash, import, _options);
 		}
 	}
 }
diff --git a/Engine/Source/Programs/Shared/EpicGames.Horde/Storage/IBlobRef.cs b/Engine/Source/Programs/Shared/EpicGames.Horde/Storage/IBlobRef.cs
new file mode 100644
--- /dev/null
+++ b/Engine/Source/Programs/Shared/EpicGames.Horde/Storage/IBlobRef.cs
@@ -0,0 +1,217 @@
+// Copyright Epic Games, Inc. All Rights Reserved.
+
+using System;
+using System.Diagnostics.CodeAnalysis;
+using System.Threading;
+using System.Threading.Tasks;
+using EpicGames.Core;
+using EpicGames.Horde.Accounts;
+
+namespace EpicGames.Horde.Storage
+{
+	/// <summary>
+	/// Reference to another node in storage. This type is similar to <see cref="IHashedBlobRef"/>, but without a hash.
+	/// </summary>
+	public interface IBlobRef
+	{
+		/// <summary>
+		/// Accessor for the innermost import
+		/// </summary>
+		IBlobRef Innermost { get; }
+
+		/// <summary>
+		/// Flush the referenced data to underlying storage
+		/// </summary>
+		ValueTask FlushAsync(CancellationToken cancellationToken = default);
+
+		/// <summary>
+		/// Reads the blob's data
+		/// </summary>
+		/// <param name="cancellationToken">Cancellation token for the operation</param>
+		ValueTask<BlobData> ReadBlobDataAsync(CancellationToken cancellationToken = default);
+
+		/// <summary>
+		/// Attempt to get a path for this blob.
+		/// </summary>
+		/// <param name="locator">Receives the blob path on success.</param>
+		/// <returns>True if a path was available, false if the blob has not yet been flushed to storage.</returns>
+		bool TryGetLocator([NotNullWhen(true)] out BlobLocator locator);
+	}
+
+	/// <summary>
+	/// Typed interface to a particular blob handle
+	/// </summary>
+	/// <typeparam name="T">Type of the deserialized blob</typeparam>
+	public interface IBlobRef<out T> : IBlobRef
+	{
+		/// <summary>
+		/// Options for deserializing the blob
+		/// </summary>
+		BlobSerializerOptions? SerializerOptions { get; }
+	}
+
+	/// <summary>
+	/// Static methods for IBlobRef types.
+	/// </summary>
+	public static class BlobRef
+	{
+		class TypedBlobRef<T> : IBlobRef<T>
+		{
+			readonly IBlobRef _inner;
+
+			/// <inheritdoc/>
+			public BlobSerializerOptions? SerializerOptions { get; }
+
+			public IBlobRef Innermost => throw new NotImplementedException();
+
+			public TypedBlobRef(IBlobRef inner, BlobSerializerOptions? serializerOptions)
+			{
+				_inner = inner;
+				SerializerOptions = serializerOptions;
+			}
+
+			public ValueTask FlushAsync(CancellationToken cancellationToken = default)
+				=> _inner.FlushAsync(cancellationToken);
+
+			public ValueTask<BlobData> ReadBlobDataAsync(CancellationToken cancellationToken = default)
+				=> _inner.ReadBlobDataAsync(cancellationToken);
+
+			public bool TryGetLocator([NotNullWhen(true)] out BlobLocator locator)
+				=> _inner.TryGetLocator(out locator);
+		}
+
+		/// <summary>
+		/// Create a typed blob reference
+		/// </summary>
+		/// <typeparam name="T">Target type</typeparam>
+		/// <param name="blobRef">Blob referfence to wrap</param>
+		/// <param name="serializerOptions">Options for deserializing the blob</param>
+		public static IBlobRef<T> Create<T>(IBlobRef blobRef, BlobSerializerOptions? serializerOptions = null)
+			=> new TypedBlobRef<T>(blobRef, serializerOptions);
+	}
+
+	/// <summary>
+	/// Handle to a node. Can be used to reference nodes that have not been flushed yet.
+	/// </summary>
+	public interface IHashedBlobRef : IBlobRef
+	{
+		/// <summary>
+		/// Hash of the target node
+		/// </summary>
+		IoHash Hash { get; }
+	}
+
+	/// <summary>
+	/// Typed interface to a particular blob handle
+	/// </summary>
+	/// <typeparam name="T">Type of the deserialized blob</typeparam>
+	public interface IHashedBlobRef<out T> : IHashedBlobRef, IBlobRef<T>
+	{
+	}
+
+	/// <summary>
+	/// Helper methods for creating blob handles
+	/// </summary>
+	public static class HashedBlobRef
+	{
+		class HashedBlobRefImpl : IHashedBlobRef
+		{
+			readonly IoHash _hash;
+			readonly IBlobRef _handle;
+
+			public IBlobRef Innermost => _handle.Innermost;
+			public IoHash Hash => _hash;
+
+			public HashedBlobRefImpl(IoHash hash, IBlobRef handle)
+			{
+				_hash = hash;
+				_handle = handle;
+			}
+
+			public ValueTask FlushAsync(CancellationToken cancellationToken = default)
+				=> _handle.FlushAsync(cancellationToken);
+
+			public ValueTask<BlobData> ReadBlobDataAsync(CancellationToken cancellationToken = default)
+				=> _handle.ReadBlobDataAsync(cancellationToken);
+
+			public bool TryGetLocator([NotNullWhen(true)] out BlobLocator locator)
+				=> _handle.TryGetLocator(out locator);
+		}
+
+		class HashedBlobRefImpl<T> : HashedBlobRefImpl, IHashedBlobRef<T>
+		{
+			readonly BlobSerializerOptions _options;
+
+			public BlobSerializerOptions SerializerOptions => _options;
+
+			public HashedBlobRefImpl(IoHash hash, IBlobRef handle, BlobSerializerOptions options)
+				: base(hash, handle)
+			{
+				_options = options;
+			}
+		}
+
+		/// <summary>
+		/// Create an untyped blob handle
+		/// </summary>
+		/// <param name="handle">Imported blob interface</param>
+		/// <param name="hash">Hash of the blob</param>
+		/// <returns>Handle to the blob</returns>
+		public static IHashedBlobRef Create(IoHash hash, IBlobRef handle)
+			=> new HashedBlobRefImpl(hash, handle);
+
+		/// <summary>
+		/// Create a typed blob handle
+		/// </summary>
+		/// <typeparam name="T"></typeparam>
+		/// <param name="blobRef">Existing blob reference</param>
+		/// <param name="options">Options for deserializing the target blob</param>
+		/// <returns>Handle to the blob</returns>
+		public static IHashedBlobRef<T> Create<T>(IHashedBlobRef blobRef, BlobSerializerOptions? options = null)
+			=> new HashedBlobRefImpl<T>(blobRef.Hash, blobRef, options ?? BlobSerializerOptions.Default);
+
+		/// <summary>
+		/// Create a typed blob handle
+		/// </summary>
+		/// <typeparam name="T"></typeparam>
+		/// <param name="hash">Hash of the blob</param>
+		/// <param name="handle">Imported blob interface</param>
+		/// <param name="options">Options for deserializing the target blob</param>
+		/// <returns>Handle to the blob</returns>
+		public static IHashedBlobRef<T> Create<T>(IoHash hash, IBlobRef handle, BlobSerializerOptions? options = null)
+			=> new HashedBlobRefImpl<T>(hash, handle, options ?? BlobSerializerOptions.Default);
+	}
+
+	/// <summary>
+	/// Contains the value for a blob ref
+	/// </summary>
+	public record class HashedBlobRefValue(IoHash Hash, BlobLocator Locator);
+
+	/// <summary>
+	/// Extension methods for <see cref="IBlobRef"/>
+	/// </summary>
+	public static class BlobRefExtensions
+	{
+		/// <summary>
+		/// Gets a path to this blob that can be used to describe blob references over the wire.
+		/// </summary>
+		/// <param name="import">Handle to query</param>
+		public static BlobLocator GetLocator(this IBlobRef import)
+		{
+			BlobLocator locator;
+			if (!import.TryGetLocator(out locator))
+			{
+				throw new InvalidOperationException("Blob has not yet been written to storage");
+			}
+			return locator;
+		}
+
+		/// <summary>
+		/// Gets a BlobRefValue from an IBlobRef
+		/// </summary>
+		public static HashedBlobRefValue GetRefValue(this IHashedBlobRef blobRef)
+		{
+			return new HashedBlobRefValue(blobRef.Hash, blobRef.GetLocator());
+		}
+	}
+}
diff --git a/Engine/Source/Programs/Shared/EpicGames.Horde/Storage/IBlobWriter.cs b/Engine/Source/Programs/Shared/EpicGames.Horde/Storage/IBlobWriter.cs
--- a/Engine/Source/Programs/Shared/EpicGames.Horde/Storage/IBlobWriter.cs
+++ b/Engine/Source/Programs/Shared/EpicGames.Horde/Storage/IBlobWriter.cs
@@ -53,7 +53,7 @@ namespace EpicGames.Horde.Storage
 		/// <param name="type">Type of the node that was written</param>
 		/// <param name="cancellationToken">Cancellation token for the operation</param>
 		/// <returns>Handle to the written node</returns>
-		ValueTask<IBlobRef> CompleteAsync(BlobType type, CancellationToken cancellationToken = default);
+		ValueTask<IHashedBlobRef> CompleteAsync(BlobType type, CancellationToken cancellationToken = default);
 
 		/// <summary>
 		/// Finish writing a blob that has been written into the output buffer.
@@ -61,20 +61,20 @@ namespace EpicGames.Horde.Storage
 		/// <param name="type">Type of the node that was written</param>
 		/// <param name="cancellationToken">Cancellation token for the operation</param>
 		/// <returns>Handle to the written node</returns>
-		ValueTask<IBlobRef<T>> CompleteAsync<T>(BlobType type, CancellationToken cancellationToken = default);
+		ValueTask<IHashedBlobRef<T>> CompleteAsync<T>(BlobType type, CancellationToken cancellationToken = default);
 
 		/// <summary>
 		/// Writes a reference to another blob. NOTE: This does not write anything to the underlying output stream, which prevents the data forming a Merkle tree
 		/// unless guaranteed uniqueness via a hash being written separately.
 		/// </summary>
 		/// <param name="handle">Referenced blob</param>
-		void WriteBlobHandleDangerous(IBlobHandle handle);
+		void WriteBlobHandleDangerous(IBlobRef handle);
 
 		/// <summary>
 		/// Writes a reference to another blob. The blob's hash is serialized to the output stream.
 		/// </summary>
 		/// <param name="blobRef">Referenced blob</param>
-		void WriteBlobRef(IBlobRef blobRef);
+		void WriteBlobRef(IHashedBlobRef blobRef);
 	}
 
 	/// <summary>
@@ -92,7 +92,7 @@ namespace EpicGames.Horde.Storage
 	{
 		Memory<byte> _memory;
 		readonly List<AliasInfo> _aliases = new List<AliasInfo>();
-		readonly List<IBlobHandle> _imports = new List<IBlobHandle>();
+		readonly List<IBlobRef> _imports = new List<IBlobRef>();
 		readonly BlobSerializerOptions _options;
 		int _length;
 
@@ -120,7 +120,7 @@ namespace EpicGames.Horde.Storage
 		public IoHash ComputeHash() => IoHash.Compute(_memory.Span.Slice(0, _length));
 
 		/// <inheritdoc/>
-		public void WriteBlobHandleDangerous(IBlobHandle target)
+		public void WriteBlobHandleDangerous(IBlobRef target)
 		{
 			_imports.Add(target);
 		}
@@ -128,7 +128,7 @@ namespace EpicGames.Horde.Storage
 		/// <summary>
 		/// Writes a handle to another node
 		/// </summary>
-		public void WriteBlobRef(IBlobRef target)
+		public void WriteBlobRef(IHashedBlobRef target)
 		{
 			this.WriteIoHash(target.Hash);
 			WriteBlobHandleDangerous(target);
@@ -169,7 +169,7 @@ namespace EpicGames.Horde.Storage
 		/// <param name="aliases">Aliases for the new blob</param>
 		/// <param name="cancellationToken">Cancellation token for the operation</param>
 		/// <returns>New buffer</returns>
-		public abstract ValueTask<IBlobRef> WriteBlobAsync(BlobType type, int size, IReadOnlyList<IBlobHandle> imports, IReadOnlyList<AliasInfo> aliases, CancellationToken cancellationToken);
+		public abstract ValueTask<IHashedBlobRef> WriteBlobAsync(BlobType type, int size, IReadOnlyList<IBlobRef> imports, IReadOnlyList<AliasInfo> aliases, CancellationToken cancellationToken);
 
 		/// <inheritdoc/>
 		public void AddAlias(string name, int rank, ReadOnlyMemory<byte> data)
@@ -182,9 +182,9 @@ namespace EpicGames.Horde.Storage
 		public abstract IBlobWriter Fork();
 
 		/// <inheritdoc/>
-		public async ValueTask<IBlobRef> CompleteAsync(BlobType type, CancellationToken cancellationToken = default)
+		public async ValueTask<IHashedBlobRef> CompleteAsync(BlobType type, CancellationToken cancellationToken = default)
 		{
-			IBlobRef blobRef = await WriteBlobAsync(type, _length, _imports, _aliases, cancellationToken);
+			IHashedBlobRef blobRef = await WriteBlobAsync(type, _length, _imports, _aliases, cancellationToken);
 
 			_memory = default;
 			_length = 0;
@@ -195,11 +195,11 @@ namespace EpicGames.Horde.Storage
 		}
 
 		/// <inheritdoc/>
-		public async ValueTask<IBlobRef<T>> CompleteAsync<T>(BlobType type, CancellationToken cancellationToken = default)
+		public async ValueTask<IHashedBlobRef<T>> CompleteAsync<T>(BlobType type, CancellationToken cancellationToken = default)
 		{
 			IoHash hash = IoHash.Compute(_memory.Span.Slice(0, _length));
-			IBlobRef blobRef = await CompleteAsync(type, cancellationToken);
-			return BlobRef.Create<T>(hash, blobRef, _options);
+			IHashedBlobRef blobRef = await CompleteAsync(type, cancellationToken);
+			return HashedBlobRef.Create<T>(hash, blobRef, _options);
 		}
 
 		/// <inheritdoc/>
@@ -211,7 +211,7 @@ namespace EpicGames.Horde.Storage
 	/// </summary>
 	public class MemoryBlobWriter : BlobWriter
 	{
-		class BlobRef : IBlobRef
+		class BlobRef : IHashedBlobRef
 		{
 			readonly int _index;
 			readonly IoHash _hash;
@@ -230,7 +230,7 @@ namespace EpicGames.Horde.Storage
 			public int Index
 				=> _index;
 
-			public IBlobHandle Innermost
+			public IBlobRef Innermost
 				=> this;
 
 			public ValueTask FlushAsync(CancellationToken cancellationToken = default)
@@ -282,17 +282,17 @@ namespace EpicGames.Horde.Storage
 			=> _memoryWriter.GetMemory(usedSize, desiredSize);
 
 		/// <inheritdoc/>
-		public override ValueTask<IBlobRef> WriteBlobAsync(BlobType type, int size, IReadOnlyList<IBlobHandle> imports, IReadOnlyList<AliasInfo> aliases, CancellationToken cancellationToken)
+		public override ValueTask<IHashedBlobRef> WriteBlobAsync(BlobType type, int size, IReadOnlyList<IBlobRef> imports, IReadOnlyList<AliasInfo> aliases, CancellationToken cancellationToken)
 		{
 			Memory<byte> memory = _memoryWriter.GetMemoryAndAdvance(size);
 			BlobData data = new BlobData(type, memory, imports.ToArray());
-			return new ValueTask<IBlobRef>(new BlobRef(++_nextIndex, data));
+			return new ValueTask<IHashedBlobRef>(new BlobRef(++_nextIndex, data));
 		}
 
 		/// <summary>
 		/// Helper function to get the index of a blob
 		/// </summary>
-		public static int GetIndex(IBlobRef handle)
+		public static int GetIndex(IHashedBlobRef handle)
 			=> ((BlobRef)handle.Innermost).Index;
 	}
 }
diff --git a/Engine/Source/Programs/Shared/EpicGames.Horde/Storage/IStorageBackend.cs b/Engine/Source/Programs/Shared/EpicGames.Horde/Storage/IStorageBackend.cs
--- a/Engine/Source/Programs/Shared/EpicGames.Horde/Storage/IStorageBackend.cs
+++ b/Engine/Source/Programs/Shared/EpicGames.Horde/Storage/IStorageBackend.cs
@@ -113,7 +113,7 @@ namespace EpicGames.Horde.Storage
 		/// <param name="cacheTime">Minimum coherency for any cached value to be returned</param>
 		/// <param name="cancellationToken">Cancellation token for the operation</param>
 		/// <returns>Blob pointed to by the ref</returns>
-		Task<BlobRefValue?> TryReadRefAsync(RefName name, RefCacheTime cacheTime = default, CancellationToken cancellationToken = default);
+		Task<HashedBlobRefValue?> TryReadRefAsync(RefName name, RefCacheTime cacheTime = default, CancellationToken cancellationToken = default);
 
 		/// <summary>
 		/// Writes a new ref to the store
@@ -123,7 +123,7 @@ namespace EpicGames.Horde.Storage
 		/// <param name="options">Options for the new ref</param>
 		/// <param name="cancellationToken">Cancellation token for the operation</param>
 		/// <returns>Unique identifier for the blob</returns>
-		Task WriteRefAsync(RefName name, BlobRefValue value, RefOptions? options = null, CancellationToken cancellationToken = default);
+		Task WriteRefAsync(RefName name, HashedBlobRefValue value, RefOptions? options = null, CancellationToken cancellationToken = default);
 
 		/// <summary>
 		/// Reads data for a ref from the store
diff --git a/Engine/Source/Programs/Shared/EpicGames.Horde/Storage/IStorageClient.cs b/Engine/Source/Programs/Shared/EpicGames.Horde/Storage/IStorageClient.cs
--- a/Engine/Source/Programs/Shared/EpicGames.Horde/Storage/IStorageClient.cs
+++ b/Engine/Source/Programs/Shared/EpicGames.Horde/Storage/IStorageClient.cs
@@ -2,6 +2,7 @@
 
 using System;
 using System.Collections.Generic;
+using System.Diagnostics.CodeAnalysis;
 using System.Linq;
 using System.Threading;
 using System.Threading.Tasks;
@@ -29,7 +30,15 @@ namespace EpicGames.Horde.Storage
 		/// </summary>
 		/// <param name="locator">Path to the blob</param>
 		/// <returns>New handle to the blob</returns>
-		IBlobHandle CreateBlobHandle(BlobLocator locator);
+		IBlobRef CreateBlobRef(BlobLocator locator);
+
+		/// <summary>
+		/// Creates a new blob handle by parsing a locator
+		/// </summary>
+		/// <param name="locator">Path to the blob</param>
+		/// <param name="serializerOptions">Options for deserializing the blob</param>
+		/// <returns>New handle to the blob</returns>
+		IBlobRef<T> CreateBlobRef<T>(BlobLocator locator, BlobSerializerOptions? serializerOptions = null);
 
 		/// <summary>
 		/// Creates a new blob reference from a locator and hash
@@ -37,7 +46,7 @@ namespace EpicGames.Horde.Storage
 		/// <param name="hash">Hash of the target blob</param>
 		/// <param name="locator">Path to the blob</param>
 		/// <returns>New handle to the blob</returns>
-		IBlobRef CreateBlobRef(IoHash hash, BlobLocator locator);
+		IHashedBlobRef CreateBlobRef(IoHash hash, BlobLocator locator);
 
 		/// <summary>
 		/// Creates a new blob reference from a locator and hash
@@ -46,7 +55,7 @@ namespace EpicGames.Horde.Storage
 		/// <param name="locator">Path to the blob</param>
 		/// <param name="serializerOptions">Options for deserializing the blob</param>
 		/// <returns>New handle to the blob</returns>
-		IBlobRef<T> CreateBlobRef<T>(IoHash hash, BlobLocator locator, BlobSerializerOptions? serializerOptions = null);
+		IHashedBlobRef<T> CreateBlobRef<T>(IoHash hash, BlobLocator locator, BlobSerializerOptions? serializerOptions = null);
 
 		/// <summary>
 		/// Creates a new writer for storage blobs
@@ -68,7 +77,7 @@ namespace EpicGames.Horde.Storage
 		/// <param name="rank">Rank for this alias. In situations where an alias has multiple mappings, the alias with the highest rank will be returned by default.</param>
 		/// <param name="data">Additional data to be stored inline with the alias</param>
 		/// <param name="cancellationToken">Cancellation token for the operation</param>
-		Task AddAliasAsync(string name, IBlobHandle handle, int rank = 0, ReadOnlyMemory<byte> data = default, CancellationToken cancellationToken = default);
+		Task AddAliasAsync(string name, IBlobRef handle, int rank = 0, ReadOnlyMemory<byte> data = default, CancellationToken cancellationToken = default);
 
 		/// <summary>
 		/// Removes an alias from a blob
@@ -76,7 +85,7 @@ namespace EpicGames.Horde.Storage
 		/// <param name="name">Name of the alias</param>
 		/// <param name="handle">Locator for the blob</param>
 		/// <param name="cancellationToken">Cancellation token for the operation</param>
-		Task RemoveAliasAsync(string name, IBlobHandle handle, CancellationToken cancellationToken = default);
+		Task RemoveAliasAsync(string name, IBlobRef handle, CancellationToken cancellationToken = default);
 
 		/// <summary>
 		/// Finds blobs with the given alias. Unlike refs, aliases do not serve as GC roots.
@@ -98,7 +107,7 @@ namespace EpicGames.Horde.Storage
 		/// <param name="cacheTime">Minimum coherency for any cached value to be returned</param>
 		/// <param name="cancellationToken">Cancellation token for the operation</param>
 		/// <returns>Blob pointed to by the ref</returns>
-		Task<IBlobRef?> TryReadRefAsync(RefName name, RefCacheTime cacheTime = default, CancellationToken cancellationToken = default);
+		Task<IHashedBlobRef?> TryReadRefAsync(RefName name, RefCacheTime cacheTime = default, CancellationToken cancellationToken = default);
 
 		/// <summary>
 		/// Writes a new ref to the store
@@ -108,7 +117,7 @@ namespace EpicGames.Horde.Storage
 		/// <param name="options">Options for the new ref</param>
 		/// <param name="cancellationToken">Cancellation token for the operation</param>
 		/// <returns>Unique identifier for the blob</returns>
-		Task WriteRefAsync(RefName name, IBlobRef target, RefOptions? options = null, CancellationToken cancellationToken = default);
+		Task WriteRefAsync(RefName name, IHashedBlobRef target, RefOptions? options = null, CancellationToken cancellationToken = default);
 
 		/// <summary>
 		/// Reads data for a ref from the store
@@ -251,13 +260,13 @@ namespace EpicGames.Horde.Storage
 		/// <summary>
 		/// Create a blob ref from a RefValue
 		/// </summary>
-		public static IBlobRef CreateBlobRef(this IStorageClient store, BlobRefValue refValue)
+		public static IHashedBlobRef CreateBlobRef(this IStorageClient store, HashedBlobRefValue refValue)
 			=> store.CreateBlobRef(refValue.Hash, refValue.Locator);
 
 		/// <summary>
 		/// Create a typed blob ref from a RefValue
 		/// </summary>
-		public static IBlobRef<T> CreateBlobRef<T>(this IStorageClient store, BlobRefValue refValue, BlobSerializerOptions? options)
+		public static IHashedBlobRef<T> CreateBlobRef<T>(this IStorageClient store, HashedBlobRefValue refValue, BlobSerializerOptions? options)
 			=> store.CreateBlobRef<T>(refValue.Hash, refValue.Locator, options);
 
 		/// <summary>
@@ -298,7 +307,7 @@ namespace EpicGames.Horde.Storage
 		/// <returns>True if the ref exists, false if it did not exist</returns>
 		public static async Task<bool> RefExistsAsync(this IStorageClient store, RefName name, RefCacheTime cacheTime = default, CancellationToken cancellationToken = default)
 		{
-			IBlobHandle? target = await store.TryReadRefAsync(name, cacheTime, cancellationToken);
+			IBlobRef? target = await store.TryReadRefAsync(name, cacheTime, cancellationToken);
 			return target != null;
 		}
 
@@ -310,9 +319,9 @@ namespace EpicGames.Horde.Storage
 		/// <param name="cacheTime">Minimum coherency of any cached result</param>
 		/// <param name="cancellationToken">Cancellation token for the operation</param>
 		/// <returns>The ref target</returns>
-		public static async Task<IBlobRef> ReadRefAsync(this IStorageClient store, RefName name, RefCacheTime cacheTime = default, CancellationToken cancellationToken = default)
+		public static async Task<IHashedBlobRef> ReadRefAsync(this IStorageClient store, RefName name, RefCacheTime cacheTime = default, CancellationToken cancellationToken = default)
 		{
-			IBlobRef? refTarget = await store.TryReadRefAsync(name, cacheTime, cancellationToken);
+			IHashedBlobRef? refTarget = await store.TryReadRefAsync(name, cacheTime, cancellationToken);
 			return refTarget ?? throw new RefNameNotFoundException(name);
 		}
 
diff --git a/Engine/Source/Programs/Shared/EpicGames.Horde/Storage/Nodes/CbNode.cs b/Engine/Source/Programs/Shared/EpicGames.Horde/Storage/Nodes/CbNode.cs
--- a/Engine/Source/Programs/Shared/EpicGames.Horde/Storage/Nodes/CbNode.cs
+++ b/Engine/Source/Programs/Shared/EpicGames.Horde/Storage/Nodes/CbNode.cs
@@ -27,14 +27,14 @@ namespace EpicGames.Horde.Storage.Nodes
 		/// <summary>
 		/// Imported nodes
 		/// </summary>
-		public IReadOnlyList<IBlobHandle> Imports { get; }
+		public IReadOnlyList<IBlobRef> Imports { get; }
 
 		/// <summary>
 		/// Constructor
 		/// </summary>
 		/// <param name="obj">The compact binary object</param>
 		/// <param name="imports">List of imports for attachments</param>
-		public CbNode(CbObject obj, IReadOnlyList<IBlobHandle> imports)
+		public CbNode(CbObject obj, IReadOnlyList<IBlobRef> imports)
 		{
 			Object = obj;
 			Imports = imports;
diff --git a/Engine/Source/Programs/Shared/EpicGames.Horde/Storage/Nodes/ChunkedDataNode.cs b/Engine/Source/Programs/Shared/EpicGames.Horde/Storage/Nodes/ChunkedDataNode.cs
--- a/Engine/Source/Programs/Shared/EpicGames.Horde/Storage/Nodes/ChunkedDataNode.cs
+++ b/Engine/Source/Programs/Shared/EpicGames.Horde/Storage/Nodes/ChunkedDataNode.cs
@@ -34,7 +34,7 @@ namespace EpicGames.Horde.Storage.Nodes
 		/// <param name="handle">Handle to the data to read</param>
 		/// <param name="outputStream">The output stream to receive the data</param>
 		/// <param name="cancellationToken">Cancellation token for the operation</param>
-		public static async Task CopyToStreamAsync(IBlobHandle handle, Stream outputStream, CancellationToken cancellationToken)
+		public static async Task CopyToStreamAsync(IBlobRef handle, Stream outputStream, CancellationToken cancellationToken)
 		{
 			using BlobData blobData = await handle.ReadBlobDataAsync(cancellationToken);
 			if (blobData.Type.Guid == LeafChunkedDataNode.BlobTypeGuid)
@@ -58,7 +58,7 @@ namespace EpicGames.Horde.Storage.Nodes
 		/// <param name="file">File to write with the contents of this node</param>
 		/// <param name="cancellationToken">Cancellation token for the operation</param>
 		/// <returns></returns>
-		public static async Task CopyToFileAsync(IBlobHandle handle, FileInfo file, CancellationToken cancellationToken)
+		public static async Task CopyToFileAsync(IBlobRef handle, FileInfo file, CancellationToken cancellationToken)
 		{
 			if (file.Exists && (file.Attributes & FileAttributes.ReadOnly) != 0)
 			{
@@ -151,12 +151,12 @@ namespace EpicGames.Horde.Storage.Nodes
 	/// <param name="Length">Length of the data stream within this node</param>
 	/// <param name="RollingHash">Rolling hash for this chunk. Only serialized for leaf node references.</param>
 	/// <param name="Handle">Handle to the target node</param>
-	public record class ChunkedDataNodeRef(ChunkedDataNodeType Type, long Length, uint RollingHash, IBlobRef<ChunkedDataNode> Handle)
+	public record class ChunkedDataNodeRef(ChunkedDataNodeType Type, long Length, uint RollingHash, IHashedBlobRef<ChunkedDataNode> Handle)
 	{
 		/// <summary>
 		/// Leaf node constructor
 		/// </summary>
-		public ChunkedDataNodeRef(long length, uint rollingHash, IBlobRef<LeafChunkedDataNode> handle)
+		public ChunkedDataNodeRef(long length, uint rollingHash, IHashedBlobRef<LeafChunkedDataNode> handle)
 			: this(ChunkedDataNodeType.Leaf, length, rollingHash, handle)
 		{
 		}
@@ -164,7 +164,7 @@ namespace EpicGames.Horde.Storage.Nodes
 		/// <summary>
 		/// Interior node constructor
 		/// </summary>
-		public ChunkedDataNodeRef(long length, IBlobRef<InteriorChunkedDataNode> handle)
+		public ChunkedDataNodeRef(long length, IHashedBlobRef<InteriorChunkedDataNode> handle)
 			: this(ChunkedDataNodeType.Interior, length, 0, handle)
 		{
 		}
@@ -172,25 +172,25 @@ namespace EpicGames.Horde.Storage.Nodes
 		/// <summary>
 		/// Gets the target interior node handle
 		/// </summary>
-		public IBlobRef<InteriorChunkedDataNode> GetInteriorHandle()
+		public IHashedBlobRef<InteriorChunkedDataNode> GetInteriorHandle()
 		{
 			if (Type != ChunkedDataNodeType.Interior)
 			{
 				throw new InvalidOperationException("Node is not an interior node");
 			}
-			return BlobRef.Create<InteriorChunkedDataNode>(Handle.Hash, Handle, Handle.SerializerOptions);
+			return HashedBlobRef.Create<InteriorChunkedDataNode>(Handle.Hash, Handle, Handle.SerializerOptions);
 		}
 
 		/// <summary>
 		/// Gets the target leaf node handle
 		/// </summary>
-		public IBlobRef<LeafChunkedDataNode> GetLeafHandle()
+		public IHashedBlobRef<LeafChunkedDataNode> GetLeafHandle()
 		{
 			if (Type != ChunkedDataNodeType.Leaf)
 			{
 				throw new InvalidOperationException("Node is not a leaf node");
 			}
-			return BlobRef.Create<LeafChunkedDataNode>(Handle.Hash, Handle, Handle.SerializerOptions);
+			return HashedBlobRef.Create<LeafChunkedDataNode>(Handle.Hash, Handle, Handle.SerializerOptions);
 		}
 
 		/// <summary>
@@ -347,7 +347,7 @@ namespace EpicGames.Horde.Storage.Nodes
 					sizeSinceProgressUpdate = 0;
 				}
 
-				IBlobRef<LeafChunkedDataNode> blobHandle = await writer.CompleteAsync<LeafChunkedDataNode>(LeafChunkedDataNodeConverter.BlobType, cancellationToken);
+				IHashedBlobRef<LeafChunkedDataNode> blobHandle = await writer.CompleteAsync<LeafChunkedDataNode>(LeafChunkedDataNodeConverter.BlobType, cancellationToken);
 				leafNodeRefs.Add(new ChunkedDataNodeRef(nextLength, rollingHash, blobHandle));
 
 				readBuffer.Memory.Slice(nextLength, size - nextLength).CopyTo(readBuffer.Memory);
@@ -478,7 +478,7 @@ namespace EpicGames.Horde.Storage.Nodes
 				handleBuffer.Clear();
 				foreach ((InteriorChunkedDataNode interiorNode, long interiorLength) in interiorNodes)
 				{
-					IBlobRef<InteriorChunkedDataNode> interiorHandle = await memoryWriter.WriteBlobAsync(interiorNode, cancellationToken);
+					IHashedBlobRef<InteriorChunkedDataNode> interiorHandle = await memoryWriter.WriteBlobAsync(interiorNode, cancellationToken);
 					handleBuffer.Add(new ChunkedDataNodeRef(interiorLength, interiorHandle));
 				}
 
@@ -504,7 +504,7 @@ namespace EpicGames.Horde.Storage.Nodes
 			}
 
 			InteriorChunkedDataNode targetNode = new InteriorChunkedDataNode(children);
-			IBlobRef<InteriorChunkedDataNode> targetHandle = await writer.WriteBlobAsync(targetNode, cancellationToken);
+			IHashedBlobRef<InteriorChunkedDataNode> targetHandle = await writer.WriteBlobAsync(targetNode, cancellationToken);
 
 			return new ChunkedDataNodeRef(source.Length, targetHandle);
 		}
@@ -565,7 +565,7 @@ namespace EpicGames.Horde.Storage.Nodes
 			BlobReader nodeReader = new BlobReader(nodeData, null);
 			while (nodeReader.GetMemory(0).Length > 0)
 			{
-				IBlobRef<ChunkedDataNode> handle = nodeReader.ReadBlobRef<ChunkedDataNode>();
+				IHashedBlobRef<ChunkedDataNode> handle = nodeReader.ReadBlobRef<ChunkedDataNode>();
 
 				ChunkedDataNodeType type = ChunkedDataNodeType.Unknown;
 				if (nodeReader.Version >= 2)
@@ -616,7 +616,7 @@ namespace EpicGames.Horde.Storage.Nodes
 			List<ChunkedDataNodeRef> children = new List<ChunkedDataNodeRef>();
 			while (reader.GetMemory().Length > 0)
 			{
-				IBlobRef<ChunkedDataNode> handle = reader.ReadBlobRef<ChunkedDataNode>();
+				IHashedBlobRef<ChunkedDataNode> handle = reader.ReadBlobRef<ChunkedDataNode>();
 
 				ChunkedDataNodeType type = ChunkedDataNodeType.Unknown;
 				if (reader.Version >= 2) // Pre sync with HordeApiVersion
diff --git a/Engine/Source/Programs/Shared/EpicGames.Horde/Storage/Nodes/ChunkedDataWriter.cs b/Engine/Source/Programs/Shared/EpicGames.Horde/Storage/Nodes/ChunkedDataWriter.cs
--- a/Engine/Source/Programs/Shared/EpicGames.Horde/Storage/Nodes/ChunkedDataWriter.cs
+++ b/Engine/Source/Programs/Shared/EpicGames.Horde/Storage/Nodes/ChunkedDataWriter.cs
@@ -338,7 +338,7 @@ namespace EpicGames.Horde.Storage.Nodes
 		async ValueTask FlushLeafNodeAsync(CancellationToken cancellationToken)
 		{
 			int leafLength = _writer.WrittenMemory.Length;
-			IBlobRef<LeafChunkedDataNode> leafHandle = await _writer.CompleteAsync<LeafChunkedDataNode>(LeafChunkedDataNodeConverter.BlobType, cancellationToken);
+			IHashedBlobRef<LeafChunkedDataNode> leafHandle = await _writer.CompleteAsync<LeafChunkedDataNode>(LeafChunkedDataNodeConverter.BlobType, cancellationToken);
 			_leafHandles.Add(new ChunkedDataNodeRef(leafLength, _leafHash, leafHandle));
 			ResetLeafState();
 		}
diff --git a/Engine/Source/Programs/Shared/EpicGames.Horde/Storage/Nodes/CommitNode.cs b/Engine/Source/Programs/Shared/EpicGames.Horde/Storage/Nodes/CommitNode.cs
--- a/Engine/Source/Programs/Shared/EpicGames.Horde/Storage/Nodes/CommitNode.cs
+++ b/Engine/Source/Programs/Shared/EpicGames.Horde/Storage/Nodes/CommitNode.cs
@@ -25,7 +25,7 @@ namespace EpicGames.Horde.Storage.Nodes
 		/// <summary>
 		/// Reference to the parent commit
 		/// </summary>
-		public IBlobRef<CommitNode>? Parent { get; set; }
+		public IHashedBlobRef<CommitNode>? Parent { get; set; }
 
 		/// <summary>
 		/// Human readable name of the author of this change
@@ -65,12 +65,12 @@ namespace EpicGames.Horde.Storage.Nodes
 		/// <summary>
 		/// Metadata for this commit, keyed by arbitrary GUID
 		/// </summary>
-		public Dictionary<Guid, IBlobRef> Metadata { get; } = new Dictionary<Guid, IBlobRef>();
+		public Dictionary<Guid, IHashedBlobRef> Metadata { get; } = new Dictionary<Guid, IHashedBlobRef>();
 
 		/// <summary>
 		/// Constructor
 		/// </summary>
-		public CommitNode(int number, IBlobRef<CommitNode>? parent, string author, string? authorId, string? committer, string? commiterId, string message, DateTime time, DirectoryNodeRef contents, Dictionary<Guid, IBlobRef> metadata)
+		public CommitNode(int number, IHashedBlobRef<CommitNode>? parent, string author, string? authorId, string? committer, string? commiterId, string message, DateTime time, DirectoryNodeRef contents, Dictionary<Guid, IHashedBlobRef> metadata)
 		{
 			Number = number;
 			Parent = parent;
@@ -93,7 +93,7 @@ namespace EpicGames.Horde.Storage.Nodes
 		{
 			int number = (int)reader.ReadUnsignedVarInt();
 
-			IBlobRef<CommitNode>? parent;
+			IHashedBlobRef<CommitNode>? parent;
 			if (reader.ReadBoolean())
 			{
 				parent = reader.ReadBlobRef<CommitNode>();
@@ -110,11 +110,11 @@ namespace EpicGames.Horde.Storage.Nodes
 			string message = reader.ReadString();
 			DateTime time = reader.ReadDateTime();
 
-			IBlobRef<DirectoryNode> contentsNode = reader.ReadBlobRef<DirectoryNode>();
+			IHashedBlobRef<DirectoryNode> contentsNode = reader.ReadBlobRef<DirectoryNode>();
 			long length = (long)reader.ReadUnsignedVarInt();
 			DirectoryNodeRef contents = new DirectoryNodeRef(length, contentsNode);
 
-			Dictionary<Guid, IBlobRef> metadata = reader.ReadDictionary(() => reader.ReadGuidUnrealOrder(), () => reader.ReadBlobRef());
+			Dictionary<Guid, IHashedBlobRef> metadata = reader.ReadDictionary(() => reader.ReadGuidUnrealOrder(), () => reader.ReadBlobRef());
 
 			return new CommitNode(number, parent, author, authorId, committer, committerId, message, time, contents, metadata);
 		}
diff --git a/Engine/Source/Programs/Shared/EpicGames.Horde/Storage/Nodes/DdcRefNode.cs b/Engine/Source/Programs/Shared/EpicGames.Horde/Storage/Nodes/DdcRefNode.cs
--- a/Engine/Source/Programs/Shared/EpicGames.Horde/Storage/Nodes/DdcRefNode.cs
+++ b/Engine/Source/Programs/Shared/EpicGames.Horde/Storage/Nodes/DdcRefNode.cs
@@ -25,7 +25,7 @@ namespace EpicGames.Horde.Storage.Nodes
 		/// <summary>
 		/// References to attachments. We embed this in the ref node to ensure any aliased blobs have a hard reference from the root.
 		/// </summary>
-		public List<IBlobRef> References { get; }
+		public List<IHashedBlobRef> References { get; }
 
 		/// <summary>
 		/// Constructor
@@ -33,16 +33,16 @@ namespace EpicGames.Horde.Storage.Nodes
 		public DdcRefNode(IoHash rootHash)
 		{
 			RootHash = rootHash;
-			References = new List<IBlobRef>();
+			References = new List<IHashedBlobRef>();
 		}
 
 		/// <summary>
 		/// Constructor
 		/// </summary>
-		public DdcRefNode(IoHash rootHash, IEnumerable<IBlobRef> references)
+		public DdcRefNode(IoHash rootHash, IEnumerable<IHashedBlobRef> references)
 		{
 			RootHash = rootHash;
-			References = new List<IBlobRef>(references);
+			References = new List<IHashedBlobRef>(references);
 		}
 	}
 
@@ -53,7 +53,7 @@ namespace EpicGames.Horde.Storage.Nodes
 		public override DdcRefNode Read(IBlobReader reader, BlobSerializerOptions options)
 		{
 			IoHash rootHash = reader.ReadIoHash();
-			List<IBlobRef> references = reader.ReadList(x => reader.ReadBlobRef());
+			List<IHashedBlobRef> references = reader.ReadList(x => reader.ReadBlobRef());
 			return new DdcRefNode(rootHash, references);
 		}
 
diff --git a/Engine/Source/Programs/Shared/EpicGames.Horde/Storage/Nodes/DirectoryNode.Extract.cs b/Engine/Source/Programs/Shared/EpicGames.Horde/Storage/Nodes/DirectoryNode.Extract.cs
--- a/Engine/Source/Programs/Shared/EpicGames.Horde/Storage/Nodes/DirectoryNode.Extract.cs
+++ b/Engine/Source/Programs/Shared/EpicGames.Horde/Storage/Nodes/DirectoryNode.Extract.cs
@@ -227,7 +227,7 @@ namespace EpicGames.Horde.Storage.Nodes
 			}
 		}
 
-		record class OutputChunk(OutputFile File, long Offset, long Length, IBlobHandle Handle);
+		record class OutputChunk(OutputFile File, long Offset, long Length, IBlobRef Handle);
 
 		record class OutputBatch(List<OutputChunk> Chunks);
 
@@ -246,7 +246,7 @@ namespace EpicGames.Horde.Storage.Nodes
 				_chunkWriter = chunkWriter;
 			}
 
-			public async Task WriteAsync(long offset, long length, IBlobHandle handle, CancellationToken cancellationToken)
+			public async Task WriteAsync(long offset, long length, IBlobRef handle, CancellationToken cancellationToken)
 			{
 				if (_bufferedChunk != null)
 				{
@@ -268,7 +268,7 @@ namespace EpicGames.Horde.Storage.Nodes
 		}
 
 #pragma warning disable IDE0060
-		static void TraceBlobRead(string type, string path, IBlobHandle handle, ILogger logger)
+		static void TraceBlobRead(string type, string path, IBlobRef handle, ILogger logger)
 		{
 			//			logger.LogTrace(KnownLogEvents.Horde_BlobRead, "Blob [{Type,-20}] Path=\"{Path}\", Locator={Locator}", type, path, handle.GetLocator());
 		}
@@ -438,7 +438,7 @@ namespace EpicGames.Horde.Storage.Nodes
 		/// <summary>
 		/// Compare two blob handles to allow sorting for optimized reads
 		/// </summary>
-		public static int CompareBlobs(IBlobHandle a, IBlobHandle b)
+		public static int CompareBlobs(IBlobRef a, IBlobRef b)
 		{
 			if (a.Innermost is not ExportHandle exportHandleA || exportHandleA.Packet is not FlushedPacketHandle packetHandleA)
 			{
diff --git a/Engine/Source/Programs/Shared/EpicGames.Horde/Storage/Nodes/DirectoryNode.Update.cs b/Engine/Source/Programs/Shared/EpicGames.Horde/Storage/Nodes/DirectoryNode.Update.cs
--- a/Engine/Source/Programs/Shared/EpicGames.Horde/Storage/Nodes/DirectoryNode.Update.cs
+++ b/Engine/Source/Programs/Shared/EpicGames.Horde/Storage/Nodes/DirectoryNode.Update.cs
@@ -350,7 +350,7 @@ namespace EpicGames.Horde.Storage.Nodes
 		/// <summary>
 		/// Writes a tree of files to a storage writer
 		/// </summary>
-		public static async Task<IBlobRef<DirectoryNode>> WriteFilesAsync(this IBlobWriter writer, DirectoryReference baseDir, ChunkingOptions? options = null, IProgress<IUpdateStats>? progress = null, CancellationToken cancellationToken = default)
+		public static async Task<IHashedBlobRef<DirectoryNode>> WriteFilesAsync(this IBlobWriter writer, DirectoryReference baseDir, ChunkingOptions? options = null, IProgress<IUpdateStats>? progress = null, CancellationToken cancellationToken = default)
 		{
 			DirectoryNode outputNode = new DirectoryNode();
 			await outputNode.AddFilesAsync(baseDir, DirectoryReference.EnumerateFiles(baseDir, "*", SearchOption.AllDirectories), writer, options, progress, cancellationToken);
@@ -360,7 +360,7 @@ namespace EpicGames.Horde.Storage.Nodes
 		/// <summary>
 		/// Writes a tree of files to a storage writer
 		/// </summary>
-		public static async Task<IBlobRef<DirectoryNode>> WriteFilesAsync(this IBlobWriter writer, DirectoryInfo baseDir, IReadOnlyList<FileInfo> files, ChunkingOptions? options = null, IProgress<IUpdateStats>? progress = null, CancellationToken cancellationToken = default)
+		public static async Task<IHashedBlobRef<DirectoryNode>> WriteFilesAsync(this IBlobWriter writer, DirectoryInfo baseDir, IReadOnlyList<FileInfo> files, ChunkingOptions? options = null, IProgress<IUpdateStats>? progress = null, CancellationToken cancellationToken = default)
 		{
 			DirectoryNode outputNode = new DirectoryNode();
 			await outputNode.AddFilesAsync(baseDir, files, writer, options, progress, cancellationToken);
@@ -370,7 +370,7 @@ namespace EpicGames.Horde.Storage.Nodes
 		/// <summary>
 		/// Writes a tree of files to a storage writer
 		/// </summary>
-		public static async Task<IBlobRef<DirectoryNode>> WriteFilesAsync(this IBlobWriter writer, DirectoryReference baseDir, IReadOnlyList<FileReference> files, ChunkingOptions? options = null, IProgress<IUpdateStats>? progress = null, CancellationToken cancellationToken = default)
+		public static async Task<IHashedBlobRef<DirectoryNode>> WriteFilesAsync(this IBlobWriter writer, DirectoryReference baseDir, IReadOnlyList<FileReference> files, ChunkingOptions? options = null, IProgress<IUpdateStats>? progress = null, CancellationToken cancellationToken = default)
 		{
 			DirectoryNode outputNode = new DirectoryNode();
 			await outputNode.AddFilesAsync(baseDir, files, writer, options, progress, cancellationToken);
@@ -429,7 +429,7 @@ namespace EpicGames.Horde.Storage.Nodes
 					DirectoryNode? childNode = await directoryNode.TryOpenDirectoryAsync(name, cancellationToken);
 					childNode ??= new DirectoryNode();
 					await childNode.UpdateAsync(directory, writer, cancellationToken);
-					IBlobRef<DirectoryNode> handle = await writer.WriteBlobAsync<DirectoryNode>(childNode, cancellationToken);
+					IHashedBlobRef<DirectoryNode> handle = await writer.WriteBlobAsync<DirectoryNode>(childNode, cancellationToken);
 					directoryNode.DeleteDirectory(name);
 					directoryNode.AddDirectory(new DirectoryEntry(name, childNode.Length, handle));
 				}
diff --git a/Engine/Source/Programs/Shared/EpicGames.Horde/Storage/Nodes/DirectoryNode.ZipStream.cs b/Engine/Source/Programs/Shared/EpicGames.Horde/Storage/Nodes/DirectoryNode.ZipStream.cs
--- a/Engine/Source/Programs/Shared/EpicGames.Horde/Storage/Nodes/DirectoryNode.ZipStream.cs
+++ b/Engine/Source/Programs/Shared/EpicGames.Horde/Storage/Nodes/DirectoryNode.ZipStream.cs
@@ -45,7 +45,7 @@ namespace EpicGames.Horde.Storage.Nodes
 		/// <param name="node">Root node to copy from</param>
 		/// <param name="filter">Filter for files to include in the zip</param>
 		/// <param name="logger">Optional logger for debug tracing</param>
-		public DirectoryNodeZipStream(IBlobRef<DirectoryNode> node, FileFilter? filter, ILogger? logger)
+		public DirectoryNodeZipStream(IHashedBlobRef<DirectoryNode> node, FileFilter? filter, ILogger? logger)
 		{
 			_pipe = new Pipe();
 			_backgroundTask = BackgroundTask.StartNew(ctx => CopyToPipeAsync(node, filter, _pipe.Writer, logger, ctx));
@@ -109,14 +109,14 @@ namespace EpicGames.Horde.Storage.Nodes
 			return length;
 		}
 
-		static async Task CopyToPipeAsync(IBlobRef<DirectoryNode> node, FileFilter? filter, PipeWriter writer, ILogger? logger, CancellationToken cancellationToken)
+		static async Task CopyToPipeAsync(IHashedBlobRef<DirectoryNode> node, FileFilter? filter, PipeWriter writer, ILogger? logger, CancellationToken cancellationToken)
 		{
 			using Stream outputStream = writer.AsStream();
 			using ZipArchive archive = new ZipArchive(outputStream, ZipArchiveMode.Create);
 			await CopyFilesAsync(node, "", filter, archive, logger, cancellationToken);
 		}
 
-		static async Task CopyFilesAsync(IBlobRef<DirectoryNode> directoryRef, string prefix, FileFilter? filter, ZipArchive archive, ILogger? logger, CancellationToken cancellationToken)
+		static async Task CopyFilesAsync(IHashedBlobRef<DirectoryNode> directoryRef, string prefix, FileFilter? filter, ZipArchive archive, ILogger? logger, CancellationToken cancellationToken)
 		{
 			DirectoryNode directory = await directoryRef.ReadBlobAsync(cancellationToken);
 
@@ -192,7 +192,7 @@ namespace EpicGames.Horde.Storage.Nodes
 		/// <param name="filter">Filter for files to include in the zip</param>
 		/// <param name="logger">Logger for diagnostic output</param>
 		/// <returns>Stream containing zipped archive data</returns>
-		public static Stream AsZipStream(this IBlobRef<DirectoryNode> directoryRef, FileFilter? filter = null, ILogger? logger = null)
+		public static Stream AsZipStream(this IHashedBlobRef<DirectoryNode> directoryRef, FileFilter? filter = null, ILogger? logger = null)
 			=> new DirectoryNodeZipStream(directoryRef, filter, logger);
 	}
 }
diff --git a/Engine/Source/Programs/Shared/EpicGames.Horde/Storage/Nodes/DirectoryNode.cs b/Engine/Source/Programs/Shared/EpicGames.Horde/Storage/Nodes/DirectoryNode.cs
--- a/Engine/Source/Programs/Shared/EpicGames.Horde/Storage/Nodes/DirectoryNode.cs
+++ b/Engine/Source/Programs/Shared/EpicGames.Horde/Storage/Nodes/DirectoryNode.cs
@@ -364,7 +364,7 @@ namespace EpicGames.Horde.Storage.Nodes
 			int fileCount = (int)reader.ReadUnsignedVarInt();
 			for (int idx = 0; idx < fileCount; idx++)
 			{
-				IBlobRef<ChunkedDataNode> targetHandle = reader.ReadBlobRef<ChunkedDataNode>();
+				IHashedBlobRef<ChunkedDataNode> targetHandle = reader.ReadBlobRef<ChunkedDataNode>();
 
 				ChunkedDataNodeType targetType = ChunkedDataNodeType.Unknown;
 				if (reader.Version >= 2) // Pre-unification with HordeApiVersion
@@ -402,7 +402,7 @@ namespace EpicGames.Horde.Storage.Nodes
 			int directoryCount = (int)reader.ReadUnsignedVarInt();
 			for (int idx = 0; idx < directoryCount; idx++)
 			{
-				IBlobRef<DirectoryNode> directoryHandle = reader.ReadBlobRef<DirectoryNode>();
+				IHashedBlobRef<DirectoryNode> directoryHandle = reader.ReadBlobRef<DirectoryNode>();
 				long length = (long)reader.ReadUnsignedVarInt();
 				string name = reader.ReadString();
 
@@ -489,13 +489,13 @@ namespace EpicGames.Horde.Storage.Nodes
 	/// </summary>
 	/// <param name="Length">Sum total of all the file lengths in this directory tree</param>
 	/// <param name="Handle">Handle to the target node</param>
-	public record class DirectoryNodeRef(long Length, IBlobRef<DirectoryNode> Handle);
+	public record class DirectoryNodeRef(long Length, IHashedBlobRef<DirectoryNode> Handle);
 
 	/// <summary>
 	/// Entry for a directory within a directory node
 	/// </summary>
 	[DebuggerDisplay("{Name}")]
-	public record class DirectoryEntry(string Name, long Length, IBlobRef<DirectoryNode> Handle) : DirectoryNodeRef(Length, Handle)
+	public record class DirectoryEntry(string Name, long Length, IHashedBlobRef<DirectoryNode> Handle) : DirectoryNodeRef(Length, Handle)
 	{
 	}
 }
diff --git a/Engine/Source/Programs/Shared/EpicGames.Horde/Storage/Nodes/RedirectNode.cs b/Engine/Source/Programs/Shared/EpicGames.Horde/Storage/Nodes/RedirectNode.cs
--- a/Engine/Source/Programs/Shared/EpicGames.Horde/Storage/Nodes/RedirectNode.cs
+++ b/Engine/Source/Programs/Shared/EpicGames.Horde/Storage/Nodes/RedirectNode.cs
@@ -24,13 +24,13 @@ namespace EpicGames.Horde.Storage.Nodes
 		/// <summary>
 		/// The target handle
 		/// </summary>
-		public IBlobRef<T> Target { get; }
+		public IHashedBlobRef<T> Target { get; }
 
 		/// <summary>
 		/// Constructor
 		/// </summary>
 		/// <param name="handle">Target node for the redirect</param>
-		public RedirectNode(IBlobRef<T> handle) => Target = handle;
+		public RedirectNode(IHashedBlobRef<T> handle) => Target = handle;
 	}
 
 	class RedirectNodeConverter<T> : BlobConverter<RedirectNode<T>>
@@ -40,7 +40,7 @@ namespace EpicGames.Horde.Storage.Nodes
 		/// <inheritdoc/>
 		public override RedirectNode<T> Read(IBlobReader reader, BlobSerializerOptions options)
 		{
-			IBlobRef<T> handle = reader.ReadBlobRef<T>();
+			IHashedBlobRef<T> handle = reader.ReadBlobRef<T>();
 			return new RedirectNode<T>(handle);
 		}
 
diff --git a/Engine/Source/Programs/Shared/EpicGames.Horde/Storage/StorageBackendCache.cs b/Engine/Source/Programs/Shared/EpicGames.Horde/Storage/StorageBackendCache.cs
--- a/Engine/Source/Programs/Shared/EpicGames.Horde/Storage/StorageBackendCache.cs
+++ b/Engine/Source/Programs/Shared/EpicGames.Horde/Storage/StorageBackendCache.cs
@@ -160,10 +160,10 @@ namespace EpicGames.Horde.Storage
 
 			#region Refs
 
-			public Task<BlobRefValue?> TryReadRefAsync(RefName name, RefCacheTime cacheTime = default, CancellationToken cancellationToken = default)
+			public Task<HashedBlobRefValue?> TryReadRefAsync(RefName name, RefCacheTime cacheTime = default, CancellationToken cancellationToken = default)
 				=> _inner.TryReadRefAsync(name, cacheTime, cancellationToken);
 
-			public Task WriteRefAsync(RefName name, BlobRefValue value, RefOptions? options = null, CancellationToken cancellationToken = default)
+			public Task WriteRefAsync(RefName name, HashedBlobRefValue value, RefOptions? options = null, CancellationToken cancellationToken = default)
 				=> _inner.WriteRefAsync(name, value, options, cancellationToken);
 
 			public Task<bool> DeleteRefAsync(RefName name, CancellationToken cancellationToken = default)
diff --git a/Engine/Source/Programs/Shared/EpicGames.Horde/Tools/ToolMessages.cs b/Engine/Source/Programs/Shared/EpicGames.Horde/Tools/ToolMessages.cs
--- a/Engine/Source/Programs/Shared/EpicGames.Horde/Tools/ToolMessages.cs
+++ b/Engine/Source/Programs/Shared/EpicGames.Horde/Tools/ToolMessages.cs
@@ -65,7 +65,7 @@ namespace EpicGames.Horde.Tools
 	/// <param name="Duration">Number of minutes over which to do the deployment</param>
 	/// <param name="CreatePaused">Whether to create the deployment in a paused state</param>
 	/// <param name="Content">Handle to a directory node with the content for the deployment</param>
-	public record CreateToolDeploymentRequest(string Version, double? Duration, bool? CreatePaused, BlobRefValue Content);
+	public record CreateToolDeploymentRequest(string Version, double? Duration, bool? CreatePaused, HashedBlobRefValue Content);
 
 	/// <summary>
 	/// Response from creating a deployment
diff --git a/Engine/Source/Programs/UnrealBuildTool/Artifacts/HordeStorageArtifactCache.cs b/Engine/Source/Programs/UnrealBuildTool/Artifacts/HordeStorageArtifactCache.cs
--- a/Engine/Source/Programs/UnrealBuildTool/Artifacts/HordeStorageArtifactCache.cs
+++ b/Engine/Source/Programs/UnrealBuildTool/Artifacts/HordeStorageArtifactCache.cs
@@ -29,7 +29,7 @@ namespace UnrealBuildTool.Artifacts
 		/// Collection of output file references.  There should be exactly the same number
 		/// of file references as outputs in the action
 		/// </summary>
-		public readonly IBlobRef<ChunkedDataNode>[] OutputRefs;
+		public readonly IHashedBlobRef<ChunkedDataNode>[] OutputRefs;
 
 		/// <summary>
 		/// Construct a new horde artifact number
@@ -39,7 +39,7 @@ namespace UnrealBuildTool.Artifacts
 		public HordeArtifactAction(ArtifactAction artifactAction)
 		{
 			ArtifactAction = artifactAction;
-			OutputRefs = new IBlobRef<ChunkedDataNode>[ArtifactAction.Outputs.Length];
+			OutputRefs = new IHashedBlobRef<ChunkedDataNode>[ArtifactAction.Outputs.Length];
 		}
 
 		/// <summary>
@@ -309,7 +309,7 @@ namespace UnrealBuildTool.Artifacts
 						output[index] = true;
 
 						int refIndex = 0;
-						foreach (IBlobRef<ChunkedDataNode> artifactRef in hordeArtifactAction.OutputRefs)
+						foreach (IHashedBlobRef<ChunkedDataNode> artifactRef in hordeArtifactAction.OutputRefs)
 						{
 							if (artifactRef == null)
 							{
@@ -478,7 +478,7 @@ namespace UnrealBuildTool.Artifacts
 					// Save the artifact action file
 					await using IBlobWriter writer = _store!.CreateBlobWriter();
 					await hordeArtifactAction.WriteFilesAsync(writer, cancellationToken);
-					IBlobRef<ArtifactActionCollectionNode> nodeRef = await writer.WriteBlobAsync(node);
+					IHashedBlobRef<ArtifactActionCollectionNode> nodeRef = await writer.WriteBlobAsync(node);
 					await writer.FlushAsync();
 
 					// Save the collection
diff --git a/Engine/Source/Programs/UnrealBuildTool/Executors/UnrealBuildAccelerator/UBAAgentCoordinatorHorde.cs b/Engine/Source/Programs/UnrealBuildTool/Executors/UnrealBuildAccelerator/UBAAgentCoordinatorHorde.cs
--- a/Engine/Source/Programs/UnrealBuildTool/Executors/UnrealBuildAccelerator/UBAAgentCoordinatorHorde.cs
+++ b/Engine/Source/Programs/UnrealBuildTool/Executors/UnrealBuildAccelerator/UBAAgentCoordinatorHorde.cs
@@ -181,7 +181,7 @@ namespace UnrealBuildTool
 			await using IBlobWriter writer = _storage.CreateBlobWriter(serializerOptions: serializerOptions);
 			DirectoryNode sandbox = new();
 			await sandbox.AddFilesAsync(baseDir, files, writer, cancellationToken: cancellationToken);
-			IBlobRef<DirectoryNode> handle = await writer.WriteBlobAsync(sandbox, cancellationToken);
+			IHashedBlobRef<DirectoryNode> handle = await writer.WriteBlobAsync(sandbox, cancellationToken);
 			await writer.FlushAsync(cancellationToken);
 			return handle.GetLocator();
 		}
diff --git a/Engine/Source/Programs/UnrealGameSync/UnrealGameSync/Launcher.cs b/Engine/Source/Programs/UnrealGameSync/UnrealGameSync/Launcher.cs
--- a/Engine/Source/Programs/UnrealGameSync/UnrealGameSync/Launcher.cs
+++ b/Engine/Source/Programs/UnrealGameSync/UnrealGameSync/Launcher.cs
@@ -187,7 +187,7 @@ namespace UnrealGameSync
 
 						// Download and extract the data
 						using IStorageClient storageClient = hordeClient.CreateStorageClient(toolId);
-						IBlobHandle handle = storageClient.CreateBlobHandle(deployment.Locator);
+						IBlobRef handle = storageClient.CreateBlobRef(deployment.Locator);
 
 						DirectoryNode rootNode = await handle.ReadBlobAsync<DirectoryNode>(cancellationToken: cancellationToken);
 						await rootNode.CopyToDirectoryAsync(new DirectoryInfo(applicationFolder), new ExtractStatsLogger(0, rootNode.Length, logger), logger, cancellationToken);
