Horde: Add extension methods to allow extracting IBlobRef<DirectoryNode> directly into directories.

diff --git a/Engine/Source/Programs/Horde/Drivers/JobDriver/Execution/JobExecutor.cs b/Engine/Source/Programs/Horde/Drivers/JobDriver/Execution/JobExecutor.cs
--- a/Engine/Source/Programs/Horde/Drivers/JobDriver/Execution/JobExecutor.cs
+++ b/Engine/Source/Programs/Horde/Drivers/JobDriver/Execution/JobExecutor.cs
@@ -1084,7 +1084,7 @@ namespace JobDriver.Execution
 				DirectoryNode? buildGraphDir = await node.TryOpenDirectoryAsync(BuildGraphTempStorageDir, cancellationToken: cancellationToken);
 				if (buildGraphDir != null)
 				{
-					await buildGraphDir.CopyToDirectoryAsync(new DirectoryInfo(workspaceDir.FullName), null, logger, cancellationToken);
+					await buildGraphDir.ExtractAsync(new DirectoryInfo(workspaceDir.FullName), logger, cancellationToken);
 					logger.LogInformation("Copying preprocessed script from {BuildGraphFolderName} into {OutputDir}", BuildGraphTempStorageDir, workspaceDir);
 				}
 				else
diff --git a/Engine/Source/Programs/Horde/Drivers/JobDriver/Utility/TempStorage.cs b/Engine/Source/Programs/Horde/Drivers/JobDriver/Utility/TempStorage.cs
--- a/Engine/Source/Programs/Horde/Drivers/JobDriver/Utility/TempStorage.cs
+++ b/Engine/Source/Programs/Horde/Drivers/JobDriver/Utility/TempStorage.cs
@@ -783,8 +783,7 @@ namespace JobDriver.Utility
 				Stopwatch timer = Stopwatch.StartNew();
 
 				// Add all the files and flush the ref
-				DirectoryNode rootDirNode = await rootDirEntry.Handle.ReadBlobAsync(cancellationToken: cancellationToken);
-				await rootDirNode.CopyToDirectoryAsync(rootDir.ToDirectoryInfo(), new ExtractStatsLogger(logger), logger, cancellationToken);
+				await rootDirEntry.Handle.ExtractAsync(rootDir.ToDirectoryInfo(), new ExtractStatsLogger(logger), logger, cancellationToken);
 
 				StorageStats deltaStats = StorageStats.GetDelta(initialStats, storageClient.GetStats());
 				logger.LogInformation("{Stats}", $"Elapsed: {(int)timer.Elapsed.TotalSeconds}s, {String.Join(", ", deltaStats.Values.Select(x => $"{x.Item1}: {x.Item2:n0}"))}");
diff --git a/Engine/Source/Programs/Horde/HordeCmd/Commands/Artifacts/ArtifactDownload.cs b/Engine/Source/Programs/Horde/HordeCmd/Commands/Artifacts/ArtifactDownload.cs
--- a/Engine/Source/Programs/Horde/HordeCmd/Commands/Artifacts/ArtifactDownload.cs
+++ b/Engine/Source/Programs/Horde/HordeCmd/Commands/Artifacts/ArtifactDownload.cs
@@ -51,12 +51,11 @@ namespace Horde.Commands.Artifacts
 			}
 
 			using IStorageClient store = _hordeClient.CreateStorageClient(artifact.Id);
-			IBlobRef handle = await store.ReadRefAsync(new RefName("default"));
 
 			Stopwatch timer = Stopwatch.StartNew();
 
-			DirectoryNode node = await handle.ReadBlobAsync<DirectoryNode>();
-			await node.CopyToDirectoryAsync(OutputDir.ToDirectoryInfo(), new ExtractStatsLogger(logger), logger, CancellationToken.None);
+			IBlobRef<DirectoryNode> handle = await store.ReadRefAsync<DirectoryNode>(new RefName("default"));
+			await handle.ExtractAsync(OutputDir.ToDirectoryInfo(), new ExtractStatsLogger(logger), logger, CancellationToken.None);
 
 			logger.LogInformation("Elapsed: {Time}s", timer.Elapsed.TotalSeconds);
 
diff --git a/Engine/Source/Programs/Horde/HordeCmd/Commands/Bundles/BundleExtract.cs b/Engine/Source/Programs/Horde/HordeCmd/Commands/Bundles/BundleExtract.cs
--- a/Engine/Source/Programs/Horde/HordeCmd/Commands/Bundles/BundleExtract.cs
+++ b/Engine/Source/Programs/Horde/HordeCmd/Commands/Bundles/BundleExtract.cs
@@ -55,19 +55,19 @@ namespace Horde.Commands.Bundles
 			if (File != null)
 			{
 				using IStorageClient store = BundleStorageClient.CreateFromDirectory(File.Directory, BundleCache, logger);
-				IBlobRef handle = store.CreateBlobRef(await FileStorageBackend.ReadRefAsync(File));
+				IBlobRef<DirectoryNode> handle = store.CreateBlobRef<DirectoryNode>(await FileStorageBackend.ReadRefAsync(File));
 				await ExecuteInternalAsync(store, handle, logger);
 			}
 			else if (Ref != null)
 			{
 				using IStorageClient store = CreateStorageClient();
-				IBlobRef handle = await store.ReadRefAsync(new RefName(Ref));
+				IBlobRef<DirectoryNode> handle = await store.ReadRefAsync<DirectoryNode>(new RefName(Ref));
 				await ExecuteInternalAsync(store, handle, logger);
 			}
 			else if (Node != null)
 			{
 				using IStorageClient store = CreateStorageClient();
-				IBlobRef handle = store.CreateBlobRef(new BlobLocator(Node));
+				IBlobRef<DirectoryNode> handle = store.CreateBlobRef<DirectoryNode>(new BlobLocator(Node));
 				await ExecuteInternalAsync(store, handle, logger);
 			}
 			else
@@ -78,12 +78,11 @@ namespace Horde.Commands.Bundles
 			return 0;
 		}
 
-		protected async Task ExecuteInternalAsync(IStorageClient store, IBlobRef handle, ILogger logger)
+		protected async Task ExecuteInternalAsync(IStorageClient store, IBlobRef<DirectoryNode> handle, ILogger logger)
 		{
 			Stopwatch timer = Stopwatch.StartNew();
 
-			DirectoryNode node = await handle.ReadBlobAsync<DirectoryNode>();
-			await node.CopyToDirectoryAsync(OutputDir.ToDirectoryInfo(), new ExtractStatsLogger(logger), logger, CancellationToken.None);
+			await handle.ExtractAsync(OutputDir.ToDirectoryInfo(), new ExtractStatsLogger(logger), logger, CancellationToken.None);
 
 			logger.LogInformation("Elapsed: {Time}s", timer.Elapsed.TotalSeconds);
 
diff --git a/Engine/Source/Programs/Shared/EpicGames.Horde.Tests/ChunkingTests.cs b/Engine/Source/Programs/Shared/EpicGames.Horde.Tests/ChunkingTests.cs
--- a/Engine/Source/Programs/Shared/EpicGames.Horde.Tests/ChunkingTests.cs
+++ b/Engine/Source/Programs/Shared/EpicGames.Horde.Tests/ChunkingTests.cs
@@ -195,8 +195,7 @@ namespace EpicGames.Horde.Tests
 			DirectoryInfo tempDir = new DirectoryInfo(Path.Combine(Path.GetTempPath(), Guid.NewGuid().ToString()));
 			try
 			{
-				DirectoryNode expandedNode = await handle.ReadBlobAsync();
-				await expandedNode.CopyToDirectoryAsync(tempDir, null, NullLogger.Instance, CancellationToken.None);
+				await handle.ExtractAsync(tempDir, null, NullLogger.Instance, CancellationToken.None);
 
 				byte[] outputData = await File.ReadAllBytesAsync(Path.Combine(tempDir.FullName, "test"));
 				Assert.IsTrue(outputData.SequenceEqual(data));
diff --git a/Engine/Source/Programs/Shared/EpicGames.Horde/Compute/AgentMessageHandler.cs b/Engine/Source/Programs/Shared/EpicGames.Horde/Compute/AgentMessageHandler.cs
--- a/Engine/Source/Programs/Shared/EpicGames.Horde/Compute/AgentMessageHandler.cs
+++ b/Engine/Source/Programs/Shared/EpicGames.Horde/Compute/AgentMessageHandler.cs
@@ -193,8 +193,7 @@ namespace EpicGames.Horde.Compute
 			BundleOptions bundleOptions = ComputeProtocolUtilities.GetBundleOptions(channel.Protocol);
 			using BundleStorageClient store = new BundleStorageClient(innerStore, cache, bundleOptions, _logger);
 
-			IBlobRef handle = store.CreateBlobRef(locator);
-			DirectoryNode directoryNode = await handle.ReadBlobAsync<DirectoryNode>(options, cancellationToken);
+			IBlobRef<DirectoryNode> directoryRef = store.CreateBlobRef<DirectoryNode>(locator, options);
 
 			DirectoryReference outputDir = DirectoryReference.Combine(_sandboxDir, path);
 			if (!outputDir.IsUnderDirectory(_sandboxDir))
@@ -202,8 +201,8 @@ namespace EpicGames.Horde.Compute
 				throw new InvalidOperationException("Cannot write files outside sandbox");
 			}
 
-			await directoryNode.CopyToDirectoryAsync(outputDir.ToDirectoryInfo(), _logger, cancellationToken);
-			await VerifyFilesAsync(outputDir, directoryNode, cancellationToken);
+			await directoryRef.ExtractAsync(outputDir.ToDirectoryInfo(), _logger, cancellationToken);
+			await VerifyFilesAsync(outputDir, directoryRef, cancellationToken);
 
 			using (IAgentMessageBuilder message = await channel.CreateMessageAsync(AgentMessageType.WriteFilesResponse, cancellationToken))
 			{
@@ -211,10 +210,11 @@ namespace EpicGames.Horde.Compute
 			}
 		}
 
-		async Task<bool> VerifyFilesAsync(DirectoryReference outputDir, DirectoryNode directoryNode, CancellationToken cancellationToken = default)
+		async Task<bool> VerifyFilesAsync(DirectoryReference outputDir, IBlobRef<DirectoryNode> directoryRef, CancellationToken cancellationToken = default)
 		{
 			bool result = true;
 
+			DirectoryNode directoryNode = await directoryRef.ReadBlobAsync(cancellationToken);
 			foreach (FileEntry fileEntry in directoryNode.Files)
 			{
 				FileReference file = FileReference.Combine(outputDir, fileEntry.Name);
@@ -242,8 +242,7 @@ namespace EpicGames.Horde.Compute
 
 			foreach (DirectoryEntry directoryEntry in directoryNode.Directories)
 			{
-				DirectoryNode subNode = await directoryEntry.Handle.ReadBlobAsync(cancellationToken: cancellationToken);
-				result &= await VerifyFilesAsync(DirectoryReference.Combine(outputDir, directoryEntry.Name), subNode, cancellationToken);
+				result &= await VerifyFilesAsync(DirectoryReference.Combine(outputDir, directoryEntry.Name), directoryEntry.Handle, cancellationToken);
 			}
 
 			return result;
diff --git a/Engine/Source/Programs/Shared/EpicGames.Horde/Storage/BlobSerializer.cs b/Engine/Source/Programs/Shared/EpicGames.Horde/Storage/BlobSerializer.cs
--- a/Engine/Source/Programs/Shared/EpicGames.Horde/Storage/BlobSerializer.cs
+++ b/Engine/Source/Programs/Shared/EpicGames.Horde/Storage/BlobSerializer.cs
@@ -223,6 +223,18 @@ namespace EpicGames.Horde.Storage
 			return BlobSerializer.Deserialize<T>(data, options);
 		}
 
+		/// <summary>
+		/// Deserialize an object
+		/// </summary>
+		/// <typeparam name="T">Return type for the deserialized object</typeparam>
+		/// <param name="handle">Handle to the blob to deserialize</param>
+		/// <param name="cancellationToken">Cancellation token for the operation</param>
+		public static async ValueTask<T> ReadBlobAsync<T>(this IBlobRef<T> handle, CancellationToken cancellationToken = default)
+		{
+			using BlobData data = await handle.ReadBlobDataAsync(cancellationToken);
+			return BlobSerializer.Deserialize<T>(data, handle.SerializerOptions);
+		}
+
 		/// <summary>
 		/// Deserialize an object
 		/// </summary>
diff --git a/Engine/Source/Programs/Shared/EpicGames.Horde/Storage/Nodes/DirectoryNode.Extract.cs b/Engine/Source/Programs/Shared/EpicGames.Horde/Storage/Nodes/DirectoryNode.Extract.cs
--- a/Engine/Source/Programs/Shared/EpicGames.Horde/Storage/Nodes/DirectoryNode.Extract.cs
+++ b/Engine/Source/Programs/Shared/EpicGames.Horde/Storage/Nodes/DirectoryNode.Extract.cs
@@ -159,11 +159,12 @@ namespace EpicGames.Horde.Storage.Nodes
 		/// <summary>
 		/// Utility function to allow extracting a packed directory to disk
 		/// </summary>
-		/// <param name="directoryNode">Directory to update</param>
+		/// <param name="directoryRef">Directory to update</param>
 		/// <param name="directoryInfo"></param>
 		/// <param name="logger"></param>
 		/// <param name="cancellationToken"></param>
-		public static Task CopyToDirectoryAsync(this DirectoryNode directoryNode, DirectoryInfo directoryInfo, ILogger logger, CancellationToken cancellationToken) => CopyToDirectoryAsync(directoryNode, directoryInfo, null, logger, cancellationToken);
+		public static Task ExtractAsync(this IBlobRef<DirectoryNode> directoryRef, DirectoryInfo directoryInfo, ILogger logger, CancellationToken cancellationToken) 
+			=> ExtractAsync(directoryRef, directoryInfo, null, logger, cancellationToken);
 
 		class OutputFile
 		{
@@ -277,26 +278,51 @@ namespace EpicGames.Horde.Storage.Nodes
 		/// <summary>
 		/// Utility function to allow extracting a packed directory to disk
 		/// </summary>
-		/// <param name="directoryNode">Directory to update</param>
+		/// <param name="directoryRef">Directory to extract</param>
 		/// <param name="directoryInfo">Direcotry to write to</param>
 		/// <param name="progress">Sink for progress updates</param>
 		/// <param name="logger">Logger for output</param>
 		/// <param name="cancellationToken">Cancellation token for the operation</param>
-		public static Task CopyToDirectoryAsync(this DirectoryNode directoryNode, DirectoryInfo directoryInfo, IProgress<IExtractStats>? progress, ILogger logger, CancellationToken cancellationToken)
+		public static Task ExtractAsync(this IBlobRef<DirectoryNode> directoryRef, DirectoryInfo directoryInfo, IProgress<IExtractStats>? progress, ILogger logger, CancellationToken cancellationToken)
 		{
-			return CopyToDirectoryAsync(directoryNode, directoryInfo, progress, TimeSpan.FromSeconds(5.0), logger, cancellationToken);
+			return ExtractAsync(directoryRef, directoryInfo, progress, TimeSpan.FromSeconds(5.0), logger, cancellationToken);
 		}
 
 		/// <summary>
 		/// Utility function to allow extracting a packed directory to disk
 		/// </summary>
-		/// <param name="directoryNode">Directory to update</param>
+		/// <param name="directoryRef">Directory to extract</param>
 		/// <param name="directoryInfo">Direcotry to write to</param>
 		/// <param name="progress">Sink for progress updates</param>
 		/// <param name="frequency">Frequency for progress updates</param>
 		/// <param name="logger">Logger for output</param>
 		/// <param name="cancellationToken">Cancellation token for the operation</param>
-		public static async Task CopyToDirectoryAsync(this DirectoryNode directoryNode, DirectoryInfo directoryInfo, IProgress<IExtractStats>? progress, TimeSpan frequency, ILogger logger, CancellationToken cancellationToken)
+		public static async Task ExtractAsync(this IBlobRef<DirectoryNode> directoryRef, DirectoryInfo directoryInfo, IProgress<IExtractStats>? progress, TimeSpan frequency, ILogger logger, CancellationToken cancellationToken)
+		{
+			DirectoryNode directoryNode = await directoryRef.ReadBlobAsync(cancellationToken);
+			await ExtractAsync(directoryNode, directoryInfo, progress, frequency, logger, cancellationToken);
+		}
+
+		/// <summary>
+		/// Utility function to allow extracting a packed directory to disk
+		/// </summary>
+		/// <param name="directoryNode">Directory to extract</param>
+		/// <param name="directoryInfo">Direcotry to write to</param>
+		/// <param name="logger">Logger for output</param>
+		/// <param name="cancellationToken">Cancellation token for the operation</param>
+		public static Task ExtractAsync(this DirectoryNode directoryNode, DirectoryInfo directoryInfo, ILogger logger, CancellationToken cancellationToken)
+			=> ExtractAsync(directoryNode, directoryInfo, null, TimeSpan.FromDays(1.0), logger, cancellationToken);
+
+		/// <summary>
+		/// Utility function to allow extracting a packed directory to disk
+		/// </summary>
+		/// <param name="directoryNode">Directory to extract</param>
+		/// <param name="directoryInfo">Direcotry to write to</param>
+		/// <param name="progress">Sink for progress updates</param>
+		/// <param name="frequency">Frequency for progress updates</param>
+		/// <param name="logger">Logger for output</param>
+		/// <param name="cancellationToken">Cancellation token for the operation</param>
+		public static async Task ExtractAsync(this DirectoryNode directoryNode, DirectoryInfo directoryInfo, IProgress<IExtractStats>? progress, TimeSpan frequency, ILogger logger, CancellationToken cancellationToken)
 		{
 			int numTasks = Math.Min(1 + (int)(directoryNode.Length / (16 * 1024 * 1024)), 16);
 			logger.LogInformation("Splitting read into {NumThreads} threads", numTasks);
diff --git a/Engine/Source/Programs/UnrealGameSync/UnrealGameSync/ArtifactDownload.cs b/Engine/Source/Programs/UnrealGameSync/UnrealGameSync/ArtifactDownload.cs
--- a/Engine/Source/Programs/UnrealGameSync/UnrealGameSync/ArtifactDownload.cs
+++ b/Engine/Source/Programs/UnrealGameSync/UnrealGameSync/ArtifactDownload.cs
@@ -59,10 +59,10 @@ namespace UnrealGameSync
 
 				progress.Report("Connecting to server...");
 
-				DirectoryNode? node = await storageClient.ReadRefTargetAsync<DirectoryNode>(refName, cancellationToken: cancellationToken);
+				IBlobRef<DirectoryNode>? blobRef = await storageClient.ReadRefAsync<DirectoryNode>(refName, cancellationToken: cancellationToken);
 
 				progress.Report("Starting...");
-				await node.CopyToDirectoryAsync(outputDir.ToDirectoryInfo(), new CopyProgressAdapter(progress), TimeSpan.FromSeconds(0.2), serviceProvider.GetRequiredService<ILogger<ArtifactDownload>>(), cancellationToken);
+				await blobRef.ExtractAsync(outputDir.ToDirectoryInfo(), new CopyProgressAdapter(progress), TimeSpan.FromSeconds(0.2), serviceProvider.GetRequiredService<ILogger<ArtifactDownload>>(), cancellationToken);
 			}
 		}
 
diff --git a/Engine/Source/Programs/UnrealGameSync/UnrealGameSync/Launcher.cs b/Engine/Source/Programs/UnrealGameSync/UnrealGameSync/Launcher.cs
--- a/Engine/Source/Programs/UnrealGameSync/UnrealGameSync/Launcher.cs
+++ b/Engine/Source/Programs/UnrealGameSync/UnrealGameSync/Launcher.cs
@@ -187,10 +187,10 @@ namespace UnrealGameSync
 
 						// Download and extract the data
 						using IStorageClient storageClient = hordeClient.CreateStorageClient(toolId);
-						IBlobRef handle = storageClient.CreateBlobRef(deployment.Locator);
+						IBlobRef<DirectoryNode> handle = storageClient.CreateBlobRef<DirectoryNode>(deployment.Locator);
 
-						DirectoryNode rootNode = await handle.ReadBlobAsync<DirectoryNode>(cancellationToken: cancellationToken);
-						await rootNode.CopyToDirectoryAsync(new DirectoryInfo(applicationFolder), new ExtractStatsLogger(0, rootNode.Length, logger), logger, cancellationToken);
+						DirectoryNode rootNode = await handle.ReadBlobAsync(cancellationToken: cancellationToken);
+						await rootNode.ExtractAsync(new DirectoryInfo(applicationFolder), new ExtractStatsLogger(0, rootNode.Length, logger), TimeSpan.FromSeconds(5.0), logger, cancellationToken);
 
 						// Update the version
 						if (!TryWriteAllText(syncVersionFile, requiredSyncText))
