Horde: Add extension methods for creating blob refs using ref names. Also convert other variants of CreateBlobRef to extension methods.

#rnx

diff --git a/Engine/Source/Programs/Shared/EpicGames.Horde/Storage/Bundles/BundleStorageClient.cs b/Engine/Source/Programs/Shared/EpicGames.Horde/Storage/Bundles/BundleStorageClient.cs
--- a/Engine/Source/Programs/Shared/EpicGames.Horde/Storage/Bundles/BundleStorageClient.cs
+++ b/Engine/Source/Programs/Shared/EpicGames.Horde/Storage/Bundles/BundleStorageClient.cs
@@ -183,24 +183,6 @@ namespace EpicGames.Horde.Storage.Bundles
 			return new Bundles.V2.ExportHandle(packetHandle, fragment.Slice(ampIdx + 1));
 		}
 
-		/// <inheritdoc/>
-		public IBlobRef<T> CreateBlobRef<T>(BlobLocator locator, BlobSerializerOptions? options)
-		{
-			return BlobRef.Create<T>(CreateBlobRef(locator), options);
-		}
-
-		/// <inheritdoc/>
-		public IHashedBlobRef CreateBlobRef(IoHash hash, BlobLocator locator)
-		{
-			return HashedBlobRef.Create(hash, CreateBlobRef(locator));
-		}
-
-		/// <inheritdoc/>
-		public IHashedBlobRef<T> CreateBlobRef<T>(IoHash hash, BlobLocator locator, BlobSerializerOptions? options)
-		{
-			return HashedBlobRef.Create<T>(hash, CreateBlobRef(locator), options);
-		}
-
 		/// <inheritdoc/>
 		public IBlobWriter CreateBlobWriter(string? basePath = null, BlobSerializerOptions? serializerOptions = null)
 		{
@@ -267,7 +249,7 @@ namespace EpicGames.Horde.Storage.Bundles
 			{
 				return null;
 			}
-			return CreateBlobRef(target.Hash, target.Locator);
+			return this.CreateBlobRef(target.Hash, target.Locator);
 		}
 
 		/// <inheritdoc/>
diff --git a/Engine/Source/Programs/Shared/EpicGames.Horde/Storage/Clients/KeyValueStorageClient.cs b/Engine/Source/Programs/Shared/EpicGames.Horde/Storage/Clients/KeyValueStorageClient.cs
--- a/Engine/Source/Programs/Shared/EpicGames.Horde/Storage/Clients/KeyValueStorageClient.cs
+++ b/Engine/Source/Programs/Shared/EpicGames.Horde/Storage/Clients/KeyValueStorageClient.cs
@@ -145,24 +145,6 @@ namespace EpicGames.Horde.Storage.Clients
 			return new Handle(this, locator);
 		}
 
-		/// <inheritdoc/>
-		public IBlobRef<T> CreateBlobRef<T>(BlobLocator locator, BlobSerializerOptions? options)
-		{
-			return BlobRef.Create<T>(CreateBlobRef(locator), options);
-		}
-
-		/// <inheritdoc/>
-		public IHashedBlobRef CreateBlobRef(IoHash hash, BlobLocator locator)
-		{
-			return HashedBlobRef.Create(hash, CreateBlobRef(locator));
-		}
-
-		/// <inheritdoc/>
-		public IHashedBlobRef<T> CreateBlobRef<T>(IoHash hash, BlobLocator locator, BlobSerializerOptions? options)
-		{
-			return HashedBlobRef.Create<T>(hash, CreateBlobRef(locator), options);
-		}
-
 		/// <inheritdoc/>
 		public IBlobWriter CreateBlobWriter(string? basePath = null, BlobSerializerOptions? options = null) => new Writer(this, basePath, options);
 
@@ -226,7 +208,7 @@ namespace EpicGames.Horde.Storage.Clients
 			{
 				return null;
 			}
-			return CreateBlobRef(value.Hash, value.Locator);
+			return this.CreateBlobRef(value.Hash, value.Locator);
 		}
 
 		/// <inheritdoc/>
diff --git a/Engine/Source/Programs/Shared/EpicGames.Horde/Storage/IBlobRef.cs b/Engine/Source/Programs/Shared/EpicGames.Horde/Storage/IBlobRef.cs
--- a/Engine/Source/Programs/Shared/EpicGames.Horde/Storage/IBlobRef.cs
+++ b/Engine/Source/Programs/Shared/EpicGames.Horde/Storage/IBlobRef.cs
@@ -49,9 +49,9 @@ namespace EpicGames.Horde.Storage
 	}
 
 	/// <summary>
-	/// Static methods for IBlobRef types.
+	/// Extension methods for <see cref="IBlobRef"/>
 	/// </summary>
-	public static class BlobRef
+	public static class BlobRefExtensions
 	{
 		class TypedBlobRef<T> : IBlobRef<T>
 		{
@@ -84,15 +84,9 @@ namespace EpicGames.Horde.Storage
 		/// <typeparam name="T">Target type</typeparam>
 		/// <param name="blobRef">Blob referfence to wrap</param>
 		/// <param name="serializerOptions">Options for deserializing the blob</param>
-		public static IBlobRef<T> Create<T>(IBlobRef blobRef, BlobSerializerOptions? serializerOptions = null)
+		public static IBlobRef<T> ForType<T>(this IBlobRef blobRef, BlobSerializerOptions? serializerOptions = null)
 			=> new TypedBlobRef<T>(blobRef, serializerOptions);
-	}
 
-	/// <summary>
-	/// Extension methods for <see cref="IBlobRef"/>
-	/// </summary>
-	public static class BlobRefExtensions
-	{
 		/// <summary>
 		/// Gets a path to this blob that can be used to describe blob references over the wire.
 		/// </summary>
diff --git a/Engine/Source/Programs/Shared/EpicGames.Horde/Storage/IStorageClient.cs b/Engine/Source/Programs/Shared/EpicGames.Horde/Storage/IStorageClient.cs
--- a/Engine/Source/Programs/Shared/EpicGames.Horde/Storage/IStorageClient.cs
+++ b/Engine/Source/Programs/Shared/EpicGames.Horde/Storage/IStorageClient.cs
@@ -2,6 +2,7 @@
 
 using System;
 using System.Collections.Generic;
+using System.Diagnostics.CodeAnalysis;
 using System.Linq;
 using System.Threading;
 using System.Threading.Tasks;
@@ -31,31 +32,6 @@ namespace EpicGames.Horde.Storage
 		/// <returns>New handle to the blob</returns>
 		IBlobRef CreateBlobRef(BlobLocator locator);
 
-		/// <summary>
-		/// Creates a new blob handle by parsing a locator
-		/// </summary>
-		/// <param name="locator">Path to the blob</param>
-		/// <param name="serializerOptions">Options for deserializing the blob</param>
-		/// <returns>New handle to the blob</returns>
-		IBlobRef<T> CreateBlobRef<T>(BlobLocator locator, BlobSerializerOptions? serializerOptions = null);
-
-		/// <summary>
-		/// Creates a new blob reference from a locator and hash
-		/// </summary>
-		/// <param name="hash">Hash of the target blob</param>
-		/// <param name="locator">Path to the blob</param>
-		/// <returns>New handle to the blob</returns>
-		IHashedBlobRef CreateBlobRef(IoHash hash, BlobLocator locator);
-
-		/// <summary>
-		/// Creates a new blob reference from a locator and hash
-		/// </summary>
-		/// <param name="hash">Hash of the target blob</param>
-		/// <param name="locator">Path to the blob</param>
-		/// <param name="serializerOptions">Options for deserializing the blob</param>
-		/// <returns>New handle to the blob</returns>
-		IHashedBlobRef<T> CreateBlobRef<T>(IoHash hash, BlobLocator locator, BlobSerializerOptions? serializerOptions = null);
-
 		/// <summary>
 		/// Creates a new writer for storage blobs
 		/// </summary>
@@ -256,6 +232,37 @@ namespace EpicGames.Horde.Storage
 	{
 		#region Blobs
 
+		/// <summary>
+		/// Creates a new blob handle by parsing a locator
+		/// </summary>
+		/// <param name="storageClient">Storage client to operate on</param>
+		/// <param name="locator">Path to the blob</param>
+		/// <param name="serializerOptions">Options for deserializing the blob</param>
+		/// <returns>New handle to the blob</returns>
+		public static IBlobRef<T> CreateBlobRef<T>(this IStorageClient storageClient, BlobLocator locator, BlobSerializerOptions? serializerOptions = null)
+			=> storageClient.CreateBlobRef(locator).ForType<T>(serializerOptions);
+
+		/// <summary>
+		/// Creates a new blob reference from a locator and hash
+		/// </summary>
+		/// <param name="storageClient">Storage client to operate on</param>
+		/// <param name="hash">Hash of the target blob</param>
+		/// <param name="locator">Path to the blob</param>
+		/// <returns>New handle to the blob</returns>
+		public static IHashedBlobRef CreateBlobRef(this IStorageClient storageClient, IoHash hash, BlobLocator locator)
+			=> HashedBlobRef.Create(hash, storageClient.CreateBlobRef(locator));
+
+		/// <summary>
+		/// Creates a new blob reference from a locator and hash
+		/// </summary>
+		/// <param name="storageClient">Storage client to operate on</param>
+		/// <param name="hash">Hash of the target blob</param>
+		/// <param name="locator">Path to the blob</param>
+		/// <param name="serializerOptions">Options for deserializing the blob</param>
+		/// <returns>New handle to the blob</returns>
+		public static IHashedBlobRef<T> CreateBlobRef<T>(this IStorageClient storageClient, IoHash hash, BlobLocator locator, BlobSerializerOptions? serializerOptions = null)
+			=> HashedBlobRef.Create<T>(hash, storageClient.CreateBlobRef(locator), serializerOptions);
+
 		/// <summary>
 		/// Create a blob ref from a RefValue
 		/// </summary>
@@ -268,6 +275,58 @@ namespace EpicGames.Horde.Storage
 		public static IHashedBlobRef<T> CreateBlobRef<T>(this IStorageClient store, HashedBlobRefValue refValue, BlobSerializerOptions? options)
 			=> store.CreateBlobRef<T>(refValue.Hash, refValue.Locator, options);
 
+		class NamedBlobRef : IBlobRef
+		{
+			readonly IStorageClient _storageClient;
+			readonly RefName _name;
+			readonly RefCacheTime _cacheTime;
+
+			public NamedBlobRef(IStorageClient storageClient, RefName name, RefCacheTime cacheTime = default)
+			{
+				_storageClient = storageClient;
+				_name = name;
+				_cacheTime = cacheTime;
+			}
+
+			public IBlobRef Innermost => this;
+
+			public ValueTask FlushAsync(CancellationToken cancellationToken = default)
+				=> default;
+
+			public async ValueTask<BlobData> ReadBlobDataAsync(CancellationToken cancellationToken = default)
+			{
+				IHashedBlobRef handle = await _storageClient.ReadRefAsync(_name, _cacheTime, cancellationToken);
+				return await handle.ReadBlobDataAsync(cancellationToken);
+			}
+
+			public bool TryGetLocator([NotNullWhen(true)] out BlobLocator locator)
+			{
+				locator = default;
+				return false;
+			}
+		}
+
+		/// <summary>
+		/// Creates a new blob ref from a ref name
+		/// </summary>
+		/// <param name="storageClient">The store instance to read from</param>
+		/// <param name="name">Name of the reference</param>
+		/// <param name="cacheTime">Maximum age for cached responses</param>
+		/// <returns>New handle to the blob</returns>
+		public static IBlobRef CreateBlobRef(this IStorageClient storageClient, RefName name, RefCacheTime cacheTime = default)
+			=> new NamedBlobRef(storageClient, name, cacheTime);
+
+		/// <summary>
+		/// Creates a new blob ref from a ref name
+		/// </summary>
+		/// <param name="storageClient">The store instance to read from</param>
+		/// <param name="name">Name of the reference</param>
+		/// <param name="cacheTime">Maximum age for cached responses</param>
+		/// <param name="serializerOptions">Options for deserializing the blob</param>
+		/// <returns>New handle to the blob</returns>
+		public static IBlobRef<T> CreateBlobRef<T>(this IStorageClient storageClient, RefName name, RefCacheTime cacheTime = default, BlobSerializerOptions? serializerOptions = null)
+			=> CreateBlobRef(storageClient, name, cacheTime).ForType<T>(serializerOptions);
+
 		/// <summary>
 		/// Creates a writer using a refname as a base path
 		/// </summary>
