Horde: Move artifact interfaces into EpicGames.Horde, and implement them in terms of HTTP responses. This allows a much more idiomatic OOP syntax for interacting with the server.

#rnx

diff --git a/Engine/Source/Programs/Horde/Drivers/JobDriver.Tests/Execution/WorkspaceExecutorTest.cs b/Engine/Source/Programs/Horde/Drivers/JobDriver.Tests/Execution/WorkspaceExecutorTest.cs
--- a/Engine/Source/Programs/Horde/Drivers/JobDriver.Tests/Execution/WorkspaceExecutorTest.cs
+++ b/Engine/Source/Programs/Horde/Drivers/JobDriver.Tests/Execution/WorkspaceExecutorTest.cs
@@ -20,6 +20,7 @@ using HordeCommon.Rpc.Tasks;
 using Microsoft.Extensions.Logging;
 using Microsoft.Extensions.Logging.Abstractions;
 using Microsoft.VisualStudio.TestTools.UnitTesting;
+using EpicGames.Horde.Artifacts;
 
 namespace JobDriver.Tests.Execution
 {
@@ -82,6 +83,8 @@ namespace JobDriver.Tests.Execution
 			public Uri ServerUrl
 				=> new Uri("http://fake-horde-server");
 
+			public IArtifactCollection Artifacts => throw new NotImplementedException();
+
 			public IComputeClient CreateComputeClient()
 				=> throw new NotImplementedException();
 
diff --git a/Engine/Source/Programs/Horde/HordeCmd/Commands/Artifacts/ArtifactFind.cs b/Engine/Source/Programs/Horde/HordeCmd/Commands/Artifacts/ArtifactFind.cs
--- a/Engine/Source/Programs/Horde/HordeCmd/Commands/Artifacts/ArtifactFind.cs
+++ b/Engine/Source/Programs/Horde/HordeCmd/Commands/Artifacts/ArtifactFind.cs
@@ -15,17 +15,16 @@ namespace Horde.Commands.Bundles
 		[Description("Artifact keys to search for. Multiple keys may be added to artifacts at upload time, eg. 'job:63dd5487c67f8a45453361c5/step:62ce'.")]
 		public List<string> Keys { get; } = new List<string>();
 
-		readonly HordeHttpClient _hordeHttpClient;
+		readonly IHordeClient _hordeClient;
 
-		public ArtifactFind(HordeHttpClient httpClient)
+		public ArtifactFind(IHordeClient hordeClient)
 		{
-			_hordeHttpClient = httpClient;
+			_hordeClient = hordeClient;
 		}
 
 		public override async Task<int> ExecuteAsync(ILogger logger)
 		{
-			List<GetArtifactResponse> artifacts = await _hordeHttpClient.FindArtifactsAsync(keys: Keys);
-			foreach (GetArtifactResponse artifact in artifacts)
+			await foreach(IArtifact artifact in _hordeClient.Artifacts.FindAsync(keys: Keys))
 			{
 				logger.LogInformation("");
 				logger.LogInformation("Artifact {Id} ({Type})", artifact.Id, artifact.Type);
diff --git a/Engine/Source/Programs/Horde/Plugins/Build/HordeServer.Build.Tests/Artifacts/UnsyncCacheTests.cs b/Engine/Source/Programs/Horde/Plugins/Build/HordeServer.Build.Tests/Artifacts/UnsyncCacheTests.cs
--- a/Engine/Source/Programs/Horde/Plugins/Build/HordeServer.Build.Tests/Artifacts/UnsyncCacheTests.cs
+++ b/Engine/Source/Programs/Horde/Plugins/Build/HordeServer.Build.Tests/Artifacts/UnsyncCacheTests.cs
@@ -1,6 +1,7 @@
 ﻿// Copyright Epic Games, Inc. All Rights Reserved.
 
 using EpicGames.Core;
+using EpicGames.Horde.Artifacts;
 using EpicGames.Horde.Storage;
 using EpicGames.Horde.Storage.Bundles;
 using EpicGames.Horde.Storage.Nodes;
diff --git a/Engine/Source/Programs/Horde/Plugins/Build/HordeServer.Build.Tests/BuildTestSetup.cs b/Engine/Source/Programs/Horde/Plugins/Build/HordeServer.Build.Tests/BuildTestSetup.cs
--- a/Engine/Source/Programs/Horde/Plugins/Build/HordeServer.Build.Tests/BuildTestSetup.cs
+++ b/Engine/Source/Programs/Horde/Plugins/Build/HordeServer.Build.Tests/BuildTestSetup.cs
@@ -6,6 +6,7 @@ using Amazon.CloudWatch;
 using Amazon.EC2;
 using EpicGames.Horde.Agents;
 using EpicGames.Horde.Agents.Pools;
+using EpicGames.Horde.Artifacts;
 using HordeServer.Acls;
 using HordeServer.Agents;
 using HordeServer.Agents.Enrollment;
diff --git a/Engine/Source/Programs/Horde/Plugins/Build/HordeServer.Build/Artifacts/ArtifactExpirationService.cs b/Engine/Source/Programs/Horde/Plugins/Build/HordeServer.Build/Artifacts/ArtifactExpirationService.cs
--- a/Engine/Source/Programs/Horde/Plugins/Build/HordeServer.Build/Artifacts/ArtifactExpirationService.cs
+++ b/Engine/Source/Programs/Horde/Plugins/Build/HordeServer.Build/Artifacts/ArtifactExpirationService.cs
@@ -2,6 +2,7 @@
 
 using System.Diagnostics;
 using EpicGames.Core;
+using EpicGames.Horde.Artifacts;
 using EpicGames.Horde.Storage;
 using EpicGames.Horde.Streams;
 using HordeServer.Storage;
diff --git a/Engine/Source/Programs/Horde/Plugins/Build/HordeServer.Build/Artifacts/ArtifactsController.cs b/Engine/Source/Programs/Horde/Plugins/Build/HordeServer.Build/Artifacts/ArtifactsController.cs
--- a/Engine/Source/Programs/Horde/Plugins/Build/HordeServer.Build/Artifacts/ArtifactsController.cs
+++ b/Engine/Source/Programs/Horde/Plugins/Build/HordeServer.Build/Artifacts/ArtifactsController.cs
@@ -202,7 +202,7 @@ namespace HordeServer.Artifacts
 			{
 				if (_buildConfig.AuthorizeArtifact(artifact.Type, artifact.StreamId, ArtifactAclAction.ReadArtifact, User))
 				{
-					response.Artifacts.Add(new GetArtifactResponse(artifact.Id, artifact.Name, artifact.Type, artifact.Description, artifact.StreamId, artifact.Keys, artifact.Metadata, artifact.CreatedAtUtc) { CommitId = artifact.CommitId });
+					response.Artifacts.Add(new GetArtifactResponse(artifact));
 				}
 			}
 
@@ -230,7 +230,7 @@ namespace HordeServer.Artifacts
 				return Forbid(ArtifactAclAction.ReadArtifact, artifact.StreamId);
 			}
 
-			return PropertyFilter.Apply(new GetArtifactResponse(artifact.Id, artifact.Name, artifact.Type, artifact.Description, artifact.StreamId, artifact.Keys, artifact.Metadata, artifact.CreatedAtUtc) { CommitId = artifact.CommitId }, filter);
+			return PropertyFilter.Apply(new GetArtifactResponse(artifact), filter);
 		}
 
 		/// <summary>
diff --git a/Engine/Source/Programs/Horde/Plugins/Build/HordeServer.Build/Artifacts/IArtifact.cs b/Engine/Source/Programs/Horde/Plugins/Build/HordeServer.Build/Artifacts/IArtifact.cs
deleted file mode 100644
diff --git a/Engine/Source/Programs/Horde/Plugins/Build/HordeServer.Build/Artifacts/IArtifactCollection.cs b/Engine/Source/Programs/Horde/Plugins/Build/HordeServer.Build/Artifacts/IArtifactCollection.cs
deleted file mode 100644
diff --git a/Engine/Source/Programs/Horde/Plugins/Build/HordeServer.Build/Artifacts/UnsyncCache.cs b/Engine/Source/Programs/Horde/Plugins/Build/HordeServer.Build/Artifacts/UnsyncCache.cs
--- a/Engine/Source/Programs/Horde/Plugins/Build/HordeServer.Build/Artifacts/UnsyncCache.cs
+++ b/Engine/Source/Programs/Horde/Plugins/Build/HordeServer.Build/Artifacts/UnsyncCache.cs
@@ -4,6 +4,7 @@ using System.Collections.Frozen;
 using System.Diagnostics;
 using System.Text.Json;
 using EpicGames.Core;
+using EpicGames.Horde.Artifacts;
 using EpicGames.Horde.Storage;
 using EpicGames.Horde.Storage.Nodes;
 using HordeServer.Storage;
diff --git a/Engine/Source/Programs/Horde/Plugins/Build/HordeServer.Build/BuildPlugin.cs b/Engine/Source/Programs/Horde/Plugins/Build/HordeServer.Build/BuildPlugin.cs
--- a/Engine/Source/Programs/Horde/Plugins/Build/HordeServer.Build/BuildPlugin.cs
+++ b/Engine/Source/Programs/Horde/Plugins/Build/HordeServer.Build/BuildPlugin.cs
@@ -2,6 +2,7 @@
 
 using System.Diagnostics.CodeAnalysis;
 using EpicGames.Core;
+using EpicGames.Horde.Artifacts;
 using HordeServer.Acls;
 using HordeServer.Agents;
 using HordeServer.Agents.Fleet;
diff --git a/Engine/Source/Programs/Shared/EpicGames.Horde.Tests/ServerLogPacketBuilderTests.cs b/Engine/Source/Programs/Shared/EpicGames.Horde.Tests/ServerLogPacketBuilderTests.cs
--- a/Engine/Source/Programs/Shared/EpicGames.Horde.Tests/ServerLogPacketBuilderTests.cs
+++ b/Engine/Source/Programs/Shared/EpicGames.Horde.Tests/ServerLogPacketBuilderTests.cs
@@ -8,6 +8,7 @@ using System.Threading;
 using System.Threading.Channels;
 using System.Threading.Tasks;
 using EpicGames.Core;
+using EpicGames.Horde.Artifacts;
 using EpicGames.Horde.Compute;
 using EpicGames.Horde.Logs;
 using EpicGames.Horde.Storage;
@@ -214,6 +215,8 @@ namespace EpicGames.Horde.Tests
 
 			public Uri ServerUrl => throw new NotImplementedException();
 
+			public IArtifactCollection Artifacts => throw new NotImplementedException();
+
 			public Task<bool> LoginAsync(bool allowLogin, CancellationToken cancellationToken)
 				=> throw new NotImplementedException();
 
diff --git a/Engine/Source/Programs/Shared/EpicGames.Horde/Artifacts/ArtifactHttpCollection.cs b/Engine/Source/Programs/Shared/EpicGames.Horde/Artifacts/ArtifactHttpCollection.cs
new file mode 100644
--- /dev/null
+++ b/Engine/Source/Programs/Shared/EpicGames.Horde/Artifacts/ArtifactHttpCollection.cs
@@ -0,0 +1,82 @@
+// Copyright Epic Games, Inc. All Rights Reserved.
+
+using System;
+using System.Collections.Generic;
+using System.Runtime.CompilerServices;
+using System.Threading;
+using System.Threading.Tasks;
+using EpicGames.Horde.Acls;
+using EpicGames.Horde.Commits;
+using EpicGames.Horde.Storage;
+using EpicGames.Horde.Streams;
+
+namespace EpicGames.Horde.Artifacts
+{
+	class ArtifactHttpCollection : IArtifactCollection
+	{
+		class Artifact : IArtifact
+		{
+			readonly GetArtifactResponse _response;
+
+			public Artifact(GetArtifactResponse response)
+			{
+				_response = response;
+			}
+
+			public ArtifactId Id => _response.Id;
+			public ArtifactName Name => _response.Name;
+			public ArtifactType Type => _response.Type;
+			public string? Description => _response.Description;
+			public StreamId StreamId => _response.StreamId;
+			public CommitIdWithOrder CommitId => _response.CommitId;
+			public IReadOnlyList<string> Keys => _response.Keys;
+			public IReadOnlyList<string> Metadata => _response.Metadata;
+			public NamespaceId NamespaceId => _response.NamespaceId;
+			public RefName RefName => _response.RefName;
+			public DateTime CreatedAtUtc => _response.CreatedAtUtc;
+		}
+
+		readonly IHordeClient _hordeClient;
+
+		public ArtifactHttpCollection(IHordeClient hordeClient)
+			=> _hordeClient = hordeClient;
+
+		/// <inheritdoc/>
+		public Task<IArtifact> AddAsync(ArtifactName name, ArtifactType type, string? description, StreamId streamId, CommitId commitId, IEnumerable<string> keys, IEnumerable<string> metadata, AclScopeName scopeName, CancellationToken cancellationToken = default)
+		{
+			throw new NotImplementedException();
+		}
+
+		/// <inheritdoc/>
+		public Task DeleteAsync(IEnumerable<ArtifactId> ids, CancellationToken cancellationToken = default)
+		{
+			throw new NotImplementedException();
+		}
+
+		/// <inheritdoc/>
+		public async IAsyncEnumerable<IArtifact> FindAsync(StreamId? streamId = null, CommitId? minCommitId = null, CommitId? maxCommitId = null, ArtifactName? name = null, ArtifactType? type = null, IEnumerable<string>? keys = null, int maxResults = 100, [EnumeratorCancellation] CancellationToken cancellationToken = default)
+		{
+			HordeHttpClient hordeHttpClient = _hordeClient.CreateHttpClient();
+
+			List<GetArtifactResponse> responses = await hordeHttpClient.FindArtifactsAsync(streamId, minCommitId, maxCommitId, name, type, keys, maxResults, cancellationToken);
+			foreach (GetArtifactResponse response in responses)
+			{
+				yield return new Artifact(response);
+			}
+		}
+
+		/// <inheritdoc/>
+		public IAsyncEnumerable<IEnumerable<IArtifact>> FindExpiredAsync(ArtifactType type, DateTime? expireAtUtc, CancellationToken cancellationToken = default)
+		{
+			throw new NotImplementedException();
+		}
+
+		/// <inheritdoc/>
+		public async Task<IArtifact?> GetAsync(ArtifactId artifactId, CancellationToken cancellationToken = default)
+		{
+			HordeHttpClient hordeHttpClient = _hordeClient.CreateHttpClient();
+			GetArtifactResponse? response = await hordeHttpClient.GetArtifactAsync(artifactId, cancellationToken);
+			return new Artifact(response);
+		}
+	}
+}
diff --git a/Engine/Source/Programs/Shared/EpicGames.Horde/Artifacts/ArtifactHttpMessages.cs b/Engine/Source/Programs/Shared/EpicGames.Horde/Artifacts/ArtifactHttpMessages.cs
new file mode 100644
--- /dev/null
+++ b/Engine/Source/Programs/Shared/EpicGames.Horde/Artifacts/ArtifactHttpMessages.cs
@@ -0,0 +1,256 @@
+﻿// Copyright Epic Games, Inc. All Rights Reserved.
+
+using System;
+using System.Collections.Generic;
+using System.Text.Json.Serialization;
+using EpicGames.Core;
+using EpicGames.Horde.Commits;
+using EpicGames.Horde.Storage;
+using EpicGames.Horde.Streams;
+
+#pragma warning disable CA2227
+
+namespace EpicGames.Horde.Artifacts
+{
+	/// <summary>
+	/// Creates a new artifact
+	/// </summary>
+	/// <param name="Name">Name of the artifact</param>
+	/// <param name="Type">Additional search keys tagged on the artifact</param>
+	/// <param name="Description">Description for the artifact</param>
+	/// <param name="StreamId">Stream to create the artifact for</param>
+	/// <param name="Keys">Keys used to identify the artifact</param>
+	/// <param name="Metadata">Metadata for the artifact</param>
+	public record CreateArtifactRequest(ArtifactName Name, ArtifactType Type, string? Description, StreamId? StreamId, List<string> Keys, List<string> Metadata)
+	{
+		/// <summary>
+		/// Legacy Perforce changelist number.
+		/// </summary>
+		[Obsolete("Use CommitId instead")]
+		public int Change
+		{
+			get => _change ?? _commitId?.TryGetPerforceChange() ?? -1;
+			set => _change = value;
+		}
+		int? _change;
+
+		/// <summary>
+		/// Commit for the new artifact
+		/// </summary>
+		public CommitId CommitId
+		{
+			get => _commitId ?? CommitId.FromPerforceChange(_change) ?? CommitId.Empty;
+			set => _commitId = value;
+		}
+		CommitId? _commitId;
+	}
+
+	/// <summary>
+	/// Information about a created artifact
+	/// </summary>
+	/// <param name="ArtifactId">Identifier for the new artifact</param>
+	/// <param name="NamespaceId">Namespace that should be written to with artifact data</param>
+	/// <param name="RefName">Ref to write to</param>
+	/// <param name="PrevRefName">Ref for the artifact at the changelist prior to this one. Can be used to deduplicate against.</param>
+	/// <param name="Token">Token which can be used to upload blobs for the artifact, and read blobs from the previous artifact</param>
+	public record CreateArtifactResponse(ArtifactId ArtifactId, NamespaceId NamespaceId, RefName RefName, RefName? PrevRefName, string Token);
+
+	/// <summary>
+	/// Type of data to download for an artifact
+	/// </summary>
+	public enum DownloadArtifactFormat
+	{
+		/// <summary>
+		/// Download as a zip file
+		/// </summary>
+		Zip,
+
+		/// <summary>
+		/// Download as a UGS link
+		/// </summary>
+		Ugs
+	}
+
+	/// <summary>
+	/// Describes an artifact
+	/// </summary>
+	public class GetArtifactResponse
+	{
+		/// <inheritdoc cref="IArtifact.Id"/>
+		public ArtifactId Id { get; set; }
+
+		/// <inheritdoc cref="IArtifact.Name"/>
+		public ArtifactName Name { get; set; }
+
+		/// <inheritdoc cref="IArtifact.Type"/>
+		public ArtifactType Type { get; set; }
+
+		/// <inheritdoc cref="IArtifact.Description"/>
+		public string? Description { get; }
+
+		/// <inheritdoc cref="IArtifact.StreamId"/>
+		public StreamId StreamId { get; set; }
+
+		/// <summary>
+		/// Change number
+		/// </summary>
+		[Obsolete("Use CommitId instead")]
+		public int Change
+		{
+			get => _change ?? _commitId?.TryGetPerforceChange() ?? -1;
+			set => _change = value;
+		}
+		int? _change;
+
+		/// <inheritdoc cref="IArtifact.CommitId"/>
+		public CommitIdWithOrder CommitId
+		{
+			get => _commitId ?? CommitIdWithOrder.FromPerforceChange(_change) ?? CommitIdWithOrder.Empty;
+			set => _commitId = value;
+		}
+		CommitIdWithOrder? _commitId;
+
+		/// <inheritdoc cref="IArtifact.Keys"/>
+		public IReadOnlyList<string> Keys { get; set; }
+
+		/// <inheritdoc cref="IArtifact.Metadata"/>
+		public IReadOnlyList<string> Metadata { get; set; }
+
+		/// <inheritdoc cref="IArtifact.NamespaceId"/>
+		public NamespaceId NamespaceId { get; set; }
+
+		/// <inheritdoc cref="IArtifact.RefName"/>
+		public RefName RefName { get; set; }
+
+		/// <inheritdoc cref="IArtifact.CreatedAtUtc"/>
+		public DateTime CreatedAtUtc { get; }
+
+		/// <summary>
+		/// Default constructor
+		/// </summary>
+		[JsonConstructor]
+		public GetArtifactResponse()
+		{
+			Keys = Array.Empty<string>();
+			Metadata = Array.Empty<string>();
+		}
+
+		/// <summary>
+		/// Constructor
+		/// </summary>
+		/// <param name="artifact"></param>
+		public GetArtifactResponse(IArtifact artifact)
+		{
+			Id = artifact.Id;
+			Name = artifact.Name;
+			Type = artifact.Type;
+			Description = artifact.Description;
+			StreamId = artifact.StreamId;
+			Keys = artifact.Keys;
+			Metadata = artifact.Metadata;
+			NamespaceId = artifact.NamespaceId;
+			RefName = artifact.RefName;
+			CreatedAtUtc = artifact.CreatedAtUtc;
+		}
+	}
+
+	/// <summary>
+	/// Result of an artifact search
+	/// </summary>
+	public class FindArtifactsResponse
+	{
+		/// <summary>
+		/// List of artifacts matching the search criteria
+		/// </summary>
+		public List<GetArtifactResponse> Artifacts { get; set; } = new List<GetArtifactResponse>();
+	}
+
+	/// <summary>
+	/// Describes a file within an artifact
+	/// </summary>
+	public class GetArtifactFileEntryResponse
+	{
+		/// <summary>
+		/// Name of this file
+		/// </summary>
+		public string Name { get; }
+
+		/// <summary>
+		/// Length of this entry
+		/// </summary>
+		public long Length { get; }
+
+		/// <summary>
+		/// Hash of the target node
+		/// </summary>
+		public IoHash Hash { get; }
+
+		/// <summary>
+		/// Constructor
+		/// </summary>
+		public GetArtifactFileEntryResponse(string name, long length, IoHash hash)
+		{
+			Name = name;
+			Length = length;
+			Hash = hash;
+		}
+	}
+
+	/// <summary>
+	/// Describes a file within an artifact
+	/// </summary>
+	public class GetArtifactDirectoryEntryResponse : GetArtifactDirectoryResponse
+	{
+		/// <summary>
+		/// Name of this file
+		/// </summary>
+		public string Name { get; }
+
+		/// <summary>
+		/// Length of this entry
+		/// </summary>
+		public long Length { get; }
+
+		/// <summary>
+		/// Hash of the target node
+		/// </summary>
+		public IoHash Hash { get; }
+
+		/// <summary>
+		/// Constructor
+		/// </summary>
+		public GetArtifactDirectoryEntryResponse(string name, long length, IoHash hash)
+		{
+			Name = name;
+			Length = length;
+			Hash = hash;
+		}
+	}
+
+	/// <summary>
+	/// Describes a directory within an artifact
+	/// </summary>
+	public class GetArtifactDirectoryResponse
+	{
+		/// <summary>
+		/// Names of sub-directories
+		/// </summary>
+		public List<GetArtifactDirectoryEntryResponse>? Directories { get; set; }
+
+		/// <summary>
+		/// Files within the directory
+		/// </summary>
+		public List<GetArtifactFileEntryResponse>? Files { get; set; }
+	}
+
+	/// <summary>
+	/// Request to create a zip file with artifact data
+	/// </summary>
+	public class CreateZipRequest
+	{
+		/// <summary>
+		/// Filter lines for the zip. Uses standard <see cref="FileFilter"/> syntax.
+		/// </summary>
+		public List<string> Filter { get; set; } = new List<string>();
+	}
+}
diff --git a/Engine/Source/Programs/Shared/EpicGames.Horde/Artifacts/ArtifactMessages.cs b/Engine/Source/Programs/Shared/EpicGames.Horde/Artifacts/ArtifactMessages.cs
deleted file mode 100644
diff --git a/Engine/Source/Programs/Shared/EpicGames.Horde/Artifacts/IArtifact.cs b/Engine/Source/Programs/Shared/EpicGames.Horde/Artifacts/IArtifact.cs
new file mode 100644
--- /dev/null
+++ b/Engine/Source/Programs/Shared/EpicGames.Horde/Artifacts/IArtifact.cs
@@ -0,0 +1,71 @@
+// Copyright Epic Games, Inc. All Rights Reserved.
+
+using System;
+using System.Collections.Generic;
+using EpicGames.Horde.Commits;
+using EpicGames.Horde.Storage;
+using EpicGames.Horde.Streams;
+
+namespace EpicGames.Horde.Artifacts
+{
+	/// <summary>
+	/// Information about an artifact
+	/// </summary>
+	public interface IArtifact
+	{
+		/// <summary>
+		/// Identifier for the Artifact. Randomly generated.
+		/// </summary>
+		public ArtifactId Id { get; }
+
+		/// <summary>
+		/// Name of the artifact
+		/// </summary>
+		public ArtifactName Name { get; }
+
+		/// <summary>
+		/// Type of artifact
+		/// </summary>
+		public ArtifactType Type { get; }
+
+		/// <summary>
+		/// Description for the artifact
+		/// </summary>
+		public string? Description { get; }
+
+		/// <summary>
+		/// Identifier for the stream that produced the artifact
+		/// </summary>
+		public StreamId StreamId { get; }
+
+		/// <summary>
+		/// Change that the artifact corresponds to
+		/// </summary>
+		public CommitIdWithOrder CommitId { get; }
+
+		/// <summary>
+		/// Keys used to collate artifacts
+		/// </summary>
+		public IReadOnlyList<string> Keys { get; }
+
+		/// <summary>
+		/// Metadata for the artifact
+		/// </summary>
+		public IReadOnlyList<string> Metadata { get; }
+
+		/// <summary>
+		/// Storage namespace containing the data
+		/// </summary>
+		public NamespaceId NamespaceId { get; }
+
+		/// <summary>
+		/// Name of the ref containing the root data object
+		/// </summary>
+		public RefName RefName { get; }
+
+		/// <summary>
+		/// Time at which the artifact was created
+		/// </summary>
+		public DateTime CreatedAtUtc { get; }
+	}
+}
diff --git a/Engine/Source/Programs/Shared/EpicGames.Horde/Artifacts/IArtifactCollection.cs b/Engine/Source/Programs/Shared/EpicGames.Horde/Artifacts/IArtifactCollection.cs
new file mode 100644
--- /dev/null
+++ b/Engine/Source/Programs/Shared/EpicGames.Horde/Artifacts/IArtifactCollection.cs
@@ -0,0 +1,84 @@
+// Copyright Epic Games, Inc. All Rights Reserved.
+
+using System;
+using System.Collections.Generic;
+using System.Threading;
+using System.Threading.Tasks;
+using EpicGames.Horde.Acls;
+using EpicGames.Horde.Commits;
+using EpicGames.Horde.Streams;
+
+namespace EpicGames.Horde.Artifacts
+{
+	/// <summary>
+	/// Interface for a collection of artifacts
+	/// </summary>
+	public interface IArtifactCollection
+	{
+		/// <summary>
+		/// Creates a new artifact
+		/// </summary>
+		/// <param name="name">Name of the artifact</param>
+		/// <param name="type">Type identifier for the artifact</param>
+		/// <param name="description">Description for the artifact</param>
+		/// <param name="streamId">Stream that the artifact was built from</param>
+		/// <param name="commitId">Commit that the artifact was built from</param>
+		/// <param name="keys">Keys for the artifact</param>
+		/// <param name="metadata">Metadata for the artifact</param>
+		/// <param name="scopeName">Inherited scope used for permissions</param>
+		/// <param name="cancellationToken">Cancellation token for the operation</param>
+		/// <returns>The new log file document</returns>
+		Task<IArtifact> AddAsync(ArtifactName name, ArtifactType type, string? description, StreamId streamId, CommitId commitId, IEnumerable<string> keys, IEnumerable<string> metadata, AclScopeName scopeName, CancellationToken cancellationToken = default);
+
+		/// <summary>
+		/// Deletes artifacts
+		/// </summary>
+		/// <param name="ids">Ids to search for</param>
+		/// <param name="cancellationToken">Cancellation token for the operation</param>
+		/// <returns>Sequence of artifacts</returns>
+		Task DeleteAsync(IEnumerable<ArtifactId> ids, CancellationToken cancellationToken = default);
+
+		/// <summary>
+		/// Finds artifacts which are ready for expiry
+		/// </summary>
+		/// <param name="type">Type of artifacts to find</param>
+		/// <param name="expireAtUtc">Number of artifacts to keep</param>
+		/// <param name="cancellationToken">Cancellation token for the operation</param>
+		/// <returns>Sequence of artifacts</returns>
+		IAsyncEnumerable<IEnumerable<IArtifact>> FindExpiredAsync(ArtifactType type, DateTime? expireAtUtc, CancellationToken cancellationToken = default);
+
+		/// <summary>
+		/// Finds artifacts with the given keys.
+		/// </summary>
+		/// <param name="streamId">Stream to find artifacts for</param>
+		/// <param name="minCommitId">Minimum commit for the artifacts (inclusive)</param>
+		/// <param name="maxCommitId">Maximum commit for the artifacts (inclusive)</param>
+		/// <param name="name">Name of the artifact to search for</param>
+		/// <param name="type">The artifact type</param>
+		/// <param name="keys">Set of keys, all of which must all be present on any returned artifacts</param>
+		/// <param name="maxResults">Maximum number of results to return</param>
+		/// <param name="cancellationToken">Cancellation token for the operation</param>
+		/// <returns>Sequence of artifacts. Ordered by descending CL order, then by descending order in which they were created.</returns>
+		IAsyncEnumerable<IArtifact> FindAsync(StreamId? streamId = null, CommitId? minCommitId = null, CommitId? maxCommitId = null, ArtifactName? name = null, ArtifactType? type = null, IEnumerable<string>? keys = null, int maxResults = 100, CancellationToken cancellationToken = default);
+
+		/// <summary>
+		/// Gets an artifact by ID
+		/// </summary>
+		/// <param name="artifactId">Unique id of the artifact</param>
+		/// <param name="cancellationToken">Cancellation token for the operation</param>
+		/// <returns>The artifact document</returns>
+		Task<IArtifact?> GetAsync(ArtifactId artifactId, CancellationToken cancellationToken = default);
+	}
+
+	/// <summary>
+	/// Extension methods for <see cref="IArtifactCollection"/>
+	/// </summary>
+	public static class ArtifactCollectionExtensions
+	{
+		/// <inheritdoc cref="IArtifactCollection.DeleteAsync(IEnumerable{ArtifactId}, CancellationToken)"/>>
+		public static Task DeleteAsync(this IArtifactCollection artifactCollection, ArtifactId artifactId, CancellationToken cancellationToken = default)
+		{
+			return artifactCollection.DeleteAsync(new[] { artifactId }, cancellationToken);
+		}
+	}
+}
diff --git a/Engine/Source/Programs/Shared/EpicGames.Horde/HordeClient.cs b/Engine/Source/Programs/Shared/EpicGames.Horde/HordeClient.cs
--- a/Engine/Source/Programs/Shared/EpicGames.Horde/HordeClient.cs
+++ b/Engine/Source/Programs/Shared/EpicGames.Horde/HordeClient.cs
@@ -8,6 +8,7 @@ using System.Net.Http.Json;
 using System.Threading;
 using System.Threading.Tasks;
 using EpicGames.Core;
+using EpicGames.Horde.Artifacts;
 using EpicGames.Horde.Compute;
 using EpicGames.Horde.Compute.Clients;
 using EpicGames.Horde.Logs;
@@ -45,6 +46,9 @@ namespace EpicGames.Horde
 		/// <inheritdoc/>
 		public Uri ServerUrl => _serverUrl;
 
+		/// <inheritdoc/>
+		public IArtifactCollection Artifacts { get; }
+
 		/// <summary>
 		/// Constructor
 		/// </summary>
@@ -55,6 +59,8 @@ namespace EpicGames.Horde
 			_hordeOptions = hordeOptions.Value;
 			_loggerFactory = loggerFactory;
 			_logger = _loggerFactory.CreateLogger<HordeClient>();
+
+			Artifacts = new ArtifactHttpCollection(this);
 		}
 
 		/// <inheritdoc/>
diff --git a/Engine/Source/Programs/Shared/EpicGames.Horde/HordeHttpClient.cs b/Engine/Source/Programs/Shared/EpicGames.Horde/HordeHttpClient.cs
--- a/Engine/Source/Programs/Shared/EpicGames.Horde/HordeHttpClient.cs
+++ b/Engine/Source/Programs/Shared/EpicGames.Horde/HordeHttpClient.cs
@@ -162,63 +162,45 @@ namespace EpicGames.Horde
 			return await _httpClient.GetStreamAsync($"api/v2/artifacts/{id}/zip", cancellationToken);
 		}
 
-		/// <summary>
-		/// Finds artifacts with a set of ids or keys
-		/// </summary>
-		/// <param name="ids">Artifact ids to return</param>
-		/// <param name="keys">Keys to find</param>
-		/// <param name="cancellationToken">Cancellation token for the operation</param>
-		/// <returns>Information about all the artifacts</returns>
-		public async Task<List<GetArtifactResponse>> FindArtifactsAsync(IEnumerable<ArtifactId>? ids = null, IEnumerable<string>? keys = null, CancellationToken cancellationToken = default)
-		{
-			QueryStringBuilder queryParams = new QueryStringBuilder();
-			if (ids != null)
-			{
-				foreach (ArtifactId id in ids)
-				{
-					queryParams.Add("id", id.ToString());
-				}
-			}
-			if (keys != null)
-			{
-				foreach (string key in keys)
-				{
-					queryParams.Add("key", key);
-				}
-			}
-
-			FindArtifactsResponse response = await GetAsync<FindArtifactsResponse>(_httpClient, $"api/v2/artifacts?{queryParams}", cancellationToken);
-			return response.Artifacts;
-		}
-
 		/// <summary>
 		/// Finds artifacts with a certain type with an optional streamId
 		/// </summary>
-		/// <param name="type">Type to find</param>
 		/// <param name="streamId">Stream to look for the artifact in</param>
-		/// <param name="minChange">The minimum change number for the artifacts</param>
-		/// <param name="maxChange">The minimum change number for the artifacts</param>
+		/// <param name="minCommitId">The minimum change number for the artifacts</param>
+		/// <param name="maxCommitId">The minimum change number for the artifacts</param>
+		/// <param name="name">Name of the artifact</param>
+		/// <param name="type">Type to find</param>
 		/// <param name="keys">Keys for artifacts to return</param>
+		/// <param name="maxResults">Maximum number of results to return</param>
 		/// <param name="cancellationToken">Cancellation token for the operation</param>
 		/// <returns>Information about all the artifacts</returns>
-		public async Task<List<GetArtifactResponse>> FindArtifactsByTypeAsync(ArtifactType type, StreamId? streamId = null, int? minChange = null, int? maxChange = null, IEnumerable<string>? keys = null, CancellationToken cancellationToken = default)
+		public async Task<List<GetArtifactResponse>> FindArtifactsAsync(StreamId? streamId = null, CommitId? minCommitId = null, CommitId? maxCommitId = null, ArtifactName? name = null, ArtifactType? type = null, IEnumerable<string>? keys = null, int maxResults = 100, CancellationToken cancellationToken = default)
 		{
 			QueryStringBuilder queryParams = new QueryStringBuilder();
-			queryParams.Add("type", type.ToString());
 
 			if (streamId != null)
 			{
 				queryParams.Add("streamId", streamId.ToString()!);
 			}
 
-			if (minChange != null)
+			if (minCommitId != null)
+			{
+				queryParams.Add("minChange", minCommitId.ToString()!);
+			}
+
+			if (maxCommitId != null)
 			{
-				queryParams.Add("minChange", minChange.ToString()!);
+				queryParams.Add("maxChange", maxCommitId.ToString()!);
 			}
 
-			if (maxChange != null)
+			if (name != null)
 			{
-				queryParams.Add("maxChange", maxChange.ToString()!);
+				queryParams.Add("name", name.Value.ToString());
+			}
+
+			if (type != null)
+			{
+				queryParams.Add("type", type.Value.ToString());
 			}
 
 			if (keys != null)
@@ -229,6 +211,8 @@ namespace EpicGames.Horde
 				}
 			}
 
+			queryParams.Add("maxResults", maxResults.ToString());
+
 			FindArtifactsResponse response = await GetAsync<FindArtifactsResponse>(_httpClient, $"api/v2/artifacts?{queryParams}", cancellationToken);
 			return response.Artifacts;
 		}
diff --git a/Engine/Source/Programs/Shared/EpicGames.Horde/IHordeClient.cs b/Engine/Source/Programs/Shared/EpicGames.Horde/IHordeClient.cs
--- a/Engine/Source/Programs/Shared/EpicGames.Horde/IHordeClient.cs
+++ b/Engine/Source/Programs/Shared/EpicGames.Horde/IHordeClient.cs
@@ -24,6 +24,11 @@ namespace EpicGames.Horde
 		/// </summary>
 		Uri ServerUrl { get; }
 
+		/// <summary>
+		/// Accessor for the artifact collection
+		/// </summary>
+		IArtifactCollection Artifacts { get; }
+
 		/// <summary>
 		/// Connect to the Horde server
 		/// </summary>
diff --git a/Engine/Source/Programs/UnrealGameSync/UnrealGameSyncShared/ArchiveInfo.cs b/Engine/Source/Programs/UnrealGameSync/UnrealGameSyncShared/ArchiveInfo.cs
--- a/Engine/Source/Programs/UnrealGameSync/UnrealGameSyncShared/ArchiveInfo.cs
+++ b/Engine/Source/Programs/UnrealGameSync/UnrealGameSyncShared/ArchiveInfo.cs
@@ -317,7 +317,7 @@ namespace UnrealGameSync
 				ArtifactType artifactType = new ArtifactType("ugs-pcb");
 				string[] artifactKeys = new[] { $"ugs-project={projectIdentifier}" };
 
-				List<GetArtifactResponse> artifactResponses = await hordeHttpClient.FindArtifactsByTypeAsync(artifactType, keys: artifactKeys, cancellationToken: cancellationToken);
+				List<GetArtifactResponse> artifactResponses = await hordeHttpClient.FindArtifactsAsync(type: artifactType, keys: artifactKeys, cancellationToken: cancellationToken);
 				foreach (IGrouping<ArtifactName, GetArtifactResponse> group in artifactResponses.GroupBy(x => x.Name))
 				{
 					GetArtifactResponse? first = group.FirstOrDefault();
