Horde: Fix changelist number not being returned in artifact responses.

diff --git a/Engine/Source/Programs/Shared/EpicGames.Horde/Artifacts/ArtifactHttpMessages.cs b/Engine/Source/Programs/Shared/EpicGames.Horde/Artifacts/ArtifactHttpMessages.cs
--- a/Engine/Source/Programs/Shared/EpicGames.Horde/Artifacts/ArtifactHttpMessages.cs
+++ b/Engine/Source/Programs/Shared/EpicGames.Horde/Artifacts/ArtifactHttpMessages.cs
@@ -146,6 +146,7 @@ namespace EpicGames.Horde.Artifacts
 			Type = artifact.Type;
 			Description = artifact.Description;
 			StreamId = artifact.StreamId;
+			CommitId = artifact.CommitId;
 			Keys = artifact.Keys;
 			Metadata = artifact.Metadata;
 			NamespaceId = artifact.NamespaceId;
