UHT: Fixed code generation to support enums with no enumerators

diff --git a/Engine/Source/Programs/Shared/EpicGames.UHT/Exporters/CodeGen/UhtHeaderCodeGeneratorCppFile.cs b/Engine/Source/Programs/Shared/EpicGames.UHT/Exporters/CodeGen/UhtHeaderCodeGeneratorCppFile.cs
--- a/Engine/Source/Programs/Shared/EpicGames.UHT/Exporters/CodeGen/UhtHeaderCodeGeneratorCppFile.cs
+++ b/Engine/Source/Programs/Shared/EpicGames.UHT/Exporters/CodeGen/UhtHeaderCodeGeneratorCppFile.cs
@@ -475,18 +475,21 @@ namespace EpicGames.UHT.Exporters.CodeGen
 					builder.AppendMetaDataDecl(enumObj, null, null, MetaDataParamsName, 1);
 
 					// Enumerators
-					builder.Append("\tstatic constexpr UECodeGen_Private::FEnumeratorParam Enumerators[] = {\r\n");
-					int enumIndex = 0;
-					foreach (UhtEnumValue value in enumObj.EnumValues)
+					if (enumObj.EnumValues.Count > 0)
 					{
-						if (!enumObj.MetaData.TryGetValue("OverrideName", enumIndex, out string? keyName))
+						builder.Append("\tstatic constexpr UECodeGen_Private::FEnumeratorParam Enumerators[] = {\r\n");
+						int enumIndex = 0;
+						foreach (UhtEnumValue value in enumObj.EnumValues)
 						{
-							keyName = value.Name.ToString();
+							if (!enumObj.MetaData.TryGetValue("OverrideName", enumIndex, out string? keyName))
+							{
+								keyName = value.Name.ToString();
+							}
+							builder.Append("\t\t{ ").AppendUTF8LiteralString(keyName).Append(", (int64)").Append(value.Name).Append(" },\r\n");
+							++enumIndex;
 						}
-						builder.Append("\t\t{ ").AppendUTF8LiteralString(keyName).Append(", (int64)").Append(value.Name).Append(" },\r\n");
-						++enumIndex;
+						builder.Append("\t};\r\n");
 					}
-					builder.Append("\t};\r\n");
 					builder.Append("\tstatic const UECodeGen_Private::FEnumParams EnumParams;\r\n");
 					builder.Append("};\r\n");
 				}
@@ -498,9 +501,18 @@ namespace EpicGames.UHT.Exporters.CodeGen
 					builder.Append('\t').Append(enumDisplayNameFn).Append(",\r\n");
 					builder.Append('\t').AppendUTF8LiteralString(enumObj.SourceName).Append(",\r\n");
 					builder.Append('\t').AppendUTF8LiteralString(enumObj.CppType).Append(",\r\n");
-					builder.Append('\t').Append(staticsName).Append("::Enumerators,\r\n");
-					builder.Append('\t').Append(ObjectFlags).Append(",\r\n");
-					builder.Append("\tUE_ARRAY_COUNT(").Append(staticsName).Append("::Enumerators),\r\n");
+					if (enumObj.EnumValues.Count > 0)
+					{
+						builder.Append('\t').Append(staticsName).Append("::Enumerators,\r\n");
+						builder.Append('\t').Append(ObjectFlags).Append(",\r\n");
+						builder.Append("\tUE_ARRAY_COUNT(").Append(staticsName).Append("::Enumerators),\r\n");
+					}
+					else
+					{
+						builder.Append('\t').Append("nullptr,\r\n");
+						builder.Append('\t').Append(ObjectFlags).Append(",\r\n");
+						builder.Append("\t0,\r\n");
+					}
 					builder.Append('\t').Append(enumObj.EnumFlags.HasAnyFlags(EEnumFlags.Flags) ? "EEnumFlags::Flags" : "EEnumFlags::None").Append(",\r\n");
 					builder.Append("\t(uint8)UEnum::ECppForm::").Append(enumObj.CppForm.ToString()).Append(",\r\n");
 					builder.Append('\t').AppendMetaDataParams(enumObj, staticsName, MetaDataParamsName).Append("\r\n");
