EpicGames.Redis: Rename files containing key classes to *Key.cs.

diff --git a/Engine/Source/Programs/Shared/EpicGames.Redis/RedisHash.cs b/Engine/Source/Programs/Shared/EpicGames.Redis/RedisHash.cs
deleted file mode 100644
diff --git a/Engine/Source/Programs/Shared/EpicGames.Redis/RedisHashKey.cs b/Engine/Source/Programs/Shared/EpicGames.Redis/RedisHashKey.cs
new file mode 100644
--- /dev/null
+++ b/Engine/Source/Programs/Shared/EpicGames.Redis/RedisHashKey.cs
@@ -0,0 +1,360 @@
+// Copyright Epic Games, Inc. All Rights Reserved.
+
+using System;
+using System.Collections.Generic;
+using System.Linq;
+using System.Linq.Expressions;
+using System.Threading.Tasks;
+using StackExchange.Redis;
+
+namespace EpicGames.Redis
+{
+	/// <summary>
+	/// Represents a redis hash key, with members corresponding to the property names of a type
+	/// </summary>
+	/// <typeparam name="T">Type of the hash fields</typeparam>
+	public record struct RedisHashKey<T>(RedisKey Inner)
+	{
+		/// <summary>
+		/// Implicit conversion to typed redis key.
+		/// </summary>
+		/// <param name="key">Key to convert</param>
+		public static implicit operator RedisHashKey<T>(string key) => new RedisHashKey<T>(new RedisKey(key));
+
+		/// <summary>
+		/// Implicit conversion to untyped redis keys.
+		/// </summary>
+		/// <param name="key">Key to convert</param>
+		public static implicit operator TypedRedisKey(RedisHashKey<T> key) => new TypedRedisKey(key.Inner);
+	}
+
+	/// <summary>
+	/// Represents a typed Redis hash with given key/value types
+	/// </summary>
+	/// <typeparam name="TName">Type of the hash key</typeparam>
+	/// <typeparam name="TValue">Type of the hash value</typeparam>
+	public record struct RedisHashKey<TName, TValue>(RedisKey Inner)
+	{
+		/// <summary>
+		/// Implicit conversion to typed redis key.
+		/// </summary>
+		/// <param name="key">Key to convert</param>
+		public static implicit operator RedisHashKey<TName, TValue>(string key) => new RedisHashKey<TName, TValue>(new RedisKey(key));
+
+		/// <summary>
+		/// Implicit conversion to untyped redis keys.
+		/// </summary>
+		/// <param name="key">Key to convert</param>
+		public static implicit operator TypedRedisKey(RedisHashKey<TName, TValue> key) => new TypedRedisKey(key.Inner);
+	}
+
+	/// <inheritdoc cref="HashEntry"/>
+	public readonly struct HashEntry<TName, TValue>
+	{
+		/// <inheritdoc cref="HashEntry.Name"/>
+		public readonly TName Name { get; }
+
+		/// <inheritdoc cref="HashEntry.Value"/>
+		public readonly TValue Value { get; }
+
+		/// <summary>
+		/// Constructor
+		/// </summary>
+		/// <param name="name"></param>
+		/// <param name="value"></param>
+		public HashEntry(TName name, TValue value)
+		{
+			Name = name;
+			Value = value;
+		}
+
+		/// <summary>
+		/// Implicit conversion to a <see cref="HashEntry"/>
+		/// </summary>
+		/// <param name="entry"></param>
+		public static implicit operator HashEntry(HashEntry<TName, TValue> entry)
+		{
+			return new HashEntry(RedisSerializer.Serialize(entry.Name), RedisSerializer.Serialize(entry.Value));
+		}
+
+		/// <summary>
+		/// Implicit conversion to a <see cref="KeyValuePair{TName, TValue}"/>
+		/// </summary>
+		/// <param name="entry"></param>
+		public static implicit operator KeyValuePair<TName, TValue>(HashEntry<TName, TValue> entry)
+		{
+			return new KeyValuePair<TName, TValue>(entry.Name, entry.Value);
+		}
+	}
+
+	/// <summary>
+	/// Extension methods for hashes
+	/// </summary>
+	public static class RedisHashKeyExtensions
+	{
+		/// <summary>
+		/// Helper method to convert an array of hash entries into a dictionary
+		/// </summary>
+		public static async Task<Dictionary<TName, TValue>> ToDictionaryAsync<TName, TValue>(this Task<HashEntry<TName, TValue>[]> entries) where TName : notnull
+			=> (await entries).ToDictionary(x => x.Name, x => x.Value);
+
+		/// <summary>
+		/// Helper method to convert an array of hash entries into a dictionary
+		/// </summary>
+		public static async Task<Dictionary<TName, TValue>> ToDictionaryAsync<TName, TValue>(this Task<HashEntry<TName, TValue>[]> entries, IEqualityComparer<TName>? comparer) where TName : notnull
+			=> (await entries).ToDictionary(x => x.Name, x => x.Value, comparer);
+
+		#region Conditions
+
+		/// <inheritdoc cref="Condition.HashEqual(RedisKey, RedisValue, RedisValue)"/>
+		public static Condition Equal<TRecord, TValue>(this RedisHashKey<TRecord> key, Expression<Func<TRecord, TValue>> selector, TValue value)
+		{
+			MemberExpression memberExpression = (selector.Body as MemberExpression) ?? throw new InvalidOperationException("Expression must be a property accessor");
+			return Condition.HashEqual(key.Inner, memberExpression.Member.Name, RedisSerializer.Serialize(value));
+		}
+
+		/// <inheritdoc cref="Condition.HashEqual(RedisKey, RedisValue, RedisValue)"/>
+		public static Condition Equal<TName, TValue>(this RedisHashKey<TName, TValue> key, TName name, TValue value)
+			=> Condition.HashEqual(key.Inner, RedisSerializer.Serialize(name), RedisSerializer.Serialize(value));
+
+		/// <inheritdoc cref="Condition.HashExists(RedisKey, RedisValue)"/>
+		public static Condition Exists<TRecord, TValue>(this RedisHashKey<TRecord> key, Expression<Func<TRecord, TValue>> selector)
+		{
+			MemberExpression memberExpression = (selector.Body as MemberExpression) ?? throw new InvalidOperationException("Expression must be a property accessor");
+			return Condition.HashExists(key.Inner, memberExpression.Member.Name);
+		}
+
+		/// <inheritdoc cref="Condition.HashExists(RedisKey, RedisValue)"/>
+		public static Condition Exists<TName, TValue>(this RedisHashKey<TName, TValue> key, TName name)
+			=> Condition.HashExists(key.Inner, RedisSerializer.Serialize(name));
+
+		/// <inheritdoc cref="Condition.HashLengthEqual(RedisKey, Int64)"/>
+		public static Condition LengthEqual<TName, TValue>(this RedisHashKey<TName, TValue> key, long length)
+			=> Condition.HashLengthEqual(key.Inner, length);
+
+		/// <inheritdoc cref="Condition.HashLengthGreaterThan(RedisKey, Int64)"/>
+		public static Condition LengthGreaterThan<TName, TValue>(this RedisHashKey<TName, TValue> key, long length)
+			=> Condition.HashLengthGreaterThan(key.Inner, length);
+
+		/// <inheritdoc cref="Condition.HashLengthLessThan(RedisKey, Int64)"/>
+		public static Condition LengthLessThan<TName, TValue>(this RedisHashKey<TName, TValue> key, long length)
+			=> Condition.HashLengthLessThan(key.Inner, length);
+
+		/// <inheritdoc cref="Condition.HashNotExists(RedisKey, RedisValue)"/>
+		public static Condition NotExists<TRecord, TValue>(this RedisHashKey<TRecord> key, Expression<Func<TRecord, TValue>> selector)
+		{
+			MemberExpression memberExpression = (selector.Body as MemberExpression) ?? throw new InvalidOperationException("Expression must be a property accessor");
+			return Condition.HashNotExists(key.Inner, memberExpression.Member.Name);
+		}
+
+		/// <inheritdoc cref="Condition.HashNotExists(RedisKey, RedisValue)"/>
+		public static Condition NotExists<TName, TValue>(this RedisHashKey<TName, TValue> key, TName name)
+			=> Condition.HashNotExists(key.Inner, RedisSerializer.Serialize(name));
+
+		/// <inheritdoc cref="Condition.HashEqual(RedisKey, RedisValue, RedisValue)"/>
+		public static Condition NotEqual<TRecord, TValue>(this RedisHashKey<TRecord> key, Expression<Func<TRecord, TValue>> selector, TValue value)
+		{
+			MemberExpression memberExpression = (selector.Body as MemberExpression) ?? throw new InvalidOperationException("Expression must be a property accessor");
+			return Condition.HashNotEqual(key.Inner, memberExpression.Member.Name, RedisSerializer.Serialize(value));
+		}
+
+		/// <inheritdoc cref="Condition.HashNotEqual(RedisKey, RedisValue, RedisValue)"/>
+		public static Condition NotEqual<TName, TValue>(this RedisHashKey<TName, TValue> key, TName name, TValue value)
+			=> Condition.HashNotEqual(key.Inner, RedisSerializer.Serialize(name), RedisSerializer.Serialize(value));
+
+		#endregion
+
+		#region HashDecrementAsync
+
+		/// <inheritdoc cref="IDatabaseAsync.HashDecrementAsync(RedisKey, RedisValue, Int64, CommandFlags)"/>
+		public static Task<long> HashDecrementAsync<TRecord>(this IDatabaseAsync target, RedisHashKey<TRecord> key, Expression<Func<TRecord, long>> selector, long value = 1L, CommandFlags flags = CommandFlags.None)
+		{
+			MemberExpression memberExpression = (selector.Body as MemberExpression) ?? throw new InvalidOperationException("Expression must be a property accessor");
+			return target.HashDecrementAsync(key.Inner, memberExpression.Member.Name, value, flags);
+		}
+
+		/// <inheritdoc cref="IDatabaseAsync.HashDecrementAsync(RedisKey, RedisValue, Double, CommandFlags)"/>
+		public static Task<double> HashDecrementAsync<TRecord>(this IDatabaseAsync target, RedisHashKey<TRecord> key, Expression<Func<TRecord, double>> selector, double value = 1.0, CommandFlags flags = CommandFlags.None)
+		{
+			MemberExpression memberExpression = (selector.Body as MemberExpression) ?? throw new InvalidOperationException("Expression must be a property accessor");
+			return target.HashDecrementAsync(key.Inner, memberExpression.Member.Name, value, flags);
+		}
+
+		/// <inheritdoc cref="IDatabaseAsync.HashDecrementAsync(RedisKey, RedisValue, Int64, CommandFlags)"/>
+		public static Task<long> HashDecrementAsync<TName>(this IDatabaseAsync target, RedisHashKey<TName, long> key, TName name, long value = 1L, CommandFlags flags = CommandFlags.None)
+		{
+			return target.HashDecrementAsync(key.Inner, RedisSerializer.Serialize<TName>(name), value, flags);
+		}
+
+		/// <inheritdoc cref="IDatabaseAsync.HashDecrementAsync(RedisKey, RedisValue, Double, CommandFlags)"/>
+		public static Task<double> HashDecrementAsync<TName>(this IDatabaseAsync target, RedisHashKey<TName, long> key, TName name, double value = 1.0, CommandFlags flags = CommandFlags.None)
+		{
+			return target.HashDecrementAsync(key.Inner, RedisSerializer.Serialize<TName>(name), value, flags);
+		}
+
+		#endregion
+
+		#region HashDeleteAsync
+
+		/// <inheritdoc cref="IDatabaseAsync.HashDeleteAsync(RedisKey, RedisValue, CommandFlags)"/>
+		public static Task<bool> HashDeleteAsync<TRecord>(this IDatabaseAsync target, RedisHashKey<TRecord> key, Expression<Func<TRecord, double>> selector, CommandFlags flags = CommandFlags.None)
+		{
+			MemberExpression memberExpression = (selector.Body as MemberExpression) ?? throw new InvalidOperationException("Expression must be a property accessor");
+			return target.HashDeleteAsync(key.Inner, memberExpression.Member.Name, flags);
+		}
+
+		/// <inheritdoc cref="IDatabaseAsync.HashDeleteAsync(RedisKey, RedisValue, CommandFlags)"/>
+		public static Task<bool> HashDeleteAsync<TName, TValue>(this IDatabaseAsync target, RedisHashKey<TName, TValue> key, TName name, CommandFlags flags = CommandFlags.None)
+		{
+			return target.HashDeleteAsync(key.Inner, RedisSerializer.Serialize(name), flags);
+		}
+
+		/// <inheritdoc cref="IDatabaseAsync.HashDeleteAsync(RedisKey, RedisValue[], CommandFlags)"/>
+		public static Task<long> HashDeleteAsync<TName, TValue>(this IDatabaseAsync target, RedisHashKey<TName, TValue> key, TName[] names, CommandFlags flags = CommandFlags.None)
+		{
+			return target.HashDeleteAsync(key.Inner, RedisSerializer.Serialize(names), flags);
+		}
+
+		#endregion
+
+		#region HashExistsAsync
+
+		/// <inheritdoc cref="IDatabaseAsync.HashExistsAsync(RedisKey, RedisValue, CommandFlags)"/>
+		public static Task<bool> HashExistsAsync<TName, TValue>(this IDatabaseAsync target, RedisHashKey<TName, TValue> key, TName name, CommandFlags flags = CommandFlags.None)
+		{
+			return target.HashExistsAsync(key.Inner, RedisSerializer.Serialize(name), flags);
+		}
+
+		#endregion
+
+		#region HashGetAsync
+
+		/// <inheritdoc cref="IDatabaseAsync.HashGetAsync(RedisKey, RedisValue, CommandFlags)"/>
+		public static Task<TValue> HashGetAsync<TRecord, TValue>(this IDatabaseAsync target, RedisHashKey<TRecord> key, Expression<Func<TRecord, TValue>> selector, CommandFlags flags = CommandFlags.None)
+		{
+			MemberExpression memberExpression = (selector.Body as MemberExpression) ?? throw new InvalidOperationException("Expression must be a property accessor");
+			return target.HashGetAsync(key.Inner, memberExpression.Member.Name, flags).DeserializeAsync<TValue>();
+		}
+
+		/// <inheritdoc cref="IDatabaseAsync.HashGetAsync(RedisKey, RedisValue, CommandFlags)"/>
+		public static Task<TValue> HashGetAsync<TName, TValue>(this IDatabaseAsync target, RedisHashKey<TName, TValue> key, TName name, CommandFlags flags = CommandFlags.None)
+		{
+			return target.HashGetAsync(key.Inner, RedisSerializer.Serialize(name), flags).DeserializeAsync<TValue>();
+		}
+
+		/// <inheritdoc cref="IDatabaseAsync.HashGetAsync(RedisKey, RedisValue[], CommandFlags)"/>
+		public static Task<TValue[]> HashGetAsync<TName, TValue>(this IDatabaseAsync target, RedisHashKey<TName, TValue> key, TName[] names, CommandFlags flags = CommandFlags.None)
+		{
+			return target.HashGetAsync(key.Inner, RedisSerializer.Serialize(names), flags).DeserializeAsync<TValue>();
+		}
+
+		#endregion
+
+		#region HashGetAllAsync
+
+		/// <inheritdoc cref="IDatabaseAsync.HashGetAllAsync(RedisKey, CommandFlags)"/>
+		public static async Task<HashEntry<TName, TValue>[]> HashGetAllAsync<TName, TValue>(this IDatabaseAsync target, RedisHashKey<TName, TValue> key, CommandFlags flags = CommandFlags.None)
+		{
+			HashEntry[] entries = await target.HashGetAllAsync(key.Inner, flags);
+			return Array.ConvertAll(entries, x => new HashEntry<TName, TValue>(RedisSerializer.Deserialize<TName>(x.Name)!, RedisSerializer.Deserialize<TValue>(x.Value)!));
+		}
+
+		#endregion
+
+		#region HashIncrementAsync
+
+		/// <inheritdoc cref="IDatabaseAsync.HashIncrementAsync(RedisKey, RedisValue, Int64, CommandFlags)"/>
+		public static Task<long> HashIncrementAsync<TRecord>(this IDatabaseAsync target, RedisHashKey<TRecord> key, Expression<Func<TRecord, long>> selector, long value = 1L, CommandFlags flags = CommandFlags.None)
+		{
+			MemberExpression memberExpression = (selector.Body as MemberExpression) ?? throw new InvalidOperationException("Expression must be a property accessor");
+			return target.HashIncrementAsync(key.Inner, memberExpression.Member.Name, value, flags);
+		}
+
+		/// <inheritdoc cref="IDatabaseAsync.HashIncrementAsync(RedisKey, RedisValue, Double, CommandFlags)"/>
+		public static Task<double> HashIncrementAsync<TRecord>(this IDatabaseAsync target, RedisHashKey<TRecord> key, Expression<Func<TRecord, double>> selector, double value = 1.0, CommandFlags flags = CommandFlags.None)
+		{
+			MemberExpression memberExpression = (selector.Body as MemberExpression) ?? throw new InvalidOperationException("Expression must be a property accessor");
+			return target.HashIncrementAsync(key.Inner, memberExpression.Member.Name, value, flags);
+		}
+
+		/// <inheritdoc cref="IDatabaseAsync.HashIncrementAsync(RedisKey, RedisValue, Int64, CommandFlags)"/>
+		public static Task<long> HashIncrementAsync<TName>(this IDatabaseAsync target, RedisHashKey<TName, long> key, TName name, long value = 1L, CommandFlags flags = CommandFlags.None)
+		{
+			return target.HashIncrementAsync(key.Inner, RedisSerializer.Serialize<TName>(name), value, flags);
+		}
+
+		/// <inheritdoc cref="IDatabaseAsync.HashIncrementAsync(RedisKey, RedisValue, Double, CommandFlags)"/>
+		public static Task<double> HashIncrementAsync<TName>(this IDatabaseAsync target, RedisHashKey<TName, double> key, TName name, double value = 1.0, CommandFlags flags = CommandFlags.None)
+		{
+			return target.HashIncrementAsync(key.Inner, RedisSerializer.Serialize<TName>(name), value, flags);
+		}
+
+		#endregion
+
+		#region HashKeysAsync
+
+		/// <inheritdoc cref="IDatabaseAsync.HashKeysAsync(RedisKey, CommandFlags)"/>
+		public static Task<TName[]> HashKeysAsync<TName, TValue>(this IDatabaseAsync target, RedisHashKey<TName, TValue> key, CommandFlags flags = CommandFlags.None)
+		{
+			return target.HashKeysAsync(key.Inner, flags).DeserializeAsync<TName>();
+		}
+
+		#endregion
+
+		#region HashLengthAsync
+
+		/// <inheritdoc cref="IDatabaseAsync.HashLengthAsync(RedisKey, CommandFlags)"/>
+		public static Task<long> HashLengthAsync(this IDatabaseAsync target, TypedRedisKey key, CommandFlags flags = CommandFlags.None)
+		{
+			return target.HashLengthAsync(key, flags);
+		}
+
+		#endregion
+
+		#region HashScanAsync
+
+		/// <inheritdoc cref="IDatabaseAsync.HashScanAsync(RedisKey, RedisValue, int, long, int, CommandFlags)"/>
+		public static async IAsyncEnumerable<HashEntry<TName, TValue>> HashScanAsync<TName, TValue>(this IDatabaseAsync target, RedisHashKey<TName, TValue> key, RedisValue pattern, int pageSize = 250, long cursor = 0, int pageOffset = 0, CommandFlags flags = CommandFlags.None)
+		{
+			await foreach (HashEntry entry in target.HashScanAsync(key.Inner, pattern, pageSize, cursor, pageOffset, flags))
+			{
+				yield return new HashEntry<TName, TValue>(RedisSerializer.Deserialize<TName>(entry.Name), RedisSerializer.Deserialize<TValue>(entry.Value));
+			}
+		}
+
+		#endregion
+
+		#region HashSetAsync
+
+		/// <inheritdoc cref="IDatabaseAsync.HashSetAsync(RedisKey, RedisValue, RedisValue, When, CommandFlags)"/>
+		public static Task HashSetAsync<TRecord, TValue>(this IDatabaseAsync target, RedisHashKey<TRecord> key, Expression<Func<TRecord, TValue>> selector, TValue value, When when = When.Always, CommandFlags flags = CommandFlags.None)
+		{
+			MemberExpression memberExpression = (selector.Body as MemberExpression) ?? throw new InvalidOperationException("Expression must be a property accessor");
+			return target.HashSetAsync(key.Inner, memberExpression.Member.Name, RedisSerializer.Serialize(value), when, flags);
+		}
+
+		/// <inheritdoc cref="IDatabaseAsync.HashSetAsync(RedisKey, RedisValue, RedisValue, When, CommandFlags)"/>
+		public static Task HashSetAsync<TName, TValue>(this IDatabaseAsync target, RedisHashKey<TName, TValue> key, TName name, TValue value, When when = When.Always, CommandFlags flags = CommandFlags.None)
+		{
+			return target.HashSetAsync(key.Inner, RedisSerializer.Serialize(name), RedisSerializer.Serialize(value), when, flags);
+		}
+
+		/// <inheritdoc cref="IDatabaseAsync.HashSetAsync(RedisKey, HashEntry[], CommandFlags)"/>
+		public static Task HashSetAsync<TName, TValue>(this IDatabaseAsync target, RedisHashKey<TName, TValue> key, IEnumerable<HashEntry<TName, TValue>> entries, CommandFlags flags = CommandFlags.None)
+		{
+			return target.HashSetAsync(key.Inner, entries.Select(x => (HashEntry)x).ToArray(), flags);
+		}
+
+		#endregion
+
+		#region HashValuesAsync
+
+		/// <inheritdoc cref="IDatabaseAsync.HashValuesAsync(RedisKey, CommandFlags)"/>
+		public static Task<TValue[]> HashValuesAsync<TName, TValue>(this IDatabaseAsync target, RedisHashKey<TName, TValue> key, CommandFlags flags = CommandFlags.None)
+		{
+			return target.HashValuesAsync(key.Inner, flags).DeserializeAsync<TValue>();
+		}
+
+		#endregion
+	}
+}
diff --git a/Engine/Source/Programs/Shared/EpicGames.Redis/RedisList.cs b/Engine/Source/Programs/Shared/EpicGames.Redis/RedisList.cs
deleted file mode 100644
diff --git a/Engine/Source/Programs/Shared/EpicGames.Redis/RedisListKey.cs b/Engine/Source/Programs/Shared/EpicGames.Redis/RedisListKey.cs
new file mode 100644
--- /dev/null
+++ b/Engine/Source/Programs/Shared/EpicGames.Redis/RedisListKey.cs
@@ -0,0 +1,200 @@
+// Copyright Epic Games, Inc. All Rights Reserved.
+
+using System;
+using System.Collections.Generic;
+using System.Threading.Tasks;
+using StackExchange.Redis;
+
+namespace EpicGames.Redis
+{
+	/// <summary>
+	/// Represents a typed Redis list with a given key
+	/// </summary>
+	/// <typeparam name="TElement">The type of element stored in the set</typeparam>
+	public record struct RedisListKey<TElement>(RedisKey Inner)
+	{
+		/// <summary>
+		/// Implicit conversion to typed redis key.
+		/// </summary>
+		/// <param name="key">Key to convert</param>
+		public static implicit operator RedisListKey<TElement>(string key) => new RedisListKey<TElement>(new RedisKey(key));
+
+		/// <summary>
+		/// Implicit conversion to untyped redis keys.
+		/// </summary>
+		/// <param name="key">Key to convert</param>
+		public static implicit operator TypedRedisKey(RedisListKey<TElement> key) => key.Inner;
+	}
+
+	/// <summary>
+	/// Extension methods for sets
+	/// </summary>
+	public static class RedisListKeyExtensions
+	{
+		#region Conditions
+
+		/// <inheritdoc cref="Condition.ListIndexEqual(RedisKey, Int64, RedisValue)"/>
+		public static Condition IndexEqual<TElement>(this RedisListKey<TElement> key, long index, TElement value)
+			=> Condition.ListIndexEqual(key.Inner, index, RedisSerializer.Serialize(value));
+
+		/// <inheritdoc cref="Condition.ListIndexExists(RedisKey, Int64)"/>
+		public static Condition IndexExists<TElement>(this RedisListKey<TElement> key, long index)
+			=> Condition.ListIndexExists(key.Inner, index);
+
+		/// <inheritdoc cref="Condition.ListIndexNotEqual(RedisKey, Int64, RedisValue)"/>
+		public static Condition IndexNotEqual<TElement>(this RedisListKey<TElement> key, long index, TElement value)
+			=> Condition.ListIndexNotEqual(key.Inner, index, RedisSerializer.Serialize(value));
+
+		/// <inheritdoc cref="Condition.ListIndexNotExists(RedisKey, Int64)"/>
+		public static Condition IndexNotExists<TElement>(this RedisListKey<TElement> key, long index)
+			=> Condition.ListIndexNotExists(key.Inner, index);
+
+		/// <inheritdoc cref="Condition.ListLengthEqual(RedisKey, Int64)"/>
+		public static Condition LengthEqual<TElement>(this RedisListKey<TElement> key, long length)
+			=> Condition.ListLengthEqual(key.Inner, length);
+
+		/// <inheritdoc cref="Condition.ListLengthGreaterThan(RedisKey, Int64)"/>
+		public static Condition LengthGreaterThan<TElement>(this RedisListKey<TElement> key, long length)
+			=> Condition.ListLengthGreaterThan(key.Inner, length);
+
+		/// <inheritdoc cref="Condition.ListLengthLessThan(RedisKey, Int64)"/>
+		public static Condition LengthLessThan<TElement>(this RedisListKey<TElement> key, long length)
+			=> Condition.ListLengthLessThan(key.Inner, length);
+
+		#endregion
+
+		#region ListGetByIndexAsync
+
+		/// <inheritdoc cref="IDatabaseAsync.ListGetByIndexAsync(RedisKey, Int64, CommandFlags)"/>
+		public static Task<TElement> ListGetByIndexAsync<TElement>(this IDatabaseAsync target, RedisListKey<TElement> key, long index, CommandFlags flags = CommandFlags.None)
+		{
+			return target.ListGetByIndexAsync(key.Inner, index, flags).DeserializeAsync<TElement>();
+		}
+
+		#endregion
+
+		#region ListInsertAfterAsync
+
+		/// <inheritdoc cref="IDatabaseAsync.ListInsertAfterAsync(RedisKey, RedisValue, RedisValue, CommandFlags)"/>
+		public static Task<long> ListInsertAfterAsync<TElement>(this IDatabaseAsync target, RedisListKey<TElement> key, TElement pivot, TElement item, CommandFlags flags = CommandFlags.None)
+		{
+			return target.ListInsertAfterAsync(key.Inner, RedisSerializer.Serialize(pivot), RedisSerializer.Serialize(item), flags);
+		}
+
+		#endregion
+
+		#region ListInsertBeforeAsync
+
+		/// <inheritdoc cref="IDatabaseAsync.ListInsertBeforeAsync(RedisKey, RedisValue, RedisValue, CommandFlags)"/>
+		public static Task<long> ListInsertBeforeAsync<TElement>(this IDatabaseAsync target, RedisListKey<TElement> key, TElement pivot, TElement item, CommandFlags flags = CommandFlags.None)
+		{
+			RedisValue pivotValue = RedisSerializer.Serialize(pivot);
+			RedisValue itemValue = RedisSerializer.Serialize(item);
+			return target.ListInsertBeforeAsync(key.Inner, pivotValue, itemValue, flags);
+		}
+
+		#endregion
+
+		#region ListLeftPopAsync
+
+		/// <inheritdoc cref="IDatabaseAsync.ListLeftPopAsync(RedisKey, CommandFlags)"/>
+		public static Task<TElement> ListLeftPopAsync<TElement>(this IDatabaseAsync target, RedisListKey<TElement> key, CommandFlags flags = CommandFlags.None)
+		{
+			return target.ListLeftPopAsync(key.Inner, flags).DeserializeAsync<TElement>();
+		}
+
+		#endregion
+
+		#region ListLeftPushAsync
+
+		/// <inheritdoc cref="IDatabaseAsync.ListLeftPushAsync(RedisKey, RedisValue, When, CommandFlags)"/>
+		public static Task<long> ListLeftPushAsync<TElement>(this IDatabaseAsync target, RedisListKey<TElement> key, TElement item, When when = When.Always, CommandFlags flags = CommandFlags.None)
+		{
+			return target.ListLeftPushAsync(key.Inner, RedisSerializer.Serialize(item), when, flags);
+		}
+
+		/// <inheritdoc cref="IDatabaseAsync.ListLeftPushAsync(RedisKey, RedisValue[], When, CommandFlags)"/>
+		public static Task<long> ListLeftPushAsync<TElement>(this IDatabaseAsync target, RedisListKey<TElement> key, TElement[] values, When when = When.Always, CommandFlags flags = CommandFlags.None)
+		{
+			return target.ListLeftPushAsync(key.Inner, RedisSerializer.Serialize(values), when, flags);
+		}
+
+		#endregion
+
+		#region ListLengthAsync
+
+		/// <inheritdoc cref="IDatabaseAsync.ListLengthAsync(RedisKey, CommandFlags)"/>
+		public static Task<long> ListLengthAsync<TElement>(this IDatabaseAsync target, RedisListKey<TElement> key)
+		{
+			return target.ListLengthAsync(key.Inner);
+		}
+
+		#endregion
+
+		#region ListRangeAsync
+
+		/// <inheritdoc cref="IDatabaseAsync.ListRangeAsync(RedisKey, Int64, Int64, CommandFlags)"/>
+		public static Task<TElement[]> ListRangeAsync<TElement>(this IDatabaseAsync target, RedisListKey<TElement> key, long start = 0, long stop = -1, CommandFlags flags = CommandFlags.None)
+		{
+			return target.ListRangeAsync(key.Inner, start, stop, flags).DeserializeAsync<TElement>();
+		}
+
+		#endregion
+
+		#region ListRemoveAsync
+
+		/// <inheritdoc cref="IDatabaseAsync.ListRemoveAsync(RedisKey, RedisValue, Int64, CommandFlags)"/>
+		public static Task<long> ListRemoveAsync<TElement>(this IDatabaseAsync target, RedisListKey<TElement> key, TElement value, long count = 0L, CommandFlags flags = CommandFlags.None)
+		{
+			return target.ListRemoveAsync(key.Inner, RedisSerializer.Serialize(value), count, flags);
+		}
+
+		#endregion
+
+		#region ListRightPopAsync
+
+		/// <inheritdoc cref="IDatabaseAsync.ListRightPopAsync(RedisKey, CommandFlags)"/>
+		public static Task<TElement> ListRightPopAsync<TElement>(this IDatabaseAsync target, RedisListKey<TElement> key, CommandFlags flags = CommandFlags.None)
+		{
+			return target.ListRightPopAsync(key.Inner, flags).DeserializeAsync<TElement>();
+		}
+
+		#endregion
+
+		#region ListRightPushAsync
+
+		/// <inheritdoc cref="IDatabaseAsync.ListRightPushAsync(RedisKey, RedisValue, When, CommandFlags)"/>
+		public static Task<long> ListRightPushAsync<TElement>(this IDatabaseAsync target, RedisListKey<TElement> key, TElement item, When when = When.Always, CommandFlags flags = CommandFlags.None)
+		{
+			return target.ListRightPushAsync(key.Inner, RedisSerializer.Serialize(item), when, flags);
+		}
+
+		/// <inheritdoc cref="IDatabaseAsync.ListRightPushAsync(RedisKey, RedisValue[], When, CommandFlags)"/>
+		public static Task<long> ListRightPushAsync<TElement>(this IDatabaseAsync target, RedisListKey<TElement> key, IEnumerable<TElement> values, When when = When.Always, CommandFlags flags = CommandFlags.None)
+		{
+			return target.ListRightPushAsync(key.Inner, RedisSerializer.Serialize(values), when, flags);
+		}
+
+		#endregion
+
+		#region ListSetByIndexAsync
+
+		/// <inheritdoc cref="IDatabaseAsync.ListSetByIndexAsync(RedisKey, Int64, RedisValue, CommandFlags)"/>
+		public static Task ListSetByIndexAsync<TElement>(this IDatabaseAsync target, RedisListKey<TElement> key, long index, TElement value, CommandFlags flags = CommandFlags.None)
+		{
+			return target.ListSetByIndexAsync(key.Inner, index, RedisSerializer.Serialize(value), flags);
+		}
+
+		#endregion
+
+		#region ListTrimAsync
+
+		/// <inheritdoc cref="IDatabaseAsync.ListTrimAsync(RedisKey, Int64, Int64, CommandFlags)"/>
+		public static Task ListTrimAsync<TElement>(this IDatabaseAsync target, RedisListKey<TElement> key, long start, long stop, CommandFlags flags = CommandFlags.None)
+		{
+			return target.ListTrimAsync(key.Inner, start, stop, flags);
+		}
+
+		#endregion
+	}
+}
diff --git a/Engine/Source/Programs/Shared/EpicGames.Redis/RedisSet.cs b/Engine/Source/Programs/Shared/EpicGames.Redis/RedisSet.cs
deleted file mode 100644
diff --git a/Engine/Source/Programs/Shared/EpicGames.Redis/RedisSetKey.cs b/Engine/Source/Programs/Shared/EpicGames.Redis/RedisSetKey.cs
new file mode 100644
--- /dev/null
+++ b/Engine/Source/Programs/Shared/EpicGames.Redis/RedisSetKey.cs
@@ -0,0 +1,149 @@
+// Copyright Epic Games, Inc. All Rights Reserved.
+
+using System;
+using System.Collections.Generic;
+using System.Threading.Tasks;
+using StackExchange.Redis;
+
+namespace EpicGames.Redis
+{
+	/// <summary>
+	/// Represents a typed Redis set with a given key
+	/// </summary>
+	/// <typeparam name="TElement">The type of element stored in the set</typeparam>
+	public record struct RedisSetKey<TElement>(RedisKey Inner)
+	{
+		/// <summary>
+		/// Implicit conversion to typed redis key.
+		/// </summary>
+		/// <param name="key">Key to convert</param>
+		public static implicit operator RedisSetKey<TElement>(string key) => new RedisSetKey<TElement>(new RedisKey(key));
+
+		/// <summary>
+		/// Implicit conversion to untyped redis keys.
+		/// </summary>
+		/// <param name="key">Key to convert</param>
+		public static implicit operator TypedRedisKey(RedisSetKey<TElement> key) => key.Inner;
+	}
+
+	/// <summary>
+	/// Extension methods for sets
+	/// </summary>
+	public static class RedisSetKeyExtensions
+	{
+		#region Conditions
+
+		/// <inheritdoc cref="Condition.SetContains(RedisKey, RedisValue)"/>
+		public static Condition Contains<TElement>(this RedisSetKey<TElement> key, TElement value)
+			=> Condition.SetContains(key.Inner, RedisSerializer.Serialize(value));
+
+		/// <inheritdoc cref="Condition.SetLengthEqual(RedisKey, Int64)"/>
+		public static Condition LengthEqual<TElement>(this RedisSetKey<TElement> key, long length)
+			=> Condition.SetLengthEqual(key.Inner, length);
+
+		/// <inheritdoc cref="Condition.SetLengthGreaterThan(RedisKey, Int64)"/>
+		public static Condition LengthGreaterThan<TElement>(this RedisSetKey<TElement> key, long length)
+			=> Condition.SetLengthGreaterThan(key.Inner, length);
+
+		/// <inheritdoc cref="Condition.SetLengthLessThan(RedisKey, Int64)"/>
+		public static Condition LengthLessThan<TElement>(this RedisSetKey<TElement> key, long length)
+			=> Condition.SetLengthLessThan(key.Inner, length);
+
+		/// <inheritdoc cref="Condition.SetContains(RedisKey, RedisValue)"/>
+		public static Condition NotContains<TElement>(this RedisSetKey<TElement> key, TElement value)
+			=> Condition.SetNotContains(key.Inner, RedisSerializer.Serialize(value));
+
+		#endregion
+
+		#region SetAddAsync
+
+		/// <inheritdoc cref="IDatabaseAsync.SetAddAsync(RedisKey, RedisValue, CommandFlags)"/>
+		public static Task<bool> SetAddAsync<TElement>(this IDatabaseAsync target, RedisSetKey<TElement> key, TElement item, CommandFlags flags = CommandFlags.None)
+		{
+			return target.SetAddAsync(key.Inner, RedisSerializer.Serialize(item), flags);
+		}
+
+		/// <inheritdoc cref="IDatabaseAsync.SetAddAsync(RedisKey, RedisValue[], CommandFlags)"/>
+		public static Task<long> SetAddAsync<TElement>(this IDatabaseAsync target, RedisSetKey<TElement> key, TElement[] values, CommandFlags flags = CommandFlags.None)
+		{
+			return target.SetAddAsync(key.Inner, RedisSerializer.Serialize(values), flags);
+		}
+
+		#endregion
+
+		#region SetContainsAsync
+
+		/// <inheritdoc cref="IDatabaseAsync.SetContainsAsync(RedisKey, RedisValue, CommandFlags)"/>
+		public static Task<bool> SetContainsAsync<TElement>(this IDatabaseAsync target, RedisSetKey<TElement> key, TElement value, CommandFlags flags = CommandFlags.None)
+		{
+			return target.SetContainsAsync(key.Inner, RedisSerializer.Serialize(value), flags);
+		}
+
+		#endregion
+
+		#region SetLengthAsync
+
+		/// <inheritdoc cref="IDatabaseAsync.SetLengthAsync(RedisKey, CommandFlags)"/>
+		public static Task<long> SetLengthAsync<TElement>(this IDatabaseAsync target, RedisSetKey<TElement> key, CommandFlags flags = CommandFlags.None)
+		{
+			return target.SetLengthAsync(key.Inner, flags);
+		}
+
+		#endregion
+
+		#region SetMembersAsync
+
+		/// <inheritdoc cref="IDatabaseAsync.SetMembersAsync(RedisKey, CommandFlags)"/>
+		public static Task<TElement[]> SetMembersAsync<TElement>(this IDatabaseAsync target, RedisSetKey<TElement> key, CommandFlags flags = CommandFlags.None)
+		{
+			return target.SetMembersAsync(key.Inner, flags).DeserializeAsync<TElement>();
+		}
+
+		#endregion
+
+		#region SetPopAsync
+
+		/// <inheritdoc cref="IDatabaseAsync.SetPopAsync(RedisKey, CommandFlags)"/>
+		public static Task<TElement> SetPopAsync<TElement>(this IDatabaseAsync target, RedisSetKey<TElement> key, CommandFlags flags = CommandFlags.None)
+		{
+			return target.SetPopAsync(key.Inner, flags).DeserializeAsync<TElement>();
+		}
+
+		/// <inheritdoc cref="IDatabaseAsync.SetPopAsync(RedisKey, Int64, CommandFlags)"/>
+		public static Task<TElement[]> SetPopAsync<TElement>(this IDatabaseAsync target, RedisSetKey<TElement> key, long count, CommandFlags flags = CommandFlags.None)
+		{
+			return target.SetPopAsync(key.Inner, count, flags).DeserializeAsync<TElement>();
+		}
+
+		#endregion
+
+		#region SetRemoveAsync
+
+		/// <inheritdoc cref="IDatabaseAsync.SetRemoveAsync(RedisKey, RedisValue, CommandFlags)"/>
+		public static Task<bool> SetRemoveAsync<TElement>(this IDatabaseAsync target, RedisSetKey<TElement> key, TElement item, CommandFlags flags = CommandFlags.None)
+		{
+			return target.SetRemoveAsync(key.Inner, RedisSerializer.Serialize(item), flags);
+		}
+
+		/// <inheritdoc cref="IDatabaseAsync.SetRemoveAsync(RedisKey, RedisValue[], CommandFlags)"/>
+		public static Task<long> SetRemoveAsync<TElement>(this IDatabaseAsync target, RedisSetKey<TElement> key, TElement[] values, CommandFlags flags = CommandFlags.None)
+		{
+			return target.SetRemoveAsync(key.Inner, RedisSerializer.Serialize(values), flags);
+		}
+
+		#endregion
+
+		#region SetScanAsync
+
+		/// <inheritdoc cref="IDatabaseAsync.SetScanAsync(RedisKey, RedisValue, int, long, int, CommandFlags)"/>
+		public static async IAsyncEnumerable<TElement> SetScanAsync<TElement>(this IDatabaseAsync target, RedisSetKey<TElement> key, RedisValue pattern = default, int pageSize = 250, long cursor = 0, int pageOffset = 0, CommandFlags flags = CommandFlags.None)
+		{
+			await foreach (RedisValue value in target.SetScanAsync(key.Inner, pattern, pageSize, cursor, pageOffset, flags))
+			{
+				yield return RedisSerializer.Deserialize<TElement>(value);
+			}
+		}
+
+		#endregion
+	}
+}
diff --git a/Engine/Source/Programs/Shared/EpicGames.Redis/RedisSortedSet.cs b/Engine/Source/Programs/Shared/EpicGames.Redis/RedisSortedSet.cs
deleted file mode 100644
diff --git a/Engine/Source/Programs/Shared/EpicGames.Redis/RedisSortedSetKey.cs b/Engine/Source/Programs/Shared/EpicGames.Redis/RedisSortedSetKey.cs
new file mode 100644
--- /dev/null
+++ b/Engine/Source/Programs/Shared/EpicGames.Redis/RedisSortedSetKey.cs
@@ -0,0 +1,314 @@
+// Copyright Epic Games, Inc. All Rights Reserved.
+
+using System;
+using System.Collections.Generic;
+using System.Threading.Tasks;
+using EpicGames.Core;
+using StackExchange.Redis;
+
+namespace EpicGames.Redis
+{
+	/// <summary>
+	/// Represents a typed Redis sorted set with a given key
+	/// </summary>
+	/// <typeparam name="TElement">The type of element stored in the set</typeparam>
+	public record struct RedisSortedSetKey<TElement>(RedisKey Inner)
+	{
+		/// <summary>
+		/// Implicit conversion to typed redis key.
+		/// </summary>
+		/// <param name="key">Key to convert</param>
+		public static implicit operator RedisSortedSetKey<TElement>(string key) => new RedisSortedSetKey<TElement>(new RedisKey(key));
+
+		/// <summary>
+		/// Implicit conversion to untyped redis keys.
+		/// </summary>
+		/// <param name="key">Key to convert</param>
+		public static implicit operator TypedRedisKey(RedisSortedSetKey<TElement> key) => key.Inner;
+	}
+
+	/// <summary>
+	/// Typed implementation of <see cref="SortedSetEntry"/>
+	/// </summary>
+	/// <typeparam name="T">The element type</typeparam>
+	public readonly struct SortedSetEntry<T>
+	{
+		/// <summary>
+		/// Accessor for the element type
+		/// </summary>
+		public readonly T Element { get; }
+
+		/// <summary>
+		/// The encoded element value
+		/// </summary>
+		public readonly RedisValue ElementValue { get; }
+
+		/// <summary>
+		/// Score for the entry
+		/// </summary>
+		public readonly double Score { get; }
+
+		/// <summary>
+		/// Constructor
+		/// </summary>
+		/// <param name="entry"></param>
+		public SortedSetEntry(SortedSetEntry entry)
+		{
+			Element = RedisSerializer.Deserialize<T>(entry.Element)!;
+			ElementValue = entry.Element;
+			Score = entry.Score;
+		}
+
+		/// <summary>
+		/// Constructor
+		/// </summary>
+		/// <param name="element"></param>
+		/// <param name="score"></param>
+		public SortedSetEntry(T element, double score)
+		{
+			Element = element;
+			ElementValue = RedisSerializer.Serialize<T>(element);
+			Score = score;
+		}
+
+		/// <summary>
+		/// Deconstruct this item into a tuple
+		/// </summary>
+		/// <param name="outElement"></param>
+		/// <param name="outScore"></param>
+		public void Deconstruct(out T outElement, out double outScore)
+		{
+			outElement = Element;
+			outScore = Score;
+		}
+	}
+
+	/// <summary>
+	/// Extension methods for sets
+	/// </summary>
+	public static class RedisSortedSetKeyExtensions
+	{
+		#region Conditions
+
+		/// <inheritdoc cref="Condition.SortedSetContains(RedisKey, RedisValue)"/>
+		public static Condition Contains<TElement>(this RedisSortedSetKey<TElement> key, TElement value)
+			=> Condition.SortedSetContains(key.Inner, RedisSerializer.Serialize(value));
+
+		/// <inheritdoc cref="Condition.SortedSetEqual(RedisKey, RedisValue, RedisValue)"/>
+		public static Condition Equal<TElement>(this RedisSortedSetKey<TElement> key, TElement value, RedisValue score)
+			=> Condition.SortedSetEqual(key.Inner, RedisSerializer.Serialize(value), score);
+
+		/// <inheritdoc cref="Condition.SortedSetLengthEqual(RedisKey, Int64)"/>
+		public static Condition LengthEqual<TElement>(this RedisSortedSetKey<TElement> key, long length)
+			=> Condition.SortedSetLengthEqual(key.Inner, length);
+
+		/// <inheritdoc cref="Condition.SortedSetLengthGreaterThan(RedisKey, Int64)"/>
+		public static Condition LengthGreaterThan<TElement>(this RedisSortedSetKey<TElement> key, long length)
+			=> Condition.SortedSetLengthGreaterThan(key.Inner, length);
+
+		/// <inheritdoc cref="Condition.SortedSetLengthLessThan(RedisKey, Int64)"/>
+		public static Condition LengthLessThan<TElement>(this RedisSortedSetKey<TElement> key, long length)
+			=> Condition.SortedSetLengthLessThan(key.Inner, length);
+
+		/// <inheritdoc cref="Condition.SortedSetNotContains(RedisKey, RedisValue)"/>
+		public static Condition NotContains<TElement>(this RedisSortedSetKey<TElement> key, TElement value)
+			=> Condition.SortedSetNotContains(key.Inner, RedisSerializer.Serialize(value));
+
+		/// <inheritdoc cref="Condition.SortedSetNotEqual(RedisKey, RedisValue, RedisValue)"/>
+		public static Condition NotEqual<TElement>(this RedisSortedSetKey<TElement> key, TElement value, RedisValue score)
+			=> Condition.SortedSetNotEqual(key.Inner, RedisSerializer.Serialize(value), score);
+
+		/// <inheritdoc cref="Condition.SortedSetScoreExists(RedisKey, RedisValue)"/>
+		public static Condition ScoreExists<TElement>(this RedisSortedSetKey<TElement> key, RedisValue score)
+			=> Condition.SortedSetScoreExists(key.Inner, score);
+
+		/// <inheritdoc cref="Condition.SortedSetScoreExists(RedisKey, RedisValue, RedisValue)"/>
+		public static Condition ScoreExists<TElement>(this RedisSortedSetKey<TElement> key, RedisValue score, RedisValue count)
+			=> Condition.SortedSetScoreExists(key.Inner, score, count);
+
+		#endregion
+
+		#region SortedSetAddAsync
+
+		/// <inheritdoc cref="IDatabaseAsync.SortedSetAddAsync(RedisKey, RedisValue, Double, CommandFlags)"/>
+		public static Task<bool> SortedSetAddAsync<TElement>(this IDatabaseAsync target, RedisSortedSetKey<TElement> key, TElement value, double score, SortedSetWhen when = SortedSetWhen.Always, CommandFlags flags = CommandFlags.None)
+		{
+			return target.SortedSetAddAsync(key.Inner, RedisSerializer.Serialize(value), score, when, flags);
+		}
+
+		/// <inheritdoc cref="IDatabaseAsync.SortedSetAddAsync(RedisKey, RedisValue, Double, When, CommandFlags)"/>
+		public static Task<long> SortedSetAddAsync<TElement>(this IDatabaseAsync target, RedisSortedSetKey<TElement> key, SortedSetEntry<TElement>[] values, SortedSetWhen when = SortedSetWhen.Always, CommandFlags flags = CommandFlags.None)
+		{
+			return target.SortedSetAddAsync(key.Inner, values.ConvertAll(x => new SortedSetEntry(x.ElementValue, x.Score)), when, flags);
+		}
+
+		#endregion
+
+		#region SortedSetLengthAsync
+
+		/// <inheritdoc cref="IDatabaseAsync.SortedSetLengthAsync(RedisKey, Double, Double, Exclude, CommandFlags)"/>
+		public static Task<long> SortedSetLengthAsync<TElement>(this IDatabaseAsync target, RedisSortedSetKey<TElement> key, double min = Double.NegativeInfinity, double max = Double.PositiveInfinity, Exclude exclude = Exclude.None, CommandFlags flags = CommandFlags.None)
+		{
+			return target.SortedSetLengthAsync(key.Inner, min, max, exclude, flags);
+		}
+
+		#endregion
+
+		#region SortedSetLengthByValueAsync
+
+		/// <inheritdoc cref="IDatabaseAsync.SortedSetLengthByValueAsync(RedisKey, RedisValue, RedisValue, Exclude, CommandFlags)"/>
+		public static Task<long> SortedSetLengthByValueAsync<TElement>(this IDatabaseAsync target, RedisSortedSetKey<TElement> key, TElement min, TElement max, Exclude exclude = Exclude.None, CommandFlags flags = CommandFlags.None)
+		{
+			return target.SortedSetLengthByValueAsync(key.Inner, RedisSerializer.Serialize(min), RedisSerializer.Serialize(max), exclude, flags);
+		}
+
+		#endregion
+
+		#region SortedSetRangeByRankAsync
+
+		/// <inheritdoc cref="IDatabaseAsync.SortedSetRangeByRankAsync(RedisKey, Int64, Int64, Order, CommandFlags)"/>
+		public static Task<TElement[]> SortedSetRangeByRankAsync<TElement>(this IDatabaseAsync target, RedisSortedSetKey<TElement> key, long start = 0, long stop = -1, Order order = Order.Ascending, CommandFlags flags = CommandFlags.None)
+		{
+			return target.SortedSetRangeByRankAsync(key.Inner, start, stop, order, flags).DeserializeAsync<TElement>();
+		}
+
+		#endregion
+
+		#region SortedSetRangeByRankWithScoresAsync
+
+		/// <inheritdoc cref="IDatabaseAsync.SortedSetRangeByRankWithScoresAsync(RedisKey, Int64, Int64, Order, CommandFlags)"/>
+		public static async Task<SortedSetEntry<TElement>[]> SortedSetRangeByRankWithScoresAsync<TElement>(this IDatabaseAsync target, RedisSortedSetKey<TElement> key, long start, long stop = -1, Order order = Order.Ascending, CommandFlags flags = CommandFlags.None)
+		{
+			SortedSetEntry[] values = await target.SortedSetRangeByRankWithScoresAsync(key.Inner, start, stop, order, flags);
+			return Array.ConvertAll(values, x => new SortedSetEntry<TElement>(x));
+		}
+
+		#endregion
+
+		#region SortedSetRangeByScoreAsync
+
+		/// <inheritdoc cref="IDatabaseAsync.SortedSetRangeByScoreAsync(RedisKey, Double, Double, Exclude, Order, Int64, Int64, CommandFlags)"/>
+		public static Task<TElement[]> SortedSetRangeByScoreAsync<TElement>(this IDatabaseAsync target, RedisSortedSetKey<TElement> key, double start = Double.NegativeInfinity, double stop = Double.PositiveInfinity, Exclude exclude = Exclude.None, Order order = Order.Ascending, long skip = 0L, long take = -1L, CommandFlags flags = CommandFlags.None)
+		{
+			return target.SortedSetRangeByScoreAsync(key.Inner, start, stop, exclude, order, skip, take, flags).DeserializeAsync<TElement>();
+		}
+
+		#endregion
+
+		#region SortedSetRangeByScoreWithScoresAsync
+
+		/// <inheritdoc cref="IDatabaseAsync.SortedSetRangeByScoreWithScoresAsync(RedisKey, Double, Double, Exclude, Order, Int64, Int64, CommandFlags)"/>
+		public static async Task<SortedSetEntry<TElement>[]> SortedSetRangeByScoreWithScoresAsync<TElement>(this IDatabaseAsync target, RedisSortedSetKey<TElement> key, double start = Double.NegativeInfinity, double stop = Double.PositiveInfinity, Exclude exclude = Exclude.None, Order order = Order.Ascending, long skip = 0L, long take = -1L, CommandFlags flags = CommandFlags.None)
+		{
+			SortedSetEntry[] values = await target.SortedSetRangeByScoreWithScoresAsync(key.Inner, start, stop, exclude, order, skip, take, flags);
+			return Array.ConvertAll(values, x => new SortedSetEntry<TElement>(x));
+		}
+
+		#endregion
+
+		#region SortedSetRangeByValueAsync
+
+		/// <inheritdoc cref="IDatabaseAsync.SortedSetRangeByValueAsync(RedisKey, RedisValue, RedisValue, Exclude, Order, Int64, Int64, CommandFlags)"/>
+		public static Task<TElement[]> SortedSetRangeByValueAsync<TElement>(this IDatabaseAsync target, RedisSortedSetKey<TElement> key, TElement min, TElement max, Exclude exclude = Exclude.None, Order order = Order.Ascending, long skip = 0L, long take = -1L, CommandFlags flags = CommandFlags.None)
+		{
+			return target.SortedSetRangeByValueAsync(key.Inner, RedisSerializer.Serialize(min), RedisSerializer.Serialize(max), exclude, order, skip, take, flags).DeserializeAsync<TElement>();
+		}
+
+		#endregion
+
+		#region SortedSetRankAsync
+
+		/// <inheritdoc cref="IDatabaseAsync.SortedSetRankAsync(RedisKey, RedisValue, Order, CommandFlags)"/>
+		public static Task<long?> SortedSetRankAsync<TElement>(this IDatabaseAsync target, RedisSortedSetKey<TElement> key, TElement item, Order order = Order.Ascending, CommandFlags flags = CommandFlags.None)
+		{
+			return target.SortedSetRankAsync(key.Inner, RedisSerializer.Serialize(item), order, flags);
+		}
+
+		#endregion
+
+		#region SortedSetRemoveAsync
+
+		/// <inheritdoc cref="IDatabaseAsync.SortedSetRemoveAsync(RedisKey, RedisValue, CommandFlags)"/>
+		public static Task<bool> SortedSetRemoveAsync<TElement>(this IDatabaseAsync target, RedisSortedSetKey<TElement> key, TElement value, CommandFlags flags = CommandFlags.None)
+		{
+			return target.SortedSetRemoveAsync(key.Inner, RedisSerializer.Serialize(value), flags);
+		}
+
+		/// <inheritdoc cref="IDatabaseAsync.SortedSetRemoveAsync(RedisKey, RedisValue[], CommandFlags)"/>
+		public static Task<long> SortedSetRemoveAsync<TElement>(this IDatabaseAsync target, RedisSortedSetKey<TElement> key, TElement[] values, CommandFlags flags = CommandFlags.None)
+		{
+			return target.SortedSetRemoveAsync(key.Inner, RedisSerializer.Serialize(values), flags);
+		}
+
+		#endregion
+
+		#region SortedSetRemoveRangeByRankAsync
+
+		/// <inheritdoc cref="IDatabaseAsync.SortedSetRemoveRangeByRankAsync(RedisKey, Int64, Int64, CommandFlags)"/>
+		public static Task<long> SortedSetRemoveRangeByRankAsync<TElement>(this IDatabaseAsync target, RedisSortedSetKey<TElement> key, long start, long stop, CommandFlags flags = CommandFlags.None)
+		{
+			return target.SortedSetRemoveRangeByRankAsync(key.Inner, start, stop, flags);
+		}
+
+		#endregion
+
+		#region SortedSetRemoveRangeByScoreAsync
+
+		/// <inheritdoc cref="IDatabaseAsync.SortedSetRemoveRangeByScoreAsync(RedisKey, Double, Double, Exclude, CommandFlags)"/>
+		public static Task<long> SortedSetRemoveRangeByScoreAsync<TElement>(this IDatabaseAsync target, RedisSortedSetKey<TElement> key, double start, double stop, Exclude exclude = Exclude.None, CommandFlags flags = CommandFlags.None)
+		{
+			return target.SortedSetRemoveRangeByScoreAsync(key.Inner, start, stop, exclude, flags);
+		}
+
+		#endregion
+
+		#region SortedSetRemoveRangeByValueAsync
+
+		/// <inheritdoc cref="IDatabaseAsync.SortedSetRemoveRangeByValueAsync(RedisKey, RedisValue, RedisValue, Exclude, CommandFlags)"/>
+		public static Task<long> SortedSetRemoveRangeByValueAsync<TElement>(this IDatabaseAsync target, RedisSortedSetKey<TElement> key, TElement min, TElement max, Exclude exclude = Exclude.None, CommandFlags flags = CommandFlags.None)
+		{
+			return target.SortedSetRemoveRangeByValueAsync(key.Inner, RedisSerializer.Serialize(min), RedisSerializer.Serialize(max), exclude, flags);
+		}
+
+		#endregion
+
+		#region SortedSetScanAsync
+
+		/// <inheritdoc cref="IDatabaseAsync.SortedSetScanAsync(RedisKey, RedisValue, Int32, Int64, Int32, CommandFlags)"/>
+		public static async IAsyncEnumerable<SortedSetEntry<TElement>> SortedSetScanAsync<TElement>(this IDatabaseAsync target, RedisSortedSetKey<TElement> key, RedisValue pattern = default, int pageSize = 250, long cursor = 0, int pageOffset = 0, CommandFlags flags = CommandFlags.None)
+		{
+			await foreach (SortedSetEntry entry in target.SortedSetScanAsync(key.Inner, pattern, pageSize, cursor, pageOffset, flags))
+			{
+				yield return new SortedSetEntry<TElement>(entry);
+			}
+		}
+
+		#endregion
+
+		#region SortedSetScoreAsync
+
+		/// <inheritdoc cref="IDatabaseAsync.SortedSetScoreAsync(RedisKey, RedisValue, CommandFlags)"/>
+		public static Task<double?> SortedSetScoreAsync<TElement>(this IDatabaseAsync target, RedisSortedSetKey<TElement> key, TElement member, CommandFlags flags = CommandFlags.None)
+		{
+			return target.SortedSetScoreAsync(key.Inner, RedisSerializer.Serialize(member), flags);
+		}
+
+		#endregion
+
+		#region SortedSetUpdateAsync
+
+		/// <inheritdoc cref="IDatabaseAsync.SortedSetUpdateAsync(RedisKey, RedisValue, double, SortedSetWhen, CommandFlags)"/>
+		public static Task<bool> SortedSetUpdateAsync<TElement>(this IDatabaseAsync target, RedisSortedSetKey<TElement> key, TElement member, double score, SortedSetWhen when = SortedSetWhen.Always, CommandFlags flags = CommandFlags.None)
+		{
+			return target.SortedSetUpdateAsync(key.Inner, RedisSerializer.Serialize(member), score, when, flags);
+		}
+
+		/// <inheritdoc cref="IDatabaseAsync.SortedSetUpdateAsync(RedisKey, RedisValue, double, SortedSetWhen, CommandFlags)"/>
+		public static Task<long> SortedSetUpdateAsync<TElement>(this IDatabaseAsync target, RedisSortedSetKey<TElement> key, SortedSetEntry<TElement>[] values, SortedSetWhen when = SortedSetWhen.Always, CommandFlags flags = CommandFlags.None)
+		{
+			return target.SortedSetUpdateAsync(key.Inner, values.ConvertAll(x => new SortedSetEntry(RedisSerializer.Serialize(x.Element), x.Score)), when, flags);
+		}
+
+		#endregion
+	}
+}
diff --git a/Engine/Source/Programs/Shared/EpicGames.Redis/RedisString.cs b/Engine/Source/Programs/Shared/EpicGames.Redis/RedisString.cs
deleted file mode 100644
diff --git a/Engine/Source/Programs/Shared/EpicGames.Redis/RedisStringKey.cs b/Engine/Source/Programs/Shared/EpicGames.Redis/RedisStringKey.cs
new file mode 100644
--- /dev/null
+++ b/Engine/Source/Programs/Shared/EpicGames.Redis/RedisStringKey.cs
@@ -0,0 +1,161 @@
+// Copyright Epic Games, Inc. All Rights Reserved.
+
+using System;
+using System.Collections.Generic;
+using System.Threading.Tasks;
+using EpicGames.Core;
+using StackExchange.Redis;
+
+namespace EpicGames.Redis
+{
+	/// <summary>
+	/// Represents a typed Redis string with a given value type
+	/// </summary>
+	/// <typeparam name="TValue">The type of element stored in the list</typeparam>
+	public record struct RedisStringKey<TValue>(RedisKey Inner)
+	{
+		/// <summary>
+		/// Implicit conversion to typed redis key.
+		/// </summary>
+		/// <param name="key">Key to convert</param>
+		public static implicit operator RedisStringKey<TValue>(string key) => new RedisStringKey<TValue>(new RedisKey(key));
+
+		/// <summary>
+		/// Implicit conversion to a regular RedisKey
+		/// </summary>
+		/// <param name="key">The key to convert</param>
+		public static implicit operator TypedRedisKey(RedisStringKey<TValue> key) => key.Inner;
+	}
+
+	/// <summary>
+	/// Extension methods for strings
+	/// </summary>
+	public static class RedisStringKeyExtensions
+	{
+		#region Conditions
+
+		/// <inheritdoc cref="Condition.StringEqual(RedisKey, RedisValue)"/>
+		public static Condition Equal<TElement>(this RedisStringKey<TElement> key, TElement value)
+			=> Condition.StringEqual(key.Inner, RedisSerializer.Serialize(value));
+
+		/// <inheritdoc cref="Condition.StringLengthEqual(RedisKey, Int64)"/>
+		public static Condition LengthEqual<TElement>(this RedisStringKey<TElement> key, long length)
+			=> Condition.StringLengthEqual(key.Inner, length);
+
+		/// <inheritdoc cref="Condition.StringLengthGreaterThan(RedisKey, Int64)"/>
+		public static Condition LengthGreaterThan<TElement>(this RedisStringKey<TElement> key, long length)
+			=> Condition.StringLengthGreaterThan(key.Inner, length);
+
+		/// <inheritdoc cref="Condition.StringLengthLessThan(RedisKey, Int64)"/>
+		public static Condition LengthLessThan<TElement>(this RedisStringKey<TElement> key, long length)
+			=> Condition.StringLengthLessThan(key.Inner, length);
+
+		/// <inheritdoc cref="Condition.StringNotEqual(RedisKey, RedisValue)"/>
+		public static Condition NotEqual<TElement>(this RedisStringKey<TElement> key, TElement value)
+			=> Condition.StringNotEqual(key.Inner, RedisSerializer.Serialize(value));
+
+		#endregion
+
+		#region StringDecrementAsync
+
+		/// <inheritdoc cref="IDatabaseAsync.StringDecrementAsync(RedisKey, Int64, CommandFlags)"/>
+		public static Task<long> StringDecrementAsync(this IDatabaseAsync target, RedisKey key, long value = 1L, CommandFlags flags = CommandFlags.None)
+		{
+			return target.StringDecrementAsync(key, value, flags);
+		}
+
+		/// <inheritdoc cref="IDatabaseAsync.StringDecrementAsync(RedisKey, Int64, CommandFlags)"/>
+		public static Task<long> StringDecrementAsync(this IDatabaseAsync target, RedisStringKey<long> key, long value = 1L, CommandFlags flags = CommandFlags.None)
+		{
+			return target.StringDecrementAsync(key.Inner, value, flags);
+		}
+
+		/// <inheritdoc cref="IDatabaseAsync.StringDecrementAsync(RedisKey, Double, CommandFlags)"/>
+		public static Task<double> StringDecrementAsync(this IDatabaseAsync target, RedisStringKey<double> key, double value = 1.0, CommandFlags flags = CommandFlags.None)
+		{
+			return target.StringDecrementAsync(key.Inner, value, flags);
+		}
+
+		#endregion
+
+		#region StringGetAsync
+
+		/// <inheritdoc cref="IDatabaseAsync.StringGetAsync(RedisKey, CommandFlags)"/>
+		public static async Task<TValue?> StringGetAsync<TValue>(this IDatabaseAsync target, RedisStringKey<TValue> key, CommandFlags flags = CommandFlags.None)
+		{
+			RedisValue value = await target.StringGetAsync(key.Inner, flags);
+			if (value.IsNullOrEmpty)
+			{
+				return default;
+			}
+			return RedisSerializer.Deserialize<TValue>(value)!;
+		}
+
+		/// <inheritdoc cref="IDatabaseAsync.StringGetAsync(RedisKey, CommandFlags)"/>
+		public static async Task<TValue> StringGetAsync<TValue>(this IDatabaseAsync target, RedisStringKey<TValue> key, TValue defaultValue, CommandFlags flags = CommandFlags.None)
+		{
+			RedisValue value = await target.StringGetAsync(key.Inner, flags);
+			if (value.IsNullOrEmpty)
+			{
+				return defaultValue;
+			}
+			return RedisSerializer.Deserialize<TValue>(value)!;
+		}
+
+		#endregion
+
+		#region StringIncrementAsync
+
+		/// <inheritdoc cref="IDatabaseAsync.StringIncrementAsync(RedisKey, Double, CommandFlags)"/>
+		public static Task<long> StringIncrementAsync(this IDatabaseAsync target, RedisKey key, long value = 1L, CommandFlags flags = CommandFlags.None)
+		{
+			return target.StringIncrementAsync(key, value, flags);
+		}
+
+		/// <inheritdoc cref="IDatabaseAsync.StringIncrementAsync(RedisKey, Double, CommandFlags)"/>
+		public static Task<long> StringIncrementAsync(this IDatabaseAsync target, RedisStringKey<long> key, long value = 1L, CommandFlags flags = CommandFlags.None)
+		{
+			return target.StringIncrementAsync(key.Inner, value, flags);
+		}
+
+		/// <inheritdoc cref="IDatabaseAsync.StringIncrementAsync(RedisKey, Double, CommandFlags)"/>
+		public static Task<double> StringIncrementAsync(this IDatabaseAsync target, RedisKey key, double value, CommandFlags flags = CommandFlags.None)
+		{
+			return target.StringIncrementAsync(key, value, flags);
+		}
+
+		/// <inheritdoc cref="IDatabaseAsync.StringIncrementAsync(RedisKey, Double, CommandFlags)"/>
+		public static Task<double> StringIncrementAsync(this IDatabaseAsync target, RedisStringKey<double> key, double value = 1.0, CommandFlags flags = CommandFlags.None)
+		{
+			return target.StringIncrementAsync(key.Inner, value, flags);
+		}
+
+		#endregion
+
+		#region StringLengthAsync
+
+		/// <inheritdoc cref="IDatabaseAsync.StringLengthAsync(RedisKey, CommandFlags)"/>
+		public static Task<long> StringLengthAsync<TValue>(this IDatabaseAsync target, RedisStringKey<TValue> key, CommandFlags flags = CommandFlags.None)
+		{
+			return target.StringLengthAsync(key.Inner, flags);
+		}
+
+		#endregion
+
+		#region StringSetAsync
+
+		/// <inheritdoc cref="IDatabaseAsync.StringSetAsync(RedisKey, RedisValue, TimeSpan?, When, CommandFlags)"/>
+		public static Task<bool> StringSetAsync<TValue>(this IDatabaseAsync target, RedisStringKey<TValue> key, TValue value, TimeSpan? expiry = null, When when = When.Always, CommandFlags flags = CommandFlags.None)
+		{
+			return target.StringSetAsync(key.Inner, RedisSerializer.Serialize(value), expiry, when, flags);
+		}
+
+		/// <inheritdoc cref="IDatabaseAsync.StringSetAsync(RedisKey, RedisValue, TimeSpan?, When, CommandFlags)"/>
+		public static Task<bool> StringSetAsync<TValue>(this IDatabaseAsync target, KeyValuePair<RedisStringKey<TValue>, TValue>[] pairs, When when = When.Always, CommandFlags flags = CommandFlags.None)
+		{
+			return target.StringSetAsync(pairs.ConvertAll(x => new KeyValuePair<RedisKey, RedisValue>(x.Key.Inner, RedisSerializer.Serialize(x.Value))), when, flags);
+		}
+
+		#endregion
+	}
+}
