EpicGames.Redis: Add wrapper types for Redis data structures that combine a database reference and key.

diff --git a/Engine/Source/Programs/Shared/EpicGames.Redis/RedisHash.cs b/Engine/Source/Programs/Shared/EpicGames.Redis/RedisHash.cs
new file mode 100644
--- /dev/null
+++ b/Engine/Source/Programs/Shared/EpicGames.Redis/RedisHash.cs
@@ -0,0 +1,210 @@
+// Copyright Epic Games, Inc. All Rights Reserved.
+
+using System;
+using System.Collections.Generic;
+using System.Linq.Expressions;
+using System.Threading.Tasks;
+using StackExchange.Redis;
+
+namespace EpicGames.Redis
+{
+	/// <summary>
+	/// Redis hash, with members corresponding to the property names of a type
+	/// </summary>
+	public record struct RedisHash<T>(IDatabaseAsync Database, RedisHashKey<T> Key);
+
+	/// <summary>
+	/// Redis hash, with members of the given key and value types
+	/// </summary>
+	public record struct RedisHash<TName, TValue>(IDatabaseAsync Database, RedisHashKey<TName, TValue> Key);
+
+	/// <summary>
+	/// Extension methods for hashes
+	/// </summary>
+	public static class RedisHashExtensions
+	{
+		#region Conditions
+
+		/// <inheritdoc cref="Condition.HashEqual(RedisKey, RedisValue, RedisValue)"/>
+		public static Condition Equal<TRecord, TValue>(this RedisHash<TRecord> target, Expression<Func<TRecord, TValue>> selector, TValue value)
+			=> target.Key.Equal(selector, value);
+
+		/// <inheritdoc cref="Condition.HashEqual(RedisKey, RedisValue, RedisValue)"/>
+		public static Condition Equal<TName, TValue>(this RedisHash<TName, TValue> target, TName name, TValue value)
+			=> target.Key.Equal(name, value);
+
+		/// <inheritdoc cref="Condition.HashExists(RedisKey, RedisValue)"/>
+		public static Condition Exists<TRecord, TValue>(this RedisHash<TRecord> target, Expression<Func<TRecord, TValue>> selector)
+			=> target.Key.Exists(selector);
+
+		/// <inheritdoc cref="Condition.HashExists(RedisKey, RedisValue)"/>
+		public static Condition Exists<TName, TValue>(this RedisHash<TName, TValue> target, TName name)
+			=> target.Key.Exists(name);
+
+		/// <inheritdoc cref="Condition.HashLengthEqual(RedisKey, Int64)"/>
+		public static Condition LengthEqual<TName, TValue>(this RedisHash<TName, TValue> target, long length)
+			=> target.Key.LengthEqual(length);
+
+		/// <inheritdoc cref="Condition.HashLengthGreaterThan(RedisKey, Int64)"/>
+		public static Condition LengthGreaterThan<TName, TValue>(this RedisHash<TName, TValue> target, long length)
+			=> target.Key.LengthGreaterThan(length);
+
+		/// <inheritdoc cref="Condition.HashLengthLessThan(RedisKey, Int64)"/>
+		public static Condition LengthLessThan<TName, TValue>(this RedisHash<TName, TValue> target, long length)
+			=> target.Key.LengthLessThan(length);
+
+		/// <inheritdoc cref="Condition.HashNotExists(RedisKey, RedisValue)"/>
+		public static Condition NotExists<TRecord, TValue>(this RedisHash<TRecord> target, Expression<Func<TRecord, TValue>> selector)
+			=> target.Key.NotExists(selector);
+
+		/// <inheritdoc cref="Condition.HashNotExists(RedisKey, RedisValue)"/>
+		public static Condition NotExists<TName, TValue>(this RedisHash<TName, TValue> target, TName name)
+			=> target.Key.NotExists(name);
+
+		/// <inheritdoc cref="Condition.HashEqual(RedisKey, RedisValue, RedisValue)"/>
+		public static Condition NotEqual<TRecord, TValue>(this RedisHash<TRecord> target, Expression<Func<TRecord, TValue>> selector, TValue value)
+			=> target.Key.NotEqual(selector, value);
+
+		/// <inheritdoc cref="Condition.HashNotEqual(RedisKey, RedisValue, RedisValue)"/>
+		public static Condition NotEqual<TName, TValue>(this RedisHash<TName, TValue> target, TName name, TValue value)
+			=> target.Key.NotEqual(name, value);
+
+		#endregion
+
+		#region HashDecrementAsync
+
+		/// <inheritdoc cref="IDatabaseAsync.HashDecrementAsync(RedisKey, RedisValue, Int64, CommandFlags)"/>
+		public static Task<long> DecrementAsync<TRecord>(this RedisHash<TRecord> target, Expression<Func<TRecord, long>> selector, long value = 1L, CommandFlags flags = CommandFlags.None)
+			=> target.Database.HashDecrementAsync(target.Key, selector, value, flags);
+
+		/// <inheritdoc cref="IDatabaseAsync.HashDecrementAsync(RedisKey, RedisValue, Double, CommandFlags)"/>
+		public static Task<double> DecrementAsync<TRecord>(this RedisHash<TRecord> target, Expression<Func<TRecord, double>> selector, double value = 1.0, CommandFlags flags = CommandFlags.None)
+			=> target.Database.HashDecrementAsync(target.Key, selector, value, flags);
+
+		/// <inheritdoc cref="IDatabaseAsync.HashDecrementAsync(RedisKey, RedisValue, Int64, CommandFlags)"/>
+		public static Task<long> DecrementAsync<TName>(this RedisHash<TName, long> target, TName name, long value = 1L, CommandFlags flags = CommandFlags.None)
+			=> target.Database.HashDecrementAsync(target.Key, name, value, flags);
+
+		/// <inheritdoc cref="IDatabaseAsync.HashDecrementAsync(RedisKey, RedisValue, Double, CommandFlags)"/>
+		public static Task<double> DecrementAsync<TName>(this RedisHash<TName, double> target, TName name, double value = 1.0, CommandFlags flags = CommandFlags.None)
+			=> target.Database.HashDecrementAsync(target.Key, name, value, flags);
+
+		#endregion
+
+		#region HashDeleteAsync
+
+		/// <inheritdoc cref="IDatabaseAsync.HashDeleteAsync(RedisKey, RedisValue, CommandFlags)"/>
+		public static Task<bool> DeleteAsync<TRecord>(this RedisHash<TRecord> target, Expression<Func<TRecord, double>> selector, CommandFlags flags = CommandFlags.None)
+			=> target.Database.HashDeleteAsync(target.Key, selector, flags);
+
+		/// <inheritdoc cref="IDatabaseAsync.HashDeleteAsync(RedisKey, RedisValue, CommandFlags)"/>
+		public static Task<bool> DeleteAsync<TName, TValue>(this RedisHash<TName, TValue> target, TName name, CommandFlags flags = CommandFlags.None)
+			=> target.Database.HashDeleteAsync(target.Key, name, flags);
+
+		/// <inheritdoc cref="IDatabaseAsync.HashDeleteAsync(RedisKey, RedisValue[], CommandFlags)"/>
+		public static Task<long> DeleteAsync<TName, TValue>(this RedisHash<TName, TValue> target, TName[] names, CommandFlags flags = CommandFlags.None)
+			=> target.Database.HashDeleteAsync(target.Key, names, flags);
+
+		#endregion
+
+		#region HashExistsAsync
+
+		/// <inheritdoc cref="IDatabaseAsync.HashExistsAsync(RedisKey, RedisValue, CommandFlags)"/>
+		public static Task<bool> ExistsAsync<TName, TValue>(this RedisHash<TName, TValue> target, TName name, CommandFlags flags = CommandFlags.None)
+			=> target.Database.HashExistsAsync(target.Key, name, flags);
+
+		#endregion
+
+		#region HashGetAsync
+
+		/// <inheritdoc cref="IDatabaseAsync.HashGetAsync(RedisKey, RedisValue, CommandFlags)"/>
+		public static Task<TValue> GetAsync<TRecord, TValue>(this RedisHash<TRecord> target, Expression<Func<TRecord, TValue>> selector, CommandFlags flags = CommandFlags.None)
+			=> target.Database.HashGetAsync(target.Key, selector, flags);
+
+		/// <inheritdoc cref="IDatabaseAsync.HashGetAsync(RedisKey, RedisValue, CommandFlags)"/>
+		public static Task<TValue> GetAsync<TName, TValue>(this RedisHash<TName, TValue> target, TName name, CommandFlags flags = CommandFlags.None)
+			=> target.Database.HashGetAsync(target.Key, name, flags);
+
+		/// <inheritdoc cref="IDatabaseAsync.HashGetAsync(RedisKey, RedisValue[], CommandFlags)"/>
+		public static Task<TValue[]> GetAsync<TName, TValue>(this RedisHash<TName, TValue> target, TName[] names, CommandFlags flags = CommandFlags.None)
+			=> target.Database.HashGetAsync(target.Key, names, flags);
+
+		#endregion
+
+		#region HashGetAllAsync
+
+		/// <inheritdoc cref="IDatabaseAsync.HashGetAllAsync(RedisKey, CommandFlags)"/>
+		public static Task<HashEntry<TName, TValue>[]> GetAllAsync<TName, TValue>(this RedisHash<TName, TValue> target, CommandFlags flags = CommandFlags.None)
+			=> target.Database.HashGetAllAsync(target.Key, flags);
+
+		#endregion
+
+		#region HashIncrementAsync
+
+		/// <inheritdoc cref="IDatabaseAsync.HashIncrementAsync(RedisKey, RedisValue, Int64, CommandFlags)"/>
+		public static Task<long> IncrementAsync<TRecord>(this RedisHash<TRecord> target, Expression<Func<TRecord, long>> selector, long value = 1L, CommandFlags flags = CommandFlags.None)
+			=> target.Database.HashIncrementAsync(target.Key, selector, value, flags);
+
+		/// <inheritdoc cref="IDatabaseAsync.HashIncrementAsync(RedisKey, RedisValue, Double, CommandFlags)"/>
+		public static Task<double> IncrementAsync<TRecord>(this RedisHash<TRecord> target, Expression<Func<TRecord, double>> selector, double value = 1.0, CommandFlags flags = CommandFlags.None)
+			=> target.Database.HashIncrementAsync(target.Key, selector, value, flags);
+
+		/// <inheritdoc cref="IDatabaseAsync.HashIncrementAsync(RedisKey, RedisValue, Double, CommandFlags)"/>
+		public static Task<long> IncrementAsync<TName>(this RedisHash<TName, long> target, TName name, long value = 1L, CommandFlags flags = CommandFlags.None)
+			=> target.Database.HashIncrementAsync(target.Key, name, value, flags);
+
+		/// <inheritdoc cref="IDatabaseAsync.HashIncrementAsync(RedisKey, RedisValue, Double, CommandFlags)"/>
+		public static Task<double> IncrementAsync<TName>(this RedisHash<TName, double> target, TName name, double value = 1.0, CommandFlags flags = CommandFlags.None)
+			=> target.Database.HashIncrementAsync(target.Key, name, value, flags);
+
+		#endregion
+
+		#region HashKeysAsync
+
+		/// <inheritdoc cref="IDatabaseAsync.HashLengthAsync(RedisKey, CommandFlags)"/>
+		public static Task<TName[]> KeysAsync<TName, TValue>(this RedisHash<TName, TValue> target, CommandFlags flags = CommandFlags.None)
+			=> target.Database.HashKeysAsync(target.Key, flags);
+
+		#endregion
+
+		#region HashLengthAsync
+
+		/// <inheritdoc cref="IDatabaseAsync.HashLengthAsync(RedisKey, CommandFlags)"/>
+		public static Task<long> LengthAsync<TName, TValue>(this RedisHash<TName, TValue> target, CommandFlags flags = CommandFlags.None)
+			=> target.Database.HashLengthAsync(target.Key, flags);
+
+		#endregion
+
+		#region HashScanAsync
+
+		/// <inheritdoc cref="IDatabaseAsync.HashScanAsync(RedisKey, RedisValue, int, long, int, CommandFlags)"/>
+		public static IAsyncEnumerable<HashEntry<TName, TValue>> ScanAsync<TName, TValue>(this RedisHash<TName, TValue> target, RedisValue pattern, int pageSize = 250, long cursor = 0, int pageOffset = 0, CommandFlags flags = CommandFlags.None)
+			=> target.Database.HashScanAsync(target.Key, pattern, pageSize, cursor, pageOffset, flags);
+
+		#endregion
+
+		#region HashSetAsync
+
+		/// <inheritdoc cref="IDatabaseAsync.HashSetAsync(RedisKey, RedisValue, RedisValue, When, CommandFlags)"/>
+		public static Task SetAsync<TRecord, TValue>(this RedisHash<TRecord> target, Expression<Func<TRecord, TValue>> selector, TValue value, When when = When.Always, CommandFlags flags = CommandFlags.None)
+			=> target.Database.HashSetAsync(target.Key, selector, value, when, flags);
+
+		/// <inheritdoc cref="IDatabaseAsync.HashSetAsync(RedisKey, RedisValue, RedisValue, When, CommandFlags)"/>
+		public static Task SetAsync<TName, TValue>(this RedisHash<TName, TValue> target, TName name, TValue value, When when = When.Always, CommandFlags flags = CommandFlags.None)
+			=> target.Database.HashSetAsync(target.Key, name, value, when, flags);
+
+		/// <inheritdoc cref="IDatabaseAsync.HashSetAsync(RedisKey, HashEntry[], CommandFlags)"/>
+		public static Task SetAsync<TName, TValue>(this RedisHash<TName, TValue> target, IEnumerable<HashEntry<TName, TValue>> entries, CommandFlags flags = CommandFlags.None)
+			=> target.Database.HashSetAsync(target.Key, entries, flags);
+
+		#endregion
+
+		#region HashValuesAsync
+
+		/// <inheritdoc cref="IDatabaseAsync.HashValuesAsync(RedisKey, CommandFlags)"/>
+		public static Task<TValue[]> ValuesAsync<TName, TValue>(this RedisHash<TName, TValue> target, CommandFlags flags = CommandFlags.None)
+			=> target.Database.HashValuesAsync(target.Key, flags);
+
+		#endregion
+	}
+}
diff --git a/Engine/Source/Programs/Shared/EpicGames.Redis/RedisHashKey.cs b/Engine/Source/Programs/Shared/EpicGames.Redis/RedisHashKey.cs
--- a/Engine/Source/Programs/Shared/EpicGames.Redis/RedisHashKey.cs
+++ b/Engine/Source/Programs/Shared/EpicGames.Redis/RedisHashKey.cs
@@ -187,7 +187,7 @@ namespace EpicGames.Redis
 		}
 
 		/// <inheritdoc cref="IDatabaseAsync.HashDecrementAsync(RedisKey, RedisValue, Double, CommandFlags)"/>
-		public static Task<double> HashDecrementAsync<TName>(this IDatabaseAsync target, RedisHashKey<TName, long> key, TName name, double value = 1.0, CommandFlags flags = CommandFlags.None)
+		public static Task<double> HashDecrementAsync<TName>(this IDatabaseAsync target, RedisHashKey<TName, double> key, TName name, double value = 1.0, CommandFlags flags = CommandFlags.None)
 		{
 			return target.HashDecrementAsync(key.Inner, RedisSerializer.Serialize<TName>(name), value, flags);
 		}
diff --git a/Engine/Source/Programs/Shared/EpicGames.Redis/RedisList.cs b/Engine/Source/Programs/Shared/EpicGames.Redis/RedisList.cs
new file mode 100644
--- /dev/null
+++ b/Engine/Source/Programs/Shared/EpicGames.Redis/RedisList.cs
@@ -0,0 +1,156 @@
+// Copyright Epic Games, Inc. All Rights Reserved.
+
+using System;
+using System.Collections.Generic;
+using System.Threading.Tasks;
+using StackExchange.Redis;
+
+namespace EpicGames.Redis
+{
+	/// <summary>
+	/// Represents a typed Redis list with a given key
+	/// </summary>
+	public record struct RedisList<TElement>(IDatabaseAsync Database, RedisListKey<TElement> Key);
+
+	/// <summary>
+	/// Extension methods for sets
+	/// </summary>
+	public static class RedisListExtensions
+	{
+		#region Conditions
+
+		/// <inheritdoc cref="Condition.ListIndexEqual(RedisKey, Int64, RedisValue)"/>
+		public static Condition IndexEqual<TElement>(this RedisList<TElement> target, long index, TElement value)
+			=> target.Key.IndexEqual(index, value);
+
+		/// <inheritdoc cref="Condition.ListIndexExists(RedisKey, Int64)"/>
+		public static Condition IndexExists<TElement>(this RedisList<TElement> target, long index)
+			=> target.Key.IndexExists(index);
+
+		/// <inheritdoc cref="Condition.ListIndexNotEqual(RedisKey, Int64, RedisValue)"/>
+		public static Condition IndexNotEqual<TElement>(this RedisList<TElement> target, long index, TElement value)
+			=> target.Key.IndexNotEqual(index, value);
+
+		/// <inheritdoc cref="Condition.ListIndexNotExists(RedisKey, Int64)"/>
+		public static Condition IndexNotExists<TElement>(this RedisList<TElement> target, long index)
+			=> target.Key.IndexNotExists(index);
+
+		/// <inheritdoc cref="Condition.ListLengthEqual(RedisKey, Int64)"/>
+		public static Condition LengthEqual<TElement>(this RedisList<TElement> target, long length)
+			=> target.LengthEqual(length);
+
+		/// <inheritdoc cref="Condition.ListLengthGreaterThan(RedisKey, Int64)"/>
+		public static Condition LengthGreaterThan<TElement>(this RedisList<TElement> target, long length)
+			=> target.LengthGreaterThan(length);
+
+		/// <inheritdoc cref="Condition.ListLengthLessThan(RedisKey, Int64)"/>
+		public static Condition LengthLessThan<TElement>(this RedisList<TElement> target, long length)
+			=> target.LengthLessThan(length);
+
+		#endregion
+
+		#region ListGetByIndexAsync
+
+		/// <inheritdoc cref="IDatabaseAsync.ListGetByIndexAsync(RedisKey, Int64, CommandFlags)"/>
+		public static Task<TElement> GetByIndexAsync<TElement>(this RedisList<TElement> target, long index, CommandFlags flags = CommandFlags.None)
+			=> target.Database.ListGetByIndexAsync(target.Key, index, flags);
+
+		#endregion
+
+		#region ListInsertAfterAsync
+
+		/// <inheritdoc cref="IDatabaseAsync.ListInsertAfterAsync(RedisKey, RedisValue, RedisValue, CommandFlags)"/>
+		public static Task<long> InsertAfterAsync<TElement>(this RedisList<TElement> target, TElement pivot, TElement item, CommandFlags flags = CommandFlags.None)
+			=> target.Database.ListInsertAfterAsync(target.Key, pivot, item, flags);
+
+		#endregion
+
+		#region ListInsertBeforeAsync
+
+		/// <inheritdoc cref="IDatabaseAsync.ListInsertBeforeAsync(RedisKey, RedisValue, RedisValue, CommandFlags)"/>
+		public static Task<long> InsertBeforeAsync<TElement>(this RedisList<TElement> target, TElement pivot, TElement item, CommandFlags flags = CommandFlags.None)
+			=> target.Database.ListInsertBeforeAsync(target.Key, pivot, item, flags);
+
+		#endregion
+
+		#region ListLeftPopAsync
+
+		/// <inheritdoc cref="IDatabaseAsync.ListLeftPopAsync(RedisKey, CommandFlags)"/>
+		public static Task<TElement> LeftPopAsync<TElement>(this RedisList<TElement> target, CommandFlags flags = CommandFlags.None)
+			=> target.Database.ListLeftPopAsync(target.Key, flags);
+
+		#endregion
+
+		#region ListLeftPushAsync
+
+		/// <inheritdoc cref="IDatabaseAsync.ListLeftPushAsync(RedisKey, RedisValue, When, CommandFlags)"/>
+		public static Task<long> LeftPushAsync<TElement>(this RedisList<TElement> target, TElement item, When when = When.Always, CommandFlags flags = CommandFlags.None)
+			=> target.Database.ListLeftPushAsync(target.Key, item, when, flags);
+
+		/// <inheritdoc cref="IDatabaseAsync.ListLeftPushAsync(RedisKey, RedisValue[], When, CommandFlags)"/>
+		public static Task<long> LeftPushAsync<TElement>(this RedisList<TElement> target, TElement[] values, When when = When.Always, CommandFlags flags = CommandFlags.None)
+			=> target.Database.ListLeftPushAsync(target.Key, values, when, flags);
+
+		#endregion
+
+		#region ListLengthAsync
+
+		/// <inheritdoc cref="IDatabaseAsync.ListLengthAsync(RedisKey, CommandFlags)"/>
+		public static Task<long> LengthAsync<TElement>(this RedisList<TElement> target)
+			=> target.Database.ListLengthAsync(target.Key);
+
+		#endregion
+
+		#region ListRangeAsync
+
+		/// <inheritdoc cref="IDatabaseAsync.ListRangeAsync(RedisKey, Int64, Int64, CommandFlags)"/>
+		public static Task<TElement[]> RangeAsync<TElement>(this RedisList<TElement> target, long start = 0, long stop = -1, CommandFlags flags = CommandFlags.None)
+			=> target.Database.ListRangeAsync(target.Key, start, stop, flags);
+
+		#endregion
+
+		#region ListRemoveAsync
+
+		/// <inheritdoc cref="IDatabaseAsync.ListRemoveAsync(RedisKey, RedisValue, Int64, CommandFlags)"/>
+		public static Task<long> RemoveAsync<TElement>(this RedisList<TElement> target, TElement value, long count = 0L, CommandFlags flags = CommandFlags.None)
+			=> target.Database.ListRemoveAsync(target.Key, value, count, flags);
+
+		#endregion
+
+		#region ListRightPopAsync
+
+		/// <inheritdoc cref="IDatabaseAsync.ListRightPopAsync(RedisKey, CommandFlags)"/>
+		public static Task<TElement> RightPopAsync<TElement>(this RedisList<TElement> target, CommandFlags flags = CommandFlags.None)
+			=> target.Database.ListRightPopAsync(target.Key, flags);
+
+		#endregion
+
+		#region ListRightPushAsync
+
+		/// <inheritdoc cref="IDatabaseAsync.ListRightPushAsync(RedisKey, RedisValue, When, CommandFlags)"/>
+		public static Task<long> RightPushAsync<TElement>(this RedisList<TElement> target, TElement item, When when = When.Always, CommandFlags flags = CommandFlags.None)
+			=> target.Database.ListRightPushAsync(target.Key, item, when, flags);
+
+		/// <inheritdoc cref="IDatabaseAsync.ListRightPushAsync(RedisKey, RedisValue[], When, CommandFlags)"/>
+		public static Task<long> RightPushAsync<TElement>(this RedisList<TElement> target, IEnumerable<TElement> values, When when = When.Always, CommandFlags flags = CommandFlags.None)
+			=> target.Database.ListRightPushAsync(target.Key, values, when, flags);
+
+		#endregion
+
+		#region ListSetByIndexAsync
+
+		/// <inheritdoc cref="IDatabaseAsync.ListSetByIndexAsync(RedisKey, Int64, RedisValue, CommandFlags)"/>
+		public static Task SetByIndexAsync<TElement>(this RedisList<TElement> target, long index, TElement value, CommandFlags flags = CommandFlags.None)
+			=> target.Database.ListSetByIndexAsync(target.Key, index, value, flags);
+
+		#endregion
+
+		#region ListTrimAsync
+
+		/// <inheritdoc cref="IDatabaseAsync.ListTrimAsync(RedisKey, Int64, Int64, CommandFlags)"/>
+		public static Task TrimAsync<TElement>(this RedisList<TElement> target, long start, long stop, CommandFlags flags = CommandFlags.None)
+			=> target.Database.ListTrimAsync(target.Key, start, stop, flags);
+
+		#endregion
+	}
+}
diff --git a/Engine/Source/Programs/Shared/EpicGames.Redis/RedisSet.cs b/Engine/Source/Programs/Shared/EpicGames.Redis/RedisSet.cs
new file mode 100644
--- /dev/null
+++ b/Engine/Source/Programs/Shared/EpicGames.Redis/RedisSet.cs
@@ -0,0 +1,121 @@
+// Copyright Epic Games, Inc. All Rights Reserved.
+
+using System;
+using System.Collections.Generic;
+using System.Threading.Tasks;
+using StackExchange.Redis;
+
+namespace EpicGames.Redis
+{
+	/// <summary>
+	/// Represents a typed Redis set with a given key
+	/// </summary>
+	/// <typeparam name="TElement"></typeparam>
+	public record struct RedisSet<TElement>(IDatabaseAsync Database, RedisSetKey<TElement> Key);
+
+	/// <summary>
+	/// Extension methods for sets
+	/// </summary>
+	public static class RedisSetExtensions
+	{
+		#region Conditions
+
+		/// <inheritdoc cref="Condition.SetContains(RedisKey, RedisValue)"/>
+		public static Condition Contains<TElement>(this RedisSet<TElement> target, TElement value)
+			=> target.Key.Contains(value);
+
+		/// <inheritdoc cref="Condition.SetLengthEqual(RedisKey, Int64)"/>
+		public static Condition LengthEqual<TElement>(this RedisSet<TElement> target, long length)
+			=> target.Key.LengthEqual(length);
+
+		/// <inheritdoc cref="Condition.SetLengthGreaterThan(RedisKey, Int64)"/>
+		public static Condition LengthGreaterThan<TElement>(this RedisSet<TElement> target, long length)
+			=> target.Key.LengthGreaterThan(length);
+
+		/// <inheritdoc cref="Condition.SetLengthLessThan(RedisKey, Int64)"/>
+		public static Condition LengthLessThan<TElement>(this RedisSet<TElement> target, long length)
+			=> target.LengthLessThan(length);
+
+		/// <inheritdoc cref="Condition.SetContains(RedisKey, RedisValue)"/>
+		public static Condition NotContains<TElement>(this RedisSet<TElement> target, TElement value)
+			=> target.NotContains(value);
+
+		#endregion
+
+		#region SetAddAsync
+
+		/// <inheritdoc cref="IDatabaseAsync.SetAddAsync(RedisKey, RedisValue, CommandFlags)"/>
+		public static Task<bool> AddAsync<TElement>(this RedisSet<TElement> target, TElement item, CommandFlags flags = CommandFlags.None)
+			=> target.Database.SetAddAsync(target.Key, item, flags);
+
+		/// <inheritdoc cref="IDatabaseAsync.SetAddAsync(RedisKey, RedisValue[], CommandFlags)"/>
+		public static Task<long> AddAsync<TElement>(this RedisSet<TElement> target, TElement[] values, CommandFlags flags = CommandFlags.None)
+			=> target.Database.SetAddAsync(target.Key, values, flags);
+
+		#endregion
+
+		#region SetContainsAsync
+
+		/// <inheritdoc cref="IDatabaseAsync.SetContainsAsync(RedisKey, RedisValue, CommandFlags)"/>
+		public static Task<bool> ContainsAsync<TElement>(this RedisSet<TElement> target, TElement value, CommandFlags flags = CommandFlags.None)
+			=> target.Database.SetContainsAsync(target.Key, value, flags);
+
+		#endregion
+
+		#region SetLengthAsync
+
+		/// <inheritdoc cref="IDatabaseAsync.SetLengthAsync(RedisKey, CommandFlags)"/>
+		public static Task<long> LengthAsync<TElement>(this RedisSet<TElement> target, CommandFlags flags = CommandFlags.None)
+			=> target.Database.SetLengthAsync(target.Key, flags);
+
+		#endregion
+
+		#region SetMembersAsync
+
+		/// <inheritdoc cref="IDatabaseAsync.SetMembersAsync(RedisKey, CommandFlags)"/>
+		public static Task<TElement[]> MembersAsync<TElement>(this RedisSet<TElement> target, CommandFlags flags = CommandFlags.None)
+			=> target.Database.SetMembersAsync(target.Key, flags);
+
+		#endregion
+
+		#region SetPopAsync
+
+		/// <inheritdoc cref="IDatabaseAsync.SetPopAsync(RedisKey, CommandFlags)"/>
+		public static Task<TElement> PopAsync<TElement>(this RedisSet<TElement> target, CommandFlags flags = CommandFlags.None)
+			=> target.Database.SetPopAsync(target.Key, flags);
+
+		/// <inheritdoc cref="IDatabaseAsync.SetPopAsync(RedisKey, Int64, CommandFlags)"/>
+		public static Task<TElement[]> PopAsync<TElement>(this RedisSet<TElement> target, long count, CommandFlags flags = CommandFlags.None)
+			=> target.Database.SetPopAsync(target.Key, count, flags);
+
+		#endregion
+
+		#region SetRandomMemberAsync
+
+		/// <inheritdoc cref="IDatabaseAsync.SetRandomMemberAsync(RedisKey, CommandFlags)"/>
+		public static Task<TElement> RandomMemberAsync<TElement>(this RedisSet<TElement> target, CommandFlags flags = CommandFlags.None)
+			=> target.Database.SetRandomMemberAsync(target.Key, flags);
+
+		#endregion
+
+		#region SetRemoveAsync
+
+		/// <inheritdoc cref="IDatabaseAsync.SetRemoveAsync(RedisKey, RedisValue, CommandFlags)"/>
+		public static Task<bool> RemoveAsync<TElement>(this RedisSet<TElement> target, TElement item, CommandFlags flags = CommandFlags.None)
+			=> target.Database.SetRemoveAsync(target.Key, item, flags);
+
+		/// <inheritdoc cref="IDatabaseAsync.SetRemoveAsync(RedisKey, RedisValue[], CommandFlags)"/>
+		public static Task<long> RemoveAsync<TElement>(this RedisSet<TElement> target, TElement[] values, CommandFlags flags = CommandFlags.None)
+			=> target.Database.SetRemoveAsync(target.Key, values, flags);
+
+		#endregion
+
+		#region SetScanAsync
+
+		/// <inheritdoc cref="IDatabaseAsync.SetScanAsync(RedisKey, RedisValue, int, long, int, CommandFlags)"/>
+		public static IAsyncEnumerable<TElement> ScanAsync<TElement>(this RedisSet<TElement> target, RedisValue pattern = default, int pageSize = 250, long cursor = 0, int pageOffset = 0, CommandFlags flags = CommandFlags.None)
+			=> target.Database.SetScanAsync(target.Key, pattern, pageSize, cursor, pageOffset, flags);
+
+		#endregion
+	}
+}
diff --git a/Engine/Source/Programs/Shared/EpicGames.Redis/RedisSetKey.cs b/Engine/Source/Programs/Shared/EpicGames.Redis/RedisSetKey.cs
--- a/Engine/Source/Programs/Shared/EpicGames.Redis/RedisSetKey.cs
+++ b/Engine/Source/Programs/Shared/EpicGames.Redis/RedisSetKey.cs
@@ -117,6 +117,22 @@ namespace EpicGames.Redis
 
 		#endregion
 
+		#region SetRandomMemberAsync
+
+		/// <inheritdoc cref="IDatabaseAsync.SetRandomMemberAsync(RedisKey, CommandFlags)"/>
+		public static Task<TElement> SetRandomMemberAsync<TElement>(this IDatabaseAsync target, RedisSetKey<TElement> key, CommandFlags flags = CommandFlags.None)
+			=> target.SetRandomMemberAsync(key, flags);
+
+		#endregion
+
+		#region SetRandomMembersAsync
+
+		/// <inheritdoc cref="IDatabaseAsync.SetRandomMembersAsync(RedisKey, long, CommandFlags)"/>
+		public static Task<TElement[]> SetRandomMembersAsync<TElement>(this IDatabaseAsync target, RedisSetKey<TElement> key, long count, CommandFlags flags = CommandFlags.None)
+			=> target.SetRandomMembersAsync(key, count, flags);
+
+		#endregion
+
 		#region SetRemoveAsync
 
 		/// <inheritdoc cref="IDatabaseAsync.SetRemoveAsync(RedisKey, RedisValue, CommandFlags)"/>
diff --git a/Engine/Source/Programs/Shared/EpicGames.Redis/RedisSortedSet.cs b/Engine/Source/Programs/Shared/EpicGames.Redis/RedisSortedSet.cs
new file mode 100644
--- /dev/null
+++ b/Engine/Source/Programs/Shared/EpicGames.Redis/RedisSortedSet.cs
@@ -0,0 +1,200 @@
+// Copyright Epic Games, Inc. All Rights Reserved.
+
+using System;
+using System.Collections.Generic;
+using System.Threading.Tasks;
+using StackExchange.Redis;
+
+namespace EpicGames.Redis
+{
+	/// <summary>
+	/// Represents a typed Redis sorted set with a given key
+	/// </summary>
+	public record struct RedisSortedSet<TElement>(IDatabaseAsync Database, RedisSortedSetKey<TElement> Key);
+
+	/// <summary>
+	/// Extension methods for sets
+	/// </summary>
+	public static class RedisSortedSetExtensions
+	{
+		#region Conditions
+
+		/// <inheritdoc cref="Condition.SortedSetContains(RedisKey, RedisValue)"/>
+		public static Condition Contains<TElement>(this RedisSortedSet<TElement> target, TElement value)
+			=> target.Key.Contains(value);
+
+		/// <inheritdoc cref="Condition.SortedSetEqual(RedisKey, RedisValue, RedisValue)"/>
+		public static Condition Equal<TElement>(this RedisSortedSet<TElement> target, TElement value, RedisValue score)
+			=> target.Key.Equal(value, score);
+
+		/// <inheritdoc cref="Condition.SortedSetLengthEqual(RedisKey, Int64)"/>
+		public static Condition LengthEqual<TElement>(this RedisSortedSet<TElement> target, long length)
+			=> target.Key.LengthLessThan(length);
+
+		/// <inheritdoc cref="Condition.SortedSetLengthGreaterThan(RedisKey, Int64)"/>
+		public static Condition LengthGreaterThan<TElement>(this RedisSortedSet<TElement> target, long length)
+			=> target.Key.LengthGreaterThan(length);
+
+		/// <inheritdoc cref="Condition.SortedSetLengthLessThan(RedisKey, Int64)"/>
+		public static Condition LengthLessThan<TElement>(this RedisSortedSet<TElement> target, long length)
+			=> target.Key.LengthLessThan(length);
+
+		/// <inheritdoc cref="Condition.SortedSetNotContains(RedisKey, RedisValue)"/>
+		public static Condition NotContains<TElement>(this RedisSortedSet<TElement> target, TElement value)
+			=> target.Key.NotContains(value);
+
+		/// <inheritdoc cref="Condition.SortedSetNotEqual(RedisKey, RedisValue, RedisValue)"/>
+		public static Condition NotEqual<TElement>(this RedisSortedSet<TElement> target, TElement value, RedisValue score)
+			=> target.Key.NotEqual(value, score);
+
+		/// <inheritdoc cref="Condition.SortedSetScoreExists(RedisKey, RedisValue)"/>
+		public static Condition ScoreExists<TElement>(this RedisSortedSet<TElement> target, RedisValue score)
+			=> target.Key.ScoreExists(score);
+
+		/// <inheritdoc cref="Condition.SortedSetScoreExists(RedisKey, RedisValue, RedisValue)"/>
+		public static Condition ScoreExists<TElement>(this RedisSortedSet<TElement> target, RedisValue score, RedisValue count)
+			=> target.Key.ScoreExists(score, count);
+
+		#endregion
+
+		#region SortedSetAddAsync
+
+		/// <inheritdoc cref="IDatabaseAsync.SortedSetAddAsync(RedisKey, RedisValue, Double, CommandFlags)"/>
+		public static Task<bool> AddAsync<TElement>(this RedisSortedSet<TElement> target, TElement value, double score, SortedSetWhen when = SortedSetWhen.Always, CommandFlags flags = CommandFlags.None)
+			=> target.Database.SortedSetAddAsync(target.Key, value, score, when, flags);
+
+		/// <inheritdoc cref="IDatabaseAsync.SortedSetAddAsync(RedisKey, RedisValue, Double, When, CommandFlags)"/>
+		public static Task<long> AddAsync<TElement>(this RedisSortedSet<TElement> target, SortedSetEntry<TElement>[] values, SortedSetWhen when = SortedSetWhen.Always, CommandFlags flags = CommandFlags.None)
+			=> target.Database.SortedSetAddAsync(target.Key, values, when, flags);
+
+		#endregion
+
+		#region SortedSetLengthAsync
+
+		/// <inheritdoc cref="IDatabaseAsync.SortedSetLengthAsync(RedisKey, Double, Double, Exclude, CommandFlags)"/>
+		public static Task<long> LengthAsync<TElement>(this RedisSortedSet<TElement> target, double min = Double.NegativeInfinity, double max = Double.PositiveInfinity, Exclude exclude = Exclude.None, CommandFlags flags = CommandFlags.None)
+			=> target.Database.SortedSetLengthAsync(target.Key, min, max, exclude, flags);
+
+		#endregion
+
+		#region SortedSetLengthByValueAsync
+
+		/// <inheritdoc cref="IDatabaseAsync.SortedSetLengthByValueAsync(RedisKey, RedisValue, RedisValue, Exclude, CommandFlags)"/>
+		public static Task<long> LengthByValueAsync<TElement>(this RedisSortedSet<TElement> target, TElement min, TElement max, Exclude exclude = Exclude.None, CommandFlags flags = CommandFlags.None)
+			=> target.Database.SortedSetLengthByValueAsync(target.Key, min, max, exclude, flags);
+
+		#endregion
+
+		#region SortedSetRangeByRankAsync
+
+		/// <inheritdoc cref="IDatabaseAsync.SortedSetRangeByRankAsync(RedisKey, Int64, Int64, Order, CommandFlags)"/>
+		public static Task<TElement[]> RangeByRankAsync<TElement>(this RedisSortedSet<TElement> target, long start = 0, long stop = -1, Order order = Order.Ascending, CommandFlags flags = CommandFlags.None)
+			=> target.Database.SortedSetRangeByRankAsync(target.Key, start, stop, order, flags);
+
+		#endregion
+
+		#region SortedSetRangeByRankWithScoresAsync
+
+		/// <inheritdoc cref="IDatabaseAsync.SortedSetRangeByRankWithScoresAsync(RedisKey, Int64, Int64, Order, CommandFlags)"/>
+		public static Task<SortedSetEntry<TElement>[]> RangeByRankWithScoresAsync<TElement>(this RedisSortedSet<TElement> target, long start, long stop = -1, Order order = Order.Ascending, CommandFlags flags = CommandFlags.None)
+			=> target.Database.SortedSetRangeByRankWithScoresAsync(target.Key, start, stop, order, flags);
+
+		#endregion
+
+		#region SortedSetRangeByScoreAsync
+
+		/// <inheritdoc cref="IDatabaseAsync.SortedSetRangeByScoreAsync(RedisKey, Double, Double, Exclude, Order, Int64, Int64, CommandFlags)"/>
+		public static Task<TElement[]> RangeByScoreAsync<TElement>(this RedisSortedSet<TElement> target, double start = Double.NegativeInfinity, double stop = Double.PositiveInfinity, Exclude exclude = Exclude.None, Order order = Order.Ascending, long skip = 0L, long take = -1L, CommandFlags flags = CommandFlags.None)
+			=> target.Database.SortedSetRangeByScoreAsync(target.Key, start, stop, exclude, order, skip, take, flags);
+
+		#endregion
+
+		#region SortedSetRangeByScoreWithScoresAsync
+
+		/// <inheritdoc cref="IDatabaseAsync.SortedSetRangeByScoreWithScoresAsync(RedisKey, Double, Double, Exclude, Order, Int64, Int64, CommandFlags)"/>
+		public static Task<SortedSetEntry<TElement>[]> SortedSetRangeByScoreWithScoresAsync<TElement>(this RedisSortedSet<TElement> target, double start = Double.NegativeInfinity, double stop = Double.PositiveInfinity, Exclude exclude = Exclude.None, Order order = Order.Ascending, long skip = 0L, long take = -1L, CommandFlags flags = CommandFlags.None)
+			=> target.Database.SortedSetRangeByScoreWithScoresAsync(target.Key, start, stop, exclude, order, skip, take, flags);
+
+		#endregion
+
+		#region SortedSetRangeByValueAsync
+
+		/// <inheritdoc cref="IDatabaseAsync.SortedSetRangeByValueAsync(RedisKey, RedisValue, RedisValue, Exclude, Order, Int64, Int64, CommandFlags)"/>
+		public static Task<TElement[]> RangeByValueAsync<TElement>(this RedisSortedSet<TElement> target, TElement min, TElement max, Exclude exclude = Exclude.None, Order order = Order.Ascending, long skip = 0L, long take = -1L, CommandFlags flags = CommandFlags.None)
+			=> target.Database.SortedSetRangeByValueAsync(target.Key, min, max, exclude, order, skip, take, flags);
+
+		#endregion
+
+		#region SortedSetRankAsync
+
+		/// <inheritdoc cref="IDatabaseAsync.SortedSetRankAsync(RedisKey, RedisValue, Order, CommandFlags)"/>
+		public static Task<long?> RankAsync<TElement>(this RedisSortedSet<TElement> target, TElement item, Order order = Order.Ascending, CommandFlags flags = CommandFlags.None)
+			=> target.Database.SortedSetRankAsync(target.Key, item, order, flags);
+
+		#endregion
+
+		#region SortedSetRemoveAsync
+
+		/// <inheritdoc cref="IDatabaseAsync.SortedSetRemoveAsync(RedisKey, RedisValue, CommandFlags)"/>
+		public static Task<bool> RemoveAsync<TElement>(this RedisSortedSet<TElement> target, TElement value, CommandFlags flags = CommandFlags.None)
+			=> target.Database.SortedSetRemoveAsync(target.Key, value, flags);
+
+		/// <inheritdoc cref="IDatabaseAsync.SortedSetRemoveAsync(RedisKey, RedisValue[], CommandFlags)"/>
+		public static Task<long> RemoveAsync<TElement>(this RedisSortedSet<TElement> target, TElement[] values, CommandFlags flags = CommandFlags.None)
+			=> target.Database.SortedSetRemoveAsync(target.Key, values, flags);
+
+		#endregion
+
+		#region SortedSetRemoveRangeByRankAsync
+
+		/// <inheritdoc cref="IDatabaseAsync.SortedSetRemoveRangeByRankAsync(RedisKey, Int64, Int64, CommandFlags)"/>
+		public static Task<long> RemoveRangeByRankAsync<TElement>(this RedisSortedSet<TElement> target, long start, long stop, CommandFlags flags = CommandFlags.None)
+			=> target.Database.SortedSetRemoveRangeByRankAsync(target.Key, start, stop, flags);
+
+		#endregion
+
+		#region SortedSetRemoveRangeByScoreAsync
+
+		/// <inheritdoc cref="IDatabaseAsync.SortedSetRemoveRangeByScoreAsync(RedisKey, Double, Double, Exclude, CommandFlags)"/>
+		public static Task<long> RemoveRangeByScoreAsync<TElement>(this RedisSortedSet<TElement> target, double start, double stop, Exclude exclude = Exclude.None, CommandFlags flags = CommandFlags.None)
+			=> target.Database.SortedSetRemoveRangeByScoreAsync(target.Key, start, stop, exclude, flags);
+
+		#endregion
+
+		#region SortedSetRemoveRangeByValueAsync
+
+		/// <inheritdoc cref="IDatabaseAsync.SortedSetRemoveRangeByValueAsync(RedisKey, RedisValue, RedisValue, Exclude, CommandFlags)"/>
+		public static Task<long> RemoveRangeByValueAsync<TElement>(this RedisSortedSet<TElement> target, TElement min, TElement max, Exclude exclude = Exclude.None, CommandFlags flags = CommandFlags.None)
+			=> target.Database.SortedSetRemoveRangeByValueAsync(target.Key, min, max, exclude, flags);
+
+		#endregion
+
+		#region SortedSetScanAsync
+
+		/// <inheritdoc cref="IDatabaseAsync.SortedSetScanAsync(RedisKey, RedisValue, Int32, Int64, Int32, CommandFlags)"/>
+		public static IAsyncEnumerable<SortedSetEntry<TElement>> ScanAsync<TElement>(this RedisSortedSet<TElement> target, RedisValue pattern = default, int pageSize = 250, long cursor = 0, int pageOffset = 0, CommandFlags flags = CommandFlags.None)
+			=> target.Database.SortedSetScanAsync(target.Key, pattern, pageSize, cursor, pageOffset, flags);
+
+		#endregion
+
+		#region SortedSetScoreAsync
+
+		/// <inheritdoc cref="IDatabaseAsync.SortedSetScoreAsync(RedisKey, RedisValue, CommandFlags)"/>
+		public static Task<double?> ScoreAsync<TElement>(this RedisSortedSet<TElement> target, TElement member, CommandFlags flags = CommandFlags.None)
+			=> target.Database.SortedSetScoreAsync(target.Key, member, flags);
+
+		#endregion
+
+		#region SortedSetUpdateAsync
+
+		/// <inheritdoc cref="IDatabaseAsync.SortedSetUpdateAsync(RedisKey, RedisValue, double, SortedSetWhen, CommandFlags)"/>
+		public static Task<bool> UpdateAsync<TElement>(this RedisSortedSet<TElement> target, TElement member, double score, SortedSetWhen when = SortedSetWhen.Always, CommandFlags flags = CommandFlags.None)
+			=> target.Database.SortedSetUpdateAsync(target.Key, member, score, when, flags);
+
+		/// <inheritdoc cref="IDatabaseAsync.SortedSetUpdateAsync(RedisKey, RedisValue, double, SortedSetWhen, CommandFlags)"/>
+		public static Task<long> UpdateAsync<TElement>(this RedisSortedSet<TElement> target, SortedSetEntry<TElement>[] values, SortedSetWhen when = SortedSetWhen.Always, CommandFlags flags = CommandFlags.None)
+			=> target.Database.SortedSetUpdateAsync(target.Key, values, when, flags);
+
+		#endregion
+	}
+}
diff --git a/Engine/Source/Programs/Shared/EpicGames.Redis/RedisSortedSetKey.cs b/Engine/Source/Programs/Shared/EpicGames.Redis/RedisSortedSetKey.cs
--- a/Engine/Source/Programs/Shared/EpicGames.Redis/RedisSortedSetKey.cs
+++ b/Engine/Source/Programs/Shared/EpicGames.Redis/RedisSortedSetKey.cs
@@ -164,6 +164,26 @@ namespace EpicGames.Redis
 
 		#endregion
 
+		#region SortedSetRandomMemberAsync
+
+		/// <inheritdoc cref="IDatabaseAsync.SortedSetRandomMemberAsync(RedisKey, CommandFlags)"/>
+		public static Task<TElement> SortedSetRandomMemberAsync<TElement>(this IDatabaseAsync target, RedisSortedSetKey<TElement> key, CommandFlags flags = CommandFlags.None)
+		{
+			return target.SortedSetRandomMemberAsync(key.Inner, flags).DeserializeAsync<TElement>();
+		}
+
+		#endregion
+
+		#region SortedSetRandomMembersAsync
+
+		/// <inheritdoc cref="IDatabaseAsync.SortedSetRandomMemberAsync(RedisKey, CommandFlags)"/>
+		public static Task<TElement[]> SortedSetRandomMembersAsync<TElement>(this IDatabaseAsync target, RedisSortedSetKey<TElement> key, long count, CommandFlags flags = CommandFlags.None)
+		{
+			return target.SortedSetRandomMembersAsync(key.Inner, count, flags).DeserializeAsync<TElement>();
+		}
+
+		#endregion
+
 		#region SortedSetRangeByRankAsync
 
 		/// <inheritdoc cref="IDatabaseAsync.SortedSetRangeByRankAsync(RedisKey, Int64, Int64, Order, CommandFlags)"/>
diff --git a/Engine/Source/Programs/Shared/EpicGames.Redis/RedisString.cs b/Engine/Source/Programs/Shared/EpicGames.Redis/RedisString.cs
new file mode 100644
--- /dev/null
+++ b/Engine/Source/Programs/Shared/EpicGames.Redis/RedisString.cs
@@ -0,0 +1,94 @@
+// Copyright Epic Games, Inc. All Rights Reserved.
+
+using System;
+using System.Threading.Tasks;
+using StackExchange.Redis;
+
+namespace EpicGames.Redis
+{
+	/// <summary>
+	/// Accessor for a typed Redis string
+	/// </summary>
+	/// <param name="Database">Database to operate on</param>
+	/// <param name="Key">Key for the string</param>
+	/// <typeparam name="TValue">Type of element stored in the string</typeparam>
+	public record struct RedisString<TValue>(IDatabaseAsync Database, RedisStringKey<TValue> Key);
+
+	/// <summary>
+	/// Extension methods for strings
+	/// </summary>
+	public static class RedisStringExtensions
+	{
+		#region Conditions
+
+		/// <inheritdoc cref="Condition.StringEqual(RedisKey, RedisValue)"/>
+		public static Condition Equal<TElement>(this RedisString<TElement> target, TElement value)
+			=> target.Key.Equal(value);
+
+		/// <inheritdoc cref="Condition.StringLengthEqual(RedisKey, Int64)"/>
+		public static Condition LengthEqual<TElement>(this RedisString<TElement> target, long length)
+			=> target.Key.LengthEqual(length);
+
+		/// <inheritdoc cref="Condition.StringLengthGreaterThan(RedisKey, Int64)"/>
+		public static Condition LengthGreaterThan<TElement>(this RedisString<TElement> target, long length)
+			=> target.Key.LengthGreaterThan(length);
+
+		/// <inheritdoc cref="Condition.StringLengthLessThan(RedisKey, Int64)"/>
+		public static Condition LengthLessThan<TElement>(this RedisString<TElement> target, long length)
+			=> target.Key.LengthLessThan(length);
+
+		/// <inheritdoc cref="Condition.StringNotEqual(RedisKey, RedisValue)"/>
+		public static Condition NotEqual<TElement>(this RedisString<TElement> target, TElement value)
+			=> target.Key.NotEqual(value);
+
+		#endregion
+
+		#region StringDecrementAsync
+
+		/// <inheritdoc cref="IDatabaseAsync.StringDecrementAsync(RedisKey, Int64, CommandFlags)"/>
+		public static Task<long> StringDecrementAsync(this RedisString<long> target, long value = 1L, CommandFlags flags = CommandFlags.None)
+			=> target.Database.StringDecrementAsync(target.Key, value, flags);
+
+		/// <inheritdoc cref="IDatabaseAsync.StringDecrementAsync(RedisKey, Int64, CommandFlags)"/>
+		public static Task<double> StringDecrementAsync(this RedisString<double> target, double value = 1.0, CommandFlags flags = CommandFlags.None)
+			=> target.Database.StringDecrementAsync(target.Key, value, flags);
+
+		#endregion
+
+		#region StringGetAsync
+
+		/// <inheritdoc cref="IDatabaseAsync.StringGetAsync(RedisKey, CommandFlags)"/>
+		public static Task<TValue?> GetAsync<TValue>(this RedisString<TValue> target, CommandFlags flags = CommandFlags.None)
+			=> target.Database.StringGetAsync(target.Key, flags);
+
+		/// <inheritdoc cref="IDatabaseAsync.StringGetAsync(RedisKey, CommandFlags)"/>
+		public static Task<TValue> GetAsync<TValue>(this RedisString<TValue> target, TValue defaultValue, CommandFlags flags = CommandFlags.None)
+			=> target.Database.StringGetAsync(target.Key, defaultValue, flags);
+
+		#endregion
+
+		#region StringIncrementAsync
+
+		/// <inheritdoc cref="IDatabaseAsync.StringIncrementAsync(RedisKey, Double, CommandFlags)"/>
+		public static Task<long> IncrementAsync(this RedisString<long> target, long value = 1L, CommandFlags flags = CommandFlags.None)
+			=> target.Database.StringIncrementAsync(target.Key, value, flags);
+
+		#endregion
+
+		#region StringLengthAsync
+
+		/// <inheritdoc cref="IDatabaseAsync.StringLengthAsync(RedisKey, CommandFlags)"/>
+		public static Task<long> LengthAsync<TValue>(this RedisString<TValue> target, CommandFlags flags = CommandFlags.None)
+			=> target.Database.StringLengthAsync(target.Key, flags);
+
+		#endregion
+
+		#region StringSetAsync
+
+		/// <inheritdoc cref="IDatabaseAsync.StringSetAsync(RedisKey, RedisValue, TimeSpan?, When, CommandFlags)"/>
+		public static Task<bool> SetAsync<TValue>(this RedisString<TValue> target, TValue value, TimeSpan? expiry = null, When when = When.Always, CommandFlags flags = CommandFlags.None)
+			=> target.Database.StringSetAsync<TValue>(target.Key, value, expiry, when, flags);
+
+		#endregion
+	}
+}
