EpicGames.Redis: Support for serializing regular Protobuf messages as well as types annotated with Protobuf.Net attributes.

diff --git a/Engine/Source/Programs/Shared/EpicGames.Redis/Converters/RedisProtobufConverter.cs b/Engine/Source/Programs/Shared/EpicGames.Redis/Converters/RedisProtobufConverter.cs
--- a/Engine/Source/Programs/Shared/EpicGames.Redis/Converters/RedisProtobufConverter.cs
+++ b/Engine/Source/Programs/Shared/EpicGames.Redis/Converters/RedisProtobufConverter.cs
@@ -2,16 +2,47 @@
 
 using System.IO;
 using EpicGames.Core;
+using Google.Protobuf;
+using Google.Protobuf.Reflection;
 using ProtoBuf;
 using StackExchange.Redis;
 
 namespace EpicGames.Redis.Converters
 {
+	/// <summary>
+	/// Converter for records to Redis values using protobuf serialization.
+	/// </summary>
+	/// <typeparam name="T">The record type</typeparam>
+	public class RedisProtobufConverter<T> : IRedisConverter<T> where T : IMessage<T>, new()
+	{
+		readonly MessageDescriptor _descriptor;
+
+		/// <summary>
+		/// Constructor
+		/// </summary>
+		public RedisProtobufConverter()
+		{
+			_descriptor = new T().Descriptor;
+		}
+
+		/// <inheritdoc/>
+		public T FromRedisValue(RedisValue value)
+		{
+			return (T)_descriptor.Parser.ParseFrom((byte[])value!);
+		}
+
+		/// <inheritdoc/>
+		public RedisValue ToRedisValue(T value)
+		{
+			return value.ToByteArray();
+		}
+	}
+
 	/// <summary>
 	/// Converter for records to Redis values using ProtoBuf-Net annotations.
 	/// </summary>
 	/// <typeparam name="T">The record type</typeparam>
-	public class RedisProtobufConverter<T> : IRedisConverter<T>
+	public class RedisProtobufNetConverter<T> : IRedisConverter<T>
 	{
 		/// <inheritdoc/>
 		public RedisValue ToRedisValue(T value)
diff --git a/Engine/Source/Programs/Shared/EpicGames.Redis/EpicGames.Redis.csproj b/Engine/Source/Programs/Shared/EpicGames.Redis/EpicGames.Redis.csproj
--- a/Engine/Source/Programs/Shared/EpicGames.Redis/EpicGames.Redis.csproj
+++ b/Engine/Source/Programs/Shared/EpicGames.Redis/EpicGames.Redis.csproj
@@ -36,6 +36,7 @@
   </ItemGroup>
   <ItemGroup>
     <PackageReference Include="Microsoft.Extensions.Logging" Version="8.0.0" />
+    <PackageReference Include="Google.Protobuf" Version="3.25.1" />
     <PackageReference Include="protobuf-net" Version="3.2.30" />
     <PackageReference Include="StackExchange.Redis" Version="2.7.4" />
   </ItemGroup>
diff --git a/Engine/Source/Programs/Shared/EpicGames.Redis/RedisSerializer.cs b/Engine/Source/Programs/Shared/EpicGames.Redis/RedisSerializer.cs
--- a/Engine/Source/Programs/Shared/EpicGames.Redis/RedisSerializer.cs
+++ b/Engine/Source/Programs/Shared/EpicGames.Redis/RedisSerializer.cs
@@ -9,6 +9,7 @@ using System.Threading.Tasks;
 using EpicGames.Core;
 using EpicGames.Redis.Converters;
 using EpicGames.Serialization;
+using Google.Protobuf;
 using ProtoBuf;
 using StackExchange.Redis;
 
@@ -221,11 +222,18 @@ namespace EpicGames.Redis
 				return (IRedisConverter<T>)nativeConverter;
 			}
 
+			// Check if the type is a protobuf message
+			if (type.IsAssignableTo(typeof(IMessage)))
+			{
+				Type converterType = typeof(RedisProtobufConverter<>).MakeGenericType(type);
+				return (IRedisConverter<T>)Activator.CreateInstance(converterType)!;
+			}
+
 			// Check if the type supports protobuf serialization
 			ProtoContractAttribute? protoAttribute = type.GetCustomAttribute<ProtoContractAttribute>();
 			if (protoAttribute != null)
 			{
-				return new RedisProtobufConverter<T>();
+				return new RedisProtobufNetConverter<T>();
 			}
 
 			// Check if there's a regular converter we can use to convert to/from a string
