[UBT]
* Fixed so cached build steps can contain log lines. This to preserve warnings etc

diff --git a/Engine/Source/Programs/Shared/EpicGames.UBA/ICacheClient.cs b/Engine/Source/Programs/Shared/EpicGames.UBA/ICacheClient.cs
--- a/Engine/Source/Programs/Shared/EpicGames.UBA/ICacheClient.cs
+++ b/Engine/Source/Programs/Shared/EpicGames.UBA/ICacheClient.cs
@@ -1,5 +1,7 @@
 // Copyright Epic Games, Inc. All Rights Reserved.
 
+using System.Collections.Generic;
+
 namespace EpicGames.UBA
 {
 	/// <summary>
@@ -32,6 +34,13 @@ namespace EpicGames.UBA
 		}
 	}
 
+	/// <summary>
+	/// Struct containing results from artifact fetch
+	/// </summary>
+	/// <param name="Success">Is set to true if succeeded in fetching artifacts</param>
+	/// <param name="LogLines">Contains log lines if any</param>
+	public record struct FetchFromCacheResult(bool Success, List<string> LogLines);
+
 	/// <summary>
 	/// Base interface for a cache client
 	/// </summary>
@@ -65,7 +74,7 @@ namespace EpicGames.UBA
 		/// <param name="bucket">Bucket to search for cache entry</param>
 		/// <param name="info">Process start info</param>
 		/// <returns>True if successful</returns>
-		public abstract bool FetchFromCache(IRootPaths rootPaths, uint bucket, ProcessStartInfo info);
+		public abstract FetchFromCacheResult FetchFromCache(IRootPaths rootPaths, uint bucket, ProcessStartInfo info);
 
 		/// <summary>
 		/// Request the connected server to shutdown
diff --git a/Engine/Source/Programs/Shared/EpicGames.UBA/Impl/CacheClientImpl.cs b/Engine/Source/Programs/Shared/EpicGames.UBA/Impl/CacheClientImpl.cs
--- a/Engine/Source/Programs/Shared/EpicGames.UBA/Impl/CacheClientImpl.cs
+++ b/Engine/Source/Programs/Shared/EpicGames.UBA/Impl/CacheClientImpl.cs
@@ -1,6 +1,7 @@
 // Copyright Epic Games, Inc. All Rights Reserved.
 
 using System;
+using System.Collections.Generic;
 using System.Runtime.InteropServices;
 
 namespace EpicGames.UBA
@@ -90,7 +91,16 @@ namespace EpicGames.UBA
 		static extern bool CacheClient_WriteToCache(IntPtr cacheClient, IntPtr rootPaths, uint bucket, IntPtr info, byte[] inputs, uint inputsSize, byte[] outputs, uint outputsSize);
 
 		[DllImport("UbaHost", CharSet = CharSet.Auto)]
-		static extern bool CacheClient_FetchFromCache(IntPtr cacheClient, IntPtr rootPaths, uint bucket, IntPtr info);
+		static extern IntPtr CacheClient_FetchFromCache2(IntPtr cacheClient, IntPtr rootPaths, uint bucket, IntPtr info);
+
+		[DllImport("UbaHost", CharSet = CharSet.Auto)]
+		static extern IntPtr CacheResult_GetLogLine(IntPtr result, uint index);
+
+		[DllImport("UbaHost", CharSet = CharSet.Auto)]
+		static extern uint CacheResult_GetLogLineType(IntPtr result, uint index);
+
+		[DllImport("UbaHost", CharSet = CharSet.Auto)]
+		static extern void CacheResult_Delete(IntPtr result);
 
 		[DllImport("UbaHost", CharSet = CharSet.Auto)]
 		static extern void CacheClient_RequestServerShutdown(IntPtr cacheClient, string reason);
@@ -138,12 +148,32 @@ namespace EpicGames.UBA
 			return CacheClient_WriteToCache(_handle, rootPaths.GetHandle(), bucket, process.GetHandle(), inputs, inputsSize, outputs, outputsSize);
 		}
 
-		public bool FetchFromCache(IRootPaths rootPaths, uint bucket, ProcessStartInfo info)
+		public FetchFromCacheResult FetchFromCache(IRootPaths rootPaths, uint bucket, ProcessStartInfo info)
 		{
 			IntPtr si = ProcessStartInfo_Create(info.Application, info.Arguments, info.WorkingDirectory, info.Description, (uint)info.Priority, info.OutputStatsThresholdMs, info.TrackInputs, info.LogFile ?? String.Empty, null);
-			bool result = CacheClient_FetchFromCache(_handle, rootPaths.GetHandle(), bucket, si);
+			IntPtr result = CacheClient_FetchFromCache2(_handle, rootPaths.GetHandle(), bucket, si);
 			ProcessStartInfo_Destroy(si);
-			return result;
+
+			if (result == IntPtr.Zero)
+			{
+				return new FetchFromCacheResult(false, new List<string>());
+			}
+
+			string? line = Marshal.PtrToStringAuto(CacheResult_GetLogLine(result, 0));
+			if (line == null)
+			{
+				return new FetchFromCacheResult(true, new List<string>());
+			}
+
+			List<string> logLines = new();
+			while (line != null)
+			{
+				logLines.Add(line);
+				line = Marshal.PtrToStringAuto(CacheResult_GetLogLine(result, (uint)logLines.Count));
+			}
+
+
+			return new FetchFromCacheResult(true, logLines);
 		}
 
 		public void RequestServerShutdown(string reason)
diff --git a/Engine/Source/Programs/UnrealBuildTool/Artifacts/ActionArtifactCache.cs b/Engine/Source/Programs/UnrealBuildTool/Artifacts/ActionArtifactCache.cs
--- a/Engine/Source/Programs/UnrealBuildTool/Artifacts/ActionArtifactCache.cs
+++ b/Engine/Source/Programs/UnrealBuildTool/Artifacts/ActionArtifactCache.cs
@@ -185,16 +185,16 @@ namespace UnrealBuildTool.Artifacts
 		}
 
 		/// <inheritdoc/>
-		public async Task<bool> CompleteActionFromCacheAsync(LinkedAction action, CancellationToken cancellationToken)
+		public async Task<ActionArtifactResult> CompleteActionFromCacheAsync(LinkedAction action, CancellationToken cancellationToken)
 		{
 			if (!EnableReads)
 			{
-				return false;
+				return new ActionArtifactResult(false, new List<string>());
 			}
 
 			if (!action.ArtifactMode.HasFlag(ArtifactMode.Enabled))
 			{
-				return false;
+				return new ActionArtifactResult(false, new List<string>());
 			}
 
 			ArtifactDirectoryMapping directoryMapping = GetDirectoryMapping(action);
@@ -215,7 +215,7 @@ namespace UnrealBuildTool.Artifacts
 				{
 					_logger.LogInformation("Artifact Cache Miss: No artifact actions found for {ActionDescription}", actionDescription);
 				}
-				return false;
+				return new ActionArtifactResult(false, new List<string>());
 			}
 
 			foreach (ArtifactAction artifactAction in artifactActions)
@@ -254,7 +254,7 @@ namespace UnrealBuildTool.Artifacts
 
 					if (readResults == null || readResults.Length == 0 || !readResults[0])
 					{
-						return false;
+						return new ActionArtifactResult(false, new List<string>());
 					}
 					else
 					{
@@ -265,11 +265,11 @@ namespace UnrealBuildTool.Artifacts
 							item.ResetCachedInfo(); // newly created outputs need refreshing
 							_fileHasher.SetDigest(item, output.ContentHash);
 						}
-						return true;
+						return new ActionArtifactResult(true, new List<string>());
 					}
 				}
 			}
-			return false;
+			return new ActionArtifactResult(false, new List<string>());
 		}
 
 		/// <inheritdoc/>
diff --git a/Engine/Source/Programs/UnrealBuildTool/Artifacts/Artifacts.cs b/Engine/Source/Programs/UnrealBuildTool/Artifacts/Artifacts.cs
--- a/Engine/Source/Programs/UnrealBuildTool/Artifacts/Artifacts.cs
+++ b/Engine/Source/Programs/UnrealBuildTool/Artifacts/Artifacts.cs
@@ -1,7 +1,9 @@
 // Copyright Epic Games, Inc. All Rights Reserved.
 
 using System;
+using System.Collections.Generic;
 using System.IO;
+using System.Linq;
 using System.Threading;
 using System.Threading.Tasks;
 using EpicGames.Core;
@@ -228,6 +230,13 @@ namespace UnrealBuildTool.Artifacts
 		public Task FlushChangesAsync(CancellationToken cancellationToken);
 	}
 
+	/// <summary>
+	/// Struct containing results from artifact fetch
+	/// </summary>
+	/// <param name="Success">Is set to true if succeeded in fetching artifacts</param>
+	/// <param name="LogLines">Contains log lines if any</param>
+	public record struct ActionArtifactResult(bool Success, List<string> LogLines);
+
 	/// <summary>
 	/// Interface for action specific support of artifacts.
 	/// </summary>
@@ -271,7 +280,7 @@ namespace UnrealBuildTool.Artifacts
 		/// <param name="action">Action to be completed</param>
 		/// <param name="cancellationToken">Token to be used to cancel operations</param>
 		/// <returns>True if it has been completed, false if not</returns>
-		public Task<bool> CompleteActionFromCacheAsync(LinkedAction action, CancellationToken cancellationToken);
+		public Task<ActionArtifactResult> CompleteActionFromCacheAsync(LinkedAction action, CancellationToken cancellationToken);
 
 		/// <summary>
 		/// Save the output for a completed action
diff --git a/Engine/Source/Programs/UnrealBuildTool/Executors/ImmediateActionQueue.cs b/Engine/Source/Programs/UnrealBuildTool/Executors/ImmediateActionQueue.cs
--- a/Engine/Source/Programs/UnrealBuildTool/Executors/ImmediateActionQueue.cs
+++ b/Engine/Source/Programs/UnrealBuildTool/Executors/ImmediateActionQueue.cs
@@ -384,8 +384,6 @@ namespace UnrealBuildTool
 		/// </summary>
 		private readonly IActionArtifactCache? _actionArtifactCache;
 
-		static readonly ExecuteResults s_copiedFromCacheResults = new(new List<string>(), 0, TimeSpan.Zero, TimeSpan.Zero, "[Cache]");
-
 		/// <summary>
 		/// Construct a new instance of the action queue
 		/// </summary>
@@ -432,11 +430,12 @@ namespace UnrealBuildTool
 				{
 					return new Func<Task>(async () =>
 					{
-						bool success = await _actionArtifactCache!.CompleteActionFromCacheAsync(action, CancellationToken);
-						if (success)
+						ActionArtifactResult result = await _actionArtifactCache!.CompleteActionFromCacheAsync(action, CancellationToken);
+						if (result.Success)
 						{
+							ExecuteResults results = new(result.LogLines, 0, TimeSpan.Zero, TimeSpan.Zero, "[Cache]");
 							Interlocked.Increment(ref _cacheHitActions);
-							OnActionCompleted(action, success, s_copiedFromCacheResults);
+							OnActionCompleted(action, true, results);
 							OnArtifactsRead?.Invoke(action);
 						}
 						else
diff --git a/Engine/Source/Programs/UnrealBuildTool/Executors/UnrealBuildAccelerator/UBAExecutor.cs b/Engine/Source/Programs/UnrealBuildTool/Executors/UnrealBuildAccelerator/UBAExecutor.cs
--- a/Engine/Source/Programs/UnrealBuildTool/Executors/UnrealBuildAccelerator/UBAExecutor.cs
+++ b/Engine/Source/Programs/UnrealBuildTool/Executors/UnrealBuildAccelerator/UBAExecutor.cs
@@ -237,7 +237,7 @@ namespace UnrealBuildTool
 			public bool LogCacheMisses { get => true; set { } }
 			public DirectoryReference? EngineRoot { get => null; set { } }
 			public DirectoryReference[]? DirectoryRoots { get => null; set { } }
-			public Task<bool> CompleteActionFromCacheAsync(LinkedAction action, CancellationToken cancellationToken)
+			public Task<ActionArtifactResult> CompleteActionFromCacheAsync(LinkedAction action, CancellationToken cancellationToken)
 			{
 				return Task.Factory.StartNew(() =>
 				{
@@ -245,7 +245,8 @@ namespace UnrealBuildTool
 					uint bucket = UBAExecutor.GetActionCacheBucket(action);
 					using (IRootPaths rootPaths = _executor.GetActionRootPaths(action))
 					{
-						return _executor._cacheClient!.FetchFromCache(rootPaths, bucket, startInfo);
+						FetchFromCacheResult result = _executor._cacheClient!.FetchFromCache(rootPaths, bucket, startInfo);
+						return new ActionArtifactResult(result.Success, result.LogLines);
 					}
 				}, cancellationToken, TaskCreationOptions.LongRunning | TaskCreationOptions.PreferFairness, TaskScheduler.Default);
 			}
