EpicGames.Redis: Add extension methods to allow creating conditions from typed Redis keys.

diff --git a/Engine/Source/Programs/Shared/EpicGames.Redis/IRedisTypedKey.cs b/Engine/Source/Programs/Shared/EpicGames.Redis/IRedisTypedKey.cs
new file mode 100644
--- /dev/null
+++ b/Engine/Source/Programs/Shared/EpicGames.Redis/IRedisTypedKey.cs
@@ -0,0 +1,75 @@
+﻿// Copyright Epic Games, Inc. All Rights Reserved.
+
+using System;
+using System.Collections.Generic;
+using System.Linq;
+using System.Threading.Tasks;
+using StackExchange.Redis;
+
+namespace EpicGames.Redis
+{
+	/// <summary>
+	/// Interface for typed redis keys
+	/// </summary>
+	public interface IRedisTypedKey
+	{
+		/// <summary>
+		/// The inner untyped key
+		/// </summary>
+		public RedisKey Inner { get; }
+	}
+
+	/// <summary>
+	/// Extension methods for <see cref="IRedisTypedKey"/>
+	/// </summary>
+	public static class RedisTypedKeyExtensions
+	{
+		#region Conditions
+
+		/// <inheritdoc cref="Condition.KeyExists(RedisKey)"/>
+		public static Condition KeyExists(this IRedisTypedKey key)
+			=> Condition.KeyExists(key.Inner);
+
+		/// <inheritdoc cref="Condition.KeyExists(RedisKey)"/>
+		public static Condition KeyNotExists(this IRedisTypedKey key)
+			=> Condition.KeyNotExists(key.Inner);
+
+		#endregion
+
+		#region KeyDeleteAsync
+
+		/// <inheritdoc cref="IDatabaseAsync.KeyDeleteAsync(RedisKey, CommandFlags)"/>
+		public static Task KeyDeleteAsync(this IDatabaseAsync target, IRedisTypedKey key, CommandFlags flags = CommandFlags.None)
+			=> target.KeyDeleteAsync(key.Inner, flags);
+
+		/// <inheritdoc cref="IDatabaseAsync.KeyDeleteAsync(RedisKey[], CommandFlags)"/>
+		public static Task KeyDeleteAsync(this IDatabaseAsync target, IRedisTypedKey[] keys, CommandFlags flags = CommandFlags.None)
+			=> target.KeyDeleteAsync(keys, flags);
+
+		#endregion
+
+		#region KeyExistsAsync
+
+		/// <inheritdoc cref="IDatabaseAsync.KeyExistsAsync(RedisKey, CommandFlags)"/>
+		public static Task<bool> KeyExistsAsync(this IDatabaseAsync target, IRedisTypedKey key, CommandFlags flags = CommandFlags.None)
+			=> target.KeyExistsAsync(key.Inner, flags);
+
+		/// <inheritdoc cref="IDatabaseAsync.KeyExistsAsync(RedisKey[], CommandFlags)"/>
+		public static Task<long> KeyExistsAsync(this IDatabaseAsync target, IEnumerable<IRedisTypedKey> keys, CommandFlags flags = CommandFlags.None)
+			=> target.KeyExistsAsync(keys.ToArray(), flags);
+
+		#endregion
+
+		#region KeyExpireAsync
+
+		/// <inheritdoc cref="IDatabaseAsync.KeyExpireAsync(RedisKey, DateTime?, ExpireWhen, CommandFlags)"/>
+		public static Task KeyExpireAsync(this IDatabaseAsync target, IRedisTypedKey key, DateTime? expiry = null, ExpireWhen when = ExpireWhen.Always, CommandFlags flags = CommandFlags.None)
+			=> target.KeyExpireAsync(key.Inner, expiry, when, flags);
+
+		/// <inheritdoc cref="IDatabaseAsync.KeyExpireAsync(RedisKey, TimeSpan?, ExpireWhen, CommandFlags)"/>
+		public static Task KeyExpireAsync(this IDatabaseAsync target, IRedisTypedKey key, TimeSpan? expiry, ExpireWhen when = ExpireWhen.Always, CommandFlags flags = CommandFlags.None)
+			=> target.KeyExpireAsync(key.Inner, expiry, when, flags);
+
+		#endregion
+	}
+}
diff --git a/Engine/Source/Programs/Shared/EpicGames.Redis/RedisHash.cs b/Engine/Source/Programs/Shared/EpicGames.Redis/RedisHash.cs
--- a/Engine/Source/Programs/Shared/EpicGames.Redis/RedisHash.cs
+++ b/Engine/Source/Programs/Shared/EpicGames.Redis/RedisHash.cs
@@ -26,48 +26,48 @@ namespace EpicGames.Redis
 		#region Conditions
 
 		/// <inheritdoc cref="Condition.HashEqual(RedisKey, RedisValue, RedisValue)"/>
-		public static Condition Equal<TRecord, TValue>(this RedisHash<TRecord> target, Expression<Func<TRecord, TValue>> selector, TValue value)
-			=> target.Key.Equal(selector, value);
+		public static Condition HashEqual<TRecord, TValue>(this RedisHash<TRecord> target, Expression<Func<TRecord, TValue>> selector, TValue value)
+			=> target.Key.HashEqual(selector, value);
 
 		/// <inheritdoc cref="Condition.HashEqual(RedisKey, RedisValue, RedisValue)"/>
-		public static Condition Equal<TName, TValue>(this RedisHash<TName, TValue> target, TName name, TValue value)
-			=> target.Key.Equal(name, value);
+		public static Condition HashEqual<TName, TValue>(this RedisHash<TName, TValue> target, TName name, TValue value)
+			=> target.Key.HashEqual(name, value);
 
 		/// <inheritdoc cref="Condition.HashExists(RedisKey, RedisValue)"/>
-		public static Condition Exists<TRecord, TValue>(this RedisHash<TRecord> target, Expression<Func<TRecord, TValue>> selector)
-			=> target.Key.Exists(selector);
+		public static Condition HashExists<TRecord, TValue>(this RedisHash<TRecord> target, Expression<Func<TRecord, TValue>> selector)
+			=> target.Key.HashExists(selector);
 
 		/// <inheritdoc cref="Condition.HashExists(RedisKey, RedisValue)"/>
-		public static Condition Exists<TName, TValue>(this RedisHash<TName, TValue> target, TName name)
-			=> target.Key.Exists(name);
+		public static Condition HashExists<TName, TValue>(this RedisHash<TName, TValue> target, TName name)
+			=> target.Key.HashExists(name);
 
 		/// <inheritdoc cref="Condition.HashLengthEqual(RedisKey, Int64)"/>
-		public static Condition LengthEqual<TName, TValue>(this RedisHash<TName, TValue> target, long length)
-			=> target.Key.LengthEqual(length);
+		public static Condition HashLengthEqual<TName, TValue>(this RedisHash<TName, TValue> target, long length)
+			=> target.Key.HashLengthEqual(length);
 
 		/// <inheritdoc cref="Condition.HashLengthGreaterThan(RedisKey, Int64)"/>
-		public static Condition LengthGreaterThan<TName, TValue>(this RedisHash<TName, TValue> target, long length)
-			=> target.Key.LengthGreaterThan(length);
+		public static Condition HashLengthGreaterThan<TName, TValue>(this RedisHash<TName, TValue> target, long length)
+			=> target.Key.HashLengthGreaterThan(length);
 
 		/// <inheritdoc cref="Condition.HashLengthLessThan(RedisKey, Int64)"/>
-		public static Condition LengthLessThan<TName, TValue>(this RedisHash<TName, TValue> target, long length)
-			=> target.Key.LengthLessThan(length);
+		public static Condition HashLengthLessThan<TName, TValue>(this RedisHash<TName, TValue> target, long length)
+			=> target.Key.HashLengthLessThan(length);
 
 		/// <inheritdoc cref="Condition.HashNotExists(RedisKey, RedisValue)"/>
-		public static Condition NotExists<TRecord, TValue>(this RedisHash<TRecord> target, Expression<Func<TRecord, TValue>> selector)
-			=> target.Key.NotExists(selector);
+		public static Condition HashNotExists<TRecord, TValue>(this RedisHash<TRecord> target, Expression<Func<TRecord, TValue>> selector)
+			=> target.Key.HashNotExists(selector);
 
 		/// <inheritdoc cref="Condition.HashNotExists(RedisKey, RedisValue)"/>
-		public static Condition NotExists<TName, TValue>(this RedisHash<TName, TValue> target, TName name)
-			=> target.Key.NotExists(name);
+		public static Condition HashNotExists<TName, TValue>(this RedisHash<TName, TValue> target, TName name)
+			=> target.Key.HashNotExists(name);
 
 		/// <inheritdoc cref="Condition.HashEqual(RedisKey, RedisValue, RedisValue)"/>
-		public static Condition NotEqual<TRecord, TValue>(this RedisHash<TRecord> target, Expression<Func<TRecord, TValue>> selector, TValue value)
-			=> target.Key.NotEqual(selector, value);
+		public static Condition HashNotEqual<TRecord, TValue>(this RedisHash<TRecord> target, Expression<Func<TRecord, TValue>> selector, TValue value)
+			=> target.Key.HashNotEqual(selector, value);
 
 		/// <inheritdoc cref="Condition.HashNotEqual(RedisKey, RedisValue, RedisValue)"/>
-		public static Condition NotEqual<TName, TValue>(this RedisHash<TName, TValue> target, TName name, TValue value)
-			=> target.Key.NotEqual(name, value);
+		public static Condition HashNotEqual<TName, TValue>(this RedisHash<TName, TValue> target, TName name, TValue value)
+			=> target.Key.HashNotEqual(name, value);
 
 		#endregion
 
diff --git a/Engine/Source/Programs/Shared/EpicGames.Redis/RedisHashKey.cs b/Engine/Source/Programs/Shared/EpicGames.Redis/RedisHashKey.cs
--- a/Engine/Source/Programs/Shared/EpicGames.Redis/RedisHashKey.cs
+++ b/Engine/Source/Programs/Shared/EpicGames.Redis/RedisHashKey.cs
@@ -13,19 +13,13 @@ namespace EpicGames.Redis
 	/// Represents a redis hash key, with members corresponding to the property names of a type
 	/// </summary>
 	/// <typeparam name="T">Type of the hash fields</typeparam>
-	public record struct RedisHashKey<T>(RedisKey Inner)
+	public record struct RedisHashKey<T>(RedisKey Inner) : IRedisTypedKey
 	{
 		/// <summary>
 		/// Implicit conversion to typed redis key.
 		/// </summary>
 		/// <param name="key">Key to convert</param>
 		public static implicit operator RedisHashKey<T>(string key) => new RedisHashKey<T>(new RedisKey(key));
-
-		/// <summary>
-		/// Implicit conversion to untyped redis keys.
-		/// </summary>
-		/// <param name="key">Key to convert</param>
-		public static implicit operator TypedRedisKey(RedisHashKey<T> key) => new TypedRedisKey(key.Inner);
 	}
 
 	/// <summary>
@@ -33,19 +27,13 @@ namespace EpicGames.Redis
 	/// </summary>
 	/// <typeparam name="TName">Type of the hash key</typeparam>
 	/// <typeparam name="TValue">Type of the hash value</typeparam>
-	public record struct RedisHashKey<TName, TValue>(RedisKey Inner)
+	public record struct RedisHashKey<TName, TValue>(RedisKey Inner) : IRedisTypedKey
 	{
 		/// <summary>
 		/// Implicit conversion to typed redis key.
 		/// </summary>
 		/// <param name="key">Key to convert</param>
 		public static implicit operator RedisHashKey<TName, TValue>(string key) => new RedisHashKey<TName, TValue>(new RedisKey(key));
-
-		/// <summary>
-		/// Implicit conversion to untyped redis keys.
-		/// </summary>
-		/// <param name="key">Key to convert</param>
-		public static implicit operator TypedRedisKey(RedisHashKey<TName, TValue> key) => new TypedRedisKey(key.Inner);
 	}
 
 	/// <inheritdoc cref="HashEntry"/>
@@ -107,59 +95,59 @@ namespace EpicGames.Redis
 		#region Conditions
 
 		/// <inheritdoc cref="Condition.HashEqual(RedisKey, RedisValue, RedisValue)"/>
-		public static Condition Equal<TRecord, TValue>(this RedisHashKey<TRecord> key, Expression<Func<TRecord, TValue>> selector, TValue value)
+		public static Condition HashEqual<TRecord, TValue>(this RedisHashKey<TRecord> key, Expression<Func<TRecord, TValue>> selector, TValue value)
 		{
 			MemberExpression memberExpression = (selector.Body as MemberExpression) ?? throw new InvalidOperationException("Expression must be a property accessor");
 			return Condition.HashEqual(key.Inner, memberExpression.Member.Name, RedisSerializer.Serialize(value));
 		}
 
 		/// <inheritdoc cref="Condition.HashEqual(RedisKey, RedisValue, RedisValue)"/>
-		public static Condition Equal<TName, TValue>(this RedisHashKey<TName, TValue> key, TName name, TValue value)
+		public static Condition HashEqual<TName, TValue>(this RedisHashKey<TName, TValue> key, TName name, TValue value)
 			=> Condition.HashEqual(key.Inner, RedisSerializer.Serialize(name), RedisSerializer.Serialize(value));
 
 		/// <inheritdoc cref="Condition.HashExists(RedisKey, RedisValue)"/>
-		public static Condition Exists<TRecord, TValue>(this RedisHashKey<TRecord> key, Expression<Func<TRecord, TValue>> selector)
+		public static Condition HashExists<TRecord, TValue>(this RedisHashKey<TRecord> key, Expression<Func<TRecord, TValue>> selector)
 		{
 			MemberExpression memberExpression = (selector.Body as MemberExpression) ?? throw new InvalidOperationException("Expression must be a property accessor");
 			return Condition.HashExists(key.Inner, memberExpression.Member.Name);
 		}
 
 		/// <inheritdoc cref="Condition.HashExists(RedisKey, RedisValue)"/>
-		public static Condition Exists<TName, TValue>(this RedisHashKey<TName, TValue> key, TName name)
+		public static Condition HashExists<TName, TValue>(this RedisHashKey<TName, TValue> key, TName name)
 			=> Condition.HashExists(key.Inner, RedisSerializer.Serialize(name));
 
 		/// <inheritdoc cref="Condition.HashLengthEqual(RedisKey, Int64)"/>
-		public static Condition LengthEqual<TName, TValue>(this RedisHashKey<TName, TValue> key, long length)
+		public static Condition HashLengthEqual<TName, TValue>(this RedisHashKey<TName, TValue> key, long length)
 			=> Condition.HashLengthEqual(key.Inner, length);
 
 		/// <inheritdoc cref="Condition.HashLengthGreaterThan(RedisKey, Int64)"/>
-		public static Condition LengthGreaterThan<TName, TValue>(this RedisHashKey<TName, TValue> key, long length)
+		public static Condition HashLengthGreaterThan<TName, TValue>(this RedisHashKey<TName, TValue> key, long length)
 			=> Condition.HashLengthGreaterThan(key.Inner, length);
 
 		/// <inheritdoc cref="Condition.HashLengthLessThan(RedisKey, Int64)"/>
-		public static Condition LengthLessThan<TName, TValue>(this RedisHashKey<TName, TValue> key, long length)
+		public static Condition HashLengthLessThan<TName, TValue>(this RedisHashKey<TName, TValue> key, long length)
 			=> Condition.HashLengthLessThan(key.Inner, length);
 
 		/// <inheritdoc cref="Condition.HashNotExists(RedisKey, RedisValue)"/>
-		public static Condition NotExists<TRecord, TValue>(this RedisHashKey<TRecord> key, Expression<Func<TRecord, TValue>> selector)
+		public static Condition HashNotExists<TRecord, TValue>(this RedisHashKey<TRecord> key, Expression<Func<TRecord, TValue>> selector)
 		{
 			MemberExpression memberExpression = (selector.Body as MemberExpression) ?? throw new InvalidOperationException("Expression must be a property accessor");
 			return Condition.HashNotExists(key.Inner, memberExpression.Member.Name);
 		}
 
 		/// <inheritdoc cref="Condition.HashNotExists(RedisKey, RedisValue)"/>
-		public static Condition NotExists<TName, TValue>(this RedisHashKey<TName, TValue> key, TName name)
+		public static Condition HashNotExists<TName, TValue>(this RedisHashKey<TName, TValue> key, TName name)
 			=> Condition.HashNotExists(key.Inner, RedisSerializer.Serialize(name));
 
 		/// <inheritdoc cref="Condition.HashEqual(RedisKey, RedisValue, RedisValue)"/>
-		public static Condition NotEqual<TRecord, TValue>(this RedisHashKey<TRecord> key, Expression<Func<TRecord, TValue>> selector, TValue value)
+		public static Condition HashNotEqual<TRecord, TValue>(this RedisHashKey<TRecord> key, Expression<Func<TRecord, TValue>> selector, TValue value)
 		{
 			MemberExpression memberExpression = (selector.Body as MemberExpression) ?? throw new InvalidOperationException("Expression must be a property accessor");
 			return Condition.HashNotEqual(key.Inner, memberExpression.Member.Name, RedisSerializer.Serialize(value));
 		}
 
 		/// <inheritdoc cref="Condition.HashNotEqual(RedisKey, RedisValue, RedisValue)"/>
-		public static Condition NotEqual<TName, TValue>(this RedisHashKey<TName, TValue> key, TName name, TValue value)
+		public static Condition HashNotEqual<TName, TValue>(this RedisHashKey<TName, TValue> key, TName name, TValue value)
 			=> Condition.HashNotEqual(key.Inner, RedisSerializer.Serialize(name), RedisSerializer.Serialize(value));
 
 		#endregion
@@ -304,7 +292,7 @@ namespace EpicGames.Redis
 		#region HashLengthAsync
 
 		/// <inheritdoc cref="IDatabaseAsync.HashLengthAsync(RedisKey, CommandFlags)"/>
-		public static Task<long> HashLengthAsync(this IDatabaseAsync target, TypedRedisKey key, CommandFlags flags = CommandFlags.None)
+		public static Task<long> HashLengthAsync<TName, TValue>(this IDatabaseAsync target, RedisHashKey<TName, TValue> key, CommandFlags flags = CommandFlags.None)
 		{
 			return target.HashLengthAsync(key, flags);
 		}
diff --git a/Engine/Source/Programs/Shared/EpicGames.Redis/RedisList.cs b/Engine/Source/Programs/Shared/EpicGames.Redis/RedisList.cs
--- a/Engine/Source/Programs/Shared/EpicGames.Redis/RedisList.cs
+++ b/Engine/Source/Programs/Shared/EpicGames.Redis/RedisList.cs
@@ -20,32 +20,32 @@ namespace EpicGames.Redis
 		#region Conditions
 
 		/// <inheritdoc cref="Condition.ListIndexEqual(RedisKey, Int64, RedisValue)"/>
-		public static Condition IndexEqual<TElement>(this RedisList<TElement> target, long index, TElement value)
-			=> target.Key.IndexEqual(index, value);
+		public static Condition ListIndexEqual<TElement>(this RedisList<TElement> target, long index, TElement value)
+			=> target.Key.ListIndexEqual(index, value);
 
 		/// <inheritdoc cref="Condition.ListIndexExists(RedisKey, Int64)"/>
-		public static Condition IndexExists<TElement>(this RedisList<TElement> target, long index)
-			=> target.Key.IndexExists(index);
+		public static Condition ListIndexExists<TElement>(this RedisList<TElement> target, long index)
+			=> target.Key.ListIndexExists(index);
 
 		/// <inheritdoc cref="Condition.ListIndexNotEqual(RedisKey, Int64, RedisValue)"/>
-		public static Condition IndexNotEqual<TElement>(this RedisList<TElement> target, long index, TElement value)
-			=> target.Key.IndexNotEqual(index, value);
+		public static Condition ListIndexNotEqual<TElement>(this RedisList<TElement> target, long index, TElement value)
+			=> target.Key.ListIndexNotEqual(index, value);
 
 		/// <inheritdoc cref="Condition.ListIndexNotExists(RedisKey, Int64)"/>
-		public static Condition IndexNotExists<TElement>(this RedisList<TElement> target, long index)
-			=> target.Key.IndexNotExists(index);
+		public static Condition ListIndexNotExists<TElement>(this RedisList<TElement> target, long index)
+			=> target.Key.ListIndexNotExists(index);
 
 		/// <inheritdoc cref="Condition.ListLengthEqual(RedisKey, Int64)"/>
-		public static Condition LengthEqual<TElement>(this RedisList<TElement> target, long length)
-			=> target.LengthEqual(length);
+		public static Condition ListLengthEqual<TElement>(this RedisList<TElement> target, long length)
+			=> target.Key.ListLengthEqual(length);
 
 		/// <inheritdoc cref="Condition.ListLengthGreaterThan(RedisKey, Int64)"/>
-		public static Condition LengthGreaterThan<TElement>(this RedisList<TElement> target, long length)
-			=> target.LengthGreaterThan(length);
+		public static Condition ListLengthGreaterThan<TElement>(this RedisList<TElement> target, long length)
+			=> target.Key.ListLengthGreaterThan(length);
 
 		/// <inheritdoc cref="Condition.ListLengthLessThan(RedisKey, Int64)"/>
-		public static Condition LengthLessThan<TElement>(this RedisList<TElement> target, long length)
-			=> target.LengthLessThan(length);
+		public static Condition ListLengthLessThan<TElement>(this RedisList<TElement> target, long length)
+			=> target.Key.ListLengthLessThan(length);
 
 		#endregion
 
diff --git a/Engine/Source/Programs/Shared/EpicGames.Redis/RedisListKey.cs b/Engine/Source/Programs/Shared/EpicGames.Redis/RedisListKey.cs
--- a/Engine/Source/Programs/Shared/EpicGames.Redis/RedisListKey.cs
+++ b/Engine/Source/Programs/Shared/EpicGames.Redis/RedisListKey.cs
@@ -11,19 +11,13 @@ namespace EpicGames.Redis
 	/// Represents a typed Redis list with a given key
 	/// </summary>
 	/// <typeparam name="TElement">The type of element stored in the set</typeparam>
-	public record struct RedisListKey<TElement>(RedisKey Inner)
+	public record struct RedisListKey<TElement>(RedisKey Inner) : IRedisTypedKey
 	{
 		/// <summary>
 		/// Implicit conversion to typed redis key.
 		/// </summary>
 		/// <param name="key">Key to convert</param>
 		public static implicit operator RedisListKey<TElement>(string key) => new RedisListKey<TElement>(new RedisKey(key));
-
-		/// <summary>
-		/// Implicit conversion to untyped redis keys.
-		/// </summary>
-		/// <param name="key">Key to convert</param>
-		public static implicit operator TypedRedisKey(RedisListKey<TElement> key) => key.Inner;
 	}
 
 	/// <summary>
@@ -34,31 +28,31 @@ namespace EpicGames.Redis
 		#region Conditions
 
 		/// <inheritdoc cref="Condition.ListIndexEqual(RedisKey, Int64, RedisValue)"/>
-		public static Condition IndexEqual<TElement>(this RedisListKey<TElement> key, long index, TElement value)
+		public static Condition ListIndexEqual<TElement>(this RedisListKey<TElement> key, long index, TElement value)
 			=> Condition.ListIndexEqual(key.Inner, index, RedisSerializer.Serialize(value));
 
 		/// <inheritdoc cref="Condition.ListIndexExists(RedisKey, Int64)"/>
-		public static Condition IndexExists<TElement>(this RedisListKey<TElement> key, long index)
+		public static Condition ListIndexExists<TElement>(this RedisListKey<TElement> key, long index)
 			=> Condition.ListIndexExists(key.Inner, index);
 
 		/// <inheritdoc cref="Condition.ListIndexNotEqual(RedisKey, Int64, RedisValue)"/>
-		public static Condition IndexNotEqual<TElement>(this RedisListKey<TElement> key, long index, TElement value)
+		public static Condition ListIndexNotEqual<TElement>(this RedisListKey<TElement> key, long index, TElement value)
 			=> Condition.ListIndexNotEqual(key.Inner, index, RedisSerializer.Serialize(value));
 
 		/// <inheritdoc cref="Condition.ListIndexNotExists(RedisKey, Int64)"/>
-		public static Condition IndexNotExists<TElement>(this RedisListKey<TElement> key, long index)
+		public static Condition ListIndexNotExists<TElement>(this RedisListKey<TElement> key, long index)
 			=> Condition.ListIndexNotExists(key.Inner, index);
 
 		/// <inheritdoc cref="Condition.ListLengthEqual(RedisKey, Int64)"/>
-		public static Condition LengthEqual<TElement>(this RedisListKey<TElement> key, long length)
+		public static Condition ListLengthEqual<TElement>(this RedisListKey<TElement> key, long length)
 			=> Condition.ListLengthEqual(key.Inner, length);
 
 		/// <inheritdoc cref="Condition.ListLengthGreaterThan(RedisKey, Int64)"/>
-		public static Condition LengthGreaterThan<TElement>(this RedisListKey<TElement> key, long length)
+		public static Condition ListLengthGreaterThan<TElement>(this RedisListKey<TElement> key, long length)
 			=> Condition.ListLengthGreaterThan(key.Inner, length);
 
 		/// <inheritdoc cref="Condition.ListLengthLessThan(RedisKey, Int64)"/>
-		public static Condition LengthLessThan<TElement>(this RedisListKey<TElement> key, long length)
+		public static Condition ListLengthLessThan<TElement>(this RedisListKey<TElement> key, long length)
 			=> Condition.ListLengthLessThan(key.Inner, length);
 
 		#endregion
diff --git a/Engine/Source/Programs/Shared/EpicGames.Redis/RedisSet.cs b/Engine/Source/Programs/Shared/EpicGames.Redis/RedisSet.cs
--- a/Engine/Source/Programs/Shared/EpicGames.Redis/RedisSet.cs
+++ b/Engine/Source/Programs/Shared/EpicGames.Redis/RedisSet.cs
@@ -21,24 +21,24 @@ namespace EpicGames.Redis
 		#region Conditions
 
 		/// <inheritdoc cref="Condition.SetContains(RedisKey, RedisValue)"/>
-		public static Condition Contains<TElement>(this RedisSet<TElement> target, TElement value)
-			=> target.Key.Contains(value);
+		public static Condition SetContains<TElement>(this RedisSet<TElement> target, TElement value)
+			=> target.Key.SetContains(value);
 
 		/// <inheritdoc cref="Condition.SetLengthEqual(RedisKey, Int64)"/>
-		public static Condition LengthEqual<TElement>(this RedisSet<TElement> target, long length)
-			=> target.Key.LengthEqual(length);
+		public static Condition SetLengthEqual<TElement>(this RedisSet<TElement> target, long length)
+			=> target.Key.SetLengthEqual(length);
 
 		/// <inheritdoc cref="Condition.SetLengthGreaterThan(RedisKey, Int64)"/>
-		public static Condition LengthGreaterThan<TElement>(this RedisSet<TElement> target, long length)
-			=> target.Key.LengthGreaterThan(length);
+		public static Condition SetLengthGreaterThan<TElement>(this RedisSet<TElement> target, long length)
+			=> target.Key.SetLengthGreaterThan(length);
 
 		/// <inheritdoc cref="Condition.SetLengthLessThan(RedisKey, Int64)"/>
-		public static Condition LengthLessThan<TElement>(this RedisSet<TElement> target, long length)
-			=> target.LengthLessThan(length);
+		public static Condition SetLengthLessThan<TElement>(this RedisSet<TElement> target, long length)
+			=> target.Key.SetLengthLessThan(length);
 
-		/// <inheritdoc cref="Condition.SetContains(RedisKey, RedisValue)"/>
-		public static Condition NotContains<TElement>(this RedisSet<TElement> target, TElement value)
-			=> target.NotContains(value);
+		/// <inheritdoc cref="Condition.SetNotContains(RedisKey, RedisValue)"/>
+		public static Condition SetNotContains<TElement>(this RedisSet<TElement> target, TElement value)
+			=> target.Key.SetNotContains(value);
 
 		#endregion
 
diff --git a/Engine/Source/Programs/Shared/EpicGames.Redis/RedisSetKey.cs b/Engine/Source/Programs/Shared/EpicGames.Redis/RedisSetKey.cs
--- a/Engine/Source/Programs/Shared/EpicGames.Redis/RedisSetKey.cs
+++ b/Engine/Source/Programs/Shared/EpicGames.Redis/RedisSetKey.cs
@@ -11,19 +11,13 @@ namespace EpicGames.Redis
 	/// Represents a typed Redis set with a given key
 	/// </summary>
 	/// <typeparam name="TElement">The type of element stored in the set</typeparam>
-	public record struct RedisSetKey<TElement>(RedisKey Inner)
+	public record struct RedisSetKey<TElement>(RedisKey Inner) : IRedisTypedKey
 	{
 		/// <summary>
 		/// Implicit conversion to typed redis key.
 		/// </summary>
 		/// <param name="key">Key to convert</param>
 		public static implicit operator RedisSetKey<TElement>(string key) => new RedisSetKey<TElement>(new RedisKey(key));
-
-		/// <summary>
-		/// Implicit conversion to untyped redis keys.
-		/// </summary>
-		/// <param name="key">Key to convert</param>
-		public static implicit operator TypedRedisKey(RedisSetKey<TElement> key) => key.Inner;
 	}
 
 	/// <summary>
@@ -34,23 +28,23 @@ namespace EpicGames.Redis
 		#region Conditions
 
 		/// <inheritdoc cref="Condition.SetContains(RedisKey, RedisValue)"/>
-		public static Condition Contains<TElement>(this RedisSetKey<TElement> key, TElement value)
+		public static Condition SetContains<TElement>(this RedisSetKey<TElement> key, TElement value)
 			=> Condition.SetContains(key.Inner, RedisSerializer.Serialize(value));
 
 		/// <inheritdoc cref="Condition.SetLengthEqual(RedisKey, Int64)"/>
-		public static Condition LengthEqual<TElement>(this RedisSetKey<TElement> key, long length)
+		public static Condition SetLengthEqual<TElement>(this RedisSetKey<TElement> key, long length)
 			=> Condition.SetLengthEqual(key.Inner, length);
 
 		/// <inheritdoc cref="Condition.SetLengthGreaterThan(RedisKey, Int64)"/>
-		public static Condition LengthGreaterThan<TElement>(this RedisSetKey<TElement> key, long length)
+		public static Condition SetLengthGreaterThan<TElement>(this RedisSetKey<TElement> key, long length)
 			=> Condition.SetLengthGreaterThan(key.Inner, length);
 
 		/// <inheritdoc cref="Condition.SetLengthLessThan(RedisKey, Int64)"/>
-		public static Condition LengthLessThan<TElement>(this RedisSetKey<TElement> key, long length)
+		public static Condition SetLengthLessThan<TElement>(this RedisSetKey<TElement> key, long length)
 			=> Condition.SetLengthLessThan(key.Inner, length);
 
 		/// <inheritdoc cref="Condition.SetContains(RedisKey, RedisValue)"/>
-		public static Condition NotContains<TElement>(this RedisSetKey<TElement> key, TElement value)
+		public static Condition SetNotContains<TElement>(this RedisSetKey<TElement> key, TElement value)
 			=> Condition.SetNotContains(key.Inner, RedisSerializer.Serialize(value));
 
 		#endregion
diff --git a/Engine/Source/Programs/Shared/EpicGames.Redis/RedisSortedSet.cs b/Engine/Source/Programs/Shared/EpicGames.Redis/RedisSortedSet.cs
--- a/Engine/Source/Programs/Shared/EpicGames.Redis/RedisSortedSet.cs
+++ b/Engine/Source/Programs/Shared/EpicGames.Redis/RedisSortedSet.cs
@@ -20,40 +20,40 @@ namespace EpicGames.Redis
 		#region Conditions
 
 		/// <inheritdoc cref="Condition.SortedSetContains(RedisKey, RedisValue)"/>
-		public static Condition Contains<TElement>(this RedisSortedSet<TElement> target, TElement value)
-			=> target.Key.Contains(value);
+		public static Condition SortedSetContains<TElement>(this RedisSortedSet<TElement> target, TElement value)
+			=> target.Key.SortedSetContains(value);
 
 		/// <inheritdoc cref="Condition.SortedSetEqual(RedisKey, RedisValue, RedisValue)"/>
-		public static Condition Equal<TElement>(this RedisSortedSet<TElement> target, TElement value, RedisValue score)
-			=> target.Key.Equal(value, score);
+		public static Condition SortedSetEqual<TElement>(this RedisSortedSet<TElement> target, TElement value, RedisValue score)
+			=> target.Key.SortedSetEqual(value, score);
 
 		/// <inheritdoc cref="Condition.SortedSetLengthEqual(RedisKey, Int64)"/>
-		public static Condition LengthEqual<TElement>(this RedisSortedSet<TElement> target, long length)
-			=> target.Key.LengthLessThan(length);
+		public static Condition SortedSetLengthEqual<TElement>(this RedisSortedSet<TElement> target, long length)
+			=> target.Key.SortedSetLengthLessThan(length);
 
 		/// <inheritdoc cref="Condition.SortedSetLengthGreaterThan(RedisKey, Int64)"/>
-		public static Condition LengthGreaterThan<TElement>(this RedisSortedSet<TElement> target, long length)
-			=> target.Key.LengthGreaterThan(length);
+		public static Condition SortedSetLengthGreaterThan<TElement>(this RedisSortedSet<TElement> target, long length)
+			=> target.Key.SortedSetLengthGreaterThan(length);
 
 		/// <inheritdoc cref="Condition.SortedSetLengthLessThan(RedisKey, Int64)"/>
-		public static Condition LengthLessThan<TElement>(this RedisSortedSet<TElement> target, long length)
-			=> target.Key.LengthLessThan(length);
+		public static Condition SortedSetLengthLessThan<TElement>(this RedisSortedSet<TElement> target, long length)
+			=> target.Key.SortedSetLengthLessThan(length);
 
 		/// <inheritdoc cref="Condition.SortedSetNotContains(RedisKey, RedisValue)"/>
-		public static Condition NotContains<TElement>(this RedisSortedSet<TElement> target, TElement value)
-			=> target.Key.NotContains(value);
+		public static Condition SortedSetNotContains<TElement>(this RedisSortedSet<TElement> target, TElement value)
+			=> target.Key.SortedSetNotContains(value);
 
 		/// <inheritdoc cref="Condition.SortedSetNotEqual(RedisKey, RedisValue, RedisValue)"/>
-		public static Condition NotEqual<TElement>(this RedisSortedSet<TElement> target, TElement value, RedisValue score)
-			=> target.Key.NotEqual(value, score);
+		public static Condition SortedSetNotEqual<TElement>(this RedisSortedSet<TElement> target, TElement value, RedisValue score)
+			=> target.Key.SortedSetNotEqual(value, score);
 
 		/// <inheritdoc cref="Condition.SortedSetScoreExists(RedisKey, RedisValue)"/>
-		public static Condition ScoreExists<TElement>(this RedisSortedSet<TElement> target, RedisValue score)
-			=> target.Key.ScoreExists(score);
+		public static Condition SortedSetScoreExists<TElement>(this RedisSortedSet<TElement> target, RedisValue score)
+			=> target.Key.SortedSetScoreExists(score);
 
 		/// <inheritdoc cref="Condition.SortedSetScoreExists(RedisKey, RedisValue, RedisValue)"/>
-		public static Condition ScoreExists<TElement>(this RedisSortedSet<TElement> target, RedisValue score, RedisValue count)
-			=> target.Key.ScoreExists(score, count);
+		public static Condition SortedSetScoreExists<TElement>(this RedisSortedSet<TElement> target, RedisValue score, RedisValue count)
+			=> target.Key.SortedSetScoreExists(score, count);
 
 		#endregion
 
diff --git a/Engine/Source/Programs/Shared/EpicGames.Redis/RedisSortedSetKey.cs b/Engine/Source/Programs/Shared/EpicGames.Redis/RedisSortedSetKey.cs
--- a/Engine/Source/Programs/Shared/EpicGames.Redis/RedisSortedSetKey.cs
+++ b/Engine/Source/Programs/Shared/EpicGames.Redis/RedisSortedSetKey.cs
@@ -12,19 +12,13 @@ namespace EpicGames.Redis
 	/// Represents a typed Redis sorted set with a given key
 	/// </summary>
 	/// <typeparam name="TElement">The type of element stored in the set</typeparam>
-	public record struct RedisSortedSetKey<TElement>(RedisKey Inner)
+	public record struct RedisSortedSetKey<TElement>(RedisKey Inner) : IRedisTypedKey
 	{
 		/// <summary>
 		/// Implicit conversion to typed redis key.
 		/// </summary>
 		/// <param name="key">Key to convert</param>
 		public static implicit operator RedisSortedSetKey<TElement>(string key) => new RedisSortedSetKey<TElement>(new RedisKey(key));
-
-		/// <summary>
-		/// Implicit conversion to untyped redis keys.
-		/// </summary>
-		/// <param name="key">Key to convert</param>
-		public static implicit operator TypedRedisKey(RedisSortedSetKey<TElement> key) => key.Inner;
 	}
 
 	/// <summary>
@@ -91,39 +85,39 @@ namespace EpicGames.Redis
 		#region Conditions
 
 		/// <inheritdoc cref="Condition.SortedSetContains(RedisKey, RedisValue)"/>
-		public static Condition Contains<TElement>(this RedisSortedSetKey<TElement> key, TElement value)
+		public static Condition SortedSetContains<TElement>(this RedisSortedSetKey<TElement> key, TElement value)
 			=> Condition.SortedSetContains(key.Inner, RedisSerializer.Serialize(value));
 
 		/// <inheritdoc cref="Condition.SortedSetEqual(RedisKey, RedisValue, RedisValue)"/>
-		public static Condition Equal<TElement>(this RedisSortedSetKey<TElement> key, TElement value, RedisValue score)
+		public static Condition SortedSetEqual<TElement>(this RedisSortedSetKey<TElement> key, TElement value, RedisValue score)
 			=> Condition.SortedSetEqual(key.Inner, RedisSerializer.Serialize(value), score);
 
 		/// <inheritdoc cref="Condition.SortedSetLengthEqual(RedisKey, Int64)"/>
-		public static Condition LengthEqual<TElement>(this RedisSortedSetKey<TElement> key, long length)
+		public static Condition SortedSetLengthEqual<TElement>(this RedisSortedSetKey<TElement> key, long length)
 			=> Condition.SortedSetLengthEqual(key.Inner, length);
 
 		/// <inheritdoc cref="Condition.SortedSetLengthGreaterThan(RedisKey, Int64)"/>
-		public static Condition LengthGreaterThan<TElement>(this RedisSortedSetKey<TElement> key, long length)
+		public static Condition SortedSetLengthGreaterThan<TElement>(this RedisSortedSetKey<TElement> key, long length)
 			=> Condition.SortedSetLengthGreaterThan(key.Inner, length);
 
 		/// <inheritdoc cref="Condition.SortedSetLengthLessThan(RedisKey, Int64)"/>
-		public static Condition LengthLessThan<TElement>(this RedisSortedSetKey<TElement> key, long length)
+		public static Condition SortedSetLengthLessThan<TElement>(this RedisSortedSetKey<TElement> key, long length)
 			=> Condition.SortedSetLengthLessThan(key.Inner, length);
 
 		/// <inheritdoc cref="Condition.SortedSetNotContains(RedisKey, RedisValue)"/>
-		public static Condition NotContains<TElement>(this RedisSortedSetKey<TElement> key, TElement value)
+		public static Condition SortedSetNotContains<TElement>(this RedisSortedSetKey<TElement> key, TElement value)
 			=> Condition.SortedSetNotContains(key.Inner, RedisSerializer.Serialize(value));
 
 		/// <inheritdoc cref="Condition.SortedSetNotEqual(RedisKey, RedisValue, RedisValue)"/>
-		public static Condition NotEqual<TElement>(this RedisSortedSetKey<TElement> key, TElement value, RedisValue score)
+		public static Condition SortedSetNotEqual<TElement>(this RedisSortedSetKey<TElement> key, TElement value, RedisValue score)
 			=> Condition.SortedSetNotEqual(key.Inner, RedisSerializer.Serialize(value), score);
 
 		/// <inheritdoc cref="Condition.SortedSetScoreExists(RedisKey, RedisValue)"/>
-		public static Condition ScoreExists<TElement>(this RedisSortedSetKey<TElement> key, RedisValue score)
+		public static Condition SortedSetScoreExists<TElement>(this RedisSortedSetKey<TElement> key, RedisValue score)
 			=> Condition.SortedSetScoreExists(key.Inner, score);
 
 		/// <inheritdoc cref="Condition.SortedSetScoreExists(RedisKey, RedisValue, RedisValue)"/>
-		public static Condition ScoreExists<TElement>(this RedisSortedSetKey<TElement> key, RedisValue score, RedisValue count)
+		public static Condition SortedSetScoreExists<TElement>(this RedisSortedSetKey<TElement> key, RedisValue score, RedisValue count)
 			=> Condition.SortedSetScoreExists(key.Inner, score, count);
 
 		#endregion
diff --git a/Engine/Source/Programs/Shared/EpicGames.Redis/RedisString.cs b/Engine/Source/Programs/Shared/EpicGames.Redis/RedisString.cs
--- a/Engine/Source/Programs/Shared/EpicGames.Redis/RedisString.cs
+++ b/Engine/Source/Programs/Shared/EpicGames.Redis/RedisString.cs
@@ -22,24 +22,24 @@ namespace EpicGames.Redis
 		#region Conditions
 
 		/// <inheritdoc cref="Condition.StringEqual(RedisKey, RedisValue)"/>
-		public static Condition Equal<TElement>(this RedisString<TElement> target, TElement value)
-			=> target.Key.Equal(value);
+		public static Condition StringEqual<TElement>(this RedisString<TElement> target, TElement value)
+			=> target.Key.StringEqual(value);
 
 		/// <inheritdoc cref="Condition.StringLengthEqual(RedisKey, Int64)"/>
-		public static Condition LengthEqual<TElement>(this RedisString<TElement> target, long length)
-			=> target.Key.LengthEqual(length);
+		public static Condition StringLengthEqual<TElement>(this RedisString<TElement> target, long length)
+			=> target.Key.StringLengthEqual(length);
 
 		/// <inheritdoc cref="Condition.StringLengthGreaterThan(RedisKey, Int64)"/>
-		public static Condition LengthGreaterThan<TElement>(this RedisString<TElement> target, long length)
-			=> target.Key.LengthGreaterThan(length);
+		public static Condition StringLengthGreaterThan<TElement>(this RedisString<TElement> target, long length)
+			=> target.Key.StringLengthGreaterThan(length);
 
 		/// <inheritdoc cref="Condition.StringLengthLessThan(RedisKey, Int64)"/>
-		public static Condition LengthLessThan<TElement>(this RedisString<TElement> target, long length)
-			=> target.Key.LengthLessThan(length);
+		public static Condition StringLengthLessThan<TElement>(this RedisString<TElement> target, long length)
+			=> target.Key.StringLengthLessThan(length);
 
 		/// <inheritdoc cref="Condition.StringNotEqual(RedisKey, RedisValue)"/>
-		public static Condition NotEqual<TElement>(this RedisString<TElement> target, TElement value)
-			=> target.Key.NotEqual(value);
+		public static Condition StringNotEqual<TElement>(this RedisString<TElement> target, TElement value)
+			=> target.Key.StringNotEqual(value);
 
 		#endregion
 
diff --git a/Engine/Source/Programs/Shared/EpicGames.Redis/RedisStringKey.cs b/Engine/Source/Programs/Shared/EpicGames.Redis/RedisStringKey.cs
--- a/Engine/Source/Programs/Shared/EpicGames.Redis/RedisStringKey.cs
+++ b/Engine/Source/Programs/Shared/EpicGames.Redis/RedisStringKey.cs
@@ -12,19 +12,13 @@ namespace EpicGames.Redis
 	/// Represents a typed Redis string with a given value type
 	/// </summary>
 	/// <typeparam name="TValue">The type of element stored in the list</typeparam>
-	public record struct RedisStringKey<TValue>(RedisKey Inner)
+	public record struct RedisStringKey<TValue>(RedisKey Inner) : IRedisTypedKey
 	{
 		/// <summary>
 		/// Implicit conversion to typed redis key.
 		/// </summary>
 		/// <param name="key">Key to convert</param>
 		public static implicit operator RedisStringKey<TValue>(string key) => new RedisStringKey<TValue>(new RedisKey(key));
-
-		/// <summary>
-		/// Implicit conversion to a regular RedisKey
-		/// </summary>
-		/// <param name="key">The key to convert</param>
-		public static implicit operator TypedRedisKey(RedisStringKey<TValue> key) => key.Inner;
 	}
 
 	/// <summary>
@@ -35,23 +29,23 @@ namespace EpicGames.Redis
 		#region Conditions
 
 		/// <inheritdoc cref="Condition.StringEqual(RedisKey, RedisValue)"/>
-		public static Condition Equal<TElement>(this RedisStringKey<TElement> key, TElement value)
+		public static Condition StringEqual<TElement>(this RedisStringKey<TElement> key, TElement value)
 			=> Condition.StringEqual(key.Inner, RedisSerializer.Serialize(value));
 
 		/// <inheritdoc cref="Condition.StringLengthEqual(RedisKey, Int64)"/>
-		public static Condition LengthEqual<TElement>(this RedisStringKey<TElement> key, long length)
+		public static Condition StringLengthEqual<TElement>(this RedisStringKey<TElement> key, long length)
 			=> Condition.StringLengthEqual(key.Inner, length);
 
 		/// <inheritdoc cref="Condition.StringLengthGreaterThan(RedisKey, Int64)"/>
-		public static Condition LengthGreaterThan<TElement>(this RedisStringKey<TElement> key, long length)
+		public static Condition StringLengthGreaterThan<TElement>(this RedisStringKey<TElement> key, long length)
 			=> Condition.StringLengthGreaterThan(key.Inner, length);
 
 		/// <inheritdoc cref="Condition.StringLengthLessThan(RedisKey, Int64)"/>
-		public static Condition LengthLessThan<TElement>(this RedisStringKey<TElement> key, long length)
+		public static Condition StringLengthLessThan<TElement>(this RedisStringKey<TElement> key, long length)
 			=> Condition.StringLengthLessThan(key.Inner, length);
 
 		/// <inheritdoc cref="Condition.StringNotEqual(RedisKey, RedisValue)"/>
-		public static Condition NotEqual<TElement>(this RedisStringKey<TElement> key, TElement value)
+		public static Condition StringNotEqual<TElement>(this RedisStringKey<TElement> key, TElement value)
 			=> Condition.StringNotEqual(key.Inner, RedisSerializer.Serialize(value));
 
 		#endregion
diff --git a/Engine/Source/Programs/Shared/EpicGames.Redis/TypedRedisKey.cs b/Engine/Source/Programs/Shared/EpicGames.Redis/TypedRedisKey.cs
deleted file mode 100644
