Horde: Refactor server Horde Analytics telemetry plugin dashboard stuff

diff --git a/Engine/Source/Programs/Horde/HordeServer/Dashboard/DashboardConfig.cs b/Engine/Source/Programs/Horde/HordeServer/Dashboard/DashboardConfig.cs
--- a/Engine/Source/Programs/Horde/HordeServer/Dashboard/DashboardConfig.cs
+++ b/Engine/Source/Programs/Horde/HordeServer/Dashboard/DashboardConfig.cs
@@ -63,11 +63,6 @@ namespace HordeServer.Dashboard
 		/// </summary>
 		public List<DashboardPoolCategoryConfig> PoolCategories { get; set; } = new List<DashboardPoolCategoryConfig>();
 
-		/// <summary>
-		/// Configuration for telemetry views
-		/// </summary>
-		public List<TelemetryViewConfig> Analytics { get; set; } = new List<TelemetryViewConfig>();
-
 		/// <summary>
 		/// Includes for other configuration files
 		/// </summary>
diff --git a/Engine/Source/Programs/Horde/HordeServer/Dashboard/DashboardController.cs b/Engine/Source/Programs/Horde/HordeServer/Dashboard/DashboardController.cs
--- a/Engine/Source/Programs/Horde/HordeServer/Dashboard/DashboardController.cs
+++ b/Engine/Source/Programs/Horde/HordeServer/Dashboard/DashboardController.cs
@@ -12,7 +12,6 @@ using EpicGames.Horde.Server;
 using HordeServer.Accounts;
 using HordeServer.Plugins;
 using HordeServer.Server;
-using HordeServer.Telemetry;
 using HordeServer.Utilities;
 using Microsoft.AspNetCore.Authentication;
 using Microsoft.AspNetCore.Authentication.Cookies;
@@ -159,43 +158,6 @@ namespace HordeServer.Dashboard
 				dashboardConfigResponse.PoolCategories.Add(new GetDashboardPoolCategoryResponse { Name = category.Name, Condition = category.Condition });
 			}
 
-			if (_globalConfig.Value.Authorize(TelemetryAclAction.QueryMetrics, User))
-			{
-				foreach (TelemetryViewConfig telemetry in _globalConfig.Value.Dashboard.Analytics)
-				{
-					GetTelemetryViewResponse rview = new GetTelemetryViewResponse();
-					rview.Id = telemetry.Id.ToString();
-					rview.Name = telemetry.Name;
-					rview.TelemetryStoreId = telemetry.TelemetryStoreId.ToString();
-
-					foreach (TelemetryVariableConfig variable in telemetry.Variables)
-					{
-						rview.Variables.Add(new GetTelemetryVariableResponse { Name = variable.Name, Group = variable.Group, Defaults = variable.Defaults });
-					}
-
-					foreach (TelemetryCategoryConfig category in telemetry.Categories)
-					{
-						GetTelemetryCategoryResponse rcategory = new GetTelemetryCategoryResponse { Name = category.Name };
-
-						foreach (TelemetryChartConfig chart in category.Charts)
-						{
-							GetTelemetryChartResponse rchart = new GetTelemetryChartResponse { Name = chart.Name, Display = chart.Display.ToString(), Graph = chart.Graph.ToString(), Min = chart.Min, Max = chart.Max, Metrics = new List<GetTelemetryChartMetricResponse>() };
-
-							foreach (TelemetryChartMetricConfig metric in chart.Metrics)
-							{
-								rchart.Metrics.Add(new GetTelemetryChartMetricResponse { MetricId = metric.Id.ToString(), Threshold = metric.Threshold, Alias = metric.Alias });
-							}
-
-							rcategory.Charts.Add(rchart);
-						}
-
-						rview.Categories.Add(rcategory);
-					}
-
-					dashboardConfigResponse.TelemetryViews.Add(rview);
-				}
-			}
-
 			foreach (IPluginResponseFilter responseFilter in _responseFilters)
 			{
 				responseFilter.Apply(HttpContext, dashboardConfigResponse);
diff --git a/Engine/Source/Programs/Horde/HordeServer/Dashboard/DashboardTelemetryConfig.cs b/Engine/Source/Programs/Horde/HordeServer/Dashboard/DashboardTelemetryConfig.cs
deleted file mode 100644
diff --git a/Engine/Source/Programs/Horde/Plugins/Analytics/HordeServer.Analytics/Telemetry/Metrics/MetricConfig.cs b/Engine/Source/Programs/Horde/Plugins/Analytics/HordeServer.Analytics/Telemetry/Metrics/MetricConfig.cs
--- a/Engine/Source/Programs/Horde/Plugins/Analytics/HordeServer.Analytics/Telemetry/Metrics/MetricConfig.cs
+++ b/Engine/Source/Programs/Horde/Plugins/Analytics/HordeServer.Analytics/Telemetry/Metrics/MetricConfig.cs
@@ -145,6 +145,11 @@ namespace HordeServer.Telemetry.Metrics
 		/// </summary>
 		public List<MetricConfig> Metrics { get; set; } = new List<MetricConfig>();
 
+		/// <summary>
+		/// Configuration for telemetry views
+		/// </summary>
+		public List<TelemetryViewConfig> Views { get; set; } = new List<TelemetryViewConfig>();
+
 		/// <summary>
 		/// Includes for other configuration files
 		/// </summary>
diff --git a/Engine/Source/Programs/Horde/Plugins/Analytics/HordeServer.Analytics/Telemetry/TelemetryController.cs b/Engine/Source/Programs/Horde/Plugins/Analytics/HordeServer.Analytics/Telemetry/TelemetryController.cs
--- a/Engine/Source/Programs/Horde/Plugins/Analytics/HordeServer.Analytics/Telemetry/TelemetryController.cs
+++ b/Engine/Source/Programs/Horde/Plugins/Analytics/HordeServer.Analytics/Telemetry/TelemetryController.cs
@@ -156,5 +156,58 @@ namespace HordeServer.Telemetry
 
 			return NoContent();
 		}
+
+		/// <summary>
+		/// Get telemetry views
+		/// </summary>
+		/// <returns>List of telemetry views</returns>
+		[HttpGet]
+		[Authorize]
+		[Route("/api/v1/telemetry/views")]
+		public ActionResult<List<GetTelemetryViewResponse>> GetDashboardViews()
+		{
+			if (!_analyticsConfig.Value.Authorize(TelemetryAclAction.QueryMetrics, User))
+			{
+				return Forbid(TelemetryAclAction.QueryMetrics);
+			}
+
+			List<GetTelemetryViewResponse> responses = new List<GetTelemetryViewResponse>();
+
+			foreach (TelemetryViewConfig telemetry in _analyticsConfig.Value.Stores.SelectMany(store => store.Views))
+			{
+				GetTelemetryViewResponse rview = new GetTelemetryViewResponse();
+				rview.Id = telemetry.Id.ToString();
+				rview.Name = telemetry.Name;
+				rview.TelemetryStoreId = telemetry.TelemetryStoreId.ToString();
+
+				foreach (TelemetryVariableConfig variable in telemetry.Variables)
+				{
+					rview.Variables.Add(new GetTelemetryVariableResponse { Name = variable.Name, Group = variable.Group, Defaults = variable.Defaults });
+				}
+
+				foreach (TelemetryCategoryConfig category in telemetry.Categories)
+				{
+					GetTelemetryCategoryResponse rcategory = new GetTelemetryCategoryResponse { Name = category.Name };
+
+					foreach (TelemetryChartConfig chart in category.Charts)
+					{
+						GetTelemetryChartResponse rchart = new GetTelemetryChartResponse { Name = chart.Name, Display = chart.Display.ToString(), Graph = chart.Graph.ToString(), Min = chart.Min, Max = chart.Max, Metrics = new List<GetTelemetryChartMetricResponse>() };
+
+						foreach (TelemetryChartMetricConfig metric in chart.Metrics)
+						{
+							rchart.Metrics.Add(new GetTelemetryChartMetricResponse { MetricId = metric.Id.ToString(), Threshold = metric.Threshold, Alias = metric.Alias });
+						}
+
+						rcategory.Charts.Add(rchart);
+					}
+
+					rview.Categories.Add(rcategory);
+				}
+
+				responses.Add(rview);
+			}
+
+			return responses;
+		}
 	}
 }
diff --git a/Engine/Source/Programs/Shared/EpicGames.Horde/Dashboard/DashboardMessages.cs b/Engine/Source/Programs/Shared/EpicGames.Horde/Dashboard/DashboardMessages.cs
--- a/Engine/Source/Programs/Shared/EpicGames.Horde/Dashboard/DashboardMessages.cs
+++ b/Engine/Source/Programs/Shared/EpicGames.Horde/Dashboard/DashboardMessages.cs
@@ -60,11 +60,6 @@ namespace EpicGames.Horde.Dashboard
 		/// </summary>
 		public List<GetDashboardPoolCategoryResponse> PoolCategories { get; set; } = new List<GetDashboardPoolCategoryResponse>();
 
-		/// <summary>
-		/// Telemetry to display on the telemetry page
-		/// </summary>
-		public List<GetTelemetryViewResponse> TelemetryViews { get; set; } = new List<GetTelemetryViewResponse>();
-
 	}
 
 	/// <summary>
@@ -233,132 +228,4 @@ namespace EpicGames.Horde.Dashboard
 		public bool NeedsAuthorization { get; set; } = true;
 
 	}
-
-	#region Telemetry
-
-	/// <summary>
-	/// Metric attached to a telemetry chart
-	/// </summary>
-	public class GetTelemetryChartMetricResponse
-	{
-		/// <summary>
-		/// Associated metric id
-		/// </summary>		
-		public string MetricId { get; set; } = null!;
-
-		/// <summary>
-		/// The threshold for KPI values
-		/// </summary>
-		public int? Threshold { get; set; }
-
-		/// <summary>
-		/// The metric alias for display purposes
-		/// </summary>
-		public string? Alias { get; set; }
-	}
-
-	/// <summary>
-	/// Telemetry chart configuraton
-	/// </summary>
-	public class GetTelemetryChartResponse
-	{
-		/// <summary>
-		/// The name of the chart, will be displayed on the dashboard
-		/// </summary>		
-		public string Name { get; set; } = null!;
-
-		/// <summary>
-		/// The unit to display
-		/// </summary>
-		public string Display { get; set; } = null!;
-
-		/// <summary>
-		/// The graph type 
-		/// </summary>
-		public string Graph { get; set; } = null!;
-
-		/// <summary>
-		/// List of configured metrics
-		/// </summary>
-		public List<GetTelemetryChartMetricResponse> Metrics { get; set; } = new List<GetTelemetryChartMetricResponse>();
-
-		/// <summary>
-		/// The min unit value for clamping chart
-		/// </summary>
-		public int? Min { get; set; }
-
-		/// <summary>
-		/// The max unit value for clamping chart
-		/// </summary>
-		public int? Max { get; set; }
-	}
-
-	/// <summary>
-	/// A chart categody, will be displayed on the dashbord under an associated pivot
-	/// </summary>
-	public class GetTelemetryCategoryResponse
-	{
-		/// <summary>
-		/// The name of the category
-		/// </summary>		
-		public string Name { get; set; } = null!;
-
-		/// <summary>
-		/// The charts contained within the category
-		/// </summary>
-		public List<GetTelemetryChartResponse> Charts { get; set; } = new List<GetTelemetryChartResponse> { };
-	}
-
-	/// <summary>
-	/// A telemetry view variable used for filtering the charting data
-	/// </summary>
-	public class GetTelemetryVariableResponse
-	{
-		/// <summary>
-		/// The name of the variable for display purposes
-		/// </summary>
-		public string Name { get; set; } = null!;
-
-		/// <summary>
-		/// The associated data group attached to the variable 
-		/// </summary>
-		public string Group { get; set; } = null!;
-
-		/// <summary>
-		/// The default values to select
-		/// </summary>
-		public List<string> Defaults { get; set; } = new List<string> { };
-	}
-
-	/// <summary>
-	/// A telemetry view of related metrics, divided into categofies
-	/// </summary>
-	public class GetTelemetryViewResponse
-	{
-		/// <summary>
-		/// Identifier for the view
-		/// </summary>
-		public string Id { get; set; } = null!;
-
-		/// <summary>
-		/// The name of the view
-		/// </summary>
-		public string Name { get; set; } = null!;
-
-		/// <summary>
-		/// The telemetry store the view uses
-		/// </summary>
-		public string TelemetryStoreId { get; set; } = null!;
-
-		/// <summary>
-		///  The variables used to filter the view data
-		/// </summary>
-		public List<GetTelemetryVariableResponse> Variables { get; set; } = new List<GetTelemetryVariableResponse> { };
-
-		/// <summary>
-		/// The categories contained within the view
-		/// </summary>
-		public List<GetTelemetryCategoryResponse> Categories { get; set; } = new List<GetTelemetryCategoryResponse> { };
-	}
-	#endregion
 }
diff --git a/Engine/Source/Programs/Shared/EpicGames.Horde/Telemetry/TelemetryMessages.cs b/Engine/Source/Programs/Shared/EpicGames.Horde/Telemetry/TelemetryMessages.cs
--- a/Engine/Source/Programs/Shared/EpicGames.Horde/Telemetry/TelemetryMessages.cs
+++ b/Engine/Source/Programs/Shared/EpicGames.Horde/Telemetry/TelemetryMessages.cs
@@ -2,7 +2,11 @@
 
 using System;
 using System.Collections.Generic;
-using System.Text.Json.Nodes;
+using System.ComponentModel.DataAnnotations;
+using System.ComponentModel;
+using EpicGames.Core;
+using EpicGames.Horde.Telemetry.Metrics;
+using EpicGames.Serialization;
 
 #pragma warning disable CA2227
 
@@ -16,7 +20,7 @@ namespace EpicGames.Horde.Telemetry
 		/// <summary>
 		/// List of telemetry events
 		/// </summary>
-		public List<JsonObject> Events { get; set; } = new List<JsonObject>();
+		public List<System.Text.Json.Nodes.JsonObject> Events { get; set; } = new List<System.Text.Json.Nodes.JsonObject>();
 	}
 
 	/// <summary>
@@ -71,4 +75,333 @@ namespace EpicGames.Horde.Telemetry
 		/// </summary>
 		public double Value { get; set; }
 	}
+
+	/// <summary>
+	/// The units used to present the telemetry
+	/// </summary>
+	public enum TelemetryMetricUnitType
+	{
+		/// <summary>
+		/// Time duration 
+		/// </summary>
+		Time,
+		/// <summary>
+		/// Ratio 0-100%
+		/// </summary>
+		Ratio,
+		/// <summary>
+		/// Artbitrary numeric value
+		/// </summary>
+		Value
+	}
+
+	/// <summary>
+	/// The type of 
+	/// </summary>
+	public enum TelemetryMetricGraphType
+	{
+		/// <summary>
+		/// A line graph
+		/// </summary>
+		Line,
+		/// <summary>
+		/// Key performance indicator (KPI) chart with thrasholds
+		/// </summary>
+		Indicator
+	}
+
+	/// <summary>
+	/// Metric attached to a telemetry chart
+	/// </summary>
+	public class TelemetryChartMetricConfig
+	{
+		/// <summary>
+		/// Associated metric id
+		/// </summary>
+		[Required]
+		public MetricId Id { get; set; }
+
+		/// <summary>
+		/// The threshold for KPI values
+		/// </summary>
+		public int? Threshold { get; set; }
+
+		/// <summary>
+		/// The metric alias for display purposes
+		/// </summary>
+		public string? Alias { get; set; }
+	}
+
+	/// <summary>
+	/// Telemetry chart configuraton
+	/// </summary>
+	public class TelemetryChartConfig
+	{
+		/// <summary>
+		/// The name of the chart, will be displayed on the dashboard
+		/// </summary>
+		[Required]
+		public string Name { get; set; } = null!;
+
+		/// <summary>
+		/// The unit to display
+		/// </summary>
+		public TelemetryMetricUnitType Display { get; set; } = TelemetryMetricUnitType.Time;
+
+		/// <summary>
+		/// The graph type 
+		/// </summary>
+		public TelemetryMetricGraphType Graph { get; set; } = TelemetryMetricGraphType.Line;
+
+		/// <summary>
+		/// List of configured metrics
+		/// </summary>
+		public List<TelemetryChartMetricConfig> Metrics { get; set; } = new List<TelemetryChartMetricConfig>();
+
+		/// <summary>
+		/// The min unit value for clamping chart
+		/// </summary>
+		public int? Min { get; set; }
+
+		/// <summary>
+		/// The max unit value for clamping chart
+		/// </summary>
+		public int? Max { get; set; }
+	}
+
+	/// <summary>
+	/// A chart categody, will be displayed on the dashbord under an associated pivot
+	/// </summary>
+	public class TelemetryCategoryConfig
+	{
+		/// <summary>
+		/// The name of the category
+		/// </summary>
+		[Required]
+		public string Name { get; set; } = null!;
+
+		/// <summary>
+		/// The charts contained within the category
+		/// </summary>
+		public List<TelemetryChartConfig> Charts { get; set; } = new List<TelemetryChartConfig> { };
+	}
+
+	/// <summary>
+	/// A telemetry view variable used for filtering the charting data
+	/// </summary>
+	public class TelemetryVariableConfig
+	{
+		/// <summary>
+		/// The name of the variable for display purposes
+		/// </summary>
+		[Required]
+		public string Name { get; set; } = null!;
+
+		/// <summary>
+		/// The associated data group attached to the variable 
+		/// </summary>
+		[Required]
+		public string Group { get; set; } = null!;
+
+		/// <summary>
+		/// The default values to select
+		/// </summary>
+		public List<string> Defaults { get; set; } = new List<string> { };
+	}
+
+	/// <summary>
+	/// A telemetry view of related metrics, divided into categofies
+	/// </summary>
+	public class TelemetryViewConfig
+	{
+		/// <summary>
+		/// Identifier for the view
+		/// </summary>
+		[Required]
+		public TelmetryViewId Id { get; set; }
+
+		/// <summary>
+		/// The name of the view
+		/// </summary>
+		[Required]
+		public string Name { get; set; } = null!;
+
+		/// <summary>
+		/// The telemetry store this view uses
+		/// </summary>
+		[Required]
+		public TelemetryStoreId TelemetryStoreId { get; set; }
+
+		/// <summary>
+		///  The variables used to filter the view data
+		/// </summary>
+		public List<TelemetryVariableConfig> Variables { get; set; } = new List<TelemetryVariableConfig> { };
+
+		/// <summary>
+		/// The categories contained within the view
+		/// </summary>
+		public List<TelemetryCategoryConfig> Categories { get; set; } = new List<TelemetryCategoryConfig> { };
+	}
+
+	/// <summary>
+	/// Identifier for a particular metric view
+	/// </summary>
+	/// <param name="Id">Id to construct from</param>
+	[LogValueType]
+	[JsonSchemaString]
+	[TypeConverter(typeof(StringIdTypeConverter<TelmetryViewId, TelmetryViewIdConverter>))]
+	[StringIdConverter(typeof(TelmetryViewIdConverter))]
+	[CbConverter(typeof(StringIdCbConverter<TelmetryViewId, TelmetryViewIdConverter>))]
+	public record struct TelmetryViewId(StringId Id)
+	{
+		/// <summary>
+		/// Constructor
+		/// </summary>
+		public TelmetryViewId(string id) : this(new StringId(id))
+		{
+		}
+
+		/// <inheritdoc cref="StringId.IsEmpty"/>
+		public bool IsEmpty => Id.IsEmpty;
+
+		/// <inheritdoc/>
+		public override string ToString() => Id.ToString();
+	}
+
+	/// <summary>
+	/// Converter to and from <see cref="StringId"/> instances.
+	/// </summary>
+	class TelmetryViewIdConverter : StringIdConverter<TelmetryViewId>
+	{
+		/// <inheritdoc/>
+		public override TelmetryViewId FromStringId(StringId id) => new TelmetryViewId(id);
+
+		/// <inheritdoc/>
+		public override StringId ToStringId(TelmetryViewId value) => value.Id;
+	}
+
+	/// <summary>
+	/// Metric attached to a telemetry chart
+	/// </summary>
+	public class GetTelemetryChartMetricResponse
+	{
+		/// <summary>
+		/// Associated metric id
+		/// </summary>		
+		public string MetricId { get; set; } = null!;
+
+		/// <summary>
+		/// The threshold for KPI values
+		/// </summary>
+		public int? Threshold { get; set; }
+
+		/// <summary>
+		/// The metric alias for display purposes
+		/// </summary>
+		public string? Alias { get; set; }
+	}
+
+	/// <summary>
+	/// Telemetry chart configuraton
+	/// </summary>
+	public class GetTelemetryChartResponse
+	{
+		/// <summary>
+		/// The name of the chart, will be displayed on the dashboard
+		/// </summary>		
+		public string Name { get; set; } = null!;
+
+		/// <summary>
+		/// The unit to display
+		/// </summary>
+		public string Display { get; set; } = null!;
+
+		/// <summary>
+		/// The graph type 
+		/// </summary>
+		public string Graph { get; set; } = null!;
+
+		/// <summary>
+		/// List of configured metrics
+		/// </summary>
+		public List<GetTelemetryChartMetricResponse> Metrics { get; set; } = new List<GetTelemetryChartMetricResponse>();
+
+		/// <summary>
+		/// The min unit value for clamping chart
+		/// </summary>
+		public int? Min { get; set; }
+
+		/// <summary>
+		/// The max unit value for clamping chart
+		/// </summary>
+		public int? Max { get; set; }
+	}
+
+	/// <summary>
+	/// A chart categody, will be displayed on the dashbord under an associated pivot
+	/// </summary>
+	public class GetTelemetryCategoryResponse
+	{
+		/// <summary>
+		/// The name of the category
+		/// </summary>		
+		public string Name { get; set; } = null!;
+
+		/// <summary>
+		/// The charts contained within the category
+		/// </summary>
+		public List<GetTelemetryChartResponse> Charts { get; set; } = new List<GetTelemetryChartResponse> { };
+	}
+
+	/// <summary>
+	/// A telemetry view variable used for filtering the charting data
+	/// </summary>
+	public class GetTelemetryVariableResponse
+	{
+		/// <summary>
+		/// The name of the variable for display purposes
+		/// </summary>
+		public string Name { get; set; } = null!;
+
+		/// <summary>
+		/// The associated data group attached to the variable 
+		/// </summary>
+		public string Group { get; set; } = null!;
+
+		/// <summary>
+		/// The default values to select
+		/// </summary>
+		public List<string> Defaults { get; set; } = new List<string> { };
+	}
+
+	/// <summary>
+	/// A telemetry view of related metrics, divided into categofies
+	/// </summary>
+	public class GetTelemetryViewResponse
+	{
+		/// <summary>
+		/// Identifier for the view
+		/// </summary>
+		public string Id { get; set; } = null!;
+
+		/// <summary>
+		/// The name of the view
+		/// </summary>
+		public string Name { get; set; } = null!;
+
+		/// <summary>
+		/// The telemetry store the view uses
+		/// </summary>
+		public string TelemetryStoreId { get; set; } = null!;
+
+		/// <summary>
+		///  The variables used to filter the view data
+		/// </summary>
+		public List<GetTelemetryVariableResponse> Variables { get; set; } = new List<GetTelemetryVariableResponse> { };
+
+		/// <summary>
+		/// The categories contained within the view
+		/// </summary>
+		public List<GetTelemetryCategoryResponse> Categories { get; set; } = new List<GetTelemetryCategoryResponse> { };
+	}
 }
