Added more info to the "Unable to upload data to redirected URL" Horde Storage error message.

diff --git a/Engine/Source/Programs/Shared/EpicGames.Horde/Storage/Backends/HttpStorageBackend.cs b/Engine/Source/Programs/Shared/EpicGames.Horde/Storage/Backends/HttpStorageBackend.cs
--- a/Engine/Source/Programs/Shared/EpicGames.Horde/Storage/Backends/HttpStorageBackend.cs
+++ b/Engine/Source/Programs/Shared/EpicGames.Horde/Storage/Backends/HttpStorageBackend.cs
@@ -151,7 +151,7 @@ namespace EpicGames.Horde.Storage.Backends
 					if (!uploadResponse.IsSuccessStatusCode)
 					{
 						string body = await uploadResponse.Content.ReadAsStringAsync(cancellationToken);
-						throw new StorageException($"Unable to upload data to redirected URL: {body}");
+						throw new StorageException($"Unable to upload data to redirected URL {redirectResponse.UploadUrl}. Status:{uploadResponse.StatusCode} Reason:{uploadResponse.ReasonPhrase} Response:{body}");
 					}
 
 					_logger.LogDebug("Written {Locator} (using redirect)", redirectResponse.Blob);
