EpicGames.Core: Fix rendering of message templates containing width and format specifiers.

diff --git a/Engine/Source/Programs/Shared/EpicGames.Core/MessageTemplate.cs b/Engine/Source/Programs/Shared/EpicGames.Core/MessageTemplate.cs
--- a/Engine/Source/Programs/Shared/EpicGames.Core/MessageTemplate.cs
+++ b/Engine/Source/Programs/Shared/EpicGames.Core/MessageTemplate.cs
@@ -2,6 +2,7 @@
 
 using System;
 using System.Collections.Generic;
+using System.Globalization;
 using System.Text;
 using System.Text.Json;
 
@@ -41,12 +42,35 @@ namespace EpicGames.Core
 			int nextOffset = 0;
 
 			List<(int, int)>? names = ParsePropertyNames(format);
-			if (names != null)
+			if (names != null && properties != null)
 			{
 				foreach((int offset, int length) in names)
 				{
+					ReadOnlySpan<char> argument = format.AsSpan(offset, length);
+
+					string? formatSpecifier = null;
+
+					// Parse the format specifier
+					int formatSpecifierIdx = argument.IndexOf(':');
+					if (formatSpecifierIdx != -1)
+					{
+						formatSpecifier = argument.Slice(formatSpecifierIdx + 1).ToString();
+						argument = argument.Slice(0, formatSpecifierIdx);
+					}
+
+					// Parse the property width
+					int alignment = 0;
+
+					int alignmentIdx = argument.IndexOf(',');
+					if (alignmentIdx != -1)
+					{
+						alignment = Int32.Parse(argument.Slice(alignmentIdx + 1), NumberStyles.Integer | NumberStyles.AllowLeadingSign);
+						argument = argument.Slice(0, alignmentIdx);
+					}
+					
+					// Try to get the property value
 					object? value;
-					if (properties != null && TryGetPropertyValue(format.AsSpan(offset, length), properties, out value))
+					if (TryGetPropertyValue(argument, properties, out value))
 					{
 						// Append the text up to this argument
 						int startOffset = offset - 1;
@@ -56,14 +80,36 @@ namespace EpicGames.Core
 						}
 						Unescape(format.AsSpan(nextOffset, startOffset - nextOffset), result);
 
-						// Append the argument
+						// Render the property
+						string? rendered;
 						if (format[offset] == '@')
 						{
-							result.Append(JsonSerializer.Serialize(value, value?.GetType() ?? typeof(object), new JsonSerializerOptions { PropertyNamingPolicy = JsonNamingPolicy.CamelCase }));
+							rendered = JsonSerializer.Serialize(value, value?.GetType() ?? typeof(object), new JsonSerializerOptions { PropertyNamingPolicy = JsonNamingPolicy.CamelCase });
+						}
+						else if (formatSpecifier != null && value is IFormattable formattable)
+						{
+							rendered = formattable.ToString(formatSpecifier, null);
 						}
 						else
 						{
-							result.Append(value?.ToString() ?? "null");
+							rendered = value?.ToString() ?? "null";
+						}
+
+						// Apply a postive width adjustment (right align)
+						int leftPad = Math.Max(alignment - rendered.Length, 0);
+						if (leftPad > 0)
+						{
+							result.Append(' ', leftPad);
+						}
+
+						// Append the argument
+						result.Append(rendered);
+
+						// Apply a negative width adjustment (left align)
+						int rightPad = Math.Max(-alignment - rendered.Length, 0);
+						if (rightPad > 0)
+						{
+							result.Append(' ', rightPad);
 						}
 
 						// Start the next plain-text run after the closing brace
