Horde: Fix infinite loop in Redis library when getting a random set member. Also support deconstructing HashEntry objects.

#rnx

diff --git a/Engine/Source/Programs/Shared/EpicGames.Redis/RedisHashKey.cs b/Engine/Source/Programs/Shared/EpicGames.Redis/RedisHashKey.cs
--- a/Engine/Source/Programs/Shared/EpicGames.Redis/RedisHashKey.cs
+++ b/Engine/Source/Programs/Shared/EpicGames.Redis/RedisHashKey.cs
@@ -56,6 +56,15 @@ namespace EpicGames.Redis
 			Value = value;
 		}
 
+		/// <summary>
+		/// Deconstructor helper method
+		/// </summary>
+		public void Deconstruct(out TName name, out TValue value)
+		{
+			name = Name;
+			value = Value;
+		}
+
 		/// <summary>
 		/// Implicit conversion to a <see cref="HashEntry"/>
 		/// </summary>
@@ -315,14 +324,14 @@ namespace EpicGames.Redis
 		#region HashSetAsync
 
 		/// <inheritdoc cref="IDatabaseAsync.HashSetAsync(RedisKey, RedisValue, RedisValue, When, CommandFlags)"/>
-		public static Task HashSetAsync<TRecord, TValue>(this IDatabaseAsync target, RedisHashKey<TRecord> key, Expression<Func<TRecord, TValue>> selector, TValue value, When when = When.Always, CommandFlags flags = CommandFlags.None)
+		public static Task<bool> HashSetAsync<TRecord, TValue>(this IDatabaseAsync target, RedisHashKey<TRecord> key, Expression<Func<TRecord, TValue>> selector, TValue value, When when = When.Always, CommandFlags flags = CommandFlags.None)
 		{
 			MemberExpression memberExpression = (selector.Body as MemberExpression) ?? throw new InvalidOperationException("Expression must be a property accessor");
 			return target.HashSetAsync(key.Inner, memberExpression.Member.Name, RedisSerializer.Serialize(value), when, flags);
 		}
 
 		/// <inheritdoc cref="IDatabaseAsync.HashSetAsync(RedisKey, RedisValue, RedisValue, When, CommandFlags)"/>
-		public static Task HashSetAsync<TName, TValue>(this IDatabaseAsync target, RedisHashKey<TName, TValue> key, TName name, TValue value, When when = When.Always, CommandFlags flags = CommandFlags.None)
+		public static Task<bool> HashSetAsync<TName, TValue>(this IDatabaseAsync target, RedisHashKey<TName, TValue> key, TName name, TValue value, When when = When.Always, CommandFlags flags = CommandFlags.None)
 		{
 			return target.HashSetAsync(key.Inner, RedisSerializer.Serialize(name), RedisSerializer.Serialize(value), when, flags);
 		}
diff --git a/Engine/Source/Programs/Shared/EpicGames.Redis/RedisSetKey.cs b/Engine/Source/Programs/Shared/EpicGames.Redis/RedisSetKey.cs
--- a/Engine/Source/Programs/Shared/EpicGames.Redis/RedisSetKey.cs
+++ b/Engine/Source/Programs/Shared/EpicGames.Redis/RedisSetKey.cs
@@ -115,7 +115,7 @@ namespace EpicGames.Redis
 
 		/// <inheritdoc cref="IDatabaseAsync.SetRandomMemberAsync(RedisKey, CommandFlags)"/>
 		public static Task<TElement> SetRandomMemberAsync<TElement>(this IDatabaseAsync target, RedisSetKey<TElement> key, CommandFlags flags = CommandFlags.None)
-			=> target.SetRandomMemberAsync(key, flags);
+			=> target.SetRandomMemberAsync(key.Inner, flags).DeserializeAsync<TElement>();
 
 		#endregion
 
@@ -123,7 +123,7 @@ namespace EpicGames.Redis
 
 		/// <inheritdoc cref="IDatabaseAsync.SetRandomMembersAsync(RedisKey, long, CommandFlags)"/>
 		public static Task<TElement[]> SetRandomMembersAsync<TElement>(this IDatabaseAsync target, RedisSetKey<TElement> key, long count, CommandFlags flags = CommandFlags.None)
-			=> target.SetRandomMembersAsync(key, count, flags);
+			=> target.SetRandomMembersAsync(key.Inner, count, flags).DeserializeAsync<TElement>();
 
 		#endregion
 
