Horde: Add extension methods for reading records from Redis hash keys.

#rnx

diff --git a/Engine/Source/Programs/Horde/HordeServer.Tests/Server/RedisTests.cs b/Engine/Source/Programs/Horde/HordeServer.Tests/Server/RedisTests.cs
new file mode 100644
--- /dev/null
+++ b/Engine/Source/Programs/Horde/HordeServer.Tests/Server/RedisTests.cs
@@ -0,0 +1,43 @@
+// Copyright Epic Games, Inc. All Rights Reserved.
+
+using System;
+using System.Linq.Expressions;
+using System.Threading.Tasks;
+using EpicGames.Redis;
+using HordeServer.Server;
+using Microsoft.VisualStudio.TestTools.UnitTesting;
+using StackExchange.Redis;
+
+namespace HordeServer.Tests.Server
+{
+	[TestClass]
+	public class RedisTests : DatabaseIntegrationTest
+	{
+		class TestClass
+		{
+			public int Foo { get; set; }
+			public int Bar { get; set; }
+			public int Baz { get; set; }
+		}
+
+		[TestMethod]
+		public async Task HashTestAsync()
+		{
+			IRedisService redisService = GetRedisServiceSingleton();
+			IDatabase database = redisService.GetDatabase();
+
+			RedisHashKey<TestClass> key = new RedisHashKey<TestClass>("test");
+			await database.HashSetAsync(key, new TestClass { Foo = 123, Bar = 456, Baz = 789 });
+
+			TestClass value = await database.HashGetAllAsync(key);
+			Assert.AreEqual(123, value.Foo);
+			Assert.AreEqual(456, value.Bar);
+			Assert.AreEqual(789, value.Baz);
+
+			TestClass value2 = await database.HashGetAsync(key, new Expression<Func<TestClass, object>>[] { x => x.Foo, x => x.Bar });
+			Assert.AreEqual(123, value2.Foo);
+			Assert.AreEqual(456, value2.Bar);
+			Assert.AreEqual(0, value2.Baz);
+		}
+	}
+}
diff --git a/Engine/Source/Programs/Shared/EpicGames.Redis/RedisHashKey.cs b/Engine/Source/Programs/Shared/EpicGames.Redis/RedisHashKey.cs
--- a/Engine/Source/Programs/Shared/EpicGames.Redis/RedisHashKey.cs
+++ b/Engine/Source/Programs/Shared/EpicGames.Redis/RedisHashKey.cs
@@ -4,6 +4,7 @@ using System;
 using System.Collections.Generic;
 using System.Linq;
 using System.Linq.Expressions;
+using System.Reflection;
 using System.Threading.Tasks;
 using StackExchange.Redis;
 
@@ -226,6 +227,35 @@ namespace EpicGames.Redis
 
 		#region HashGetAsync
 
+		/// <inheritdoc cref="IDatabaseAsync.HashGetAsync(RedisKey, RedisValue[], CommandFlags)"/>
+		public static async Task<TRecord> HashGetAsync<TRecord>(this IDatabaseAsync target, RedisHashKey<TRecord> key, Expression<Func<TRecord, object>>[] selectors, CommandFlags flags = CommandFlags.None) where TRecord : new()
+		{
+			RedisValue[] names = new RedisValue[selectors.Length];
+			for (int idx = 0; idx < selectors.Length; idx++)
+			{
+				Expression expr = selectors[idx].Body;
+				if (expr is UnaryExpression unaryExpr && unaryExpr.NodeType == ExpressionType.Convert)
+				{
+					expr = unaryExpr.Operand;
+				}
+
+				MemberExpression memberExpression = (expr as MemberExpression) ?? throw new InvalidOperationException("Expression must be a property accessor");
+				names[idx] = memberExpression.Member.Name;
+			}
+
+			RedisValue[] values = await target.HashGetAsync(key.Inner, RedisSerializer.Serialize(names), flags);
+
+			HashRecordInfo<TRecord> recordInfo = HashRecordInfo<TRecord>.Instance;
+
+			TRecord record = new TRecord();
+			for (int idx = 0; idx < selectors.Length; idx++)
+			{
+				recordInfo.SetProperty(record, names[idx], values[idx]);
+			}
+
+			return record;
+		}
+
 		/// <inheritdoc cref="IDatabaseAsync.HashGetAsync(RedisKey, RedisValue, CommandFlags)"/>
 		public static Task<TValue> HashGetAsync<TRecord, TValue>(this IDatabaseAsync target, RedisHashKey<TRecord> key, Expression<Func<TRecord, TValue>> selector, CommandFlags flags = CommandFlags.None)
 		{
@@ -256,6 +286,22 @@ namespace EpicGames.Redis
 			return Array.ConvertAll(entries, x => new HashEntry<TName, TValue>(RedisSerializer.Deserialize<TName>(x.Name)!, RedisSerializer.Deserialize<TValue>(x.Value)!));
 		}
 
+		/// <inheritdoc cref="IDatabaseAsync.HashGetAllAsync(RedisKey, CommandFlags)"/>
+		public static async Task<TRecord> HashGetAllAsync<TRecord>(this IDatabaseAsync target, RedisHashKey<TRecord> key, CommandFlags flags = CommandFlags.None) where TRecord : new()
+		{
+			HashRecordInfo<TRecord> recordInfo = HashRecordInfo<TRecord>.Instance;
+
+			TRecord value = new TRecord();
+
+			HashEntry[] entries = await target.HashGetAllAsync(key.Inner, flags);
+			foreach (HashEntry entry in entries)
+			{
+				recordInfo.SetProperty(value, entry.Name, entry.Value);
+			}
+
+			return value;
+		}
+
 		#endregion
 
 		#region HashIncrementAsync
@@ -323,6 +369,52 @@ namespace EpicGames.Redis
 
 		#region HashSetAsync
 
+		class HashRecordInfo<T>
+		{
+			public static HashRecordInfo<T> Instance { get; } = new HashRecordInfo<T>();
+
+			public IReadOnlyList<HashPropertyInfo<T>> Properties { get; }
+
+			public IReadOnlyDictionary<string, HashPropertyInfo<T>> PropertiesByName { get; }
+
+			public HashRecordInfo()
+			{
+				List<HashPropertyInfo<T>> properties = new List<HashPropertyInfo<T>>();
+				foreach (PropertyInfo propertyInfo in typeof(T).GetProperties(BindingFlags.Public | BindingFlags.Instance))
+				{
+					properties.Add(new HashPropertyInfo<T>(propertyInfo));
+				}
+				Properties = properties;
+				PropertiesByName = properties.ToDictionary(x => x.Name, x => x);
+			}
+
+			public bool SetProperty(T record, RedisValue name, RedisValue value)
+			{
+				string? nameStr = (string?)name;
+				if (nameStr != null)
+				{
+					HashPropertyInfo<T>? propertyInfo;
+					if (PropertiesByName.TryGetValue(nameStr, out propertyInfo))
+					{
+						propertyInfo.SetValue(record, value);
+						return true;
+					}
+				}
+				return false;
+			}
+		}
+
+		record class HashPropertyInfo<T>(string Name, Func<T, RedisValue> GetValue, Action<T, RedisValue> SetValue)
+		{
+			public HashPropertyInfo(PropertyInfo propertyInfo)
+				: this(
+					propertyInfo.Name,
+					record => RedisSerializer.Serialize(propertyInfo.GetValue(record), propertyInfo.PropertyType),
+					(record, value) => propertyInfo.SetValue(record, RedisSerializer.Deserialize(value, propertyInfo.PropertyType))
+				)
+			{ }
+		}
+
 		/// <inheritdoc cref="IDatabaseAsync.HashSetAsync(RedisKey, RedisValue, RedisValue, When, CommandFlags)"/>
 		public static Task<bool> HashSetAsync<TRecord, TValue>(this IDatabaseAsync target, RedisHashKey<TRecord> key, Expression<Func<TRecord, TValue>> selector, TValue value, When when = When.Always, CommandFlags flags = CommandFlags.None)
 		{
@@ -330,6 +422,21 @@ namespace EpicGames.Redis
 			return target.HashSetAsync(key.Inner, memberExpression.Member.Name, RedisSerializer.Serialize(value), when, flags);
 		}
 
+		/// <inheritdoc cref="IDatabaseAsync.HashSetAsync(RedisKey, HashEntry[], CommandFlags)"/>
+		public static Task HashSetAsync<TRecord>(this IDatabaseAsync target, RedisHashKey<TRecord> key, TRecord value, CommandFlags flags = CommandFlags.None)
+		{
+			HashRecordInfo<TRecord> recordInfo = HashRecordInfo<TRecord>.Instance;
+
+			HashEntry[] entries = new HashEntry[recordInfo.Properties.Count];
+			for (int idx = 0; idx < recordInfo.Properties.Count; idx++)
+			{
+				HashPropertyInfo<TRecord> property = recordInfo.Properties[idx];
+				entries[idx] = new HashEntry(property.Name, property.GetValue(value));
+			}
+
+			return target.HashSetAsync(key.Inner, entries, flags);
+		}
+
 		/// <inheritdoc cref="IDatabaseAsync.HashSetAsync(RedisKey, RedisValue, RedisValue, When, CommandFlags)"/>
 		public static Task<bool> HashSetAsync<TName, TValue>(this IDatabaseAsync target, RedisHashKey<TName, TValue> key, TName name, TValue value, When when = When.Always, CommandFlags flags = CommandFlags.None)
 		{
diff --git a/Engine/Source/Programs/Shared/EpicGames.Redis/RedisSerializer.cs b/Engine/Source/Programs/Shared/EpicGames.Redis/RedisSerializer.cs
--- a/Engine/Source/Programs/Shared/EpicGames.Redis/RedisSerializer.cs
+++ b/Engine/Source/Programs/Shared/EpicGames.Redis/RedisSerializer.cs
@@ -1,6 +1,7 @@
 // Copyright Epic Games, Inc. All Rights Reserved.
 
 using System;
+using System.Collections.Concurrent;
 using System.Collections.Generic;
 using System.ComponentModel;
 using System.Reflection;
@@ -173,6 +174,17 @@ namespace EpicGames.Redis
 
 		static readonly Dictionary<Type, Type> s_typeToConverterType = new Dictionary<Type, Type>();
 
+		class RedisObjectConverter<T> : IRedisConverter<object>
+		{
+			public object FromRedisValue(RedisValue value)
+				=> GetConverter<T>().FromRedisValue(value)!;
+
+			public RedisValue ToRedisValue(object value)
+				=> GetConverter<T>().ToRedisValue((T)value);
+		}
+
+		static readonly ConcurrentDictionary<Type, IRedisConverter<object>> s_typeToObjectConverter = new ConcurrentDictionary<Type, IRedisConverter<object>>();
+
 		/// <summary>
 		/// Register a custom converter for a particular type
 		/// </summary>
@@ -279,6 +291,32 @@ namespace EpicGames.Redis
 			return CachedConverter<T>.Converter;
 		}
 
+		/// <summary>
+		/// Gets a type converter which casts to/from an object value
+		/// </summary>
+		/// <param name="type">The concrete type for the converter</param>
+		/// <returns></returns>
+		public static IRedisConverter<object> GetObjectConverter(Type type)
+		{
+			IRedisConverter<object>? converter;
+			if (!s_typeToObjectConverter.TryGetValue(type, out converter))
+			{
+				converter = s_typeToObjectConverter.GetOrAdd(type, (IRedisConverter<object>)Activator.CreateInstance(typeof(RedisObjectConverter<>).MakeGenericType(type))!);
+			}
+			return converter;
+		}
+
+		/// <summary>
+		/// Serialize an object to a <see cref="RedisValue"/>
+		/// </summary>
+		/// <param name="value"></param>
+		/// <param name="type">Type of the object</param>
+		/// <returns></returns>
+		public static RedisValue Serialize(object? value, Type type)
+		{
+			return GetObjectConverter(type).ToRedisValue(value!);
+		}
+
 		/// <summary>
 		/// Serialize an object to a <see cref="RedisValue"/>
 		/// </summary>
@@ -308,6 +346,17 @@ namespace EpicGames.Redis
 			return outputs;
 		}
 
+		/// <summary>
+		/// Deserialize a <see cref="RedisValue"/>
+		/// </summary>
+		/// <param name="value"></param>
+		/// <param name="type">Type of the value to return</param>
+		/// <returns></returns>
+		public static object? Deserialize(RedisValue value, Type type)
+		{
+			return GetObjectConverter(type).FromRedisValue(value);
+		}
+
 		/// <summary>
 		/// Deserialize a <see cref="RedisValue"/>
 		/// </summary>
