Horde: Fix static analysis warning.

#rnx

diff --git a/Engine/Source/Programs/Shared/EpicGames.Horde/Agents/Sessions/SessionMessages.cs b/Engine/Source/Programs/Shared/EpicGames.Horde/Agents/Sessions/SessionMessages.cs
--- a/Engine/Source/Programs/Shared/EpicGames.Horde/Agents/Sessions/SessionMessages.cs
+++ b/Engine/Source/Programs/Shared/EpicGames.Horde/Agents/Sessions/SessionMessages.cs
@@ -1,7 +1,6 @@
 // Copyright Epic Games, Inc. All Rights Reserved.
 
 using System;
-using System.Collections.Generic;
 
 namespace EpicGames.Horde.Agents.Sessions
 {
