EpicGames.Core: Handle invalid format specifiers when formatting message templates.

#rnx

diff --git a/Engine/Source/Programs/Shared/EpicGames.Core/MessageTemplate.cs b/Engine/Source/Programs/Shared/EpicGames.Core/MessageTemplate.cs
--- a/Engine/Source/Programs/Shared/EpicGames.Core/MessageTemplate.cs
+++ b/Engine/Source/Programs/Shared/EpicGames.Core/MessageTemplate.cs
@@ -62,9 +62,8 @@ namespace EpicGames.Core
 					int alignment = 0;
 
 					int alignmentIdx = argument.IndexOf(',');
-					if (alignmentIdx != -1)
+					if (alignmentIdx != -1 && Int32.TryParse(argument.Slice(alignmentIdx + 1), NumberStyles.Integer | NumberStyles.AllowLeadingSign, null, out alignment))
 					{
-						alignment = Int32.Parse(argument.Slice(alignmentIdx + 1), NumberStyles.Integer | NumberStyles.AllowLeadingSign);
 						argument = argument.Slice(0, alignmentIdx);
 					}
 					
