auto-selection of manual SDKs does not depend on having AutoSDK set up, on platforms that support it.

#rnx

diff --git a/Engine/Source/Programs/Shared/EpicGames.Build/UEBuildPlatformSDK.cs b/Engine/Source/Programs/Shared/EpicGames.Build/UEBuildPlatformSDK.cs
--- a/Engine/Source/Programs/Shared/EpicGames.Build/UEBuildPlatformSDK.cs
+++ b/Engine/Source/Programs/Shared/EpicGames.Build/UEBuildPlatformSDK.cs
@@ -348,7 +348,8 @@ namespace EpicGames.Core
 			else
 			{
 				// allow the platform to select the most appropriate manually installed SDK if it hasn't already been set up by AutoSDK
-				if (PlatformSupportsAutoSDKs() && HasAutoSDKSystemEnabled() && !HasParentProcessSetupAutoSDK(out _)) // @todo: PlatformSupportsAutoSDKs & HasAutoSDKSystemEnabled check is not technically needed for manual sdk switching but the editor-side code is heavily tied into this and would require further refactoring first
+				bool bHasAutoSDK = PlatformSupportsAutoSDKs() && HasAutoSDKSystemEnabled(); // @todo: PlatformSupportsAutoSDKs & HasAutoSDKSystemEnabled check is not technically needed for manual sdk switching but the editor-side code is heavily tied into this and would require further refactoring first
+				if ( (bHasAutoSDK || !ManualSDKAutoSwitchRequiresAutoSDK) && !HasParentProcessSetupAutoSDK(out _)) 
 				{
 					if (TrySelectBestManualSDK( out string? SelectedSDKVersion ))
 					{
@@ -706,6 +707,13 @@ namespace EpicGames.Core
 			return false;
 		}
 
+		/// <summary>
+		/// Whether the automatic switching to a locally-installed manual SDK requires AutoSDK to be configured
+		/// This is typically true unless the editor-side binaries do not depend on a version-specific SDK environment variable or path
+		/// i.e. the SDK version is compiled into the binary and can be used directly.
+		/// </summary>
+		public virtual bool ManualSDKAutoSwitchRequiresAutoSDK => true;
+
 		/// <summary>
 		/// Allows a platform to switch a different version of a manually-installed SDK, if the platform supports side-by-side installations of different versions
 		/// </summary>
