Horde: Fix syntax warning for CommitId

diff --git a/Engine/Source/Programs/Shared/EpicGames.Horde/Commits/CommitId.cs b/Engine/Source/Programs/Shared/EpicGames.Horde/Commits/CommitId.cs
--- a/Engine/Source/Programs/Shared/EpicGames.Horde/Commits/CommitId.cs
+++ b/Engine/Source/Programs/Shared/EpicGames.Horde/Commits/CommitId.cs
@@ -156,13 +156,13 @@ namespace EpicGames.Horde.Commits
 		/// <summary>
 		/// Creates a commit it from a Perforce changelist number. Temporary helper method for migration purposes.
 		/// </summary>
-		public new static CommitIdWithOrder FromPerforceChange(int change)
+		public static new CommitIdWithOrder FromPerforceChange(int change)
 			=> new CommitIdWithOrder($"{change}", change);
 
 		/// <summary>
 		/// Creates a commit it from a Perforce changelist number. Temporary helper method for migration purposes.
 		/// </summary>
-		public new static CommitIdWithOrder? FromPerforceChange(int? change)
+		public static new CommitIdWithOrder? FromPerforceChange(int? change)
 			=> (change == null) ? null : CommitIdWithOrder.FromPerforceChange(change.Value);
 
 		/// <inheritdoc/>
