Horde: Add method for creating a blob ref from a path on the server.

#rnx

diff --git a/Engine/Source/Programs/Horde/Plugins/Tools/HordeServer.Tools/Commands/Generate/ToolDataCommand.cs b/Engine/Source/Programs/Horde/Plugins/Tools/HordeServer.Tools/Commands/Generate/ToolDataCommand.cs
--- a/Engine/Source/Programs/Horde/Plugins/Tools/HordeServer.Tools/Commands/Generate/ToolDataCommand.cs
+++ b/Engine/Source/Programs/Horde/Plugins/Tools/HordeServer.Tools/Commands/Generate/ToolDataCommand.cs
@@ -131,7 +131,7 @@ namespace HordeServer.Commands.Generate
 
 			return 0;
 		}
-
+		/*
 		static async Task PopulateCacheAsync(IStorageClient client, DedupeBlobWriter writer, DirectoryReference searchDir, CancellationToken cancellationToken)
 		{
 			foreach (RefName refName in FileStorageBackend.EnumerateRefs(searchDir))
@@ -148,5 +148,6 @@ namespace HordeServer.Commands.Generate
 				}
 			}
 		}
+		*/
 	}
 }
diff --git a/Engine/Source/Programs/Shared/EpicGames.Horde/HordeHttpClient.cs b/Engine/Source/Programs/Shared/EpicGames.Horde/HordeHttpClient.cs
--- a/Engine/Source/Programs/Shared/EpicGames.Horde/HordeHttpClient.cs
+++ b/Engine/Source/Programs/Shared/EpicGames.Horde/HordeHttpClient.cs
@@ -3,7 +3,9 @@
 using System;
 using System.Collections.Generic;
 using System.IO;
+using System.Net;
 using System.Net.Http;
+using System.Net.Http.Headers;
 using System.Net.Http.Json;
 using System.Text.Json;
 using System.Text.Json.Serialization;
@@ -348,6 +350,43 @@ namespace EpicGames.Horde
 
 		#endregion
 
+		#region Storage
+
+		/// <summary>
+		/// Attempts to read a named storage ref from the server
+		/// </summary>
+		/// <param name="path">Path to the ref</param>
+		/// <param name="cacheTime">Max allowed age for a cached value to be returned</param>
+		/// <param name="cancellationToken">Cancellation token for the operation</param>
+		public async Task<ReadRefResponse?> TryReadRefAsync(string path, RefCacheTime cacheTime = default, CancellationToken cancellationToken = default)
+		{
+			using (HttpRequestMessage request = new HttpRequestMessage(HttpMethod.Get, path))
+			{
+				if (cacheTime.IsSet())
+				{
+					request.Headers.CacheControl = new CacheControlHeaderValue { MaxAge = cacheTime.MaxAge };
+				}
+
+				using (HttpResponseMessage response = await _httpClient.SendAsync(request, cancellationToken))
+				{
+					if (response.StatusCode == HttpStatusCode.NotFound)
+					{
+						return null;
+					}
+					else if (!response.IsSuccessStatusCode)
+					{
+						throw new StorageException($"Unable to read ref '{path}' (status: {response.StatusCode}, body: {await response.Content.ReadAsStringAsync(cancellationToken)}");
+					}
+					else
+					{
+						return await response.Content.ReadFromJsonAsync<ReadRefResponse>(cancellationToken: cancellationToken);
+					}
+				}
+			}
+		}
+
+		#endregion
+
 		#region Telemetry
 
 		/// <summary>
diff --git a/Engine/Source/Programs/Shared/EpicGames.Horde/IHordeClient.cs b/Engine/Source/Programs/Shared/EpicGames.Horde/IHordeClient.cs
--- a/Engine/Source/Programs/Shared/EpicGames.Horde/IHordeClient.cs
+++ b/Engine/Source/Programs/Shared/EpicGames.Horde/IHordeClient.cs
@@ -123,5 +123,28 @@ namespace EpicGames.Horde
 		public static IStorageClient CreateStorageClient(this IHordeClient hordeClient, ToolId toolId)
 			=> hordeClient.CreateStorageClient($"api/v1/tools/{toolId}");
 
+		/// <summary>
+		/// Reads a blob storage ref from a path
+		/// </summary>
+		public static async Task<IBlobRef?> TryReadRefAsync(this IHordeClient hordeClient, string path, RefCacheTime cacheTime = default, CancellationToken cancellationToken = default)
+		{
+			ReadRefResponse? response = await hordeClient.CreateHttpClient().TryReadRefAsync(path, cacheTime, cancellationToken);
+			if (response == null)
+			{
+				return null;
+			}
+
+			IStorageClient storageClient = hordeClient.CreateStorageClient(response.BasePath);
+			return storageClient.CreateBlobRef(response.Target);
+		}
+
+		/// <summary>
+		/// Reads a typed blob storage ref from a path
+		/// </summary>
+		public static async Task<IBlobRef<T>?> TryReadRefAsync<T>(this IHordeClient hordeClient, string path, RefCacheTime cacheTime = default, CancellationToken cancellationToken = default)
+		{
+			IBlobRef? blobRef = await hordeClient.TryReadRefAsync(path, cacheTime, cancellationToken);
+			return blobRef?.ForType<T>();
+		}
 	}
 }
