Horde: Remove enumerate method from object stores.

#rnx

diff --git a/Engine/Source/Programs/Horde/Plugins/Storage/HordeServer.Storage.Tests/Storage/GcServiceTests.cs b/Engine/Source/Programs/Horde/Plugins/Storage/HordeServer.Storage.Tests/Storage/GcServiceTests.cs
--- a/Engine/Source/Programs/Horde/Plugins/Storage/HordeServer.Storage.Tests/Storage/GcServiceTests.cs
+++ b/Engine/Source/Programs/Horde/Plugins/Storage/HordeServer.Storage.Tests/Storage/GcServiceTests.cs
@@ -1,6 +1,7 @@
 // Copyright Epic Games, Inc. All Rights Reserved.
 
 using EpicGames.Horde.Storage;
+using EpicGames.Horde.Storage.ObjectStores;
 using HordeServer.Plugins;
 using HordeServer.Server;
 using HordeServer.Storage;
@@ -51,9 +52,9 @@ namespace HordeServer.Tests.Storage
 
 			await Clock.AdvanceAsync(TimeSpan.FromDays(1.0));
 
-			IObjectStore backend = ServiceProvider.GetRequiredService<IObjectStoreFactory>().CreateObjectStore(storageConfig.Backends[0]);
+			MemoryObjectStore backend = (MemoryObjectStore)ServiceProvider.GetRequiredService<IObjectStoreFactory>().CreateObjectStore(storageConfig.Backends[0]);
 
-			ObjectKey[] remaining = await backend.EnumerateAsync().ToArrayAsync();
+			ObjectKey[] remaining = backend.Blobs.Keys.ToArray();
 			Assert.AreEqual(nodes.Count, remaining.Length);
 
 			HashSet<ObjectKey> nodePaths = new HashSet<ObjectKey>(nodes.Select(x => StorageService.GetObjectKey(x.BaseLocator)));
diff --git a/Engine/Source/Programs/Horde/Plugins/Storage/HordeServer.Storage/Storage/ObjectStores/AwsObjectStore.cs b/Engine/Source/Programs/Horde/Plugins/Storage/HordeServer.Storage/Storage/ObjectStores/AwsObjectStore.cs
--- a/Engine/Source/Programs/Horde/Plugins/Storage/HordeServer.Storage/Storage/ObjectStores/AwsObjectStore.cs
+++ b/Engine/Source/Programs/Horde/Plugins/Storage/HordeServer.Storage/Storage/ObjectStores/AwsObjectStore.cs
@@ -552,43 +552,6 @@ namespace HordeServer.Storage.ObjectStores
 			}
 		}
 
-		/// <inheritdoc/>
-		public async IAsyncEnumerable<ObjectKey> EnumerateAsync([EnumeratorCancellation] CancellationToken cancellationToken = default)
-		{
-			ListObjectsV2Request request = new ListObjectsV2Request();
-			request.BucketName = _options.AwsBucketName;
-			if (_pathPrefix.Length > 0)
-			{
-				request.Prefix = _pathPrefix;
-			}
-
-			for (; ; )
-			{
-				ListObjectsV2Response response = await _client.ListObjectsV2Async(request, cancellationToken);
-				foreach (S3Object obj in response.S3Objects)
-				{
-					string path = obj.Key;
-					if (path.StartsWith(_pathPrefix, StringComparison.Ordinal))
-					{
-						yield return new ObjectKey(path.Substring(_pathPrefix.Length));
-					}
-					else
-					{
-						_logger.LogError("Unexpected object enumerated from {AwsBucketName} - expected object \"{Path}\" to start with \"{AwsBucketPath}\"", _options.AwsBucketName, path, _pathPrefix);
-					}
-				}
-
-				if (response.IsTruncated)
-				{
-					request.ContinuationToken = response.NextContinuationToken;
-				}
-				else
-				{
-					break;
-				}
-			}
-		}
-
 		/// <inheritdoc/>
 		public void GetStats(StorageStats stats) { }
 	}
diff --git a/Engine/Source/Programs/Horde/Plugins/Storage/HordeServer.Storage/Storage/ObjectStores/AzureObjectStore.cs b/Engine/Source/Programs/Horde/Plugins/Storage/HordeServer.Storage/Storage/ObjectStores/AzureObjectStore.cs
--- a/Engine/Source/Programs/Horde/Plugins/Storage/HordeServer.Storage/Storage/ObjectStores/AzureObjectStore.cs
+++ b/Engine/Source/Programs/Horde/Plugins/Storage/HordeServer.Storage/Storage/ObjectStores/AzureObjectStore.cs
@@ -141,19 +141,6 @@ namespace HordeServer.Storage.ObjectStores
 			await _blobContainer.DeleteBlobAsync(key.ToString(), cancellationToken: cancellationToken);
 		}
 
-		/// <inheritdoc/>
-		public async IAsyncEnumerable<ObjectKey> EnumerateAsync([EnumeratorCancellation] CancellationToken cancellationToken = default)
-		{
-			bool exists = await _blobContainer.ExistsAsync(cancellationToken);
-			if (exists)
-			{
-				await foreach (BlobItem? item in _blobContainer.GetBlobsAsync(BlobTraits.Metadata, cancellationToken: cancellationToken))
-				{
-					yield return new ObjectKey(item.Name);
-				}
-			}
-		}
-
 		/// <inheritdoc/>
 		public async ValueTask<Uri?> TryGetReadRedirectAsync(ObjectKey key, CancellationToken cancellationToken = default)
 		{
diff --git a/Engine/Source/Programs/Horde/Plugins/Storage/HordeServer.Storage/Storage/ObjectStores/ChainedObjectStore.cs b/Engine/Source/Programs/Horde/Plugins/Storage/HordeServer.Storage/Storage/ObjectStores/ChainedObjectStore.cs
--- a/Engine/Source/Programs/Horde/Plugins/Storage/HordeServer.Storage/Storage/ObjectStores/ChainedObjectStore.cs
+++ b/Engine/Source/Programs/Horde/Plugins/Storage/HordeServer.Storage/Storage/ObjectStores/ChainedObjectStore.cs
@@ -32,19 +32,6 @@ namespace HordeServer.Storage.Storage.ObjectStores
 		public Task DeleteAsync(ObjectKey key, CancellationToken cancellationToken = default)
 			=> _primary.DeleteAsync(key, cancellationToken);
 
-		/// <inheritdoc/>
-		public async IAsyncEnumerable<ObjectKey> EnumerateAsync([EnumeratorCancellation] CancellationToken cancellationToken = default)
-		{
-			await foreach (ObjectKey key in _primary.EnumerateAsync(cancellationToken))
-			{
-				yield return key;
-			}
-			await foreach (ObjectKey key in _secondary.EnumerateAsync(cancellationToken))
-			{
-				yield return key;
-			}
-		}
-
 		/// <inheritdoc/>
 		public async Task<bool> ExistsAsync(ObjectKey key, CancellationToken cancellationToken = default)
 			=> await _primary.ExistsAsync(key, cancellationToken) || await _secondary.ExistsAsync(key, cancellationToken);
diff --git a/Engine/Source/Programs/Shared/EpicGames.Horde/Storage/IObjectStore.cs b/Engine/Source/Programs/Shared/EpicGames.Horde/Storage/IObjectStore.cs
--- a/Engine/Source/Programs/Shared/EpicGames.Horde/Storage/IObjectStore.cs
+++ b/Engine/Source/Programs/Shared/EpicGames.Horde/Storage/IObjectStore.cs
@@ -64,13 +64,6 @@ namespace EpicGames.Horde.Storage
 		/// <returns>Async task</returns>
 		Task DeleteAsync(ObjectKey key, CancellationToken cancellationToken = default);
 
-		/// <summary>
-		/// Enumerates all the objects in the store
-		/// </summary>
-		/// <param name="cancellationToken">Cancellation token for the operation</param>
-		/// <returns>Sequence of object paths</returns>
-		IAsyncEnumerable<ObjectKey> EnumerateAsync(CancellationToken cancellationToken = default);
-
 		/// <summary>
 		/// Gets a HTTP redirect for a read request
 		/// </summary>
@@ -118,7 +111,6 @@ namespace EpicGames.Horde.Storage
 			public Task WriteAsync(ObjectKey key, Stream stream, CancellationToken cancellationToken) => _inner.WriteAsync(key, stream, cancellationToken);
 			public Task DeleteAsync(ObjectKey path, CancellationToken cancellationToken) => _inner.DeleteAsync(path, cancellationToken);
 			public Task<bool> ExistsAsync(ObjectKey path, CancellationToken cancellationToken) => _inner.ExistsAsync(path, cancellationToken);
-			public IAsyncEnumerable<ObjectKey> EnumerateAsync(CancellationToken cancellationToken = default) => _inner.EnumerateAsync(cancellationToken);
 			public ValueTask<Uri?> TryGetReadRedirectAsync(ObjectKey path, CancellationToken cancellationToken = default) => _inner.TryGetReadRedirectAsync(path, cancellationToken);
 			public ValueTask<Uri?> TryGetWriteRedirectAsync(ObjectKey path, CancellationToken cancellationToken = default) => _inner.TryGetWriteRedirectAsync(path, cancellationToken);
 			public void GetStats(StorageStats stats) => _inner.GetStats(stats);
diff --git a/Engine/Source/Programs/Shared/EpicGames.Horde/Storage/ObjectStores/FileObjectStore.cs b/Engine/Source/Programs/Shared/EpicGames.Horde/Storage/ObjectStores/FileObjectStore.cs
--- a/Engine/Source/Programs/Shared/EpicGames.Horde/Storage/ObjectStores/FileObjectStore.cs
+++ b/Engine/Source/Programs/Shared/EpicGames.Horde/Storage/ObjectStores/FileObjectStore.cs
@@ -126,45 +126,6 @@ namespace EpicGames.Horde.Storage.ObjectStores
 			return Task.CompletedTask;
 		}
 
-		/// <inheritdoc/>
-		public async IAsyncEnumerable<ObjectKey> EnumerateAsync([EnumeratorCancellation] CancellationToken cancellationToken = default)
-		{
-			Stack<IEnumerator<DirectoryInfo>> queue = new Stack<IEnumerator<DirectoryInfo>>();
-			try
-			{
-				queue.Push(new List<DirectoryInfo> { _baseDir.ToDirectoryInfo() }.GetEnumerator());
-				while (queue.Count > 0)
-				{
-					IEnumerator<DirectoryInfo> top = queue.Peek();
-					if (!top.MoveNext())
-					{
-						top.Dispose();
-						queue.Pop();
-						continue;
-					}
-
-					DirectoryInfo current = top.Current;
-					foreach (FileInfo fileInfo in current.EnumerateFiles("*"))
-					{
-						string path = fileInfo.FullName.Substring(_baseDir.FullName.Length + 1).Replace(Path.DirectorySeparatorChar, '/');
-						yield return new ObjectKey(path.Substring(0, path.Length - 5));
-					}
-
-					queue.Push(current.EnumerateDirectories().GetEnumerator());
-
-					cancellationToken.ThrowIfCancellationRequested();
-					await Task.Yield();
-				}
-			}
-			finally
-			{
-				while (queue.TryPop(out IEnumerator<DirectoryInfo>? enumerator))
-				{
-					enumerator.Dispose();
-				}
-			}
-		}
-
 		/// <inheritdoc/>
 		public ValueTask<Uri?> TryGetReadRedirectAsync(ObjectKey key, CancellationToken cancellationToken = default) => default;
 
diff --git a/Engine/Source/Programs/Shared/EpicGames.Horde/Storage/ObjectStores/MemoryObjectStore.cs b/Engine/Source/Programs/Shared/EpicGames.Horde/Storage/ObjectStores/MemoryObjectStore.cs
--- a/Engine/Source/Programs/Shared/EpicGames.Horde/Storage/ObjectStores/MemoryObjectStore.cs
+++ b/Engine/Source/Programs/Shared/EpicGames.Horde/Storage/ObjectStores/MemoryObjectStore.cs
@@ -71,17 +71,6 @@ namespace EpicGames.Horde.Storage.ObjectStores
 			return Task.CompletedTask;
 		}
 
-		/// <inheritdoc/>
-		public async IAsyncEnumerable<ObjectKey> EnumerateAsync([EnumeratorCancellation] CancellationToken cancellationToken = default)
-		{
-			foreach (ObjectKey key in _keyToData.Keys)
-			{
-				yield return key;
-				cancellationToken.ThrowIfCancellationRequested();
-				await Task.Yield();
-			}
-		}
-
 		/// <inheritdoc/>
 		public ValueTask<Uri?> TryGetReadRedirectAsync(ObjectKey key, CancellationToken cancellationToken = default) => default;
 
diff --git a/Engine/Source/Programs/Shared/EpicGames.Horde/Storage/ObjectStores/PrefixedObjectStore.cs b/Engine/Source/Programs/Shared/EpicGames.Horde/Storage/ObjectStores/PrefixedObjectStore.cs
--- a/Engine/Source/Programs/Shared/EpicGames.Horde/Storage/ObjectStores/PrefixedObjectStore.cs
+++ b/Engine/Source/Programs/Shared/EpicGames.Horde/Storage/ObjectStores/PrefixedObjectStore.cs
@@ -39,18 +39,6 @@ namespace EpicGames.Horde.Storage.ObjectStores
 		/// <inheritdoc/>
 		public Task DeleteAsync(ObjectKey locator, CancellationToken cancellationToken = default) => _inner.DeleteAsync(GetKeyWithPrefix(locator), cancellationToken);
 
-		/// <inheritdoc/>
-		public async IAsyncEnumerable<ObjectKey> EnumerateAsync([EnumeratorCancellation] CancellationToken cancellationToken = default)
-		{
-			await foreach (ObjectKey key in _inner.EnumerateAsync(cancellationToken))
-			{
-				if (key.Path.StartsWith(_prefix))
-				{
-					yield return new ObjectKey(key.Path.Substring(_prefix.Length));
-				}
-			}
-		}
-
 		/// <inheritdoc/>
 		public Task<bool> ExistsAsync(ObjectKey locator, CancellationToken cancellationToken = default) => _inner.ExistsAsync(GetKeyWithPrefix(locator), cancellationToken);
 
diff --git a/Engine/Source/Programs/Shared/EpicGames.Horde/Storage/StorageBackendCache.cs b/Engine/Source/Programs/Shared/EpicGames.Horde/Storage/StorageBackendCache.cs
--- a/Engine/Source/Programs/Shared/EpicGames.Horde/Storage/StorageBackendCache.cs
+++ b/Engine/Source/Programs/Shared/EpicGames.Horde/Storage/StorageBackendCache.cs
@@ -83,7 +83,6 @@ namespace EpicGames.Horde.Storage
 
 			public Task WriteAsync(ObjectKey key, Stream stream, CancellationToken cancellationToken = default) => _inner.WriteAsync(key, stream, cancellationToken);
 			public Task DeleteAsync(ObjectKey key, CancellationToken cancellationToken = default) => _inner.DeleteAsync(key, cancellationToken);
-			public IAsyncEnumerable<ObjectKey> EnumerateAsync(CancellationToken cancellationToken = default) => _inner.EnumerateAsync(cancellationToken);
 			public Task<bool> ExistsAsync(ObjectKey key, CancellationToken cancellationToken = default) => _inner.ExistsAsync(key, cancellationToken);
 			public ValueTask<Uri?> TryGetReadRedirectAsync(ObjectKey key, CancellationToken cancellationToken = default) => _inner.TryGetReadRedirectAsync(key, cancellationToken);
 			public ValueTask<Uri?> TryGetWriteRedirectAsync(ObjectKey key, CancellationToken cancellationToken = default) => _inner.TryGetWriteRedirectAsync(key, cancellationToken);
