Horde: Add a utility class for broadcasting events similar to AsyncEvent over Redis.

#rnx

diff --git a/Engine/Source/Programs/Horde/HordeServer.Tests/Server/RedisTests.cs b/Engine/Source/Programs/Horde/HordeServer.Tests/Server/RedisTests.cs
--- a/Engine/Source/Programs/Horde/HordeServer.Tests/Server/RedisTests.cs
+++ b/Engine/Source/Programs/Horde/HordeServer.Tests/Server/RedisTests.cs
@@ -39,5 +39,19 @@ namespace HordeServer.Tests.Server
 			Assert.AreEqual(456, value2.Bar);
 			Assert.AreEqual(0, value2.Baz);
 		}
+
+		[TestMethod]
+		public async Task AsyncEventTestAsync()
+		{
+			IRedisService redisService = GetRedisServiceSingleton();
+
+			await using RedisAsyncEvent asyncEvent = await RedisAsyncEvent.CreateAsync(redisService.GetConnection(), RedisChannel.Literal("hello-world"));
+
+			Task task = asyncEvent.Task;
+			Assert.IsFalse(task.IsCompleted);
+
+			asyncEvent.Pulse();
+			await task;
+		}
 	}
 }
diff --git a/Engine/Source/Programs/Shared/EpicGames.Redis/RedisAsyncEvent.cs b/Engine/Source/Programs/Shared/EpicGames.Redis/RedisAsyncEvent.cs
new file mode 100644
--- /dev/null
+++ b/Engine/Source/Programs/Shared/EpicGames.Redis/RedisAsyncEvent.cs
@@ -0,0 +1,61 @@
+// Copyright Epic Games, Inc. All Rights Reserved.
+
+using System;
+using System.Collections.Generic;
+using System.Linq;
+using System.Text;
+using System.Threading.Tasks;
+using EpicGames.Core;
+using StackExchange.Redis;
+
+namespace EpicGames.Redis
+{
+	/// <summary>
+	/// Asynchronous version of <see cref="AsyncEvent"/>. The event will be pulsed via a pub/sub channel in Redis.
+	/// </summary>
+	public class RedisAsyncEvent : IAsyncDisposable
+	{
+		readonly IConnectionMultiplexer _multiplexer;
+		readonly RedisChannel _channel;
+		readonly AsyncEvent _asyncEvent;
+		readonly RedisSubscription _subscription;
+
+		/// <summary>
+		/// Accessor for the inner task. Can be captured and awaited by clients.
+		/// </summary>
+		public Task Task => _asyncEvent.Task;
+
+		RedisAsyncEvent(IConnectionMultiplexer multiplexer, RedisChannel channel, AsyncEvent asyncEvent, RedisSubscription subscription)
+		{
+			_multiplexer = multiplexer;
+			_channel = channel;
+			_asyncEvent = asyncEvent;
+			_subscription = subscription;
+		}
+
+		/// <inheritdoc/>
+		public ValueTask DisposeAsync()
+			=> _subscription.DisposeAsync();
+
+		/// <summary>
+		/// Create a new async event using the given channel name
+		/// </summary>
+		/// <param name="multiplexer">Multiplexer for the </param>
+		/// <param name="channel"></param>
+		/// <returns></returns>
+		public static async Task<RedisAsyncEvent> CreateAsync(IConnectionMultiplexer multiplexer, RedisChannel channel)
+		{
+			AsyncEvent asyncEvent = new AsyncEvent();
+			RedisSubscription subscription = await multiplexer.SubscribeAsync(channel, x => asyncEvent.Pulse());
+			return new RedisAsyncEvent(multiplexer, channel, asyncEvent, subscription);
+		}
+
+		/// <summary>
+		/// Pulse the event, allowing any captured copies of the task to continue.
+		/// </summary>
+		public void Pulse()
+		{
+			_ = _multiplexer.GetDatabase().PublishAsync(_channel, RedisValue.EmptyString, CommandFlags.FireAndForget);
+		}
+	}
+}
