Horde: Static analysis fix.

#rnx

diff --git a/Engine/Source/Programs/Shared/EpicGames.Redis/RedisAsyncEvent.cs b/Engine/Source/Programs/Shared/EpicGames.Redis/RedisAsyncEvent.cs
--- a/Engine/Source/Programs/Shared/EpicGames.Redis/RedisAsyncEvent.cs
+++ b/Engine/Source/Programs/Shared/EpicGames.Redis/RedisAsyncEvent.cs
@@ -1,9 +1,6 @@
 // Copyright Epic Games, Inc. All Rights Reserved.
 
 using System;
-using System.Collections.Generic;
-using System.Linq;
-using System.Text;
 using System.Threading.Tasks;
 using EpicGames.Core;
 using StackExchange.Redis;
