Horde: Add a FIFO queue implementation to the Redis library which wraps an item list and event for waiting for updates.

#rnx

diff --git a/Engine/Source/Programs/Horde/HordeServer.Tests/Server/RedisTests.cs b/Engine/Source/Programs/Horde/HordeServer.Tests/Server/RedisTests.cs
--- a/Engine/Source/Programs/Horde/HordeServer.Tests/Server/RedisTests.cs
+++ b/Engine/Source/Programs/Horde/HordeServer.Tests/Server/RedisTests.cs
@@ -4,6 +4,7 @@ using System;
 using System.Linq.Expressions;
 using System.Threading.Tasks;
 using EpicGames.Redis;
+using EpicGames.Redis.Utility;
 using HordeServer.Server;
 using Microsoft.VisualStudio.TestTools.UnitTesting;
 using StackExchange.Redis;
@@ -45,7 +46,7 @@ namespace HordeServer.Tests.Server
 		{
 			IRedisService redisService = GetRedisServiceSingleton();
 
-			await using RedisAsyncEvent asyncEvent = await RedisAsyncEvent.CreateAsync(redisService.GetConnection(), RedisChannel.Literal("hello-world"));
+			await using RedisEvent asyncEvent = await RedisEvent.CreateAsync(redisService.GetConnection(), RedisChannel.Literal("hello-world"));
 
 			Task task = asyncEvent.Task;
 			Assert.IsFalse(task.IsCompleted);
@@ -53,5 +54,37 @@ namespace HordeServer.Tests.Server
 			asyncEvent.Pulse();
 			await task;
 		}
+
+		[TestMethod]
+		public async Task QueueTestAsync()
+		{
+			IRedisService redisService = GetRedisServiceSingleton();
+
+			await using RedisQueue<int> queue = await RedisQueue.CreateAsync<int>(redisService.GetConnection(), "queue", RedisChannel.Literal("queue-events"));
+			await queue.PushAsync(1);
+			await queue.PushAsync(2);
+			await queue.PushAsync(3);
+
+			int value = await queue.TryPopAsync();
+			Assert.AreEqual(1, value);
+
+			value = await queue.TryPopAsync();
+			Assert.AreEqual(2, value);
+
+			value = await queue.TryPopAsync();
+			Assert.AreEqual(3, value);
+
+			value = await queue.TryPopAsync();
+			Assert.AreEqual(0, value);
+
+			Task task = queue.WaitForDataAsync();
+			Assert.IsFalse(task.IsCompleted);
+
+			await queue.PushAsync(4);
+			await task.WaitAsync(TimeSpan.FromSeconds(5.0));
+
+			value = await queue.TryPopAsync();
+			Assert.AreEqual(4, value);
+		}
 	}
 }
diff --git a/Engine/Source/Programs/Shared/EpicGames.Redis/RedisAsyncEvent.cs b/Engine/Source/Programs/Shared/EpicGames.Redis/RedisAsyncEvent.cs
deleted file mode 100644
diff --git a/Engine/Source/Programs/Shared/EpicGames.Redis/Utility/RedisEvent.cs b/Engine/Source/Programs/Shared/EpicGames.Redis/Utility/RedisEvent.cs
new file mode 100644
--- /dev/null
+++ b/Engine/Source/Programs/Shared/EpicGames.Redis/Utility/RedisEvent.cs
@@ -0,0 +1,58 @@
+// Copyright Epic Games, Inc. All Rights Reserved.
+
+using System;
+using System.Threading.Tasks;
+using EpicGames.Core;
+using StackExchange.Redis;
+
+namespace EpicGames.Redis
+{
+	/// <summary>
+	/// Distributed version of <see cref="AsyncEvent"/>. The event will be pulsed via a pub/sub channel in Redis.
+	/// </summary>
+	public sealed class RedisEvent : IAsyncDisposable
+	{
+		readonly IConnectionMultiplexer _multiplexer;
+		readonly RedisChannel _channel;
+		readonly AsyncEvent _asyncEvent;
+		readonly RedisSubscription _subscription;
+
+		/// <summary>
+		/// Accessor for the inner task. Can be captured and awaited by clients.
+		/// </summary>
+		public Task Task => _asyncEvent.Task;
+
+		RedisEvent(IConnectionMultiplexer multiplexer, RedisChannel channel, AsyncEvent asyncEvent, RedisSubscription subscription)
+		{
+			_multiplexer = multiplexer;
+			_channel = channel;
+			_asyncEvent = asyncEvent;
+			_subscription = subscription;
+		}
+
+		/// <inheritdoc/>
+		public ValueTask DisposeAsync()
+			=> _subscription.DisposeAsync();
+
+		/// <summary>
+		/// Create a new async event using the given channel name
+		/// </summary>
+		/// <param name="multiplexer">Multiplexer for the </param>
+		/// <param name="channel"></param>
+		/// <returns></returns>
+		public static async Task<RedisEvent> CreateAsync(IConnectionMultiplexer multiplexer, RedisChannel channel)
+		{
+			AsyncEvent asyncEvent = new AsyncEvent();
+			RedisSubscription subscription = await multiplexer.SubscribeAsync(channel, x => asyncEvent.Pulse());
+			return new RedisEvent(multiplexer, channel, asyncEvent, subscription);
+		}
+
+		/// <summary>
+		/// Pulse the event, allowing any captured copies of the task to continue.
+		/// </summary>
+		public void Pulse()
+		{
+			_ = _multiplexer.GetDatabase().PublishAsync(_channel, RedisValue.EmptyString, CommandFlags.FireAndForget);
+		}
+	}
+}
diff --git a/Engine/Source/Programs/Shared/EpicGames.Redis/Utility/RedisQueue.cs b/Engine/Source/Programs/Shared/EpicGames.Redis/Utility/RedisQueue.cs
new file mode 100644
--- /dev/null
+++ b/Engine/Source/Programs/Shared/EpicGames.Redis/Utility/RedisQueue.cs
@@ -0,0 +1,85 @@
+// Copyright Epic Games, Inc. All Rights Reserved.
+
+using System;
+using System.Threading;
+using System.Threading.Tasks;
+using StackExchange.Redis;
+
+namespace EpicGames.Redis.Utility
+{
+	/// <summary>
+	/// Utility method for queue instances
+	/// </summary>
+	public static class RedisQueue
+	{
+		/// <summary>
+		/// Create a new queue instance
+		/// </summary>
+		/// <param name="multiplexer">Multiplexer for the connection</param>
+		/// <param name="queueKey">Key for the queue</param>
+		/// <param name="eventChannel">Channel to use for posting update notifications</param>
+		public static async Task<RedisQueue<T>> CreateAsync<T>(IConnectionMultiplexer multiplexer, RedisKey queueKey, RedisChannel eventChannel)
+		{
+			RedisEvent asyncEvent = await RedisEvent.CreateAsync(multiplexer, eventChannel);
+			return new RedisQueue<T>(multiplexer, asyncEvent, new RedisListKey<T>(queueKey));
+		}
+	}
+
+	/// <summary>
+	/// Implements a waitable FIFO queue with Redis. Wraps a pub/sub channel and list, ensuring that each item is only popped from the queue once.
+	/// </summary>
+	/// <typeparam name="T">Type of item in the queue</typeparam>
+	public sealed class RedisQueue<T> : IAsyncDisposable
+	{
+		readonly IConnectionMultiplexer _multiplexer;
+		readonly RedisEvent _asyncEvent;
+		readonly RedisListKey<T> _listKey;
+
+		internal RedisQueue(IConnectionMultiplexer multiplexer, RedisEvent asyncEvent, RedisListKey<T> listKey)
+		{
+			_multiplexer = multiplexer;
+			_asyncEvent = asyncEvent;
+			_listKey = listKey;
+		}
+
+		/// <inheritdoc/>
+		public ValueTask DisposeAsync()
+			=> _asyncEvent.DisposeAsync();
+
+		/// <summary>
+		/// Push a new item onto the queue
+		/// </summary>
+		/// <param name="item">Item to add to the queue</param>
+		/// <param name="flags">Flags for the push operation</param>
+		public async Task PushAsync(T item, CommandFlags flags = CommandFlags.None)
+		{
+			IDatabase database = _multiplexer.GetDatabase();
+			await database.ListRightPushAsync(_listKey, item, flags: flags);
+			_asyncEvent.Pulse();
+		}
+
+		/// <summary>
+		/// Attempt to pop an item from the front of the queue. Returns the default value for the item if the queue is empty.
+		/// </summary>
+		public async Task<T> TryPopAsync()
+		{
+			IDatabase database = _multiplexer.GetDatabase();
+			return await database.ListLeftPopAsync(_listKey);
+		}
+
+		/// <summary>
+		/// Wait for a new item to be pushed onto the queue
+		/// </summary>
+		/// <param name="cancellationToken">Cancellation token for the operation</param>
+		public async Task WaitForDataAsync(CancellationToken cancellationToken = default)
+		{
+			Task task = _asyncEvent.Task;
+
+			IDatabase database = _multiplexer.GetDatabase();
+			if (await database.ListLengthAsync(_listKey) == 0)
+			{
+				await task.WaitAsync(cancellationToken);
+			}
+		}
+	}
+}
