Horde: Remove unnecessary using directive.

#rnx

diff --git a/Engine/Source/Programs/Shared/EpicGames.Core/AsyncQueue.cs b/Engine/Source/Programs/Shared/EpicGames.Core/AsyncQueue.cs
--- a/Engine/Source/Programs/Shared/EpicGames.Core/AsyncQueue.cs
+++ b/Engine/Source/Programs/Shared/EpicGames.Core/AsyncQueue.cs
@@ -5,7 +5,6 @@ using System.Collections.Generic;
 using System.Threading;
 using System.Threading.Channels;
 using System.Threading.Tasks;
-using System.Xml;
 
 namespace EpicGames.Core
 {
@@ -45,6 +44,8 @@ namespace EpicGames.Core
 		{
 			_cancellationSource.Cancel();
 			await StopAsync(CancellationToken.None).ConfigureAwait(false);
+
+			_cancellationSource.Dispose();
 		}
 
 		/// <summary>
