Horde: Retry HTTP requests in the event of an internal server error.

#rnx

diff --git a/Engine/Source/Programs/Shared/EpicGames.Horde/HordeClient.cs b/Engine/Source/Programs/Shared/EpicGames.Horde/HordeClient.cs
--- a/Engine/Source/Programs/Shared/EpicGames.Horde/HordeClient.cs
+++ b/Engine/Source/Programs/Shared/EpicGames.Horde/HordeClient.cs
@@ -137,7 +137,7 @@ namespace EpicGames.Horde
 			// when allocating compute resources when none are available). This pathway is handled explicitly on the application side.
 			return Policy<HttpResponseMessage>
 				.Handle<HttpRequestException>()
-				.OrResult(x => (x.StatusCode > HttpStatusCode.InternalServerError && x.StatusCode != HttpStatusCode.ServiceUnavailable) || x.StatusCode == HttpStatusCode.RequestTimeout)
+				.OrResult(x => (x.StatusCode >= HttpStatusCode.InternalServerError && x.StatusCode != HttpStatusCode.ServiceUnavailable) || x.StatusCode == HttpStatusCode.RequestTimeout)
 				.WaitAndRetryAsync(retryTimes, OnTimeoutAsync);
 		}
 
