Added the ability for UHT to support multiple packages per module (limited to one /Script/ package).  This change will likely require changes to any UHT exporter plugs.  The UE_5_5_OR_LATER define can be used to maintain backwards compatibility with older versions of UE.

1) UhtHeaderFile is no longer part of Outer chain for types.  This makes the Outer chain match the engine.
2) UhtPackages can no longer be enumerated from UhtSession.
3) UhtModule has been added which can now be enumerated from UhtSession.  UhtPackages and UhtHeaderFiles can be enumerated from each module.
4) UhtHeaderFile contains a complete list of all root level types found in the header.  Those types can also be enumerated from the individual UhtPackages found in the module owning the header file.

diff --git a/Engine/Plugins/ScriptPlugin/Source/ScriptGeneratorUbtPlugin/ScriptCodeGeneratorBase.cs b/Engine/Plugins/ScriptPlugin/Source/ScriptGeneratorUbtPlugin/ScriptCodeGeneratorBase.cs
--- a/Engine/Plugins/ScriptPlugin/Source/ScriptGeneratorUbtPlugin/ScriptCodeGeneratorBase.cs
+++ b/Engine/Plugins/ScriptPlugin/Source/ScriptGeneratorUbtPlugin/ScriptCodeGeneratorBase.cs
@@ -35,19 +35,19 @@ namespace ScriptGeneratorUbtPlugin
 			// Loop through the packages making sure they should be exported.  Queue the export of the classes
 			List<UhtClass> classes = new();
 			List<Task?> tasks = new();
-			foreach (UhtPackage package in Session.Packages)
+			foreach (UhtModule module in Session.Modules)
 			{
-				if (package.Module.ModuleType != UHTModuleType.EngineRuntime && package.Module.ModuleType != UHTModuleType.GameRuntime)
+				if (module.Module.ModuleType != UHTModuleType.EngineRuntime && module.Module.ModuleType != UHTModuleType.GameRuntime)
 				{
 					continue;
 				}
 
-				if (supportedScriptModules != null && !supportedScriptModules.Any(x => String.Compare(x, package.Module.Name, StringComparison.OrdinalIgnoreCase) == 0))
+				if (supportedScriptModules != null && !supportedScriptModules.Any(x => String.Compare(x, module.Module.Name, StringComparison.OrdinalIgnoreCase) == 0))
 				{
 					continue;
 				}
 
-				QueueClassExports(package, package, classes, tasks);
+				QueueClassExports(module.Package, module.Package, classes, tasks);
 			}
 
 			// Wait for all the classes to export
diff --git a/Engine/Source/Programs/Shared/EpicGames.UHT/Exporters/CodeGen/UhtCodeGenerator.cs b/Engine/Source/Programs/Shared/EpicGames.UHT/Exporters/CodeGen/UhtCodeGenerator.cs
--- a/Engine/Source/Programs/Shared/EpicGames.UHT/Exporters/CodeGen/UhtCodeGenerator.cs
+++ b/Engine/Source/Programs/Shared/EpicGames.UHT/Exporters/CodeGen/UhtCodeGenerator.cs
@@ -28,7 +28,6 @@ namespace EpicGames.UHT.Exporters.CodeGen
 		public struct PackageInfo
 		{
 			public string StrippedName { get; set; }
-			public string Api { get; set; }
 		}
 		public PackageInfo[] PackageInfos { get; set; }
 
@@ -75,18 +74,19 @@ namespace EpicGames.UHT.Exporters.CodeGen
 			FastArraySerializer = Session.FindType(null, UhtFindOptions.SourceName | UhtFindOptions.ScriptStruct, "FFastArraySerializer") as UhtScriptStruct;
 
 			// Perform some startup initialization to compute things we need over and over again
+			
 			if (Session.GoWide)
 			{
-				Parallel.ForEach(Factory.Session.Packages, package =>
+				Parallel.ForEach(Factory.Session.Modules, module =>
 				{
-					InitPackageInfo(package);
+					InitModuleInfo(module);
 				});
 			}
 			else
 			{
-				foreach (UhtPackage package in Factory.Session.Packages)
+				foreach (UhtModule module in Factory.Session.Modules)
 				{
-					InitPackageInfo(package);
+					InitModuleInfo(module);
 				}
 			}
 
@@ -95,9 +95,6 @@ namespace EpicGames.UHT.Exporters.CodeGen
 			{
 				if (headerFile.ShouldExport)
 				{
-					UhtPackage package = headerFile.Package;
-					UHTManifest.Module module = package.Module;
-
 					prereqs.Clear();
 					foreach (UhtHeaderFile referenced in headerFile.ReferencedHeadersNoLock)
 					{
@@ -110,66 +107,61 @@ namespace EpicGames.UHT.Exporters.CodeGen
 					HeaderInfos[headerFile.HeaderFileTypeIndex].Task = Factory.CreateTask(prereqs,
 						(IUhtExportFactory factory) =>
 						{
-							new UhtHeaderCodeGeneratorHFile(this, package, headerFile).Generate(factory);
-							new UhtHeaderCodeGeneratorCppFile(this, package, headerFile).Generate(factory);
+							new UhtHeaderCodeGeneratorHFile(this, headerFile).Generate(factory);
+							new UhtHeaderCodeGeneratorCppFile(this, headerFile).Generate(factory);
 						});
 				}
 			}
 
-			// Generate the files for the packages
-			List<Task?> generatedPackages = new(Session.PackageTypeCount);
-			foreach (UhtPackage package in Session.Packages)
+			// Generate the files for the modules
+			List<Task?> generatedModules = new(Session.Modules.Count);
+			foreach (UhtModule module in Session.Modules)
 			{
-				UHTManifest.Module module = package.Module;
-
 				bool writeHeader = false;
 				prereqs.Clear();
-				foreach (UhtType packageChild in package.Children)
+				foreach (UhtHeaderFile headerFile in module.Headers)
 				{
-					if (packageChild is UhtHeaderFile headerFile)
+					prereqs.Add(HeaderInfos[headerFile.HeaderFileTypeIndex].Task);
+					if (!writeHeader)
 					{
-						prereqs.Add(HeaderInfos[headerFile.HeaderFileTypeIndex].Task);
-						if (!writeHeader)
+						foreach (UhtType type in headerFile.Children)
 						{
-							foreach (UhtType type in headerFile.Children)
+							if (type is UhtClass classObj)
 							{
-								if (type is UhtClass classObj)
+								if (classObj.ClassType != UhtClassType.NativeInterface &&
+									classObj.ClassFlags.HasExactFlags(EClassFlags.Native | EClassFlags.Intrinsic, EClassFlags.Native) &&
+									!classObj.ClassExportFlags.HasAllFlags(UhtClassExportFlags.NoExport))
 								{
-									if (classObj.ClassType != UhtClassType.NativeInterface &&
-										classObj.ClassFlags.HasExactFlags(EClassFlags.Native | EClassFlags.Intrinsic, EClassFlags.Native) &&
-										!classObj.ClassExportFlags.HasAllFlags(UhtClassExportFlags.NoExport))
-									{
-										writeHeader = true;
-										break;
-									}
+									writeHeader = true;
+									break;
 								}
 							}
 						}
 					}
 				}
 
-				generatedPackages.Add(Factory.CreateTask(prereqs,
+				generatedModules.Add(Factory.CreateTask(prereqs,
 					(IUhtExportFactory factory) =>
 					{
-						List<UhtHeaderFile> packageSortedHeaders = GetSortedHeaderFiles(package);
+						List<UhtHeaderFile> moduleSortedHeaders = GetSortedHeaderFiles(module);
 						if (writeHeader)
 						{
-							new UhtPackageCodeGeneratorHFile(this, package).Generate(factory, packageSortedHeaders);
+							new UhtPackageCodeGeneratorHFile(this, module).Generate(factory, moduleSortedHeaders);
 						}
-						new UhtPackageCodeGeneratorCppFile(this, package).Generate(factory, packageSortedHeaders);
+						new UhtPackageCodeGeneratorCppFile(this, module).Generate(factory, moduleSortedHeaders);
 					}));
 			}
 
 			// Wait for all the packages to complete
-			List<Task> packageTasks = new(Session.PackageTypeCount);
-			foreach (Task? output in generatedPackages)
+			List<Task> moduleTasks = new(generatedModules.Count);
+			foreach (Task? output in generatedModules)
 			{
 				if (output != null)
 				{
-					packageTasks.Add(output);
+					moduleTasks.Add(output);
 				}
 			}
-			Task.WaitAll(packageTasks.ToArray());
+			Task.WaitAll(moduleTasks.ToArray());
 		}
 
 		#region Utility functions
@@ -212,35 +204,39 @@ namespace EpicGames.UHT.Exporters.CodeGen
 		#endregion
 
 		#region Information initialization
-		private void InitPackageInfo(UhtPackage package)
+		private void InitModuleInfo(UhtModule module)
 		{
 			StringBuilder builder = new();
 
+			foreach (UhtPackage package in module.Packages)
+			{
+				InitPackageInfo(builder, package);
+			}
+			foreach (UhtHeaderFile headerFile in module.Headers)
+			{
+				InitHeaderInfo(builder, headerFile);
+			}
+		}
+
+		private void InitPackageInfo(StringBuilder builder, UhtPackage package)
+		{
 			ref PackageInfo packageInfo = ref PackageInfos[package.PackageTypeIndex];
 			packageInfo.StrippedName = package.SourceName.Replace('/', '_');
-			packageInfo.Api = $"{package.ShortName.ToString().ToUpper()}_API ";
 
 			// Construct the names used commonly during export
 			ref ObjectInfo objectInfo = ref ObjectInfos[package.ObjectTypeIndex];
+			builder.Clear();
 			builder.Append("Z_Construct_UPackage_");
 			builder.Append(packageInfo.StrippedName);
 			objectInfo.UnregisteredSingletonName = objectInfo.RegisteredSingletonName = builder.ToString();
 			objectInfo.UnregisteredExternalDecl = objectInfo.RegsiteredExternalDecl = $"\tUPackage* {objectInfo.RegisteredSingletonName}();\r\n";
-
-			foreach (UhtType packageChild in package.Children)
-			{
-				if (packageChild is UhtHeaderFile headerFile)
-				{
-					InitHeaderInfo(builder, package, ref packageInfo, headerFile);
-				}
-			}
 		}
 
-		private void InitHeaderInfo(StringBuilder builder, UhtPackage package, ref PackageInfo packageInfo, UhtHeaderFile headerFile)
+		private void InitHeaderInfo(StringBuilder builder, UhtHeaderFile headerFile)
 		{
 			ref HeaderInfo headerInfo = ref HeaderInfos[headerFile.HeaderFileTypeIndex];
 
-			headerInfo.IncludePath = Path.GetRelativePath(package.Module.IncludeBase, headerFile.FilePath).Replace('\\', '/');
+			headerInfo.IncludePath = Path.GetRelativePath(headerFile.Module.Module.IncludeBase, headerFile.FilePath).Replace('\\', '/');
 
 			// Convert the file path to a C identifier
 			string filePath = headerFile.FilePath;
@@ -287,13 +283,19 @@ namespace EpicGames.UHT.Exporters.CodeGen
 			{
 				if (headerFileChild is UhtObject obj)
 				{
-					InitObjectInfo(builder, package, ref packageInfo, ref headerInfo, obj);
+					UhtPackage? package = obj.Outer as UhtPackage;
+					if (package == null)
+					{
+						throw new UhtIceException("Expected type defined in a header to have a package outer");
+					}
+					InitObjectInfo(builder, package, ref headerInfo, obj);
 				}
 			}
 		}
 
-		private void InitObjectInfo(StringBuilder builder, UhtPackage package, ref PackageInfo packageInfo, ref HeaderInfo headerInfo, UhtObject obj)
+		private void InitObjectInfo(StringBuilder builder, UhtPackage package, ref HeaderInfo headerInfo, UhtObject obj)
 		{
+			UhtModule module = package.Module;
 			ref ObjectInfo objectInfo = ref ObjectInfos[obj.ObjectTypeIndex];
 
 			builder.Clear();
@@ -362,13 +364,13 @@ namespace EpicGames.UHT.Exporters.CodeGen
 				builder.Append("_NoRegister");
 				objectInfo.UnregisteredSingletonName = builder.ToString();
 
-				objectInfo.UnregisteredExternalDecl = $"\t{packageInfo.Api}U{engineClassName}* {objectInfo.UnregisteredSingletonName}();\r\n";
-				objectInfo.RegsiteredExternalDecl = $"\t{packageInfo.Api}U{engineClassName}* {objectInfo.RegisteredSingletonName}();\r\n";
+				objectInfo.UnregisteredExternalDecl = $"\t{module.Api}U{engineClassName}* {objectInfo.UnregisteredSingletonName}();\r\n";
+				objectInfo.RegsiteredExternalDecl = $"\t{module.Api}U{engineClassName}* {objectInfo.RegisteredSingletonName}();\r\n";
 			}
 			else
 			{
 				objectInfo.UnregisteredSingletonName = objectInfo.RegisteredSingletonName = builder.ToString();
-				objectInfo.UnregisteredExternalDecl = objectInfo.RegsiteredExternalDecl = $"\t{packageInfo.Api}U{engineClassName}* {objectInfo.RegisteredSingletonName}();\r\n";
+				objectInfo.UnregisteredExternalDecl = objectInfo.RegsiteredExternalDecl = $"\t{module.Api}U{engineClassName}* {objectInfo.RegisteredSingletonName}();\r\n";
 			}
 
 			// Init the children
@@ -376,7 +378,7 @@ namespace EpicGames.UHT.Exporters.CodeGen
 			{
 				if (child is UhtObject childObject)
 				{
-					InitObjectInfo(builder, package, ref packageInfo, ref headerInfo, childObject);
+					InitObjectInfo(builder, package, ref headerInfo, childObject);
 				}
 			}
 		}
@@ -384,21 +386,18 @@ namespace EpicGames.UHT.Exporters.CodeGen
 
 		#region Utility functions
 		/// <summary>
-		/// Return a package's sorted header file list of all header files that or referenced or have declarations.
+		/// Return a module's sorted header file list of all header files that or referenced or have declarations.
 		/// </summary>
-		/// <param name="package">The package in question</param>
+		/// <param name="module">The module in question</param>
 		/// <returns>Sorted list of the header files</returns>
-		private static List<UhtHeaderFile> GetSortedHeaderFiles(UhtPackage package)
+		private static List<UhtHeaderFile> GetSortedHeaderFiles(UhtModule module)
 		{
-			List<UhtHeaderFile> sortedHeaders = new(package.Children.Count);
-			foreach (UhtType packageChild in package.Children)
+			List<UhtHeaderFile> sortedHeaders = new(module.Headers.Count);
+			foreach (UhtHeaderFile headerFile in module.Headers)
 			{
-				if (packageChild is UhtHeaderFile headerFile)
+				if (headerFile.ShouldExport)
 				{
-					if (headerFile.ShouldExport)
-					{
-						sortedHeaders.Add(headerFile);
-					}
+					sortedHeaders.Add(headerFile);
 				}
 			}
 			sortedHeaders.Sort((lhs, rhs) => { return StringComparerUE.OrdinalIgnoreCase.Compare(lhs.FilePath, rhs.FilePath); });
diff --git a/Engine/Source/Programs/Shared/EpicGames.UHT/Exporters/CodeGen/UhtHeaderCodeGenerator.cs b/Engine/Source/Programs/Shared/EpicGames.UHT/Exporters/CodeGen/UhtHeaderCodeGenerator.cs
--- a/Engine/Source/Programs/Shared/EpicGames.UHT/Exporters/CodeGen/UhtHeaderCodeGenerator.cs
+++ b/Engine/Source/Programs/Shared/EpicGames.UHT/Exporters/CodeGen/UhtHeaderCodeGenerator.cs
@@ -39,8 +39,8 @@ namespace EpicGames.UHT.Exporters.CodeGen
 		public readonly UhtHeaderFile HeaderFile;
 		public string FileId => HeaderInfos[HeaderFile.HeaderFileTypeIndex].FileId;
 
-		public UhtHeaderCodeGenerator(UhtCodeGenerator codeGenerator, UhtPackage package, UhtHeaderFile headerFile)
-			: base(codeGenerator, package)
+		public UhtHeaderCodeGenerator(UhtCodeGenerator codeGenerator, UhtHeaderFile headerFile)
+			: base(codeGenerator, headerFile.Module)
 		{
 			HeaderFile = headerFile;
 		}
@@ -129,9 +129,9 @@ namespace EpicGames.UHT.Exporters.CodeGen
 			{
 				outerName = classObj.EngineName;
 			}
-			else if (outer is UhtHeaderFile)
+			else if (outer is UhtPackage packageObj)
 			{
-				string packageName = outer.Package.EngineName;
+				string packageName = packageObj.EngineName;
 				outerName = packageName.Replace('/', '_');
 			}
 
@@ -164,7 +164,7 @@ namespace EpicGames.UHT.Exporters.CodeGen
 					!(outerClass != null && outerClass.ClassFlags.HasAnyFlags(EClassFlags.RequiredAPI)) &&
 					exportFlags.HasAnyFlags(UhtFunctionExportFlags.RequiredAPI))
 				{
-					builder.Append(PackageApi);
+					builder.Append(Module.Api);
 				}
 
 				if (textType == UhtPropertyTextType.InterfaceFunctionArgOrRetVal)
diff --git a/Engine/Source/Programs/Shared/EpicGames.UHT/Exporters/CodeGen/UhtHeaderCodeGeneratorCppFile.cs b/Engine/Source/Programs/Shared/EpicGames.UHT/Exporters/CodeGen/UhtHeaderCodeGeneratorCppFile.cs
--- a/Engine/Source/Programs/Shared/EpicGames.UHT/Exporters/CodeGen/UhtHeaderCodeGeneratorCppFile.cs
+++ b/Engine/Source/Programs/Shared/EpicGames.UHT/Exporters/CodeGen/UhtHeaderCodeGeneratorCppFile.cs
@@ -60,6 +60,16 @@ namespace EpicGames.UHT.Exporters.CodeGen
 		}
 	}
 
+	/// <summary>
+	/// Collection of all registrations for a specific package
+	/// </summary>
+	internal class UhtRegistrations
+	{
+		public UhtUsedDefineScopes<UhtEnum> Enumerations { get; } = new();
+		public UhtUsedDefineScopes<UhtScriptStruct> ScriptStructs { get; } = new();
+		public UhtUsedDefineScopes<UhtClass> Classes { get; } = new();
+	}
+
 	internal class UhtHeaderCodeGeneratorCppFile : UhtHeaderCodeGenerator
 	{
 
@@ -67,10 +77,9 @@ namespace EpicGames.UHT.Exporters.CodeGen
 		/// Construct an instance of this generator object
 		/// </summary>
 		/// <param name="codeGenerator">The base code generator</param>
-		/// <param name="package">Package being generated</param>
 		/// <param name="headerFile">Header file being generated</param>
-		public UhtHeaderCodeGeneratorCppFile(UhtCodeGenerator codeGenerator, UhtPackage package, UhtHeaderFile headerFile)
-			: base(codeGenerator, package, headerFile)
+		public UhtHeaderCodeGeneratorCppFile(UhtCodeGenerator codeGenerator, UhtHeaderFile headerFile)
+			: base(codeGenerator,  headerFile)
 		{
 		}
 
@@ -210,16 +219,14 @@ namespace EpicGames.UHT.Exporters.CodeGen
 
 				int generatedBodyStart = builder.Length;
 
-				UhtUsedDefineScopes<UhtEnum> enums = new();
-				UhtUsedDefineScopes<UhtScriptStruct> scriptStructs = new();
-				UhtUsedDefineScopes<UhtClass> classes = new();
+				SortedDictionary<UhtPackage, UhtRegistrations> packageRegistrations = new();
 				foreach (UhtField field in HeaderFile.References.ExportTypes)
 				{
 					if (field is UhtEnum enumObj)
 					{
 						using UhtCodeBlockComment blockComment = new(builder, field);
 						AppendEnum(builder, enumObj);
-						enums.Add(enumObj);
+						GetRegistrations(packageRegistrations, field).Enumerations.Add(enumObj);
 					}
 					else if (field is UhtScriptStruct scriptStruct)
 					{
@@ -227,7 +234,7 @@ namespace EpicGames.UHT.Exporters.CodeGen
 						AppendScriptStruct(builder, scriptStruct);
 						if (scriptStruct.ScriptStructFlags.HasAnyFlags(EStructFlags.Native))
 						{
-							scriptStructs.Add(scriptStruct);
+							GetRegistrations(packageRegistrations, field).ScriptStructs.Add(scriptStruct);
 						}
 					}
 					else if (field is UhtFunction function)
@@ -251,13 +258,18 @@ namespace EpicGames.UHT.Exporters.CodeGen
 
 							using UhtCodeBlockComment blockComment = new(builder, field);
 							AppendClass(builder, classObj, functions);
-							classes.Add(classObj);
+							GetRegistrations(packageRegistrations, field).Classes.Add(classObj);
 						}
 					}
 				}
 
-				if (!enums.IsEmpty || !scriptStructs.IsEmpty || !classes.IsEmpty)
+				foreach (UhtPackage package in Module.Packages)
 				{
+					if (!packageRegistrations.TryGetValue(package, out UhtRegistrations? registrations))
+					{
+						continue;
+					}
+
 					string name = $"Z_CompiledInDeferFile_{headerInfo.FileId}";
 					string staticsName = $"{name}_Statics";
 
@@ -267,7 +279,7 @@ namespace EpicGames.UHT.Exporters.CodeGen
 					builder.Append("struct ").Append(staticsName).Append("\r\n");
 					builder.Append("{\r\n");
 
-					builder.AppendInstances(enums, UhtDefineScopeNames.Standard, 
+					builder.AppendInstances(registrations.Enumerations, UhtDefineScopeNames.Standard, 
 						builder => builder.Append("\tstatic constexpr FEnumRegisterCompiledInInfo EnumInfo[] = {\r\n"),
 						(builder, enumObj) =>
 						{
@@ -284,7 +296,7 @@ namespace EpicGames.UHT.Exporters.CodeGen
 						},
 						builder => builder.Append("\t};\r\n"));
 
-					builder.AppendInstances(scriptStructs, UhtDefineScopeNames.Standard,
+					builder.AppendInstances(registrations.ScriptStructs, UhtDefineScopeNames.Standard,
 						builder => builder.Append("\tstatic constexpr FStructRegisterCompiledInInfo ScriptStructInfo[] = {\r\n"),
 						(builder, scriptStruct) =>
 						{
@@ -305,7 +317,7 @@ namespace EpicGames.UHT.Exporters.CodeGen
 						},
 						builder => builder.Append("\t};\r\n"));
 
-					builder.AppendInstances(classes, UhtDefineScopeNames.Standard,
+					builder.AppendInstances(registrations.Classes, UhtDefineScopeNames.Standard,
 						builder => builder.Append("\tstatic constexpr FClassRegisterCompiledInInfo ClassInfo[] = {\r\n"),
 						(builder, classObj) =>
 						{
@@ -332,12 +344,12 @@ namespace EpicGames.UHT.Exporters.CodeGen
 						.Append("static FRegisterCompiledInInfo ")
 						.Append(name)
 						.Append($"_{combinedHash}(TEXT(\"")
-						.Append(Package.EngineName)
+						.Append(package.EngineName)
 						.Append("\"),\r\n");
 
-					builder.AppendArrayPtrAndCountLine(classes, UhtDefineScopeNames.Standard, staticsName, "ClassInfo", 1, ",\r\n");
-					builder.AppendArrayPtrAndCountLine(scriptStructs, UhtDefineScopeNames.Standard, staticsName, "ScriptStructInfo", 1, ",\r\n");
-					builder.AppendArrayPtrAndCountLine(enums, UhtDefineScopeNames.Standard, staticsName, "EnumInfo", 1, ");\r\n");
+					builder.AppendArrayPtrAndCountLine(registrations.Classes, UhtDefineScopeNames.Standard, staticsName, "ClassInfo", 1, ",\r\n");
+					builder.AppendArrayPtrAndCountLine(registrations.ScriptStructs, UhtDefineScopeNames.Standard, staticsName, "ScriptStructInfo", 1, ",\r\n");
+					builder.AppendArrayPtrAndCountLine(registrations.Enumerations, UhtDefineScopeNames.Standard, staticsName, "EnumInfo", 1, ");\r\n");
 				}
 
 				if (Session.IncludeDebugOutput)
@@ -454,13 +466,13 @@ namespace EpicGames.UHT.Exporters.CodeGen
 				builder.Append("\tif (!").Append(registrationName).Append(".OuterSingleton)\r\n");
 				builder.Append("\t{\r\n");
 				builder.Append("\t\t").Append(registrationName).Append(".OuterSingleton = GetStaticEnum(").Append(singletonName).Append(", (UObject*)")
-					.Append(PackageSingletonName).Append("(), TEXT(\"").Append(enumObj.SourceName).Append("\"));\r\n");
+					.Append(GetSingletonName(enumObj.Package, true)).Append("(), TEXT(\"").Append(enumObj.SourceName).Append("\"));\r\n");
 				builder.Append("\t}\r\n");
 				builder.Append("\treturn ").Append(registrationName).Append(".OuterSingleton;\r\n");
 
 				builder.Append("}\r\n");
 
-				builder.Append("template<> ").Append(PackageApi).Append("UEnum* StaticEnum<").Append(enumObj.CppType).Append(">()\r\n");
+				builder.Append("template<> ").Append(Module.Api).Append("UEnum* StaticEnum<").Append(enumObj.CppType).Append(">()\r\n");
 				builder.Append("{\r\n");
 				builder.Append("\treturn ").Append(enumObj.SourceName).Append("_StaticEnum();\r\n");
 				builder.Append("}\r\n");
@@ -497,7 +509,7 @@ namespace EpicGames.UHT.Exporters.CodeGen
 				// Statics definition
 				{
 					builder.Append("const UECodeGen_Private::FEnumParams ").Append(staticsName).Append("::EnumParams = {\r\n");
-					builder.Append("\t(UObject*(*)())").Append(PackageSingletonName).Append(",\r\n");
+					builder.Append("\t(UObject*(*)())").Append(GetSingletonName(enumObj.Package, true)).Append(",\r\n");
 					builder.Append('\t').Append(enumDisplayNameFn).Append(",\r\n");
 					builder.Append('\t').AppendUTF8LiteralString(enumObj.SourceName).Append(",\r\n");
 					builder.Append('\t').AppendUTF8LiteralString(enumObj.CppType).Append(",\r\n");
@@ -572,7 +584,7 @@ namespace EpicGames.UHT.Exporters.CodeGen
 					.Append(".OuterSingleton = GetStaticStruct(")
 					.Append(singletonName)
 					.Append(", (UObject*)")
-					.Append(PackageSingletonName)
+					.Append(GetSingletonName(scriptStruct.Package, true))
 					.Append("(), TEXT(\"")
 					.Append(scriptStruct.EngineName)
 					.Append("\"));\r\n");
@@ -652,7 +664,7 @@ namespace EpicGames.UHT.Exporters.CodeGen
 				builder.Append("}\r\n");
 
 				// Generate the StaticStruct specialization
-				builder.Append("template<> ").Append(PackageApi).Append("UScriptStruct* StaticStruct<").Append(scriptStruct.SourceName).Append(">()\r\n");
+				builder.Append("template<> ").Append(Module.Api).Append("UScriptStruct* StaticStruct<").Append(scriptStruct.SourceName).Append(">()\r\n");
 				builder.Append("{\r\n");
 				builder.Append("\treturn ").Append(scriptStruct.SourceName).Append("::StaticStruct();\r\n");
 				builder.Append("}\r\n");
@@ -712,7 +724,7 @@ namespace EpicGames.UHT.Exporters.CodeGen
 				AppendPropertiesDefs(builder, context, properties, 0);
 
 				builder.Append("const UECodeGen_Private::FStructParams ").Append(staticsName).Append("::StructParams = {\r\n");
-				builder.Append("\t(UObject* (*)())").Append(PackageSingletonName).Append(",\r\n");
+				builder.Append("\t(UObject* (*)())").Append(GetSingletonName(scriptStruct.Package, true)).Append(",\r\n");
 				builder.Append('\t').Append(GetSingletonName(scriptStruct.SuperScriptStruct, true)).Append(",\r\n");
 				if (scriptStruct.ScriptStructFlags.HasAnyFlags(EStructFlags.Native))
 				{
@@ -1008,9 +1020,9 @@ namespace EpicGames.UHT.Exporters.CodeGen
 			{
 				return "nullptr";
 			}
-			else if (function.Outer is UhtHeaderFile headerFile)
+			else if (function.Outer is UhtPackage package)
 			{
-				return GetSingletonName(headerFile.Package, true);
+				return GetSingletonName(package, true);
 			}
 			else
 			{
@@ -1711,7 +1723,7 @@ namespace EpicGames.UHT.Exporters.CodeGen
 				{
 					builder.Append("\t(UObject* (*)())").Append(GetSingletonName(classObj.SuperClass, true)).Append(",\r\n");
 				}
-				builder.Append("\t(UObject* (*)())").Append(GetSingletonName(Package, true)).Append(",\r\n");
+				builder.Append("\t(UObject* (*)())").Append(GetSingletonName(classObj.Package, true)).Append(",\r\n");
 				builder.Append("};\r\n");
 				builder.Append("static_assert(UE_ARRAY_COUNT(").Append(staticsName).Append("::DependentSingletons) < 16);\r\n");
 
@@ -1847,7 +1859,7 @@ namespace EpicGames.UHT.Exporters.CodeGen
 				}
 			}
 
-			builder.Append("template<> ").Append(PackageApi).Append("UClass* StaticClass<").Append(classObj.SourceName).Append(">()\r\n");
+			builder.Append("template<> ").Append(Module.Api).Append("UClass* StaticClass<").Append(classObj.SourceName).Append(">()\r\n");
 			builder.Append("{\r\n");
 			builder.Append("\treturn ").Append(classObj.SourceName).Append("::StaticClass();\r\n");
 			builder.Append("}\r\n");
@@ -2015,6 +2027,18 @@ namespace EpicGames.UHT.Exporters.CodeGen
 			return outScriptStructs;
 		}
 
+		private UhtRegistrations GetRegistrations(SortedDictionary<UhtPackage, UhtRegistrations> packageRegistrations, UhtField fieldObj)
+		{
+			UhtPackage package = fieldObj.Package;
+			if (packageRegistrations.TryGetValue(package, out UhtRegistrations? registrations))
+			{
+				return registrations;
+			}
+			registrations = new();
+			packageRegistrations.Add(package, registrations);
+			return registrations;
+		}
+
 		private class PropertyMemberContextImpl : IUhtPropertyMemberContext
 		{
 			private readonly UhtCodeGenerator _codeGenerator;
diff --git a/Engine/Source/Programs/Shared/EpicGames.UHT/Exporters/CodeGen/UhtHeaderCodeGeneratorHFile.cs b/Engine/Source/Programs/Shared/EpicGames.UHT/Exporters/CodeGen/UhtHeaderCodeGeneratorHFile.cs
--- a/Engine/Source/Programs/Shared/EpicGames.UHT/Exporters/CodeGen/UhtHeaderCodeGeneratorHFile.cs
+++ b/Engine/Source/Programs/Shared/EpicGames.UHT/Exporters/CodeGen/UhtHeaderCodeGeneratorHFile.cs
@@ -21,10 +21,9 @@ namespace EpicGames.UHT.Exporters.CodeGen
 		/// Construct an instance of this generator object
 		/// </summary>
 		/// <param name="codeGenerator">The base code generator</param>
-		/// <param name="package">Package being generated</param>
 		/// <param name="headerFile">Header file being generated</param>
-		public UhtHeaderCodeGeneratorHFile(UhtCodeGenerator codeGenerator, UhtPackage package, UhtHeaderFile headerFile)
-			: base(codeGenerator, package, headerFile)
+		public UhtHeaderCodeGeneratorHFile(UhtCodeGenerator codeGenerator, UhtHeaderFile headerFile)
+			: base(codeGenerator, headerFile)
 		{
 		}
 
@@ -70,7 +69,7 @@ namespace EpicGames.UHT.Exporters.CodeGen
 				builder.Append(DisableDeprecationWarnings).Append("\r\n");
 
 				string strippedName = Path.GetFileNameWithoutExtension(HeaderFile.FilePath);
-				string defineName = $"{Package.ShortName.ToString().ToUpper()}_{strippedName}_generated_h".Replace('.', '_');
+				string defineName = $"{Module.ShortName.ToUpper()}_{strippedName}_generated_h".Replace('.', '_');
 
 				if (HeaderFile.References.ForwardDeclarations.Count > 0)
 				{
@@ -181,7 +180,7 @@ namespace EpicGames.UHT.Exporters.CodeGen
 					builder.Append('\t');
 					if (!scriptStruct.ScriptStructFlags.HasAnyFlags(EStructFlags.RequiredAPI))
 					{
-						builder.Append(PackageApi);
+						builder.Append(Module.Api);
 					}
 					builder.Append("static class UScriptStruct* StaticStruct(); \\\r\n");
 
@@ -411,14 +410,14 @@ namespace EpicGames.UHT.Exporters.CodeGen
 							.Append(", ");
 						if (!scriptStruct.ScriptStructFlags.HasAnyFlags(EStructFlags.RequiredAPI))
 						{
-							builder.Append(PackageApi);
+							builder.Append(Module.Api);
 						}
 						builder.Append("); \\\r\n");
 					}
 				}
 
 				// Forward declare the StaticStruct specialization in the header
-				builder.Append("template<> ").Append(PackageApi).Append("UScriptStruct* StaticStruct<struct ").Append(scriptStruct.SourceName).Append(">();\r\n");
+				builder.Append("template<> ").Append(Module.Api).Append("UScriptStruct* StaticStruct<struct ").Append(scriptStruct.SourceName).Append(">();\r\n");
 				builder.Append("\r\n");
 			}
 			return builder;
@@ -471,7 +470,7 @@ namespace EpicGames.UHT.Exporters.CodeGen
 			if (enumObj.CppForm == UhtEnumCppForm.EnumClass || enumObj.UnderlyingType != UhtEnumUnderlyingType.Unspecified)
 			{	
 				// Forward declare the StaticEnum<> specialization for enum classes
-				builder.Append("template<> ").Append(PackageApi).Append("UEnum* StaticEnum<").Append(enumObj.CppType).Append(">();\r\n");
+				builder.Append("template<> ").Append(Module.Api).Append("UEnum* StaticEnum<").Append(enumObj.CppType).Append(">();\r\n");
 				builder.Append("\r\n");
 			}
 
@@ -504,7 +503,7 @@ namespace EpicGames.UHT.Exporters.CodeGen
 
 				if (addAPI)
 				{
-					builder.Append(PackageApi);
+					builder.Append(Module.Api);
 				}
 
 				AppendNativeFunctionHeader(builder, function, UhtPropertyTextType.EventFunctionArgOrRetVal, true, exportFunctionName, extraParameter, UhtFunctionExportFlags.None, 0, "; \\\r\n");
@@ -514,7 +513,7 @@ namespace EpicGames.UHT.Exporters.CodeGen
 
 		private StringBuilder AppendClass(StringBuilder builder, UhtClass classObj)
 		{
-			string api = classObj.ClassFlags.HasAnyFlags(EClassFlags.MinimalAPI) ? PackageApi : "NO_API ";
+			string api = classObj.ClassFlags.HasAnyFlags(EClassFlags.MinimalAPI) ? Module.Api : "NO_API ";
 			bool usesLegacy = classObj.ClassExportFlags.HasAnyFlags(UhtClassExportFlags.UsesGeneratedBodyLegacy);
 			UhtClass? nativeInterface = ObjectInfos[classObj.ObjectTypeIndex].NativeInterface;
 			if (classObj.ClassFlags.HasAnyFlags(EClassFlags.Interface) && nativeInterface == null)
@@ -640,7 +639,7 @@ namespace EpicGames.UHT.Exporters.CodeGen
 			}
 
 			// Forward declare the StaticClass specialization in the header
-			builder.Append("template<> ").Append(PackageApi).Append("UClass* StaticClass<class ").Append(classObj.SourceName).Append(">();\r\n");
+			builder.Append("template<> ").Append(Module.Api).Append("UClass* StaticClass<class ").Append(classObj.SourceName).Append(">();\r\n");
 			builder.Append("\r\n");
 			return builder;
 		}
@@ -668,7 +667,7 @@ namespace EpicGames.UHT.Exporters.CodeGen
 			return builder.AppendSingleMacro(notifyTypes, UhtDefineScopeNames.Standard, this, classObj, FieldNotifyMacroSuffix,
 				(builder, filteredTypes) =>
 				{
-					builder.Append("\tUE_FIELD_NOTIFICATION_DECLARE_CLASS_DESCRIPTOR_BEGIN(").Append(PackageApi).Append(") \\\r\n");
+					builder.Append("\tUE_FIELD_NOTIFICATION_DECLARE_CLASS_DESCRIPTOR_BEGIN(").Append(Module.Api).Append(") \\\r\n");
 
 					// UE_FIELD_NOTIFICATION_DECLARE_FIELD
 					foreach (UhtType notifyType in filteredTypes)
@@ -774,7 +773,7 @@ namespace EpicGames.UHT.Exporters.CodeGen
 		{
 			using (UhtMacroCreator macro = new(builder, this, classObj, SparseDataMacroSuffix))
 			{
-				string api = classObj.ClassFlags.HasAnyFlags(EClassFlags.MinimalAPI) ? PackageApi : "";
+				string api = classObj.ClassFlags.HasAnyFlags(EClassFlags.MinimalAPI) ? Module.Api : "";
 
 				foreach (UhtScriptStruct sparseScriptStruct in sparseScriptStructs)
 				{
@@ -873,7 +872,7 @@ namespace EpicGames.UHT.Exporters.CodeGen
 
 				if (!classObj.ClassFlags.HasAnyFlags(EClassFlags.RequiredAPI) && function.FunctionFlags.HasAnyFlags(EFunctionFlags.RequiredAPI))
 				{
-					builder.Append(PackageApi);
+					builder.Append(Module.Api);
 				}
 
 				if (!function.FunctionFlags.HasAnyFlags(EFunctionFlags.Static) && !function.FunctionExportFlags.HasAnyFlags(UhtFunctionExportFlags.Final))
@@ -1182,7 +1181,7 @@ namespace EpicGames.UHT.Exporters.CodeGen
 			{
 				builder.Append("CASTCLASS_None");
 			}
-			builder.Append(", TEXT(\"").Append(Package.SourceName).Append("\"), ").Append(api[0..^1]).Append(") \\\r\n");
+			builder.Append(", TEXT(\"").Append(classObj.Package.SourceName).Append("\"), ").Append(api[0..^1]).Append(") \\\r\n");
 
 			builder.Append("\tDECLARE_SERIALIZER(").Append(classObj.SourceName).Append(") \\\r\n");
 
diff --git a/Engine/Source/Programs/Shared/EpicGames.UHT/Exporters/CodeGen/UhtPackageCodeGenerator.cs b/Engine/Source/Programs/Shared/EpicGames.UHT/Exporters/CodeGen/UhtPackageCodeGenerator.cs
--- a/Engine/Source/Programs/Shared/EpicGames.UHT/Exporters/CodeGen/UhtPackageCodeGenerator.cs
+++ b/Engine/Source/Programs/Shared/EpicGames.UHT/Exporters/CodeGen/UhtPackageCodeGenerator.cs
@@ -23,20 +23,18 @@ namespace EpicGames.UHT.Exporters.CodeGen
 		public static string DisableDeprecationWarnings = "PRAGMA_DISABLE_DEPRECATION_WARNINGS";
 
 		public readonly UhtCodeGenerator CodeGenerator;
-		public readonly UhtPackage Package;
-		public bool SaveExportedHeaders => Package.Module.SaveExportedHeaders;
+		public readonly UhtModule Module;
+		public bool SaveExportedHeaders => Module.Module.SaveExportedHeaders;
 
 		public Utils.UhtSession Session => CodeGenerator.Session;
 		public UhtCodeGenerator.PackageInfo[] PackageInfos => CodeGenerator.PackageInfos;
 		public UhtCodeGenerator.HeaderInfo[] HeaderInfos => CodeGenerator.HeaderInfos;
 		public UhtCodeGenerator.ObjectInfo[] ObjectInfos => CodeGenerator.ObjectInfos;
-		public string PackageApi => PackageInfos[Package.PackageTypeIndex].Api;
-		public string PackageSingletonName => ObjectInfos[Package.ObjectTypeIndex].RegisteredSingletonName;
 
-		public UhtPackageCodeGenerator(UhtCodeGenerator codeGenerator, UhtPackage package)
+		public UhtPackageCodeGenerator(UhtCodeGenerator codeGenerator, UhtModule module)
 		{
 			CodeGenerator = codeGenerator;
-			Package = package;
+			Module = module;
 		}
 
 		#region Utility functions
diff --git a/Engine/Source/Programs/Shared/EpicGames.UHT/Exporters/CodeGen/UhtPackageCodeGeneratorCppFile.cs b/Engine/Source/Programs/Shared/EpicGames.UHT/Exporters/CodeGen/UhtPackageCodeGeneratorCppFile.cs
--- a/Engine/Source/Programs/Shared/EpicGames.UHT/Exporters/CodeGen/UhtPackageCodeGeneratorCppFile.cs
+++ b/Engine/Source/Programs/Shared/EpicGames.UHT/Exporters/CodeGen/UhtPackageCodeGeneratorCppFile.cs
@@ -2,8 +2,10 @@
 
 using System;
 using System.Collections.Generic;
+using System.Linq;
 using System.Text;
 using EpicGames.Core;
+using EpicGames.UHT.Tables;
 using EpicGames.UHT.Types;
 using EpicGames.UHT.Utils;
 
@@ -15,9 +17,9 @@ namespace EpicGames.UHT.Exporters.CodeGen
 		/// Construct an instance of this generator object
 		/// </summary>
 		/// <param name="codeGenerator">The base code generator</param>
-		/// <param name="package">Package being generated</param>
-		public UhtPackageCodeGeneratorCppFile(UhtCodeGenerator codeGenerator, UhtPackage package)
-			: base(codeGenerator, package)
+		/// <param name="module">Module being generated</param>
+		public UhtPackageCodeGeneratorCppFile(UhtCodeGenerator codeGenerator, UhtModule module)
+			: base(codeGenerator, module)
 		{
 		}
 
@@ -67,14 +69,12 @@ namespace EpicGames.UHT.Exporters.CodeGen
 					if (SaveExportedHeaders)
 					{
 						// We need to create the directory, otherwise UBT will think that this module has not been properly updated and won't write a Timestamp file
-						System.IO.Directory.CreateDirectory(Package.Module.OutputDirectory);
+						System.IO.Directory.CreateDirectory(Module.Module.OutputDirectory);
 					}
 					return;
 				}
 				uint declarationsHash = UhtHash.GenenerateTextHash(declarations.ToString());
 
-				string strippedName = PackageInfos[Package.PackageTypeIndex].StrippedName;
-				string singletonName = PackageSingletonName;
 				singletons.Sort((UhtField lhs, UhtField rhs) =>
 				{
 					bool lhsIsDel = IsDelegateFunction(lhs);
@@ -91,7 +91,7 @@ namespace EpicGames.UHT.Exporters.CodeGen
 				builder.Append(HeaderCopyright);
 				builder.Append(RequiredCPPIncludes);
 				builder.Append(DisableDeprecationWarnings).Append("\r\n");
-				builder.Append("void EmptyLinkFunctionForGeneratedCode").Append(Package.ShortName).Append("_init() {}\r\n");
+				builder.Append("void EmptyLinkFunctionForGeneratedCode").Append(Module.ShortName).Append("_init() {}\r\n");
 
 				if (Session.IncludeDebugOutput)
 				{
@@ -110,47 +110,53 @@ namespace EpicGames.UHT.Exporters.CodeGen
 					builder.Append(info.RegsiteredExternalDecl);
 				}
 
-				builder.Append("\tstatic FPackageRegistrationInfo Z_Registration_Info_UPackage_").Append(strippedName).Append(";\r\n");
-				builder.Append("\tFORCENOINLINE UPackage* ").Append(singletonName).Append("()\r\n");
-				builder.Append("\t{\r\n");
-				builder.Append("\t\tif (!Z_Registration_Info_UPackage_").Append(strippedName).Append(".OuterSingleton)\r\n");
-				builder.Append("\t\t{\r\n");
-
-				if (singletons.Count != 0)
+				foreach (UhtPackage package in Module.Packages)
 				{
-					builder.Append("\t\t\tstatic UObject* (*const SingletonFuncArray[])() = {\r\n");
-					foreach (UhtField field in singletons)
+					string strippedName = PackageInfos[package.PackageTypeIndex].StrippedName;
+					string singletonName = GetSingletonName(package, true);
+					builder.Append("\tstatic FPackageRegistrationInfo Z_Registration_Info_UPackage_").Append(strippedName).Append(";\r\n");
+					builder.Append("\tFORCENOINLINE UPackage* ").Append(singletonName).Append("()\r\n");
+					builder.Append("\t{\r\n");
+					builder.Append("\t\tif (!Z_Registration_Info_UPackage_").Append(strippedName).Append(".OuterSingleton)\r\n");
+					builder.Append("\t\t{\r\n");
+
+					UhtField[] packageSingletons = singletons.Where(x => x.Package == package).ToArray();
+
+					if (packageSingletons.Length != 0)
 					{
-						builder.Append("\t\t\t\t(UObject* (*)())").Append(ObjectInfos[field.ObjectTypeIndex].RegisteredSingletonName).Append(",\r\n");
+						builder.Append("\t\t\tstatic UObject* (*const SingletonFuncArray[])() = {\r\n");
+						foreach (UhtField field in packageSingletons)
+						{
+							builder.Append("\t\t\t\t(UObject* (*)())").Append(ObjectInfos[field.ObjectTypeIndex].RegisteredSingletonName).Append(",\r\n");
+						}
+						builder.Append("\t\t\t};\r\n");
 					}
+
+					EPackageFlags flags = package.PackageFlags & (EPackageFlags.ClientOptional | EPackageFlags.ServerSideOnly | EPackageFlags.EditorOnly | EPackageFlags.Developer | EPackageFlags.UncookedOnly);
+					builder.Append("\t\t\tstatic const UECodeGen_Private::FPackageParams PackageParams = {\r\n");
+					builder.Append("\t\t\t\t").AppendUTF8LiteralString(package.SourceName).Append(",\r\n");
+					builder.Append("\t\t\t\t").Append(packageSingletons.Length != 0 ? "SingletonFuncArray" : "nullptr").Append(",\r\n");
+					builder.Append("\t\t\t\t").Append(packageSingletons.Length != 0 ? "UE_ARRAY_COUNT(SingletonFuncArray)" : "0").Append(",\r\n");
+					builder.Append("\t\t\t\tPKG_CompiledIn | ").Append($"0x{(uint)flags:X8}").Append(",\r\n");
+					builder.Append("\t\t\t\t").Append($"0x{bodiesHash:X8}").Append(",\r\n");
+					builder.Append("\t\t\t\t").Append($"0x{declarationsHash:X8}").Append(",\r\n");
+					builder.Append("\t\t\t\t").AppendMetaDataParams(package, null, MetaDataParamsName).Append("\r\n");
 					builder.Append("\t\t\t};\r\n");
+					builder.Append("\t\t\tUECodeGen_Private::ConstructUPackage(Z_Registration_Info_UPackage_").Append(strippedName).Append(".OuterSingleton, PackageParams);\r\n");
+					builder.Append("\t\t}\r\n");
+					builder.Append("\t\treturn Z_Registration_Info_UPackage_").Append(strippedName).Append(".OuterSingleton;\r\n");
+					builder.Append("\t}\r\n");
+
+					// Do not change the Z_CompiledInDeferPackage_UPackage_ without changing LC_SymbolPatterns
+					builder.Append("\tstatic FRegisterCompiledInInfo Z_CompiledInDeferPackage_UPackage_").Append(strippedName).Append('(').Append(singletonName)
+						.Append(", TEXT(\"").Append(package.SourceName).Append("\"), Z_Registration_Info_UPackage_").Append(strippedName).Append(", CONSTRUCT_RELOAD_VERSION_INFO(FPackageReloadVersionInfo, ")
+						.Append($"0x{bodiesHash:X8}, 0x{declarationsHash:X8}").Append("));\r\n");
 				}
-
-				EPackageFlags flags = Package.PackageFlags & (EPackageFlags.ClientOptional | EPackageFlags.ServerSideOnly | EPackageFlags.EditorOnly | EPackageFlags.Developer | EPackageFlags.UncookedOnly);
-				builder.Append("\t\t\tstatic const UECodeGen_Private::FPackageParams PackageParams = {\r\n");
-				builder.Append("\t\t\t\t").AppendUTF8LiteralString(Package.SourceName).Append(",\r\n");
-				builder.Append("\t\t\t\t").Append(singletons.Count != 0 ? "SingletonFuncArray" : "nullptr").Append(",\r\n");
-				builder.Append("\t\t\t\t").Append(singletons.Count != 0 ? "UE_ARRAY_COUNT(SingletonFuncArray)" : "0").Append(",\r\n");
-				builder.Append("\t\t\t\tPKG_CompiledIn | ").Append($"0x{(uint)flags:X8}").Append(",\r\n");
-				builder.Append("\t\t\t\t").Append($"0x{bodiesHash:X8}").Append(",\r\n");
-				builder.Append("\t\t\t\t").Append($"0x{declarationsHash:X8}").Append(",\r\n");
-				builder.Append("\t\t\t\t").AppendMetaDataParams(Package, null, MetaDataParamsName).Append("\r\n");
-				builder.Append("\t\t\t};\r\n");
-				builder.Append("\t\t\tUECodeGen_Private::ConstructUPackage(Z_Registration_Info_UPackage_").Append(strippedName).Append(".OuterSingleton, PackageParams);\r\n");
-				builder.Append("\t\t}\r\n");
-				builder.Append("\t\treturn Z_Registration_Info_UPackage_").Append(strippedName).Append(".OuterSingleton;\r\n");
-				builder.Append("\t}\r\n");
-
-				// Do not change the Z_CompiledInDeferPackage_UPackage_ without changing LC_SymbolPatterns
-				builder.Append("\tstatic FRegisterCompiledInInfo Z_CompiledInDeferPackage_UPackage_").Append(strippedName).Append('(').Append(singletonName)
-					.Append(", TEXT(\"").Append(Package.SourceName).Append("\"), Z_Registration_Info_UPackage_").Append(strippedName).Append(", CONSTRUCT_RELOAD_VERSION_INFO(FPackageReloadVersionInfo, ")
-					.Append($"0x{bodiesHash:X8}, 0x{declarationsHash:X8}").Append("));\r\n");
-
 				builder.Append(EnableDeprecationWarnings).Append("\r\n");
 
 				if (SaveExportedHeaders)
 				{
-					string cppFilePath = factory.MakePath(Package, ".init.gen.cpp");
+					string cppFilePath = factory.MakePath(Module, ".init.gen.cpp");
 					factory.CommitOutput(cppFilePath, builder);
 				}
 			}
diff --git a/Engine/Source/Programs/Shared/EpicGames.UHT/Exporters/CodeGen/UhtPackageCodeGeneratorHFile.cs b/Engine/Source/Programs/Shared/EpicGames.UHT/Exporters/CodeGen/UhtPackageCodeGeneratorHFile.cs
--- a/Engine/Source/Programs/Shared/EpicGames.UHT/Exporters/CodeGen/UhtPackageCodeGeneratorHFile.cs
+++ b/Engine/Source/Programs/Shared/EpicGames.UHT/Exporters/CodeGen/UhtPackageCodeGeneratorHFile.cs
@@ -15,9 +15,9 @@ namespace EpicGames.UHT.Exporters.CodeGen
 		/// Construct an instance of this generator object
 		/// </summary>
 		/// <param name="codeGenerator">The base code generator</param>
-		/// <param name="package">Package being generated</param>
-		public UhtPackageCodeGeneratorHFile(UhtCodeGenerator codeGenerator, UhtPackage package)
-			: base(codeGenerator, package)
+		/// <param name="module">Module being generated</param>
+		public UhtPackageCodeGeneratorHFile(UhtCodeGenerator codeGenerator, UhtModule module)
+			: base(codeGenerator, module)
 		{
 		}
 
@@ -25,8 +25,8 @@ namespace EpicGames.UHT.Exporters.CodeGen
 		/// For a given UE header file, generated the generated H file
 		/// </summary>
 		/// <param name="factory">Requesting factory</param>
-		/// <param name="packageSortedHeaders">Sorted list of headers by name of all headers in the package</param>
-		public void Generate(IUhtExportFactory factory, List<UhtHeaderFile> packageSortedHeaders)
+		/// <param name="moduleSortedHeaders">Sorted list of headers by name of all headers in the module</param>
+		public void Generate(IUhtExportFactory factory, List<UhtHeaderFile> moduleSortedHeaders)
 		{
 			{
 				using BorrowStringBuilder borrower = new(StringBuilderCache.Big);
@@ -37,17 +37,14 @@ namespace EpicGames.UHT.Exporters.CodeGen
 				builder.Append("\r\n");
 				builder.Append("\r\n");
 
-				List<UhtHeaderFile> headerFiles = new(Package.Children.Count * 2);
-				headerFiles.AddRange(packageSortedHeaders);
+				List<UhtHeaderFile> headerFiles = new(moduleSortedHeaders.Count + Module.Headers.Count);
+				headerFiles.AddRange(moduleSortedHeaders);
 
-				foreach (UhtType type in Package.Children)
+				foreach (UhtHeaderFile headerFile in Module.Headers)
 				{
-					if (type is UhtHeaderFile headerFile)
+					if (headerFile.HeaderFileType == UhtHeaderFileType.Classes)
 					{
-						if (headerFile.HeaderFileType == UhtHeaderFileType.Classes)
-						{
-							headerFiles.Add(headerFile);
-						}
+						headerFiles.Add(headerFile);
 					}
 				}
 
@@ -66,7 +63,7 @@ namespace EpicGames.UHT.Exporters.CodeGen
 
 				if (SaveExportedHeaders)
 				{
-					string headerFilePath = factory.MakePath(Package, "Classes.h");
+					string headerFilePath = factory.MakePath(Module, "Classes.h");
 					factory.CommitOutput(headerFilePath, builder);
 				}
 			}
diff --git a/Engine/Source/Programs/Shared/EpicGames.UHT/Exporters/Json/UhtJsonExporter.cs b/Engine/Source/Programs/Shared/EpicGames.UHT/Exporters/Json/UhtJsonExporter.cs
--- a/Engine/Source/Programs/Shared/EpicGames.UHT/Exporters/Json/UhtJsonExporter.cs
+++ b/Engine/Source/Programs/Shared/EpicGames.UHT/Exporters/Json/UhtJsonExporter.cs
@@ -33,15 +33,14 @@ namespace EpicGames.UHT.Exporters.Json
 		{
 			// Generate the files for the packages
 			List<Task?> generatedPackages = new(Session.PackageTypeCount);
-			foreach (UhtPackage package in Session.Packages)
+			foreach (UhtModule module in Session.Modules)
 			{
-				UHTManifest.Module module = package.Module;
 				generatedPackages.Add(Factory.CreateTask(
 					(IUhtExportFactory factory) =>
 					{
-						string jsonPath = factory.MakePath(package, ".json");
+						string jsonPath = factory.MakePath(module, ".json");
 						JsonSerializerOptions options = new() { WriteIndented = true, DefaultIgnoreCondition = JsonIgnoreCondition.WhenWritingNull };
-						factory.CommitOutput(jsonPath, JsonSerializer.Serialize(package, options));
+						factory.CommitOutput(jsonPath, JsonSerializer.Serialize(module, options));
 					}));
 			}
 
diff --git a/Engine/Source/Programs/Shared/EpicGames.UHT/Exporters/Stats/UhtStatsExporter.cs b/Engine/Source/Programs/Shared/EpicGames.UHT/Exporters/Stats/UhtStatsExporter.cs
--- a/Engine/Source/Programs/Shared/EpicGames.UHT/Exporters/Stats/UhtStatsExporter.cs
+++ b/Engine/Source/Programs/Shared/EpicGames.UHT/Exporters/Stats/UhtStatsExporter.cs
@@ -17,9 +17,12 @@ namespace EpicGames.UHT.Exporters.Stats
 		private static void StatsExporter(IUhtExportFactory factory)
 		{
 			SortedDictionary<string, int> countByType = new();
-			foreach (UhtType type in factory.Session.Packages)
+			foreach (UhtModule module in factory.Session.Modules)
 			{
-				Collect(countByType, type);
+				foreach (UhtPackage package in module.Packages)
+				{
+					Collect(countByType, package);
+				}
 			}
 
 			ILogger logger = factory.Session.Logger;
diff --git a/Engine/Source/Programs/Shared/EpicGames.UHT/Parsers/UhtClassParser.cs b/Engine/Source/Programs/Shared/EpicGames.UHT/Parsers/UhtClassParser.cs
--- a/Engine/Source/Programs/Shared/EpicGames.UHT/Parsers/UhtClassParser.cs
+++ b/Engine/Source/Programs/Shared/EpicGames.UHT/Parsers/UhtClassParser.cs
@@ -57,7 +57,7 @@ namespace EpicGames.UHT.Parsers
 
 		private static UhtParseResult ParseUClass(UhtParsingScope parentScope, ref UhtToken token)
 		{
-			UhtClass classObj = new(parentScope.ScopeType, token.InputLine);
+			UhtClass classObj = new(parentScope.HeaderFile, parentScope.ScopeType, token.InputLine);
 			{
 				using UhtParsingScope topScope = new(parentScope, classObj, parentScope.Session.GetKeywordTable(UhtTableNames.Class), UhtAccessSpecifier.Private);
 				const string ScopeName = "class";
diff --git a/Engine/Source/Programs/Shared/EpicGames.UHT/Parsers/UhtEnumParser.cs b/Engine/Source/Programs/Shared/EpicGames.UHT/Parsers/UhtEnumParser.cs
--- a/Engine/Source/Programs/Shared/EpicGames.UHT/Parsers/UhtEnumParser.cs
+++ b/Engine/Source/Programs/Shared/EpicGames.UHT/Parsers/UhtEnumParser.cs
@@ -27,7 +27,7 @@ namespace EpicGames.UHT.Parsers
 
 		private static UhtParseResult ParseUEnum(UhtParsingScope parentScope, UhtToken keywordToken)
 		{
-			UhtEnum enumObject = new(parentScope.ScopeType, keywordToken.InputLine);
+			UhtEnum enumObject = new(parentScope.HeaderFile, parentScope.ScopeType, keywordToken.InputLine);
 			{
 				using UhtParsingScope topScope = new(parentScope, enumObject, parentScope.Session.GetKeywordTable(UhtTableNames.Enum), UhtAccessSpecifier.Public);
 				const string ScopeName = "UENUM";
@@ -120,7 +120,7 @@ namespace EpicGames.UHT.Parsers
 
 					if (enumObject.CppForm != UhtEnumCppForm.EnumClass && enumObject.UnderlyingType == UhtEnumUnderlyingType.Unspecified)
 					{
-						UhtIssueBehavior enumUnderlyingTypeBehavior = enumObject.Package.IsPartOfEngine ? topScope.Session.Config!.EngineEnumUnderlyingTypeNotSet
+						UhtIssueBehavior enumUnderlyingTypeBehavior = enumObject.Module.IsPartOfEngine ? topScope.Session.Config!.EngineEnumUnderlyingTypeNotSet
 							: topScope.Session.Config!.NonEngineEnumUnderlyingTypeNotSet;
 
 						string logMessage = $"Underlying type must be specified.";
diff --git a/Engine/Source/Programs/Shared/EpicGames.UHT/Parsers/UhtFunctionParser.cs b/Engine/Source/Programs/Shared/EpicGames.UHT/Parsers/UhtFunctionParser.cs
--- a/Engine/Source/Programs/Shared/EpicGames.UHT/Parsers/UhtFunctionParser.cs
+++ b/Engine/Source/Programs/Shared/EpicGames.UHT/Parsers/UhtFunctionParser.cs
@@ -129,7 +129,7 @@ namespace EpicGames.UHT.Parsers
 
 		private static UhtParseResult ParseUDelegate(UhtParsingScope parentScope, UhtToken token, bool hasSpecifiers)
 		{
-			UhtFunction function = new(parentScope.ScopeType, token.InputLine);
+			UhtFunction function = new(parentScope.HeaderFile, parentScope.ScopeType, token.InputLine);
 
 			{
 				using UhtParsingScope topScope = new(parentScope, function, parentScope.Session.GetKeywordTable(UhtTableNames.Function), UhtAccessSpecifier.Public);
@@ -299,7 +299,7 @@ namespace EpicGames.UHT.Parsers
 
 		private static UhtParseResult ParseUFunction(UhtParsingScope parentScope, UhtToken token)
 		{
-			UhtFunction function = new(parentScope.ScopeType, token.InputLine);
+			UhtFunction function = new(parentScope.HeaderFile, parentScope.ScopeType, token.InputLine);
 
 			{
 				using UhtParsingScope topScope = new(parentScope, function, parentScope.Session.GetKeywordTable(UhtTableNames.Function), UhtAccessSpecifier.Public);
diff --git a/Engine/Source/Programs/Shared/EpicGames.UHT/Parsers/UhtHeaderFileParser.cs b/Engine/Source/Programs/Shared/EpicGames.UHT/Parsers/UhtHeaderFileParser.cs
--- a/Engine/Source/Programs/Shared/EpicGames.UHT/Parsers/UhtHeaderFileParser.cs
+++ b/Engine/Source/Programs/Shared/EpicGames.UHT/Parsers/UhtHeaderFileParser.cs
@@ -300,14 +300,14 @@ namespace EpicGames.UHT.Parsers
 		public UhtHeaderFile HeaderFile { get; }
 
 		/// <summary>
-		/// Token reader for the header
+		/// Module containing the header file
 		/// </summary>
-		public IUhtTokenReader TokenReader { get; }
+		public UhtModule Module => HeaderFile.Module;
 
 		/// <summary>
-		/// If true, this header file belongs to the engine
+		/// Token reader for the header
 		/// </summary>
-		public bool IsPartOfEngine => HeaderFile.Package.IsPartOfEngine;
+		public IUhtTokenReader TokenReader { get; }
 
 		/// <summary>
 		/// If true, the inclusion of the generated header file was seen
@@ -348,7 +348,7 @@ namespace EpicGames.UHT.Parsers
 		public static UhtHeaderFileParser Parse(UhtHeaderFile headerFile)
 		{
 			UhtHeaderFileParser headerParser = new(headerFile);
-			using UhtParsingScope topScope = new(headerParser, headerParser.HeaderFile, headerFile.Session.GetKeywordTable(UhtTableNames.Global));
+			using UhtParsingScope topScope = new(headerParser, headerFile.Session.GetKeywordTable(UhtTableNames.Global));
 			headerParser.ParseStatements();
 
 			if (!headerParser.SpottedAutogeneratedHeaderInclude && headerParser.HeaderFile.Data.Length > 0)
@@ -379,7 +379,7 @@ namespace EpicGames.UHT.Parsers
 
 				if (missingGeneratedHeader)
 				{
-					UhtIssueBehavior missingGeneratedHeaderBehavior = headerFile.Package.IsPartOfEngine ? headerFile.Session.Config!.EngineMissingGeneratedHeaderIncludeBehavior
+					UhtIssueBehavior missingGeneratedHeaderBehavior = headerFile.Module.IsPartOfEngine ? headerFile.Session.Config!.EngineMissingGeneratedHeaderIncludeBehavior
 						: headerFile.Session.Config!.NonEngineMissingGeneratedHeaderIncludeBehavior;
 
 					switch (missingGeneratedHeaderBehavior) 
@@ -835,7 +835,7 @@ namespace EpicGames.UHT.Parsers
 						}
 						if (!includeNameString.Span.Contains(".generated.h", StringComparison.Ordinal))
 						{
-							HeaderFile.AddReferencedHeader(includeNameString.ToString(), true);
+							HeaderFile.AddReferencedHeader(includeNameString.ToString(), UhtHeaderReferenceType.Include);
 						}
 					}
 				}
diff --git a/Engine/Source/Programs/Shared/EpicGames.UHT/Parsers/UhtInterfaceClassParser.cs b/Engine/Source/Programs/Shared/EpicGames.UHT/Parsers/UhtInterfaceClassParser.cs
--- a/Engine/Source/Programs/Shared/EpicGames.UHT/Parsers/UhtInterfaceClassParser.cs
+++ b/Engine/Source/Programs/Shared/EpicGames.UHT/Parsers/UhtInterfaceClassParser.cs
@@ -50,7 +50,7 @@ namespace EpicGames.UHT.Parsers
 
 		private static UhtParseResult ParseUInterface(UhtParsingScope parentScope, ref UhtToken token)
 		{
-			UhtClass classObj = new(parentScope.ScopeType, token.InputLine);
+			UhtClass classObj = new(parentScope.HeaderFile, parentScope.ScopeType, token.InputLine);
 			classObj.ClassType = UhtClassType.Interface;
 
 			{
diff --git a/Engine/Source/Programs/Shared/EpicGames.UHT/Parsers/UhtNativeInterfaceClassParser.cs b/Engine/Source/Programs/Shared/EpicGames.UHT/Parsers/UhtNativeInterfaceClassParser.cs
--- a/Engine/Source/Programs/Shared/EpicGames.UHT/Parsers/UhtNativeInterfaceClassParser.cs
+++ b/Engine/Source/Programs/Shared/EpicGames.UHT/Parsers/UhtNativeInterfaceClassParser.cs
@@ -105,7 +105,7 @@ namespace EpicGames.UHT.Parsers
 		{
 			IUhtTokenReader tokenReader = parentScope.TokenReader;
 
-			UhtClass classObj = new(parentScope.ScopeType, token.InputLine);
+			UhtClass classObj = new(parentScope.HeaderFile, parentScope.ScopeType, token.InputLine);
 			classObj.ClassType = UhtClassType.NativeInterface;
 			classObj.SourceName = sourceName.Value.ToString();
 			classObj.ClassFlags |= EClassFlags.Native | EClassFlags.Interface;
@@ -126,16 +126,13 @@ namespace EpicGames.UHT.Parsers
 			//TODO - C++ UHT compatibility - When we know for sure that we have a native interface, then we should error out.  Due to the lack of a symbol table,
 			// we can't do this 100% reliably.  However, if we find a U class, we can assume we have a native interface.
 			bool logUnhandledKeywords = false;
-			if (classObj.Outer != null) 
+			string interfaceName = "U" + classObj.EngineName;
+			foreach (UhtType outerChild in parentScope.HeaderFile.Children)
 			{
-				string interfaceName = "U" + classObj.EngineName;
-				foreach (UhtType outerChild in classObj.Outer.Children)
+				if (outerChild.SourceName == interfaceName && outerChild is UhtClass outerChildClass && outerChildClass.ClassType == UhtClassType.Interface)
 				{
-					if (outerChild.SourceName == interfaceName && outerChild is UhtClass outerChildClass && outerChildClass.ClassType == UhtClassType.Interface)
-					{
-						logUnhandledKeywords = true;
-						break;
-					}
+					logUnhandledKeywords = true;
+					break;
 				}
 			}
 
diff --git a/Engine/Source/Programs/Shared/EpicGames.UHT/Parsers/UhtParserHelpers.cs b/Engine/Source/Programs/Shared/EpicGames.UHT/Parsers/UhtParserHelpers.cs
--- a/Engine/Source/Programs/Shared/EpicGames.UHT/Parsers/UhtParserHelpers.cs
+++ b/Engine/Source/Programs/Shared/EpicGames.UHT/Parsers/UhtParserHelpers.cs
@@ -63,7 +63,7 @@ namespace EpicGames.UHT.Parsers
 		{
 
 			// Fetch the default generation code version. If supplied, then package code version overrides the default.
-			EGeneratedCodeVersion version = structObj.Package.Module.GeneratedCodeVersion;
+			EGeneratedCodeVersion version = structObj.Module.Module.GeneratedCodeVersion;
 			if (version == EGeneratedCodeVersion.None)
 			{
 				version = config.DefaultGeneratedCodeVersion;
diff --git a/Engine/Source/Programs/Shared/EpicGames.UHT/Parsers/UhtParsingScope.cs b/Engine/Source/Programs/Shared/EpicGames.UHT/Parsers/UhtParsingScope.cs
--- a/Engine/Source/Programs/Shared/EpicGames.UHT/Parsers/UhtParsingScope.cs
+++ b/Engine/Source/Programs/Shared/EpicGames.UHT/Parsers/UhtParsingScope.cs
@@ -23,6 +23,16 @@ namespace EpicGames.UHT.Parsers
 		/// </summary>
 		public UhtHeaderFileParser HeaderParser { get; }
 
+		/// <summary>
+		/// Header file being parsed
+		/// </summary>
+		public UhtHeaderFile HeaderFile => HeaderParser.HeaderFile;
+
+		/// <summary>
+		/// Module owning the header file
+		/// </summary>
+		public UhtModule Module => HeaderFile.Module;
+
 		/// <summary>
 		/// Token reader
 		/// </summary>
@@ -34,7 +44,7 @@ namespace EpicGames.UHT.Parsers
 		public UhtParsingScope? ParentScope { get; }
 
 		/// <summary>
-		/// Type being parsed
+		/// Type being parsed.
 		/// </summary>
 		public UhtType ScopeType { get; }
 
@@ -51,7 +61,7 @@ namespace EpicGames.UHT.Parsers
 		/// <summary>
 		/// Current session
 		/// </summary>
-		public UhtSession Session => ScopeType.Session;
+		public UhtSession Session => HeaderFile.Session;
 
 		/// <summary>
 		/// Return the current class scope being compiled
@@ -73,23 +83,17 @@ namespace EpicGames.UHT.Parsers
 			}
 		}
 
-		/// <summary>
-		/// Return the current class being compiled
-		/// </summary>
-		public UhtClass CurrentClass => (UhtClass)CurrentClassScope.ScopeType;
-
 		/// <summary>
 		/// Construct a root/global scope
 		/// </summary>
 		/// <param name="headerParser">Header parser</param>
-		/// <param name="scopeType">Type being parsed</param>
 		/// <param name="keywordTable">Keyword table</param>
-		public UhtParsingScope(UhtHeaderFileParser headerParser, UhtType scopeType, UhtKeywordTable keywordTable)
+		public UhtParsingScope(UhtHeaderFileParser headerParser, UhtKeywordTable keywordTable)
 		{
 			HeaderParser = headerParser;
 			TokenReader = headerParser.TokenReader;
 			ParentScope = null;
-			ScopeType = scopeType;
+			ScopeType = headerParser.Module.ScriptPackage; // The default package for parsing will be the standard UE package
 			ScopeKeywordTable = keywordTable;
 			HeaderParser.PushScope(this);
 		}
diff --git a/Engine/Source/Programs/Shared/EpicGames.UHT/Parsers/UhtPropertyParser.cs b/Engine/Source/Programs/Shared/EpicGames.UHT/Parsers/UhtPropertyParser.cs
--- a/Engine/Source/Programs/Shared/EpicGames.UHT/Parsers/UhtPropertyParser.cs
+++ b/Engine/Source/Programs/Shared/EpicGames.UHT/Parsers/UhtPropertyParser.cs
@@ -804,8 +804,7 @@ namespace EpicGames.UHT.Parsers
 					// Checking for this error is a bit tricky given legacy code.  
 					// 1) If already wrapped in WITH_EDITORONLY_DATA (see above), then we ignore the error via the else 
 					// 2) Ignore any module that is an editor module
-					UhtPackage package = topScope.ScopeType.HeaderFile.Package;
-					UHTManifest.Module module = package.Module;
+					UHTManifest.Module module = topScope.Module.Module;
 					bool isEditorModule =
 						module.ModuleType == UHTModuleType.EngineEditor ||
 						module.ModuleType == UHTModuleType.GameEditor ||
@@ -892,7 +891,7 @@ namespace EpicGames.UHT.Parsers
 
 			if (_options.HasAnyFlags(UhtPropertyParseOptions.AddModuleRelativePath))
 			{
-				UhtParsingScope.AddModuleRelativePathToMetaData(propertySettings.MetaData, topScope.ScopeType.HeaderFile);
+				UhtParsingScope.AddModuleRelativePathToMetaData(propertySettings.MetaData, topScope.HeaderFile);
 			}
 
 			// Fetch the name of the property, bitfield and array size
@@ -1120,7 +1119,7 @@ namespace EpicGames.UHT.Parsers
 			// Force the category in non-engine projects
 			if (newProperty.PropertyCategory == UhtPropertyCategory.Member)
 			{
-				if (!newProperty.Package.IsPartOfEngine &&
+				if (!newProperty.Module.IsPartOfEngine &&
 					newProperty.PropertyFlags.HasAnyFlags(EPropertyFlags.Edit | EPropertyFlags.BlueprintVisible) &&
 					!newProperty.MetaData.ContainsKey(UhtNames.Category))
 				{
diff --git a/Engine/Source/Programs/Shared/EpicGames.UHT/Parsers/UhtScriptStructParser.cs b/Engine/Source/Programs/Shared/EpicGames.UHT/Parsers/UhtScriptStructParser.cs
--- a/Engine/Source/Programs/Shared/EpicGames.UHT/Parsers/UhtScriptStructParser.cs
+++ b/Engine/Source/Programs/Shared/EpicGames.UHT/Parsers/UhtScriptStructParser.cs
@@ -65,7 +65,7 @@ namespace EpicGames.UHT.Parsers
 
 		private static UhtParseResult ParseUScriptStruct(UhtParsingScope parentScope, UhtToken keywordToken)
 		{
-			UhtScriptStruct scriptStruct = new(parentScope.ScopeType, keywordToken.InputLine);
+			UhtScriptStruct scriptStruct = new(parentScope.HeaderFile, parentScope.ScopeType, keywordToken.InputLine);
 			{
 				using UhtParsingScope topScope = new(parentScope, scriptStruct, parentScope.Session.GetKeywordTable(UhtTableNames.ScriptStruct), UhtAccessSpecifier.Public);
 
diff --git a/Engine/Source/Programs/Shared/EpicGames.UHT/Types/Properties/UhtObjectPropertyBase.cs b/Engine/Source/Programs/Shared/EpicGames.UHT/Types/Properties/UhtObjectPropertyBase.cs
--- a/Engine/Source/Programs/Shared/EpicGames.UHT/Types/Properties/UhtObjectPropertyBase.cs
+++ b/Engine/Source/Programs/Shared/EpicGames.UHT/Types/Properties/UhtObjectPropertyBase.cs
@@ -241,11 +241,11 @@ namespace EpicGames.UHT.Types
 				return;
 			}
 
-			UhtPackage package = propertySettings.Outer.Package;
+			UhtModule module = propertySettings.Outer.Module;
 			UhtIssueBehavior behavior = nonEngineBehavior;
-			if (package.IsPartOfEngine)
+			if (module.IsPartOfEngine)
 			{
-				if (package.IsPlugin)
+				if (module.IsPlugin)
 				{
 					behavior = enginePluginBehavior;
 				}
diff --git a/Engine/Source/Programs/Shared/EpicGames.UHT/Types/UhtClass.cs b/Engine/Source/Programs/Shared/EpicGames.UHT/Types/UhtClass.cs
--- a/Engine/Source/Programs/Shared/EpicGames.UHT/Types/UhtClass.cs
+++ b/Engine/Source/Programs/Shared/EpicGames.UHT/Types/UhtClass.cs
@@ -513,9 +513,10 @@ namespace EpicGames.UHT.Types
 		/// <summary>
 		/// Construct a new instance of the class
 		/// </summary>
+		/// <param name="headerFile">Header being parsed</param>
 		/// <param name="outer">The outer type</param>
 		/// <param name="lineNumber">Line number where class begins</param>
-		public UhtClass(UhtType outer, int lineNumber) : base(outer, lineNumber)
+		public UhtClass(UhtHeaderFile headerFile, UhtType outer, int lineNumber) : base(headerFile, outer, lineNumber)
 		{
 			ClassWithin = this;
 		}
@@ -771,8 +772,8 @@ namespace EpicGames.UHT.Types
 								}
 								else
 								{
-									// Copy the children
-									AddChildren(nativeInterface.DetachChildren());
+									// Move native interface children into the interface
+									MoveChildren(nativeInterface);
 								}
 							}
 							break;
diff --git a/Engine/Source/Programs/Shared/EpicGames.UHT/Types/UhtEnum.cs b/Engine/Source/Programs/Shared/EpicGames.UHT/Types/UhtEnum.cs
--- a/Engine/Source/Programs/Shared/EpicGames.UHT/Types/UhtEnum.cs
+++ b/Engine/Source/Programs/Shared/EpicGames.UHT/Types/UhtEnum.cs
@@ -152,9 +152,10 @@ namespace EpicGames.UHT.Types
 		/// <summary>
 		/// Construct a new enumeration
 		/// </summary>
+		/// <param name="headerFile">Header being parsed</param>
 		/// <param name="outer">Outer type</param>
 		/// <param name="lineNumber">Line number of declaration</param>
-		public UhtEnum(UhtType outer, int lineNumber) : base(outer, lineNumber)
+		public UhtEnum(UhtHeaderFile headerFile, UhtType outer, int lineNumber) : base(headerFile, outer, lineNumber)
 		{
 			MetaData.KeyConversion = this;
 			EnumValues = new List<UhtEnumValue>();
diff --git a/Engine/Source/Programs/Shared/EpicGames.UHT/Types/UhtField.cs b/Engine/Source/Programs/Shared/EpicGames.UHT/Types/UhtField.cs
--- a/Engine/Source/Programs/Shared/EpicGames.UHT/Types/UhtField.cs
+++ b/Engine/Source/Programs/Shared/EpicGames.UHT/Types/UhtField.cs
@@ -14,9 +14,10 @@ namespace EpicGames.UHT.Types
 		/// <summary>
 		/// Construct a new field
 		/// </summary>
+		/// <param name="headerFile">Header file being parsed</param>
 		/// <param name="outer">Outer object</param>
 		/// <param name="lineNumber">Line number of declaration</param>
-		protected UhtField(UhtType outer, int lineNumber) : base(outer, lineNumber)
+		protected UhtField(UhtHeaderFile headerFile, UhtType outer, int lineNumber) : base(headerFile, outer, lineNumber)
 		{
 		}
 	}
diff --git a/Engine/Source/Programs/Shared/EpicGames.UHT/Types/UhtFunction.cs b/Engine/Source/Programs/Shared/EpicGames.UHT/Types/UhtFunction.cs
--- a/Engine/Source/Programs/Shared/EpicGames.UHT/Types/UhtFunction.cs
+++ b/Engine/Source/Programs/Shared/EpicGames.UHT/Types/UhtFunction.cs
@@ -375,9 +375,10 @@ namespace EpicGames.UHT.Types
 		/// <summary>
 		/// Construct a new instance of a function
 		/// </summary>
+		/// <param name="headerFile">Header being parsed</param>
 		/// <param name="outer">The parent object</param>
 		/// <param name="lineNumber">The line number where the function is defined</param>
-		public UhtFunction(UhtType outer, int lineNumber) : base(outer, lineNumber)
+		public UhtFunction(UhtHeaderFile headerFile, UhtType outer, int lineNumber) : base(headerFile, outer, lineNumber)
 		{
 		}
 
@@ -563,7 +564,7 @@ namespace EpicGames.UHT.Types
 						if (!hasMenuCategory && !internalOnly && !Deprecated && !blueprintAccessor)
 						{
 							// To allow for quick iteration, don't enforce the requirement that game functions have to be categorized
-							if (HeaderFile.Package.IsPartOfEngine)
+							if (Module.IsPartOfEngine)
 							{
 								this.LogError("An explicit Category specifier is required for Blueprint accessible functions in an Engine module.");
 							}
diff --git a/Engine/Source/Programs/Shared/EpicGames.UHT/Types/UhtHeaderFile.cs b/Engine/Source/Programs/Shared/EpicGames.UHT/Types/UhtHeaderFile.cs
--- a/Engine/Source/Programs/Shared/EpicGames.UHT/Types/UhtHeaderFile.cs
+++ b/Engine/Source/Programs/Shared/EpicGames.UHT/Types/UhtHeaderFile.cs
@@ -3,7 +3,6 @@
 using System;
 using System.Collections.Generic;
 using System.IO;
-using System.Text;
 using System.Text.Json.Serialization;
 using EpicGames.Core;
 using EpicGames.UHT.Utils;
@@ -98,14 +97,31 @@ namespace EpicGames.UHT.Types
 	}
 
 	/// <summary>
-	/// Represents a header file.  Unlike the engine, UhtHeader files will appear as a child 
-	/// of a UhtPackage and the outer of all global types found in that header.
+	/// Type of reference being added
 	/// </summary>
-	public class UhtHeaderFile : UhtType
+	public enum UhtHeaderReferenceType
+	{
+
+		/// <summary>
+		/// The header being referenced in a direct include
+		/// </summary>
+		Include,	
+
+		/// <summary>
+		/// The header being referenced is a passive reference (i.e. NoExportTypes.h or a header referenced by a header being included)
+		/// </summary>
+		Passive,
+	}
+
+	/// <summary>
+	/// Represents a header file.
+	/// </summary>
+	public class UhtHeaderFile : IUhtMessageSite
 	{
 		private readonly UhtSimpleMessageSite _messageSite;
 		private readonly UhtSourceFile _sourceFile;
 		private readonly List<UhtHeaderFile> _referencedHeaders = new();
+		private readonly List<UhtType> _children = new();
 
 		/// <summary>
 		/// Contents of the header
@@ -119,6 +135,18 @@ namespace EpicGames.UHT.Types
 		[JsonIgnore]
 		public string FilePath => _sourceFile.FilePath;
 
+		/// <summary>
+		/// Currently running session
+		/// </summary>
+		[JsonIgnore]
+		public UhtSession Session => Module.Session;
+
+		/// <summary>
+		/// Module associated with the header
+		/// </summary>
+		[JsonIgnore]
+		public UhtModule Module { get; }
+
 		/// <summary>
 		/// File name without the extension
 		/// </summary>
@@ -165,6 +193,15 @@ namespace EpicGames.UHT.Types
 		[JsonConverter(typeof(JsonStringEnumConverter))]
 		public UhtHeaderFileExportFlags HeaderFileExportFlags { get; set; } = UhtHeaderFileExportFlags.None;
 
+		/// <summary>
+		/// Children types of this type
+		/// </summary>
+		[JsonConverter(typeof(UhtNullableTypeListJsonConverter<UhtType>))]
+		public IReadOnlyList<UhtType> Children => _children;
+
+		/// <inheritdoc/>
+		public override string ToString() { return FilePath; }
+
 		/// <summary>
 		/// If true, the header file should be exported
 		/// </summary>
@@ -177,31 +214,6 @@ namespace EpicGames.UHT.Types
 		[JsonIgnore]
 		public UhtReferenceCollector References { get; } = new UhtReferenceCollector();
 
-		/// <inheritdoc/>
-		[JsonIgnore]
-		public override UhtPackage Package
-		{
-			get
-			{
-				if (Outer == null)
-				{
-					throw new UhtIceException("Attempt to fetch header file package but it has no outer");
-				}
-				return (UhtPackage)Outer;
-			}
-		}
-
-		/// <inheritdoc/>
-		[JsonIgnore]
-		public override UhtHeaderFile HeaderFile => this;
-
-		/// <inheritdoc/>
-		[JsonIgnore]
-		public override UhtEngineType EngineType => UhtEngineType.Header;
-
-		/// <inheritdoc/>
-		public override string EngineClassName => "UhtHeaderFile";
-
 		/// <summary>
 		/// Collection of headers directly included by this header
 		/// </summary>
@@ -209,30 +221,33 @@ namespace EpicGames.UHT.Types
 		public List<UhtHeaderFile> IncludedHeaders { get; } = new List<UhtHeaderFile>();
 
 		#region IUHTMessageSite implementation
+		/// <inheritdoc/>
+		[JsonIgnore]
+		public IUhtMessageSession MessageSession => _messageSite.MessageSession;
 
 		/// <inheritdoc/>
 		[JsonIgnore]
-		public override IUhtMessageSession MessageSession => _messageSite.MessageSession;
+		public IUhtMessageSource? MessageSource => _messageSite.MessageSource;
 
 		/// <inheritdoc/>
 		[JsonIgnore]
-		public override IUhtMessageSource? MessageSource => _messageSite.MessageSource;
+		public IUhtMessageLineNumber? MessageLineNumber => throw new NotImplementedException();
 		#endregion
 
 		/// <summary>
 		/// Construct a new header file
 		/// </summary>
-		/// <param name="package">Owning package</param>
+		/// <param name="module">Owning module</param>
 		/// <param name="path">Path to the header file</param>
-		public UhtHeaderFile(UhtPackage package, string path) : base(package, 1)
+		public UhtHeaderFile(UhtModule module, string path)
 		{
+			Module = module;
 			HeaderFileTypeIndex = Session.GetNextHeaderFileTypeIndex();
 			_messageSite = new UhtSimpleMessageSite(Session);
 			_sourceFile = new UhtSourceFile(Session, path);
 			_messageSite.MessageSource = _sourceFile;
 			FileNameWithoutExtension = System.IO.Path.GetFileNameWithoutExtension(_sourceFile.FilePath);
 			GeneratedHeaderFileName = FileNameWithoutExtension + ".generated.h";
-			SourceName = System.IO.Path.GetFileName(_sourceFile.FilePath);
 			IsNoExportTypes = String.Equals(_sourceFile.FileName, "NoExportTypes", StringComparison.OrdinalIgnoreCase);
 		}
 
@@ -248,13 +263,13 @@ namespace EpicGames.UHT.Types
 		/// Add a reference to the given header
 		/// </summary>
 		/// <param name="id">Path of the header</param>
-		/// <param name="isIncludedFile">True if this is a directly included file</param>
-		public void AddReferencedHeader(string id, bool isIncludedFile)
+		/// <param name="referenceType">How is the include in question referenced</param>
+		public void AddReferencedHeader(string id, UhtHeaderReferenceType referenceType)
 		{
 			UhtHeaderFile? headerFile = Session.FindHeaderFile(Path.GetFileName(id));
 			if (headerFile != null)
 			{
-				AddReferencedHeader(headerFile, isIncludedFile);
+				AddReferencedHeader(headerFile, referenceType);
 			}
 		}
 
@@ -264,19 +279,19 @@ namespace EpicGames.UHT.Types
 		/// <param name="type">Type in question</param>
 		public void AddReferencedHeader(UhtType type)
 		{
-			AddReferencedHeader(type.HeaderFile, false);
+			AddReferencedHeader(type.HeaderFile, UhtHeaderReferenceType.Passive);
 		}
 
 		/// <summary>
 		/// Add a reference to the given header file
 		/// </summary>
 		/// <param name="headerFile">Header file in question</param>
-		/// <param name="isIncludedFile">True if this is a directly included file</param>
-		public void AddReferencedHeader(UhtHeaderFile headerFile, bool isIncludedFile)
+		/// <param name="referenceType">How is the include in question referenced</param>
+		public void AddReferencedHeader(UhtHeaderFile headerFile, UhtHeaderReferenceType referenceType)
 		{
 
 			// Ignore direct references to myself
-			if (!isIncludedFile && headerFile == this)
+			if (referenceType != UhtHeaderReferenceType.Include && headerFile == this)
 			{
 				return;
 			}
@@ -294,12 +309,12 @@ namespace EpicGames.UHT.Types
 
 				// There is questionable compatibility hack where a source file will always be exported
 				// regardless of having types when it is being included by the SAME package.
-				if (headerFile.Package == Package)
+				if (headerFile.Module == Module)
 				{
 					headerFile.HeaderFileExportFlags |= UhtHeaderFileExportFlags.Referenced;
 				}
 				_referencedHeaders.Add(headerFile);
-				if (isIncludedFile)
+				if (referenceType == UhtHeaderReferenceType.Include)
 				{
 					IncludedHeaders.Add(headerFile);
 				}
@@ -330,20 +345,43 @@ namespace EpicGames.UHT.Types
 			}
 		}
 
-		/// <inheritdoc/>
-		public override void AppendPathName(StringBuilder builder, UhtType? stopOuter = null)
+		/// <summary>
+		/// Add a type as a child
+		/// </summary>
+		/// <param name="child">The child to be added.</param>
+		public void AddChild(UhtType child)
+		{
+			_children.Add(child);
+		}
+
+		/// <summary>
+		/// Resolve all types owned by the header
+		/// </summary>
+		/// <param name="resolvePhase">Phase of the resolution process</param>
+		public void Resolve(UhtResolvePhase resolvePhase)
+		{
+			UhtType.ResolveChildren(_children, resolvePhase);
+		}
+
+		/// <summary>
+		/// Bind all the super structs and base classes
+		/// </summary>
+		public void BindSuperAndBases()
 		{
-			// Headers do not contribute to path names
-			if (this != stopOuter && Outer != null)
+			foreach (UhtType child in Children)
 			{
-				Outer.AppendPathName(builder, stopOuter);
+				child.BindSuperAndBases();
 			}
 		}
 
-		/// <inheritdoc/>
-		protected override UhtValidationOptions Validate(UhtValidationOptions options)
+		/// <summary>
+		/// Validate the state of the header file
+		/// </summary>
+		/// <param name="options">Validation options</param>
+		/// <returns></returns>
+		public void Validate(UhtValidationOptions options)
 		{
-			options = base.Validate(options | UhtValidationOptions.Shadowing);
+			options |= UhtValidationOptions.Shadowing;
 
 			Dictionary<int, UhtFunction> usedRPCIds = new();
 			Dictionary<int, UhtFunction> rpcsNeedingHookup = new();
@@ -398,7 +436,26 @@ namespace EpicGames.UHT.Types
 					kvp.Value.LogError($"Request function '{kvp.Value.SourceName}' is missing a response function with the id of '{kvp.Key}'");
 				}
 			}
-			return options;
+
+			foreach (UhtType child in Children)
+			{
+				UhtType.ValidateType(child, options);
+			}
+		}
+
+		/// <summary>
+		/// Collect all things referenced by the given header
+		/// </summary>
+		public void CollectReferences()
+		{
+			foreach (UhtType child in Children)
+			{
+				child.CollectReferences(References);
+			}
+			foreach (UhtHeaderFile refHeaderFile in References.ReferencedHeaders)
+			{
+				AddReferencedHeader(refHeaderFile, UhtHeaderReferenceType.Passive);
+			}
 		}
 	}
 }
diff --git a/Engine/Source/Programs/Shared/EpicGames.UHT/Types/UhtModule.cs b/Engine/Source/Programs/Shared/EpicGames.UHT/Types/UhtModule.cs
new file mode 100644
--- /dev/null
+++ b/Engine/Source/Programs/Shared/EpicGames.UHT/Types/UhtModule.cs
@@ -0,0 +1,246 @@
+// Copyright Epic Games, Inc. All Rights Reserved.
+
+using System;
+using System.Collections.Generic;
+using System.IO;
+using System.Security.Cryptography.X509Certificates;
+using System.Text.Json.Serialization;
+using EpicGames.Core;
+using EpicGames.UHT.Utils;
+using UnrealBuildBase;
+
+namespace EpicGames.UHT.Types
+{
+
+	/// <summary>
+	/// Represents a UHT module
+	/// </summary>
+	public class UhtModule
+	{
+		private readonly List<UhtHeaderFile> _headers = new();
+
+		/// <summary>
+		/// The running session
+		/// </summary>
+		[JsonIgnore]
+		public UhtSession Session { get; }
+
+		/// <summary>
+		/// UHT module from UBT (1 to 1 relationship)
+		/// </summary>
+		public UHTManifest.Module Module { get; }
+
+		/// <summary>
+		/// Primary /Script/ package
+		/// </summary>
+		[JsonIgnore]
+		public UhtPackage ScriptPackage { get; }
+
+		/// <summary>
+		/// Name of the module based on the script package name with out the "/Script/" prefix
+		/// </summary>
+		public string ShortName { get; }
+
+		/// <summary>
+		/// C++ API name for the module with a trailing space
+		/// </summary>
+		[JsonIgnore]
+		public string Api { get; }
+
+		/// <summary>
+		/// Enumeration of all the packages in the module.
+		/// </summary>
+		public IEnumerable<UhtPackage> Packages
+		{
+			get
+			{
+				yield return ScriptPackage;
+			}
+		}
+
+		/// <summary>
+		/// Collection of headers owned by the module
+		/// </summary>
+		[JsonIgnore]
+		public IReadOnlyList<UhtHeaderFile> Headers => _headers;
+
+		/// <summary>
+		/// True if the package is part of the engine
+		/// </summary>
+		[JsonIgnore]
+		public bool IsPartOfEngine
+		{
+			get
+			{
+				switch (Module.ModuleType)
+				{
+					case UHTModuleType.Program:
+						return Module.BaseDirectory.Replace('\\', '/').StartsWith(Unreal.EngineDirectory.FullName.Replace('\\', '/'), StringComparison.Ordinal);
+					case UHTModuleType.EngineRuntime:
+					case UHTModuleType.EngineUncooked:
+					case UHTModuleType.EngineDeveloper:
+					case UHTModuleType.EngineEditor:
+					case UHTModuleType.EngineThirdParty:
+						return true;
+					case UHTModuleType.GameRuntime:
+					case UHTModuleType.GameUncooked:
+					case UHTModuleType.GameDeveloper:
+					case UHTModuleType.GameEditor:
+					case UHTModuleType.GameThirdParty:
+						return false;
+					default:
+						throw new UhtIceException("Invalid module type");
+				}
+			}
+		}
+
+		/// <summary>
+		/// True if the package is a plugin
+		/// </summary>
+		[JsonIgnore]
+		public bool IsPlugin => Module.BaseDirectory.Replace('\\', '/').Contains("/Plugins/", StringComparison.Ordinal);
+
+		/// <summary>
+		/// Create a new module from the given manifest module
+		/// </summary>
+		/// <param name="session">Running session</param>
+		/// <param name="module">Source module from UBT</param>
+		public UhtModule(UhtSession session, UHTManifest.Module module)
+		{
+			Session = session;
+			Module = module;
+
+			EPackageFlags packageFlags = EPackageFlags.ContainsScript | EPackageFlags.Compiling;
+
+			switch (module.OverrideModuleType)
+			{
+				case EPackageOverrideType.None:
+					switch (module.ModuleType)
+					{
+						case UHTModuleType.GameEditor:
+						case UHTModuleType.EngineEditor:
+							packageFlags |= EPackageFlags.EditorOnly;
+							break;
+
+						case UHTModuleType.GameDeveloper:
+						case UHTModuleType.EngineDeveloper:
+							packageFlags |= EPackageFlags.Developer;
+							break;
+
+						case UHTModuleType.GameUncooked:
+						case UHTModuleType.EngineUncooked:
+							packageFlags |= EPackageFlags.UncookedOnly;
+							break;
+					}
+					break;
+
+				case EPackageOverrideType.EditorOnly:
+					packageFlags |= EPackageFlags.EditorOnly;
+					break;
+
+				case EPackageOverrideType.EngineDeveloper:
+				case EPackageOverrideType.GameDeveloper:
+					packageFlags |= EPackageFlags.Developer;
+					break;
+
+				case EPackageOverrideType.EngineUncookedOnly:
+				case EPackageOverrideType.GameUncookedOnly:
+					packageFlags |= EPackageFlags.UncookedOnly;
+					break;
+			}
+
+			string packageName;
+			int lastSlashIndex = Module.Name.LastIndexOf('/');
+			if (lastSlashIndex == -1)
+			{
+				packageName = $"/Script/{Module.Name}";
+				ShortName = Module.Name;
+			}
+			else
+			{
+				packageName = Module.Name;
+				ShortName = packageName[(lastSlashIndex + 1)..];
+			}
+
+			Api = $"{ShortName.ToUpper()}_API ";
+
+			ScriptPackage = new(this, packageName, packageFlags);
+		}
+
+		/// <summary>
+		/// Prepare to parse all the headers in the module.  This creates all the UhtHeaderFiles and validates there
+		/// are no conflicts.
+		/// </summary>
+		/// <param name="addHeaderFileAction">Action to invoke to register the headers with the session</param>
+		public void PrepareHeaders(Action<UhtHeaderFile> addHeaderFileAction)
+		{
+			PrepareHeaders(Module.ClassesHeaders, UhtHeaderFileType.Classes, addHeaderFileAction);
+			PrepareHeaders(Module.PublicHeaders, UhtHeaderFileType.Public, addHeaderFileAction);
+			PrepareHeaders(Module.InternalHeaders, UhtHeaderFileType.Internal, addHeaderFileAction);
+			PrepareHeaders(Module.PrivateHeaders, UhtHeaderFileType.Private, addHeaderFileAction);
+		}
+
+		private void PrepareHeaders(IEnumerable<string> headerFiles, UhtHeaderFileType headerFileType, Action<UhtHeaderFile> addHeaderFileAction)
+		{
+			string typeDirectory = headerFileType.ToString() + '/';
+			string normalizedModuleBaseFullFilePath = GetNormalizedFullFilePath(Module.BaseDirectory);
+			foreach (string headerFilePath in headerFiles)
+			{
+
+				// Make sure this isn't a duplicate
+				string normalizedFullFilePath = GetNormalizedFullFilePath(headerFilePath);
+				string fileName = Path.GetFileName(normalizedFullFilePath);
+				UhtHeaderFile? existingHeaderFile = Session.FindHeaderFile(fileName);
+				if (existingHeaderFile != null)
+				{
+					string normalizedExistingFullFilePath = GetNormalizedFullFilePath(existingHeaderFile.FilePath);
+					if (!String.Equals(normalizedFullFilePath, normalizedExistingFullFilePath, StringComparison.OrdinalIgnoreCase))
+					{
+						IUhtMessageSite site = (IUhtMessageSite?)Session.ManifestFile ?? Session;
+						site.LogError($"Two headers with the same name is not allowed. '{headerFilePath}' conflicts with '{existingHeaderFile.FilePath}'");
+						continue;
+					}
+				}
+
+				// Create the header file and add to the collections
+				UhtHeaderFile headerFile = new(this, headerFilePath);
+				headerFile.HeaderFileType = headerFileType;
+				addHeaderFileAction(headerFile);
+				_headers.Add(headerFile);
+
+				// Save metadata for the class path, both for it's include path and relative to the module base directory
+				if (normalizedFullFilePath.StartsWith(normalizedModuleBaseFullFilePath, true, null))
+				{
+					int stripLength = normalizedModuleBaseFullFilePath.Length;
+					if (stripLength < normalizedFullFilePath.Length && normalizedFullFilePath[stripLength] == '/')
+					{
+						++stripLength;
+					}
+
+					headerFile.ModuleRelativeFilePath = normalizedFullFilePath[stripLength..];
+
+					if (normalizedFullFilePath[stripLength..].StartsWith(typeDirectory, true, null))
+					{
+						stripLength += typeDirectory.Length;
+					}
+
+					headerFile.IncludeFilePath = normalizedFullFilePath[stripLength..];
+				}
+			}
+		}
+
+		/// <summary>
+		/// Return the normalized path converted to a full path if possible. 
+		/// Code should NOT depend on a full path being returned.
+		/// 
+		/// In general, it is assumed that during normal UHT, all paths are already full paths.
+		/// Only the test harness deals in relative paths.
+		/// </summary>
+		/// <param name="filePath">Path to normalize</param>
+		/// <returns>Normalized path possibly converted to a full path.</returns>
+		private string GetNormalizedFullFilePath(string filePath)
+		{
+			return Session.FileManager!.GetFullFilePath(filePath).Replace('\\', '/');
+		}
+	}
+}
diff --git a/Engine/Source/Programs/Shared/EpicGames.UHT/Types/UhtObject.cs b/Engine/Source/Programs/Shared/EpicGames.UHT/Types/UhtObject.cs
--- a/Engine/Source/Programs/Shared/EpicGames.UHT/Types/UhtObject.cs
+++ b/Engine/Source/Programs/Shared/EpicGames.UHT/Types/UhtObject.cs
@@ -39,8 +39,8 @@ namespace EpicGames.UHT.Types
 		/// <summary>
 		/// Construct a new instance of the object
 		/// </summary>
-		/// <param name="session">Session of the object</param>
-		protected UhtObject(UhtSession session) : base(session)
+		/// <param name="module">Owning module</param>
+		protected UhtObject(UhtModule module) : base(module)
 		{
 			ObjectTypeIndex = Session.GetNextObjectTypeIndex();
 		}
@@ -48,9 +48,10 @@ namespace EpicGames.UHT.Types
 		/// <summary>
 		/// Construct a new instance of the object
 		/// </summary>
+		/// <param name="headerFile">Header file being compiled</param>
 		/// <param name="outer">Outer object</param>
 		/// <param name="lineNumber">Line number where object is defined</param>
-		protected UhtObject(UhtType outer, int lineNumber) : base(outer, lineNumber)
+		protected UhtObject(UhtHeaderFile headerFile, UhtType outer, int lineNumber) : base(headerFile, outer, lineNumber)
 		{
 			ObjectTypeIndex = Session.GetNextObjectTypeIndex();
 		}
diff --git a/Engine/Source/Programs/Shared/EpicGames.UHT/Types/UhtPackage.cs b/Engine/Source/Programs/Shared/EpicGames.UHT/Types/UhtPackage.cs
--- a/Engine/Source/Programs/Shared/EpicGames.UHT/Types/UhtPackage.cs
+++ b/Engine/Source/Programs/Shared/EpicGames.UHT/Types/UhtPackage.cs
@@ -4,8 +4,6 @@ using System;
 using System.Text.Json.Serialization;
 using EpicGames.Core;
 using EpicGames.UHT.Tables;
-using EpicGames.UHT.Utils;
-using UnrealBuildBase;
 
 namespace EpicGames.UHT.Types
 {
@@ -29,11 +27,6 @@ namespace EpicGames.UHT.Types
 		[JsonConverter(typeof(JsonStringEnumConverter))]
 		public EPackageFlags PackageFlags { get; set; } = EPackageFlags.None;
 
-		/// <summary>
-		/// UHT module of the package (1 to 1 relationship)
-		/// </summary>
-		public UHTManifest.Module Module { get; set; }
-
 		/// <inheritdoc/>
 		[JsonIgnore]
 		public override UhtEngineType EngineType => UhtEngineType.Package;
@@ -45,74 +38,25 @@ namespace EpicGames.UHT.Types
 		[JsonIgnore]
 		public override UhtPackage Package => this;
 
-		/// <inheritdoc/>
-		[JsonIgnore]
-		public override UhtHeaderFile HeaderFile => throw new NotImplementedException();
-
-		/// <summary>
-		/// True if the package is part of the engine
-		/// </summary>
-		[JsonIgnore]
-		public bool IsPartOfEngine
-		{
-			get
-			{
-				switch (Module.ModuleType)
-				{
-					case UHTModuleType.Program:
-						return Module.BaseDirectory.Replace('\\', '/').StartsWith(Unreal.EngineDirectory.FullName.Replace('\\', '/'), StringComparison.Ordinal);
-					case UHTModuleType.EngineRuntime:
-					case UHTModuleType.EngineUncooked:
-					case UHTModuleType.EngineDeveloper:
-					case UHTModuleType.EngineEditor:
-					case UHTModuleType.EngineThirdParty:
-						return true;
-					case UHTModuleType.GameRuntime:
-					case UHTModuleType.GameUncooked:
-					case UHTModuleType.GameDeveloper:
-					case UHTModuleType.GameEditor:
-					case UHTModuleType.GameThirdParty:
-						return false;
-					default:
-						throw new UhtIceException("Invalid module type");
-				}
-			}
-		}
-
-		/// <summary>
-		/// True if the package is a plugin
-		/// </summary>
-		[JsonIgnore]
-		public bool IsPlugin => Module.BaseDirectory.Replace('\\', '/').Contains("/Plugins/", StringComparison.Ordinal);
-
-		/// <summary>
-		/// Short name of the package (without the /Script/)
-		/// </summary>
-		public string ShortName { get; set; }
-
 		/// <summary>
 		/// Construct a new instance of a package
 		/// </summary>
-		/// <param name="session">Running session</param>
 		/// <param name="module">Source module of the package</param>
+		/// <param name="packageName">Name of the package</param>
 		/// <param name="packageFlags">Assorted package flags</param>
-		public UhtPackage(UhtSession session, UHTManifest.Module module, EPackageFlags packageFlags) : base(session)
+		public UhtPackage(UhtModule module, string packageName, EPackageFlags packageFlags) : base(module)
 		{
-			Module = module;
 			PackageFlags = packageFlags;
 			PackageTypeIndex = Session.GetNextPackageTypeIndex();
+			SourceName = packageName;
+		}
+
+		/// <inheritdoc/>
+		public override void AddChild(UhtType type)
+		{
 
-			int lastSlashIndex = module.Name.LastIndexOf('/');
-			if (lastSlashIndex == -1)
-			{
-				SourceName = $"/Script/{module.Name}";
-				ShortName = module.Name;
-			}
-			else
-			{
-				SourceName = module.Name;
-				ShortName = SourceName[(lastSlashIndex + 1)..];
-			}
+			// Package children are initially added to the header and then added to the package when we aren't going wide
+			type.HeaderFile.AddChild(type);
 		}
 	}
 }
diff --git a/Engine/Source/Programs/Shared/EpicGames.UHT/Types/UhtProperty.cs b/Engine/Source/Programs/Shared/EpicGames.UHT/Types/UhtProperty.cs
--- a/Engine/Source/Programs/Shared/EpicGames.UHT/Types/UhtProperty.cs
+++ b/Engine/Source/Programs/Shared/EpicGames.UHT/Types/UhtProperty.cs
@@ -1089,7 +1089,7 @@ namespace EpicGames.UHT.Types
 		/// </summary>
 		/// <param name="outer">Outer type of the property</param>
 		/// <param name="lineNumber">Line number where property was declared</param>
-		protected UhtProperty(UhtType outer, int lineNumber) : base(outer, lineNumber)
+		protected UhtProperty(UhtType outer, int lineNumber) : base(outer.HeaderFile, outer, lineNumber)
 		{
 			PropertyFlags = EPropertyFlags.None;
 			PropertyCaps = UhtPropertyCaps.CanBeContainerValue | UhtPropertyCaps.CanBeContainerKey | UhtPropertyCaps.CanHaveConfig;
@@ -1099,7 +1099,7 @@ namespace EpicGames.UHT.Types
 		/// Construct a new property
 		/// </summary>
 		/// <param name="propertySettings">Property settings from parsing</param>
-		protected UhtProperty(UhtPropertySettings propertySettings) : base(propertySettings.Outer, propertySettings.LineNumber, propertySettings.MetaData)
+		protected UhtProperty(UhtPropertySettings propertySettings) : base(propertySettings.Outer.HeaderFile, propertySettings.Outer, propertySettings.LineNumber, propertySettings.MetaData)
 		{
 			SourceName = propertySettings.SourceName;
 			// Engine name defaults to source name.  If it doesn't match what is coming in, then set it.
@@ -1857,7 +1857,7 @@ namespace EpicGames.UHT.Types
 				{
 					if (PropertyFlags.HasAnyFlags(EPropertyFlags.Edit | EPropertyFlags.BlueprintVisible))
 					{
-						if (Package.IsPartOfEngine)
+						if (Module.IsPartOfEngine)
 						{
 							this.LogError("An explicit Category specifier is required for any property exposed to the editor or Blueprints in an Engine module.");
 						}
diff --git a/Engine/Source/Programs/Shared/EpicGames.UHT/Types/UhtScriptStruct.cs b/Engine/Source/Programs/Shared/EpicGames.UHT/Types/UhtScriptStruct.cs
--- a/Engine/Source/Programs/Shared/EpicGames.UHT/Types/UhtScriptStruct.cs
+++ b/Engine/Source/Programs/Shared/EpicGames.UHT/Types/UhtScriptStruct.cs
@@ -701,9 +701,10 @@ namespace EpicGames.UHT.Types
 		/// <summary>
 		/// Construct a new script struct
 		/// </summary>
+		/// <param name="headerFile">Header being parsed</param>
 		/// <param name="outer">Outer type</param>
 		/// <param name="lineNumber">Line number of the definition</param>
-		public UhtScriptStruct(UhtType outer, int lineNumber) : base(outer, lineNumber)
+		public UhtScriptStruct(UhtHeaderFile headerFile, UhtType outer, int lineNumber) : base(headerFile, outer, lineNumber)
 		{
 		}
 
diff --git a/Engine/Source/Programs/Shared/EpicGames.UHT/Types/UhtStruct.cs b/Engine/Source/Programs/Shared/EpicGames.UHT/Types/UhtStruct.cs
--- a/Engine/Source/Programs/Shared/EpicGames.UHT/Types/UhtStruct.cs
+++ b/Engine/Source/Programs/Shared/EpicGames.UHT/Types/UhtStruct.cs
@@ -122,9 +122,10 @@ namespace EpicGames.UHT.Types
 		/// <summary>
 		/// Construct a new instance
 		/// </summary>
+		/// <param name="headerFile">Header file being parsed</param>
 		/// <param name="outer">Outer type</param>
 		/// <param name="lineNumber">Line number where definition begins</param>
-		protected UhtStruct(UhtType outer, int lineNumber) : base(outer, lineNumber)
+		protected UhtStruct(UhtHeaderFile headerFile, UhtType outer, int lineNumber) : base(headerFile, outer, lineNumber)
 		{
 		}
 
diff --git a/Engine/Source/Programs/Shared/EpicGames.UHT/Types/UhtType.cs b/Engine/Source/Programs/Shared/EpicGames.UHT/Types/UhtType.cs
--- a/Engine/Source/Programs/Shared/EpicGames.UHT/Types/UhtType.cs
+++ b/Engine/Source/Programs/Shared/EpicGames.UHT/Types/UhtType.cs
@@ -19,11 +19,6 @@ namespace EpicGames.UHT.Types
 	public enum UhtEngineType
 	{
 
-		/// <summary>
-		/// Type is a header file (UhtHeader)
-		/// </summary>
-		Header,
-
 		/// <summary>
 		/// Type is a package (UhtPackage)
 		/// </summary>
@@ -86,7 +81,6 @@ namespace EpicGames.UHT.Types
 		{
 			switch (engineType)
 			{
-				case UhtEngineType.Header: return "header";
 				case UhtEngineType.Package: return "package";
 				case UhtEngineType.Class:
 				case UhtEngineType.Interface:
@@ -111,7 +105,6 @@ namespace EpicGames.UHT.Types
 		{
 			switch (engineType)
 			{
-				case UhtEngineType.Header: return "Header";
 				case UhtEngineType.Package: return "Package";
 				case UhtEngineType.Class:
 				case UhtEngineType.Interface:
@@ -136,7 +129,6 @@ namespace EpicGames.UHT.Types
 		{
 			switch (engineType)
 			{
-				case UhtEngineType.Header: return "header";
 				case UhtEngineType.Package: return "package";
 				case UhtEngineType.Class: return "class";
 				case UhtEngineType.Interface: return "interface";
@@ -161,7 +153,6 @@ namespace EpicGames.UHT.Types
 		{
 			switch (engineType)
 			{
-				case UhtEngineType.Header: return "Header";
 				case UhtEngineType.Package: return "Package";
 				case UhtEngineType.Class: return "Class";
 				case UhtEngineType.Interface: return "Interface";
@@ -278,7 +269,6 @@ namespace EpicGames.UHT.Types
 		{
 			switch (engineType)
 			{
-				case UhtEngineType.Header: return 0;
 				case UhtEngineType.Package: return 0;
 				case UhtEngineType.Class: return UhtFindOptions.Class;
 				case UhtEngineType.Interface: return UhtFindOptions.Class;
@@ -729,7 +719,19 @@ namespace EpicGames.UHT.Types
 		/// All UHT runs are associated with a given session.  The session holds all the global information for a run.
 		/// </summary>
 		[JsonIgnore]
-		public UhtSession Session { get; set; }
+		public UhtSession Session => Module.Session;
+
+		/// <summary>
+		/// UHT module of the package (1 to 1 relationship)
+		/// </summary>
+		[JsonIgnore]
+		public UhtModule Module { get; }
+
+		/// <summary>
+		/// The header file containing the type
+		/// </summary>
+		[JsonIgnore]
+		public UhtHeaderFile HeaderFile => _headerFile ?? throw new UhtIceException("Attempt to fetch the header file for a type that isn't associated with a header");
 
 		/// <summary>
 		/// Every type in a session is assigned a unique, non-zero id.
@@ -760,23 +762,6 @@ namespace EpicGames.UHT.Types
 			}
 		}
 
-		/// <summary>
-		/// The header file containing the type
-		/// </summary>
-		[JsonIgnore]
-		public virtual UhtHeaderFile HeaderFile
-		{
-			get
-			{
-				UhtType type = this;
-				while (type.Outer != null && type.Outer.Outer != null)
-				{
-					type = type.Outer;
-				}
-				return (UhtHeaderFile)type;
-			}
-		}
-
 		/// <summary>
 		/// The name of the type as it appears in the source file
 		/// </summary>
@@ -815,6 +800,12 @@ namespace EpicGames.UHT.Types
 		[JsonConverter(typeof(JsonStringEnumConverter))]
 		public UhtDefineScope DefineScope { get; set; } = UhtDefineScope.None;
 
+		/// <summary>
+		/// If true, the type has not been marked as invalid.
+		/// </summary>
+		[JsonIgnore]
+		public bool IsValid { get; set; } = true;
+
 		/// <summary>
 		/// Return a combination of the engine type name followed by the path name of the type
 		/// </summary>
@@ -876,6 +867,7 @@ namespace EpicGames.UHT.Types
 
 		private string? _engineName = null;
 		private List<UhtType>? _children = null;
+		private readonly UhtHeaderFile? _headerFile = null;
 		private int _resolveState = (int)UhtResolvePhase.None * 2;
 
 		#region IUHTMessageSite implementation
@@ -900,26 +892,28 @@ namespace EpicGames.UHT.Types
 		/// Construct the base type information.  This constructor is used exclusively by UhtPackage which is at 
 		/// the root of all type hierarchies.
 		/// </summary>
-		/// <param name="session"></param>
-		protected UhtType(UhtSession session)
+		/// <param name="module>"></param>
+		protected UhtType(UhtModule module)
 		{
-			Session = session;
+			Module = module;
 			Outer = null;
 			SourceName = String.Empty;
 			LineNumber = 1;
-			TypeIndex = session.GetNextTypeIndex();
+			TypeIndex = Session.GetNextTypeIndex();
 			MetaData = new UhtMetaData(this, Session.Config);
 		}
 
 		/// <summary>
 		/// Construct instance of the type given a parent type.
 		/// </summary>
+		/// <param name="headerFile">Header file being parsed</param>
 		/// <param name="outer">The outer type of the type being constructed.  For example, a class defined in a given header will have the header as the outer.</param>
 		/// <param name="lineNumber">The line number in the source file where the type was defined.</param>
 		/// <param name="metaData">Optional meta data to be associated with the type.</param>
-		protected UhtType(UhtType outer, int lineNumber, UhtMetaData? metaData = null)
+		protected UhtType(UhtHeaderFile headerFile, UhtType outer, int lineNumber, UhtMetaData? metaData = null)
 		{
-			Session = outer.Session;
+			Module = headerFile.Module;
+			_headerFile = headerFile;
 			Outer = outer;
 			SourceName = String.Empty;
 			LineNumber = lineNumber;
@@ -935,35 +929,38 @@ namespace EpicGames.UHT.Types
 		/// <param name="child">The child to be added.</param>
 		public virtual void AddChild(UhtType child)
 		{
-			_children ??= new();
-			_children.Add(child);
+			AddChildDirectly(child);
 		}
 
 		/// <summary>
-		/// Return the list of children attached to the type.  The 
-		/// list associated with the type is cleared.
+		/// Non-virtual method that always adds the child directly to the child list
 		/// </summary>
-		/// <returns>The current list of children</returns>
-		public List<UhtType>? DetachChildren()
+		/// <param name="child">Child to be added</param>
+		public void AddChildDirectly(UhtType child)
 		{
-			List<UhtType>? children = _children;
-			_children = null;
-			return children;
+			_children ??= new();
+			_children.Add(child);
 		}
 
 		/// <summary>
-		/// Merge the list children into the type.  The child's outer is set.
+		/// Move all children from the given source type to this type
 		/// </summary>
-		/// <param name="children">List of children to add.</param>
-		public void AddChildren(List<UhtType>? children)
+		/// <param name="source">Source of the children</param>
+		public void MoveChildren(UhtType source)
 		{
-			if (children != null)
+			if (_headerFile == null || _headerFile != source._headerFile)
 			{
-				foreach (UhtType child in children)
+				throw new UhtIceException("Children can not be moved to a package or between header files");
+			}
+			if (source._children != null)
+			{
+				_children ??= new();
+				foreach (UhtType child in source._children)
 				{
 					child.Outer = this;
-					AddChild(child);
+					_children.Add(child);
 				}
+				source._children = null;
 			}
 		}
 
@@ -1047,32 +1044,71 @@ namespace EpicGames.UHT.Types
 		{
 			if (_children != null)
 			{
-				if (phase == UhtResolvePhase.InvalidCheck)
+				ResolveChildren(_children, phase);
+			}
+		}
+
+		/// <summary>
+		/// Helper method to resolve a list of types
+		/// </summary>
+		/// <param name="children">Children to resolve</param>
+		/// <param name="phase">Phase of the resolve</param>
+		public static void ResolveChildren(List<UhtType> children, UhtResolvePhase phase)
+		{
+			if (phase == UhtResolvePhase.InvalidCheck)
+			{
+				int outIndex = 0;
+				for (int index = 0; index < children.Count; ++index)
 				{
-					int outIndex = 0;
-					for (int index = 0; index < _children.Count; ++index)
+					UhtType child = children[index];
+					if (child.Resolve(phase))
 					{
-						UhtType child = _children[index];
-						if (child.Resolve(phase))
-						{
-							_children[outIndex++] = child;
-						}
+						children[outIndex++] = child;
 					}
-					if (outIndex < _children.Count)
+					else
 					{
-						_children.RemoveRange(outIndex, _children.Count - outIndex);
+						child.IsValid = false;
 					}
 				}
-				else
+				if (outIndex < children.Count)
 				{
-					foreach (UhtType child in _children)
+					children.RemoveRange(outIndex, children.Count - outIndex);
+				}
+			}
+			else
+			{
+				foreach (UhtType child in children)
+				{
+					child.Resolve(phase);
+				}
+			}
+		}
+
+		/// <summary>
+		/// Remove any invalid children that were previously marked bad.  This is only
+		/// executed on a UhtPackage
+		/// </summary>
+		public void RemoveInvalidChildren()
+		{
+			if (_children != null)
+			{
+				int outIndex = 0;
+				for (int index = 0; index < _children.Count; ++index)
+				{
+					UhtType child = _children[index];
+					if (child.IsValid)
 					{
-						child.Resolve(phase);
+						_children[outIndex++] = child;
 					}
 				}
+				if (outIndex < _children.Count)
+				{
+					_children.RemoveRange(outIndex, _children.Count - outIndex);
+				}
 			}
 		}
 
+
 		/// <summary>
 		/// Resolve all the super and base structures and classes
 		/// </summary>
@@ -1250,8 +1286,7 @@ namespace EpicGames.UHT.Types
 					outer.AppendPathName(builder, stopOuter);
 
 					// SubObjectDelimiter is used to indicate that this object's outer is not a UPackage
-					// In the C# version of UHT, we key off the header file
-					if (outer is not UhtHeaderFile && outer.Outer is UhtHeaderFile)
+					if (outer is not UhtPackage && outer.Outer is UhtPackage)
 					{
 						builder.Append(UhtFCString.SubObjectDelimiter);
 					}
@@ -1365,7 +1400,7 @@ namespace EpicGames.UHT.Types
 			else if (outer is UhtScriptStruct)
 			{
 				// Structs can also have UPackage outer.
-				if (outer.Outer is not UhtHeaderFile)
+				if (outer.Outer is not UhtPackage)
 				{
 					AppendOuterNames(builder, outer.Outer);
 				}
@@ -1375,12 +1410,7 @@ namespace EpicGames.UHT.Types
 			else if (outer is UhtPackage package)
 			{
 				builder.Append('_');
-				builder.Append(package.ShortName);
-			}
-			else if (outer is UhtHeaderFile)
-			{
-				// Pickup the package
-				AppendOuterNames(builder, outer.Outer);
+				builder.Append(package.Module.ShortName);
 			}
 			else
 			{
diff --git a/Engine/Source/Programs/Shared/EpicGames.UHT/Utils/UhtExport.cs b/Engine/Source/Programs/Shared/EpicGames.UHT/Utils/UhtExport.cs
--- a/Engine/Source/Programs/Shared/EpicGames.UHT/Utils/UhtExport.cs
+++ b/Engine/Source/Programs/Shared/EpicGames.UHT/Utils/UhtExport.cs
@@ -71,12 +71,12 @@ namespace EpicGames.UHT.Utils
 		public string MakePath(UhtHeaderFile headerFile, string suffix);
 
 		/// <summary>
-		/// Make a path for an output based on the package name.
+		/// Make a path for an output based on the module name.
 		/// </summary>
-		/// <param name="package">Package being exported</param>
+		/// <param name="module">Module being exported</param>
 		/// <param name="suffix">Suffix to be added to the file name.</param>
 		/// <returns>Output file path</returns>
-		public string MakePath(UhtPackage package, string suffix);
+		public string MakePath(UhtModule module, string suffix);
 
 		/// <summary>
 		/// Make a path for the given file name and extension.  This is only valid for plugins.
diff --git a/Engine/Source/Programs/Shared/EpicGames.UHT/Utils/UhtSession.cs b/Engine/Source/Programs/Shared/EpicGames.UHT/Utils/UhtSession.cs
--- a/Engine/Source/Programs/Shared/EpicGames.UHT/Utils/UhtSession.cs
+++ b/Engine/Source/Programs/Shared/EpicGames.UHT/Utils/UhtSession.cs
@@ -204,18 +204,18 @@ namespace EpicGames.UHT.Utils
 		/// <returns>Resulting file name</returns>
 		public string MakePath(UhtHeaderFile headerFile, string suffix)
 		{
-			return MakePath(headerFile.Package.Module, headerFile.FileNameWithoutExtension, suffix);
+			return MakePath(headerFile.Module.Module, headerFile.FileNameWithoutExtension, suffix);
 		}
 
 		/// <summary>
 		/// Given a package file, generate the output file name
 		/// </summary>
-		/// <param name="package">Package file</param>
+		/// <param name="module">Module</param>
 		/// <param name="suffix">Suffix/extension to be added to the file name.</param>
 		/// <returns>Resulting file name</returns>
-		public string MakePath(UhtPackage package, string suffix)
+		public string MakePath(UhtModule module, string suffix)
 		{
-			return MakePath(package.Module, package.ShortName, suffix);
+			return MakePath(module.Module, module.ShortName, suffix);
 		}
 
 		/// <summary>
@@ -678,9 +678,9 @@ namespace EpicGames.UHT.Utils
 		public UHTManifest? Manifest => ManifestFile?.Manifest;
 
 		/// <summary>
-		/// Collection of packages from the manifest
+		/// Collection of modules from the manifest
 		/// </summary>
-		public IReadOnlyList<UhtPackage> Packages => _packages;
+		public IReadOnlyList<UhtModule> Modules => _modules;
 
 		/// <summary>
 		/// Collection of header files from the manifest.  The header files will also appear as the children 
@@ -783,7 +783,7 @@ namespace EpicGames.UHT.Utils
 		/// </summary>
 		public UhtScriptStruct? FStateTreePropertyRef { get; set; } = null;
 
-		private readonly List<UhtPackage> _packages = new();
+		private readonly List<UhtModule> _modules = new();
 		private readonly List<UhtHeaderFile> _headerFiles = new();
 		private readonly List<UhtHeaderFile> _sortedHeaderFiles = new();
 		private readonly Dictionary<string, UhtHeaderFile> _headerFileDictionary = new(StringComparer.OrdinalIgnoreCase);
@@ -1129,6 +1129,42 @@ namespace EpicGames.UHT.Utils
 			}
 		}
 
+		/// <summary>
+		/// Execute the given action on all the headers
+		/// </summary>
+		/// <param name="action">Action to be executed</param>
+		/// <param name="allowGoWide">If true, go wide if enabled.  Otherwise single threaded</param>
+		public void ForEachHeader(Action<UhtHeaderFile> action, bool allowGoWide)
+		{
+			if (!HasErrors)
+			{
+				ForEachHeaderNoErrorCheck(action, allowGoWide);
+			}
+		}
+
+		/// <summary>
+		/// Execute the given action on all the headers
+		/// </summary>
+		/// <param name="action">Action to be executed</param>
+		/// <param name="allowGoWide">If true, go wide if enabled.  Otherwise single threaded</param>
+		public void ForEachHeaderNoErrorCheck(Action<UhtHeaderFile> action, bool allowGoWide)
+		{
+			if (GoWide && allowGoWide)
+			{
+				Parallel.ForEach(_headerFiles, headerFile =>
+				{
+					TryHeader(action, headerFile);
+				});
+			}
+			else
+			{
+				foreach (UhtHeaderFile headerFile in _headerFiles)
+				{
+					TryHeader(action, headerFile);
+				}
+			}
+		}
+
 		/// <summary>
 		/// Read the given source file
 		/// </summary>
@@ -1812,20 +1848,6 @@ namespace EpicGames.UHT.Utils
 			}
 		}
 
-		/// <summary>
-		/// Return the normalized path converted to a full path if possible. 
-		/// Code should NOT depend on a full path being returned.
-		/// 
-		/// In general, it is assumed that during normal UHT, all paths are already full paths.
-		/// Only the test harness deals in relative paths.
-		/// </summary>
-		/// <param name="filePath">Path to normalize</param>
-		/// <returns>Normalized path possibly converted to a full path.</returns>
-		private string GetNormalizedFullFilePath(string filePath)
-		{
-			return FileManager!.GetFullFilePath(filePath).Replace('\\', '/');
-		}
-
 		private UhtHeaderFileParser? ParseHeaderFile(UhtHeaderFile headerFile)
 		{
 			UhtHeaderFileParser? parser = null;
@@ -1870,107 +1892,14 @@ namespace EpicGames.UHT.Utils
 			{
 				Log.Logger.LogTrace("Step - Prepare Modules");
 
-				foreach (UHTManifest.Module module in ManifestFile.Manifest.Modules)
+				foreach (UHTManifest.Module manifestModule in ManifestFile.Manifest.Modules)
 				{
-					EPackageFlags packageFlags = EPackageFlags.ContainsScript | EPackageFlags.Compiling;
-
-					switch (module.OverrideModuleType)
-					{
-						case EPackageOverrideType.None:
-							switch (module.ModuleType)
-							{
-								case UHTModuleType.GameEditor:
-								case UHTModuleType.EngineEditor:
-									packageFlags |= EPackageFlags.EditorOnly;
-									break;
-
-								case UHTModuleType.GameDeveloper:
-								case UHTModuleType.EngineDeveloper:
-									packageFlags |= EPackageFlags.Developer;
-									break;
-
-								case UHTModuleType.GameUncooked:
-								case UHTModuleType.EngineUncooked:
-									packageFlags |= EPackageFlags.UncookedOnly;
-									break;
-							}
-							break;
-
-						case EPackageOverrideType.EditorOnly:
-							packageFlags |= EPackageFlags.EditorOnly;
-							break;
-
-						case EPackageOverrideType.EngineDeveloper:
-						case EPackageOverrideType.GameDeveloper:
-							packageFlags |= EPackageFlags.Developer;
-							break;
-
-						case EPackageOverrideType.EngineUncookedOnly:
-						case EPackageOverrideType.GameUncookedOnly:
-							packageFlags |= EPackageFlags.UncookedOnly;
-							break;
-					}
-
-					UhtPackage package = new(this, module, packageFlags);
-					_packages.Add(package);
+					UhtModule module = new(this, manifestModule);
+					_modules.Add(module);
 				}
 			});
 		}
 
-		private void StepPrepareHeaders(UhtPackage package, IEnumerable<string> headerFiles, UhtHeaderFileType headerFileType)
-		{
-			if (package.Module == null)
-			{
-				return;
-			}
-
-			string typeDirectory = headerFileType.ToString() + '/';
-			string normalizedModuleBaseFullFilePath = GetNormalizedFullFilePath(package.Module.BaseDirectory);
-			foreach (string headerFilePath in headerFiles)
-			{
-
-				// Make sure this isn't a duplicate
-				string normalizedFullFilePath = GetNormalizedFullFilePath(headerFilePath);
-				string fileName = Path.GetFileName(normalizedFullFilePath);
-				if (_headerFileDictionary.TryGetValue(fileName, out UhtHeaderFile? existingHeaderFile) && existingHeaderFile != null)
-				{
-					string normalizedExistingFullFilePath = GetNormalizedFullFilePath(existingHeaderFile.FilePath);
-					if (!String.Equals(normalizedFullFilePath, normalizedExistingFullFilePath, StringComparison.OrdinalIgnoreCase))
-					{
-						IUhtMessageSite site = (IUhtMessageSite?)ManifestFile ?? this;
-						site.LogError($"Two headers with the same name is not allowed. '{headerFilePath}' conflicts with '{existingHeaderFile.FilePath}'");
-						continue;
-					}
-				}
-
-				// Create the header file and add to the collections
-				UhtHeaderFile headerFile = new(package, headerFilePath);
-				headerFile.HeaderFileType = headerFileType;
-				_headerFiles.Add(headerFile);
-				_headerFileDictionary.Add(fileName, headerFile);
-				package.AddChild(headerFile);
-
-				// Save metadata for the class path, both for it's include path and relative to the module base directory
-				if (normalizedFullFilePath.StartsWith(normalizedModuleBaseFullFilePath, true, null))
-				{
-					int stripLength = normalizedModuleBaseFullFilePath.Length;
-					if (stripLength < normalizedFullFilePath.Length && normalizedFullFilePath[stripLength] == '/')
-					{
-						++stripLength;
-					}
-
-					headerFile.ModuleRelativeFilePath = normalizedFullFilePath[stripLength..];
-
-					if (normalizedFullFilePath[stripLength..].StartsWith(typeDirectory, true, null))
-					{
-						stripLength += typeDirectory.Length;
-					}
-
-					headerFile.IncludeFilePath = normalizedFullFilePath[stripLength..];
-				}
-			}
-		}
-
 		private void StepPrepareHeaders()
 		{
 			if (ManifestFile == null)
@@ -1982,31 +1911,23 @@ namespace EpicGames.UHT.Utils
 			{
 				Log.Logger.LogTrace("Step - Prepare Headers");
 
-				foreach (UhtPackage package in _packages)
+				foreach (UhtModule module in Modules)
 				{
-					if (package.Module != null)
+					module.PrepareHeaders(headerFile =>
 					{
-						StepPrepareHeaders(package, package.Module.ClassesHeaders, UhtHeaderFileType.Classes);
-						StepPrepareHeaders(package, package.Module.PublicHeaders, UhtHeaderFileType.Public);
-						StepPrepareHeaders(package, package.Module.InternalHeaders, UhtHeaderFileType.Internal);
-						StepPrepareHeaders(package, package.Module.PrivateHeaders, UhtHeaderFileType.Private);
-					}
+						_headerFiles.Add(headerFile);
+						_headerFileDictionary.Add(Path.GetFileName(headerFile.FilePath), headerFile);
+					});
 				}
 
 				// Locate the NoExportTypes.h file and add it to every other header file
 				if (_headerFileDictionary.TryGetValue("NoExportTypes.h", out UhtHeaderFile? noExportTypes))
 				{
-					foreach (UhtPackage package in _packages)
+					foreach (UhtHeaderFile headerFile in _headerFiles)
 					{
-						foreach (UhtType type in package.Children)
+						if (headerFile != noExportTypes)
 						{
-							if (type is UhtHeaderFile headerFile)
-							{
-								if (headerFile != noExportTypes)
-								{
-									headerFile.AddReferencedHeader(noExportTypes);
-								}
-							}
+							headerFile.AddReferencedHeader(noExportTypes, UhtHeaderReferenceType.Passive);
 						}
 					}
 				}
@@ -2021,21 +1942,11 @@ namespace EpicGames.UHT.Utils
 			}
 
 			Log.Logger.LogTrace("Step - Parse Headers");
-
-			if (GoWide)
+			ForEachHeader(headerFile =>
 			{
-				Parallel.ForEach(_headerFiles, headerFile =>
-				{
-					ParseHeaderFile(headerFile);
-				});
-			}
-			else
-			{
-				foreach (UhtHeaderFile headerFile in _headerFiles)
-				{
-					ParseHeaderFile(headerFile);
-				}
-			}
+				headerFile.Read();
+				UhtHeaderFileParser.Parse(headerFile);
+			}, true);
 		}
 
 		private void StepPopulateTypeTable()
@@ -2063,7 +1974,7 @@ namespace EpicGames.UHT.Utils
 		private void StepBindSuperAndBases()
 		{
 			Log.Logger.LogTrace("Step - Bind super and bases");
-			StepForAllHeaders(headerFile => BindSuperAndBases(headerFile), true);
+			ForEachHeader(headerFile => headerFile.BindSuperAndBases(), true);
 		}
 
 		private void StepResolveBases()
@@ -2093,78 +2004,41 @@ namespace EpicGames.UHT.Utils
 		private void StepResolveValidate()
 		{
 			Log.Logger.LogTrace("Step - Resolve validate");
-			StepForAllHeaders(headerFile => UhtType.ValidateType(headerFile, UhtValidationOptions.None), true);
+			ForEachHeader(headerFile => headerFile.Validate(UhtValidationOptions.None), true);
 		}
 
 		private void StepCollectReferences()
 		{
 			Log.Logger.LogTrace("Step - Collect references");
-			StepForAllHeaders(headerFile =>
-			{
-				foreach (UhtType child in headerFile.Children)
-				{
-					child.CollectReferences(headerFile.References);
-				}
-				foreach (UhtHeaderFile refHeaderFile in headerFile.References.ReferencedHeaders)
-				{
-					headerFile.AddReferencedHeader(refHeaderFile);
-				}
-			}, true);
+			ForEachHeader(headerFile => headerFile.CollectReferences(), true);
 		}
 
 		private void ResolveAllHeaders(UhtResolvePhase resolvePhase)
 		{
-			StepForAllHeaders(headerFile => Resolve(headerFile, resolvePhase), resolvePhase.IsMultiThreadedResolvePhase());
-		}
-
-		private void Resolve(UhtHeaderFile headerFile, UhtResolvePhase resolvePhase)
-		{
-			try
+			ForEachHeader(headerFile => headerFile.Resolve(resolvePhase), resolvePhase.IsMultiThreadedResolvePhase());
+			if (resolvePhase == UhtResolvePhase.InvalidCheck)
 			{
-				headerFile.Resolve(resolvePhase);
-			}
-			catch (Exception e)
-			{
-				HandleException(headerFile.MessageSource, e);
+				foreach (UhtModule module in Modules)
+				{
+					foreach (UhtPackage package in module.Packages)
+					{
+						package.RemoveInvalidChildren();
+					}
+				}
 			}
 		}
 
-		private void BindSuperAndBases(UhtHeaderFile headerFile)
+		private void TryHeader(Action<UhtHeaderFile> action, UhtHeaderFile headerFile)
 		{
 			try
 			{
-				headerFile.BindSuperAndBases();
+				action(headerFile);
 			}
 			catch (Exception e)
 			{
 				HandleException(headerFile.MessageSource, e);
 			}
 		}
-
-		private delegate void StepDelegate(UhtHeaderFile headerFile);
-
-		private void StepForAllHeaders(StepDelegate stepDelegate, bool allowGoWide)
-		{
-			if (HasErrors)
-			{
-				return;
-			}
-
-			if (GoWide && allowGoWide)
-			{
-				Parallel.ForEach(_headerFiles, headerFile =>
-				{
-					stepDelegate(headerFile);
-				});
-			}
-			else
-			{
-				foreach (UhtHeaderFile headerFile in _headerFiles)
-				{
-					stepDelegate(headerFile);
-				}
-			}
-		}
 		#endregion
 
 		#region Symbol table initialization
@@ -2172,7 +2046,18 @@ namespace EpicGames.UHT.Utils
 		{
 			foreach (UhtHeaderFile headerFile in _headerFiles)
 			{
-				AddTypeToSymbolTable(headerFile);
+				foreach (UhtType child in headerFile.Children)
+				{
+					if (child.Outer is UhtPackage package)
+					{
+						package.AddChildDirectly(child);
+					}
+					else
+					{
+						throw new UhtIceException($"Type \"{child.SourceName}\" is a root level type but does not have a package outer");
+					}
+					AddTypeToSymbolTable(child);
+				}
 			}
 			_symbolTablePopulated = true;
 		}
@@ -2331,6 +2216,17 @@ namespace EpicGames.UHT.Utils
 			});
 		}
 
+		private void TopologicalStructVisitChildren(List<TopologicalState> states, UhtHeaderFile visit, List<UhtStruct> structStack)
+		{
+			foreach (UhtType child in visit.Children)
+			{
+				if (child is UhtStruct childStruct)
+				{
+					TopologicalStructVisit(states, childStruct, structStack);
+				}
+			}
+		}
+
 		private void TopologicalStructVisitChildren(List<TopologicalState> states, UhtType visit, List<UhtStruct> structStack)
 		{
 			foreach (UhtType child in visit.Children)
diff --git a/Engine/Source/Programs/Shared/EpicGames.UHT/Utils/UhtStringBuilderExtensions.cs b/Engine/Source/Programs/Shared/EpicGames.UHT/Utils/UhtStringBuilderExtensions.cs
--- a/Engine/Source/Programs/Shared/EpicGames.UHT/Utils/UhtStringBuilderExtensions.cs
+++ b/Engine/Source/Programs/Shared/EpicGames.UHT/Utils/UhtStringBuilderExtensions.cs
@@ -194,7 +194,7 @@ namespace EpicGames.UHT.Utils
 		{
 			if (!type.MetaData.IsEmpty())
 			{
-				bool isPartOfEngine = type.Package.IsPartOfEngine;
+				bool isPartOfEngine = type.Module.IsPartOfEngine;
 				List<KeyValuePair<string, string>> sortedMetaData = type.MetaData.GetSorted();
 				builder.AppendTabs(tabs).Append("static constexpr UECodeGen_Private::FMetaDataPairParam ").AppendNameDecl(namePrefix, name, nameSuffix).Append(metaNameSuffix).Append("[] = {\r\n");
 				foreach (KeyValuePair<string, string> kvp in sortedMetaData)
diff --git a/Engine/Source/Programs/Shared/EpicGames.UHT/Utils/UhtSymbolTable.cs b/Engine/Source/Programs/Shared/EpicGames.UHT/Utils/UhtSymbolTable.cs
--- a/Engine/Source/Programs/Shared/EpicGames.UHT/Utils/UhtSymbolTable.cs
+++ b/Engine/Source/Programs/Shared/EpicGames.UHT/Utils/UhtSymbolTable.cs
@@ -247,7 +247,7 @@ namespace EpicGames.UHT.Utils
 			{
 				for (int index = lookup.SymbolIndex; index != 0; index = _symbols[index].NextIndex)
 				{
-					if (IsMatch(options, index, lookup.CasedIndex) && _symbols[index].Type.Outer is UhtHeaderFile)
+					if (IsMatch(options, index, lookup.CasedIndex) && _symbols[index].Type.Outer is UhtPackage)
 					{
 						return _symbols[index].Type;
 					}
@@ -376,6 +376,25 @@ namespace EpicGames.UHT.Utils
 			return null;
 		}
 
+		/// <summary>
+		/// Lookup the given name 
+		/// </summary>
+		/// <param name="header">Header file being searched.</param>
+		/// <param name="options">Options controlling what is search and what is returned.</param>
+		/// <param name="lookup">Starting lookup location.</param>
+		/// <returns>Found type or null if not found.</returns>
+		private UhtType? FindSymbolChain(UhtHeaderFile header, UhtFindOptions options, ref Lookup lookup)
+		{
+			for (int index = lookup.SymbolIndex; index != 0; index = _symbols[index].NextIndex)
+			{
+				if (IsMatch(options, index, lookup.CasedIndex) && _symbols[index].Type.HeaderFile == header && _symbols[index].Type.Outer is UhtPackage)
+				{
+					return _symbols[index].Type;
+				}
+			}
+			return null;
+		}
+
 		/// <summary>
 		/// Add a new type to the given symbol chain
 		/// </summary>
