EpicGames.Redis: Support for serializing DateTime types.

#rnx

diff --git a/Engine/Source/Programs/Shared/EpicGames.Redis/RedisSerializer.cs b/Engine/Source/Programs/Shared/EpicGames.Redis/RedisSerializer.cs
--- a/Engine/Source/Programs/Shared/EpicGames.Redis/RedisSerializer.cs
+++ b/Engine/Source/Programs/Shared/EpicGames.Redis/RedisSerializer.cs
@@ -163,6 +163,7 @@ namespace EpicGames.Redis
 				CreateNativeConverter(x => (ReadOnlyMemory<byte>)x, x => x),
 				CreateNativeConverter(x => (byte[])x!, x => x),
 				CreateNativeConverter(x => (string)x!, x => x),
+				CreateNativeConverter(x => new DateTime((long)x, DateTimeKind.Utc), x => x.ToUniversalTime().Ticks)
 			};
 			return new Dictionary<Type, object>(converters);
 		}
