Fixed some code analyzer warnings, misspellings, and minor issues.

#rnx

diff --git a/Engine/Source/Programs/Shared/EpicGames.UHT/Exporters/CodeGen/UhtCodeGenerator.cs b/Engine/Source/Programs/Shared/EpicGames.UHT/Exporters/CodeGen/UhtCodeGenerator.cs
--- a/Engine/Source/Programs/Shared/EpicGames.UHT/Exporters/CodeGen/UhtCodeGenerator.cs
+++ b/Engine/Source/Programs/Shared/EpicGames.UHT/Exporters/CodeGen/UhtCodeGenerator.cs
@@ -47,7 +47,7 @@ namespace EpicGames.UHT.Exporters.CodeGen
 		{
 			public string RegisteredSingletonName { get; set; }
 			public string UnregisteredSingletonName { get; set; }
-			public string RegsiteredExternalDecl { get; set; }
+			public string RegisteredExternalDecl { get; set; }
 			public string UnregisteredExternalDecl { get; set; }
 			public UhtClass? NativeInterface { get; set; }
 			public UhtProperty? FastArrayProperty { get; set; }
@@ -199,7 +199,7 @@ namespace EpicGames.UHT.Exporters.CodeGen
 		/// <returns>External declaration</returns>
 		public string GetExternalDecl(int objectIndex, bool registered)
 		{
-			return registered ? ObjectInfos[objectIndex].RegsiteredExternalDecl : ObjectInfos[objectIndex].UnregisteredExternalDecl;
+			return registered ? ObjectInfos[objectIndex].RegisteredExternalDecl : ObjectInfos[objectIndex].UnregisteredExternalDecl;
 		}
 		#endregion
 
@@ -229,7 +229,7 @@ namespace EpicGames.UHT.Exporters.CodeGen
 			builder.Append("Z_Construct_UPackage_");
 			builder.Append(packageInfo.StrippedName);
 			objectInfo.UnregisteredSingletonName = objectInfo.RegisteredSingletonName = builder.ToString();
-			objectInfo.UnregisteredExternalDecl = objectInfo.RegsiteredExternalDecl = $"\tUPackage* {objectInfo.RegisteredSingletonName}();\r\n";
+			objectInfo.UnregisteredExternalDecl = objectInfo.RegisteredExternalDecl = $"\tUPackage* {objectInfo.RegisteredSingletonName}();\r\n";
 		}
 
 		private void InitHeaderInfo(StringBuilder builder, UhtHeaderFile headerFile)
@@ -365,12 +365,12 @@ namespace EpicGames.UHT.Exporters.CodeGen
 				objectInfo.UnregisteredSingletonName = builder.ToString();
 
 				objectInfo.UnregisteredExternalDecl = $"\t{module.Api}U{engineClassName}* {objectInfo.UnregisteredSingletonName}();\r\n";
-				objectInfo.RegsiteredExternalDecl = $"\t{module.Api}U{engineClassName}* {objectInfo.RegisteredSingletonName}();\r\n";
+				objectInfo.RegisteredExternalDecl = $"\t{module.Api}U{engineClassName}* {objectInfo.RegisteredSingletonName}();\r\n";
 			}
 			else
 			{
 				objectInfo.UnregisteredSingletonName = objectInfo.RegisteredSingletonName = builder.ToString();
-				objectInfo.UnregisteredExternalDecl = objectInfo.RegsiteredExternalDecl = $"\t{module.Api}U{engineClassName}* {objectInfo.RegisteredSingletonName}();\r\n";
+				objectInfo.UnregisteredExternalDecl = objectInfo.RegisteredExternalDecl = $"\t{module.Api}U{engineClassName}* {objectInfo.RegisteredSingletonName}();\r\n";
 			}
 
 			// Init the children
diff --git a/Engine/Source/Programs/Shared/EpicGames.UHT/Exporters/CodeGen/UhtHeaderCodeGeneratorCppFile.cs b/Engine/Source/Programs/Shared/EpicGames.UHT/Exporters/CodeGen/UhtHeaderCodeGeneratorCppFile.cs
--- a/Engine/Source/Programs/Shared/EpicGames.UHT/Exporters/CodeGen/UhtHeaderCodeGeneratorCppFile.cs
+++ b/Engine/Source/Programs/Shared/EpicGames.UHT/Exporters/CodeGen/UhtHeaderCodeGeneratorCppFile.cs
@@ -219,7 +219,7 @@ namespace EpicGames.UHT.Exporters.CodeGen
 
 				int generatedBodyStart = builder.Length;
 
-				SortedDictionary<UhtPackage, UhtRegistrations> packageRegistrations = new();
+				Dictionary<UhtPackage, UhtRegistrations> packageRegistrations = new();
 				foreach (UhtField field in HeaderFile.References.ExportTypes)
 				{
 					if (field is UhtEnum enumObj)
@@ -2027,7 +2027,7 @@ namespace EpicGames.UHT.Exporters.CodeGen
 			return outScriptStructs;
 		}
 
-		private UhtRegistrations GetRegistrations(SortedDictionary<UhtPackage, UhtRegistrations> packageRegistrations, UhtField fieldObj)
+		private static UhtRegistrations GetRegistrations(Dictionary<UhtPackage, UhtRegistrations> packageRegistrations, UhtField fieldObj)
 		{
 			UhtPackage package = fieldObj.Package;
 			if (packageRegistrations.TryGetValue(package, out UhtRegistrations? registrations))
diff --git a/Engine/Source/Programs/Shared/EpicGames.UHT/Exporters/CodeGen/UhtPackageCodeGeneratorCppFile.cs b/Engine/Source/Programs/Shared/EpicGames.UHT/Exporters/CodeGen/UhtPackageCodeGeneratorCppFile.cs
--- a/Engine/Source/Programs/Shared/EpicGames.UHT/Exporters/CodeGen/UhtPackageCodeGeneratorCppFile.cs
+++ b/Engine/Source/Programs/Shared/EpicGames.UHT/Exporters/CodeGen/UhtPackageCodeGeneratorCppFile.cs
@@ -106,8 +106,7 @@ namespace EpicGames.UHT.Exporters.CodeGen
 
 				foreach (UhtObject obj in singletons)
 				{
-					ref UhtCodeGenerator.ObjectInfo info = ref ObjectInfos[obj.ObjectTypeIndex];
-					builder.Append(info.RegsiteredExternalDecl);
+					builder.Append(GetExternalDecl(obj, true));
 				}
 
 				foreach (UhtPackage package in Module.Packages)
diff --git a/Engine/Source/Programs/Shared/EpicGames.UHT/Utils/UhtExport.cs b/Engine/Source/Programs/Shared/EpicGames.UHT/Utils/UhtExport.cs
--- a/Engine/Source/Programs/Shared/EpicGames.UHT/Utils/UhtExport.cs
+++ b/Engine/Source/Programs/Shared/EpicGames.UHT/Utils/UhtExport.cs
@@ -73,10 +73,10 @@ namespace EpicGames.UHT.Utils
 		/// <summary>
 		/// Make a path for an output based on the module name.
 		/// </summary>
-		/// <param name="module">Module being exported</param>
+		/// <param name="moduleObj">Module being exported</param>
 		/// <param name="suffix">Suffix to be added to the file name.</param>
 		/// <returns>Output file path</returns>
-		public string MakePath(UhtModule module, string suffix);
+		public string MakePath(UhtModule moduleObj, string suffix);
 
 		/// <summary>
 		/// Make a path for the given file name and extension.  This is only valid for plugins.
