Add initial Verse support in UHT.  This does not include the runtime support.

#rnx

diff --git a/Engine/Source/Programs/Shared/EpicGames.Core/UHTTypes.cs b/Engine/Source/Programs/Shared/EpicGames.Core/UHTTypes.cs
--- a/Engine/Source/Programs/Shared/EpicGames.Core/UHTTypes.cs
+++ b/Engine/Source/Programs/Shared/EpicGames.Core/UHTTypes.cs
@@ -124,6 +124,7 @@ namespace EpicGames.Core
 			[JsonConverter(typeof(JsonStringEnumConverter))]
 			public UHTVerseScope VerseScope { get; set; } = UHTVerseScope.PublicAPI;
 			public bool HasVerse { get; set; } = false;
+			public string VersePluginName { get; set; } = "";
 
 			public override string ToString()
 			{
diff --git a/Engine/Source/Programs/Shared/EpicGames.UHT/EpicGames.UHT.csproj b/Engine/Source/Programs/Shared/EpicGames.UHT/EpicGames.UHT.csproj
--- a/Engine/Source/Programs/Shared/EpicGames.UHT/EpicGames.UHT.csproj
+++ b/Engine/Source/Programs/Shared/EpicGames.UHT/EpicGames.UHT.csproj
@@ -48,6 +48,7 @@
   <ItemGroup>
     <PackageReference Include="Microsoft.CSharp" Version="4.7.0" />
     <PackageReference Include="Microsoft.Win32.Registry" Version="5.0.0" />
+    <PackageReference Include="System.Collections.Specialized" Version="4.3.0" />
     <PackageReference Include="System.Data.DataSetExtensions" Version="4.5.0" />
   </ItemGroup>
 </Project>
\ No newline at end of file
diff --git a/Engine/Source/Programs/Shared/EpicGames.UHT/Exporters/CodeGen/UhtCodeGenerator.cs b/Engine/Source/Programs/Shared/EpicGames.UHT/Exporters/CodeGen/UhtCodeGenerator.cs
--- a/Engine/Source/Programs/Shared/EpicGames.UHT/Exporters/CodeGen/UhtCodeGenerator.cs
+++ b/Engine/Source/Programs/Shared/EpicGames.UHT/Exporters/CodeGen/UhtCodeGenerator.cs
@@ -40,6 +40,10 @@ namespace EpicGames.UHT.Exporters.CodeGen
 			public bool NeedsPushModelHeaders { get; set; }
 			public bool NeedsFastArrayHeaders { get; set; }
 			public bool NeedsVerseHeaders { get; set; }
+			public bool NeedsVerseClass { get; set; }
+			public bool NeedsVerseStruct { get; set; }
+			public bool NeedsVerseEnum { get; set; }
+			public bool NeedsVerseInterop { get; set; }
 		}
 		public HeaderInfo[] HeaderInfos { get; set; }
 
@@ -291,6 +295,8 @@ namespace EpicGames.UHT.Exporters.CodeGen
 					InitObjectInfo(builder, package, ref headerInfo, obj);
 				}
 			}
+
+			headerInfo.NeedsVerseInterop = headerInfo.NeedsVerseStruct;
 		}
 
 		private void InitObjectInfo(StringBuilder builder, UhtPackage package, ref HeaderInfo headerInfo, UhtObject obj)
@@ -315,6 +321,14 @@ namespace EpicGames.UHT.Exporters.CodeGen
 				{
 					headerInfo.NeedsPushModelHeaders = true;
 				}
+				if (classObj.IsVerseField)
+				{
+					headerInfo.NeedsVerseClass = true;
+				}
+				if (classObj.Children.Any(x => x is UhtVerseValueProperty))
+				{
+					headerInfo.NeedsVerseHeaders = true;
+				}
 				if (classObj.ClassType == UhtClassType.NativeInterface)
 				{
 					if (classObj.AlternateObject != null)
@@ -322,15 +336,23 @@ namespace EpicGames.UHT.Exporters.CodeGen
 						ObjectInfos[classObj.AlternateObject.ObjectTypeIndex].NativeInterface = classObj;
 					}
 				}
-				headerInfo.NeedsVerseHeaders = classObj.Children.Any(x => x is UhtVerseValueProperty);
 			}
-			else if (obj is UhtScriptStruct scriptStruct)
+			else if (obj is UhtScriptStruct scriptStructObj)
 			{
+				if (scriptStructObj.IsVerseField)
+				{
+					headerInfo.NeedsVerseStruct = true;
+				}
+				if (scriptStructObj.Children.Any(x => x is UhtVerseValueProperty))
+				{
+					headerInfo.NeedsVerseHeaders = true;
+				}
+
 				// Check to see if we are a FastArraySerializer and should try to deduce the FastArraySerializerItemType
 				// To fulfill that requirement the struct should be derived from FFastArraySerializer and have a single replicated TArrayProperty
-				if (scriptStruct.IsChildOf(FastArraySerializer))
+				if (scriptStructObj.IsChildOf(FastArraySerializer))
 				{
-					foreach (UhtType child in scriptStruct.Children)
+					foreach (UhtType child in scriptStructObj.Children)
 					{
 						if (child is UhtProperty property)
 						{
@@ -350,13 +372,19 @@ namespace EpicGames.UHT.Exporters.CodeGen
 						headerInfo.NeedsFastArrayHeaders = true;
 					}
 				}
-				headerInfo.NeedsVerseHeaders = scriptStruct.Children.Any(x => x is UhtVerseValueProperty);
 			}
 			else if (obj is UhtFunction)
 			{
 				// The method for EngineClassName returns type specific where in this case we need just the simple return type
 				engineClassName = "Function";
 			}
+			else if (obj is UhtEnum enumObj)
+			{
+				if (enumObj.IsVerseField)
+				{
+					headerInfo.NeedsVerseEnum = true;
+				}
+			}
 
 			if (isNonIntrinsicClass)
 			{
diff --git a/Engine/Source/Programs/Shared/EpicGames.UHT/Exporters/CodeGen/UhtHeaderCodeGeneratorCppFile.cs b/Engine/Source/Programs/Shared/EpicGames.UHT/Exporters/CodeGen/UhtHeaderCodeGeneratorCppFile.cs
--- a/Engine/Source/Programs/Shared/EpicGames.UHT/Exporters/CodeGen/UhtHeaderCodeGeneratorCppFile.cs
+++ b/Engine/Source/Programs/Shared/EpicGames.UHT/Exporters/CodeGen/UhtHeaderCodeGeneratorCppFile.cs
@@ -109,6 +109,18 @@ namespace EpicGames.UHT.Exporters.CodeGen
 				{
 					includesToAdd.Add("Net/Serialization/FastArraySerializerImplementation.h");
 				}
+				if (headerInfo.NeedsVerseClass)
+				{
+					includesToAdd.Add("VerseVM/VVMVerseClass.h");
+				}
+				if (headerInfo.NeedsVerseStruct)
+				{
+					includesToAdd.Add("VerseVM/VVMVerseStruct.h");
+				}
+				if (headerInfo.NeedsVerseEnum)
+				{
+					includesToAdd.Add("VerseVM/VVMVerseEnum.h");
+				}
 
 				foreach (UhtType type in HeaderFile.Children)
 				{
@@ -761,7 +773,14 @@ namespace EpicGames.UHT.Exporters.CodeGen
 				}
 				builder.Append("\tif (!").Append(innerSingletonName).Append(")\r\n");
 				builder.Append("\t{\r\n");
-				builder.Append("\t\tUECodeGen_Private::ConstructUScriptStruct(").Append(innerSingletonName).Append(", ").Append(staticsName).Append("::StructParams);\r\n");
+				if (scriptStruct.IsVerseField)
+				{
+					builder.Append("\t\tUECodeGen_Private::ConstructUScriptStruct(").Append(innerSingletonName).Append(", ").Append(staticsName).Append("::StructParams, Verse::Private::ConstructUScriptStructFunc);\r\n");
+				}
+				else 
+				{
+					builder.Append("\t\tUECodeGen_Private::ConstructUScriptStruct(").Append(innerSingletonName).Append(", ").Append(staticsName).Append("::StructParams);\r\n");
+				}
 				builder.Append("\t}\r\n");
 				builder.Append("\treturn ").Append(innerSingletonName).Append(";\r\n");
 				builder.Append("}\r\n");
@@ -1641,7 +1660,14 @@ namespace EpicGames.UHT.Exporters.CodeGen
 
 			bool hasInterfaces = classObj.Bases.Any(x => x is UhtClass baseClass && baseClass.ClassFlags.HasAnyFlags(EClassFlags.Interface));
 
-			builder.Append("IMPLEMENT_CLASS_NO_AUTO_REGISTRATION(").Append(classObj.SourceName).Append(");\r\n");
+			if (classObj.IsVerseField)
+			{
+				builder.Append("IMPLEMENT_CLASS_NO_AUTO_REGISTRATION_WITH_METHOD(").Append(classObj.SourceName).Append(", Verse::Private::ConstructUClass);\r\n");
+			}
+			else
+			{
+				builder.Append("IMPLEMENT_CLASS_NO_AUTO_REGISTRATION(").Append(classObj.SourceName).Append(");\r\n");
+			}
 
 			// Everything from this point on will be part of the definition hash
 			int hashCodeBlockStart = builder.Length;
diff --git a/Engine/Source/Programs/Shared/EpicGames.UHT/Exporters/CodeGen/UhtHeaderCodeGeneratorHFile.cs b/Engine/Source/Programs/Shared/EpicGames.UHT/Exporters/CodeGen/UhtHeaderCodeGeneratorHFile.cs
--- a/Engine/Source/Programs/Shared/EpicGames.UHT/Exporters/CodeGen/UhtHeaderCodeGeneratorHFile.cs
+++ b/Engine/Source/Programs/Shared/EpicGames.UHT/Exporters/CodeGen/UhtHeaderCodeGeneratorHFile.cs
@@ -65,6 +65,10 @@ namespace EpicGames.UHT.Exporters.CodeGen
 				{
 					builder.Append("#include \"UObject/VerseTypes.h\"\r\n");
 				}
+				if (headerInfo.NeedsVerseInterop)
+				{
+					builder.Append("#include \"VerseInteropMacros.h\"\r\n");
+				}
 				builder.Append("\r\n");
 				builder.Append(DisableDeprecationWarnings).Append("\r\n");
 
@@ -419,6 +423,13 @@ namespace EpicGames.UHT.Exporters.CodeGen
 				// Forward declare the StaticStruct specialization in the header
 				builder.Append("template<> ").Append(Module.Api).Append("UScriptStruct* StaticStruct<struct ").Append(scriptStruct.SourceName).Append(">();\r\n");
 				builder.Append("\r\n");
+
+				// Verse declarations
+				if (scriptStruct.IsVerseField)
+				{
+					builder.Append("VUHT_DECLARE_STRUCT_RUNTIME_TRAITS(struct ").Append(scriptStruct.SourceName).Append(")\r\n");
+					builder.Append("\r\n");
+				}
 			}
 			return builder;
 		}
diff --git a/Engine/Source/Programs/Shared/EpicGames.UHT/Parsers/UhtPropertyParser.cs b/Engine/Source/Programs/Shared/EpicGames.UHT/Parsers/UhtPropertyParser.cs
--- a/Engine/Source/Programs/Shared/EpicGames.UHT/Parsers/UhtPropertyParser.cs
+++ b/Engine/Source/Programs/Shared/EpicGames.UHT/Parsers/UhtPropertyParser.cs
@@ -1178,6 +1178,17 @@ namespace EpicGames.UHT.Parsers
 				topScope.AddFormattedCommentsAsTooltipMetaData(newProperty);
 			}
 
+			// Adjust the name for verse
+			(bool wasMangled, string result) = newProperty.GetMangledEngineName();
+			if (wasMangled)
+			{
+				if (!newProperty.MetaData.ContainsKey(UhtNames.DisplayName))
+				{
+					newProperty.MetaData.Add(UhtNames.DisplayName, newProperty.StrippedEngineName);
+				}
+				newProperty.EngineName = result;
+			}
+
 			propertyDelegate(topScope, newProperty, ref nameToken, layoutMacroType);
 
 			// Void properties don't get added when they are the return value
diff --git a/Engine/Source/Programs/Shared/EpicGames.UHT/Specifiers/UhtFieldSpecifiers.cs b/Engine/Source/Programs/Shared/EpicGames.UHT/Specifiers/UhtFieldSpecifiers.cs
new file mode 100644
--- /dev/null
+++ b/Engine/Source/Programs/Shared/EpicGames.UHT/Specifiers/UhtFieldSpecifiers.cs
@@ -0,0 +1,69 @@
+// Copyright Epic Games, Inc. All Rights Reserved.
+
+using System;
+using System.Collections.Generic;
+using System.Diagnostics.CodeAnalysis;
+using EpicGames.Core;
+using EpicGames.UHT.Tables;
+using EpicGames.UHT.Types;
+using EpicGames.UHT.Utils;
+
+namespace EpicGames.UHT.Parsers
+{
+
+	/// <summary>
+	/// Collection of UENUM specifiers
+	/// </summary>
+	[UnrealHeaderTool]
+	[SuppressMessage("CodeQuality", "IDE0051:Remove unused private members", Justification = "Attribute accessed method")]
+	public static class UhtFieldSpecifiers
+	{
+		[UhtSpecifier(Extends = UhtTableNames.Struct, ValueType = UhtSpecifierValueType.KeyValuePairList)]
+		[UhtSpecifier(Extends = UhtTableNames.Class, ValueType = UhtSpecifierValueType.KeyValuePairList)]
+		private static void VerseSpecifier(UhtSpecifierContext specifierContext, List<KeyValuePair<StringView, StringView>> values)
+		{
+			UhtField fieldObj = (UhtField)specifierContext.Type;
+
+			foreach (KeyValuePair<StringView, StringView> kvp in values)
+			{
+				ReadOnlySpan<char> key = kvp.Key.Span;
+				if (key.Equals("name", StringComparison.OrdinalIgnoreCase))
+				{
+					fieldObj.VerseName = kvp.Value.ToString();
+				}
+				else if (key.Equals("module", StringComparison.OrdinalIgnoreCase))
+				{
+					fieldObj.VerseModule = kvp.Value.ToString();
+				}
+				else
+				{
+					fieldObj.LogError($"Verse specifier option '{key}' is unknown");
+				}
+			}
+
+			if (!fieldObj.IsVerseField)
+			{
+				fieldObj.LogError($"Verse specifier must include the name");
+			}
+			else if(fieldObj.Outer is not UhtPackage)
+			{
+				fieldObj.LogError($"Verse specifier can only appear on top level classes and structures");
+			}
+			else
+			{
+				if (String.IsNullOrEmpty(fieldObj.VerseModule))
+				{
+					fieldObj.EngineName = fieldObj.VerseName!;
+				}
+				else
+				{
+					fieldObj.EngineName = $"{fieldObj.VerseModule}_{fieldObj.VerseName!}";
+				}
+
+				using BorrowStringBuilder borrowBuilder = new(StringBuilderCache.Small);
+				borrowBuilder.StringBuilder.AppendVerseUEVNIPackageName(fieldObj);
+				fieldObj.Outer = fieldObj.Module.CreatePackage(borrowBuilder.StringBuilder.ToString());
+			}
+		}
+	}
+}
diff --git a/Engine/Source/Programs/Shared/EpicGames.UHT/Specifiers/UhtPropertyMemberSpecifiers.cs b/Engine/Source/Programs/Shared/EpicGames.UHT/Specifiers/UhtPropertyMemberSpecifiers.cs
--- a/Engine/Source/Programs/Shared/EpicGames.UHT/Specifiers/UhtPropertyMemberSpecifiers.cs
+++ b/Engine/Source/Programs/Shared/EpicGames.UHT/Specifiers/UhtPropertyMemberSpecifiers.cs
@@ -469,5 +469,12 @@ namespace EpicGames.UHT.Parsers
 			UhtPropertySpecifierContext context = (UhtPropertySpecifierContext)specifierContext;
 			context.PropertySettings.PropertyExportFlags |= UhtPropertyExportFlags.FieldNotify;
 		}
+
+		[UhtSpecifier(Extends = UhtTableNames.PropertyMember, ValueType = UhtSpecifierValueType.String)]
+		private static void VerseNameSpecifier(UhtSpecifierContext specifierContext, StringView value)
+		{
+			UhtPropertySpecifierContext context = (UhtPropertySpecifierContext)specifierContext;
+			context.PropertySettings.VerseName = value.ToString();
+		}
 	}
 }
diff --git a/Engine/Source/Programs/Shared/EpicGames.UHT/Types/Properties/UhtBoolProperty.cs b/Engine/Source/Programs/Shared/EpicGames.UHT/Types/Properties/UhtBoolProperty.cs
--- a/Engine/Source/Programs/Shared/EpicGames.UHT/Types/Properties/UhtBoolProperty.cs
+++ b/Engine/Source/Programs/Shared/EpicGames.UHT/Types/Properties/UhtBoolProperty.cs
@@ -110,6 +110,12 @@ namespace EpicGames.UHT.Types
 		[JsonConverter(typeof(JsonStringEnumConverter))]
 		public UhtBoolType BoolType { get; }
 
+		/// <summary>
+		/// Return the engine name without and 'b' prefixes
+		/// </summary>
+		[JsonIgnore]
+		public override string StrippedEngineName => EngineName.StartsWith('b') ? EngineName.Substring(1) : EngineName;
+
 		/// <summary>
 		/// Construct a new boolean property
 		/// </summary>
diff --git a/Engine/Source/Programs/Shared/EpicGames.UHT/Types/UhtClass.cs b/Engine/Source/Programs/Shared/EpicGames.UHT/Types/UhtClass.cs
--- a/Engine/Source/Programs/Shared/EpicGames.UHT/Types/UhtClass.cs
+++ b/Engine/Source/Programs/Shared/EpicGames.UHT/Types/UhtClass.cs
@@ -533,6 +533,10 @@ namespace EpicGames.UHT.Types
 		{
 			get
 			{
+				if (IsVerseField)
+				{
+					return "";
+				}
 				switch (ClassType)
 				{
 					case UhtClassType.Class:
@@ -1461,10 +1465,13 @@ namespace EpicGames.UHT.Types
 			options = base.Validate(options);
 
 			// Classes must start with a valid prefix
-			string expectedClassName = EngineNamePrefix + EngineName;
-			if (expectedClassName != SourceName)
+			if (!IsVerseField)
 			{
-				this.LogError($"Class '{SourceName}' has an invalid Unreal prefix, expecting '{expectedClassName}'");
+				string expectedClassName = EngineNamePrefix + EngineName;
+				if (expectedClassName != SourceName)
+				{
+					this.LogError($"Class '{SourceName}' has an invalid Unreal prefix, expecting '{expectedClassName}'");
+				}
 			}
 
 			// If we have a super class
diff --git a/Engine/Source/Programs/Shared/EpicGames.UHT/Types/UhtField.cs b/Engine/Source/Programs/Shared/EpicGames.UHT/Types/UhtField.cs
--- a/Engine/Source/Programs/Shared/EpicGames.UHT/Types/UhtField.cs
+++ b/Engine/Source/Programs/Shared/EpicGames.UHT/Types/UhtField.cs
@@ -1,5 +1,8 @@
 // Copyright Epic Games, Inc. All Rights Reserved.
 
+using System.Text;
+using EpicGames.UHT.Utils;
+
 namespace EpicGames.UHT.Types
 {
 
@@ -11,6 +14,21 @@ namespace EpicGames.UHT.Types
 		/// <inheritdoc/>
 		public override string EngineClassName => "Field";
 
+		/// <summary>
+		/// Name of the module containing the type
+		/// </summary>
+		public string? VerseModule { get; set; } = null;
+
+		/// <summary>
+		/// Cased name of the verse field
+		/// </summary>
+		public string? VerseName { get; set; } = null;
+
+		/// <summary>
+		/// Returns true if field is a verse element
+		/// </summary>
+		public bool IsVerseField => VerseName != null;
+
 		/// <summary>
 		/// Construct a new field
 		/// </summary>
@@ -21,4 +39,93 @@ namespace EpicGames.UHT.Types
 		{
 		}
 	}
+
+	/// <summary>
+	/// Helper extension methods for fields
+	/// </summary>
+	public static class UhtFieldStringBuilderExtensions
+	{
+
+		/// <summary>
+		/// Append the Verse UE VNI package name
+		/// </summary>
+		/// <param name="builder">Destination builder</param>
+		/// <param name="fieldObj">Field to serialize</param>
+		/// <returns>Builder</returns>
+		public static StringBuilder AppendVerseUEVNIPackageName(this StringBuilder builder, UhtField fieldObj)
+		{
+			if (!fieldObj.IsVerseField)
+			{
+				throw new UhtException(fieldObj, "Attempt to write the Verse VNI package name on a field that isn't part of Verse");
+			}
+			UhtModule module = fieldObj.Module;
+			return builder.Append('/').Append(module.Module.VersePluginName).Append("/_Verse/VNI/").Append(module.Module.Name);
+		}
+
+		/// <summary>
+		/// Append the Verse UE name 
+		/// </summary>
+		/// <param name="builder">Destination builder</param>
+		/// <param name="fieldObj">Field to serialize</param>
+		/// <returns>Builder</returns>
+		public static StringBuilder AppendVerseUEName(this StringBuilder builder, UhtField fieldObj)
+		{
+			if (!fieldObj.IsVerseField)
+			{
+				throw new UhtException(fieldObj, "Attempt to write the Verse name on a field that isn't part of Verse");
+			}
+			return builder.Append(fieldObj.VerseName!);
+		}
+
+		/// <summary>
+		/// Append the Verse UE name 
+		/// </summary>
+		/// <param name="builder">Destination builder</param>
+		/// <param name="fieldObj">Field to serialize</param>
+		/// <returns>Builder</returns>
+		public static StringBuilder AppendVerseUEModuleAndName(this StringBuilder builder, UhtField fieldObj)
+		{
+			if (fieldObj.VerseModule == null)
+			{
+				throw new UhtException(fieldObj, "Attempt to write the Verse module on a field doesn't have the module name set");
+			}
+			return builder.Append(fieldObj.VerseModule!).Append('_').AppendVerseUEName(fieldObj);
+		}
+
+		/// <summary>
+		/// Append the Verse package name
+		/// </summary>
+		/// <param name="builder">Destination builder</param>
+		/// <param name="fieldObj">Field to serialize</param>
+		/// <returns>Builder</returns>
+		public static StringBuilder AppendVersePackageName(this StringBuilder builder, UhtField fieldObj)
+		{
+			if (!fieldObj.IsVerseField)
+			{
+				throw new UhtException(fieldObj, "Attempt to write the Verse package name on a field that isn't part of Verse");
+			}
+			UhtModule module = fieldObj.Module;
+			return builder.Append('/').Append(module.Module.VersePluginName).Append("/_Verse/VNI/").Append(module.Module.Name);
+		}
+
+		/// <summary>
+		/// Append the Verse module name
+		/// </summary>
+		/// <param name="builder">Destination builder</param>
+		/// <param name="fieldObj">Field to serialize</param>
+		/// <returns>Builder</returns>
+		public static StringBuilder AppendVerseModuleName(this StringBuilder builder, UhtField fieldObj)
+		{
+			if (!fieldObj.IsVerseField)
+			{
+				throw new UhtException(fieldObj, "Attempt to write the Verse module name on a field that isn't part of Verse");
+			}
+			builder.Append(fieldObj.Package.Module.Module.VersePath);
+			if (fieldObj.VerseModule != null)
+			{
+				builder.Append('/').Append(fieldObj.VerseModule!);
+			}
+			return builder;
+		}
+	}
 }
diff --git a/Engine/Source/Programs/Shared/EpicGames.UHT/Types/UhtModule.cs b/Engine/Source/Programs/Shared/EpicGames.UHT/Types/UhtModule.cs
--- a/Engine/Source/Programs/Shared/EpicGames.UHT/Types/UhtModule.cs
+++ b/Engine/Source/Programs/Shared/EpicGames.UHT/Types/UhtModule.cs
@@ -18,6 +18,7 @@ namespace EpicGames.UHT.Types
 	public class UhtModule
 	{
 		private readonly List<UhtHeaderFile> _headers = new();
+		private readonly SortedDictionary<string, UhtPackage> _packages = new();
 
 		/// <summary>
 		/// The running session
@@ -55,6 +56,10 @@ namespace EpicGames.UHT.Types
 			get
 			{
 				yield return ScriptPackage;
+				foreach (KeyValuePair<string, UhtPackage> kvp in _packages)
+				{
+					yield return kvp.Value;
+				}
 			}
 		}
 
@@ -180,6 +185,33 @@ namespace EpicGames.UHT.Types
 			PrepareHeaders(Module.PrivateHeaders, UhtHeaderFileType.Private, addHeaderFileAction);
 		}
 
+		/// <summary>
+		/// Create a new package with the given name.  Will return an existing instance if it already exists
+		/// </summary>
+		/// <param name="packageName">Name of the package</param>
+		/// <param name="packageFlags">Flags to set if the package is created</param>
+		/// <returns>Package</returns>
+		public UhtPackage CreatePackage(string packageName, EPackageFlags packageFlags = EPackageFlags.None)
+		{
+			if (packageName == ScriptPackage.SourceName)
+			{
+				return ScriptPackage;
+			}
+			if (_packages.TryGetValue(packageName, out UhtPackage? package))
+			{
+				return package;
+			}
+			package = new(this, packageName, packageFlags);
+			_packages.Add(packageName, package);
+			return package;
+		}
+
+		/// <summary>
+		/// Given a collection of headers, create the UhtHeaders
+		/// </summary>
+		/// <param name="headerFiles">Collection of header names</param>
+		/// <param name="headerFileType">Type of the headers</param>
+		/// <param name="addHeaderFileAction">Action to take to notify the caller of the header being created</param>
 		private void PrepareHeaders(IEnumerable<string> headerFiles, UhtHeaderFileType headerFileType, Action<UhtHeaderFile> addHeaderFileAction)
 		{
 			string typeDirectory = headerFileType.ToString() + '/';
diff --git a/Engine/Source/Programs/Shared/EpicGames.UHT/Types/UhtProperty.cs b/Engine/Source/Programs/Shared/EpicGames.UHT/Types/UhtProperty.cs
--- a/Engine/Source/Programs/Shared/EpicGames.UHT/Types/UhtProperty.cs
+++ b/Engine/Source/Programs/Shared/EpicGames.UHT/Types/UhtProperty.cs
@@ -676,6 +676,11 @@ namespace EpicGames.UHT.Types
 		/// </summary>
 		public string EngineName { get; set; } = String.Empty;
 
+		/// <summary>
+		/// Verse name of the property 
+		/// </summary>
+		public string VerseName { get; set; } = String.Empty;
+
 		/// <summary>
 		/// Property's meta data
 		/// </summary>
@@ -782,6 +787,7 @@ namespace EpicGames.UHT.Types
 			ParentSettings = parentPropertySettings;
 			SourceName = sourceName;
 			EngineName = sourceName;
+			VerseName = String.Empty;
 			MetaData = new UhtMetaData(messageSite, parentPropertySettings.Outer.Session.Config);
 			Outer = parentPropertySettings.Outer;
 			LineNumber = parentPropertySettings.LineNumber;
@@ -813,6 +819,7 @@ namespace EpicGames.UHT.Types
 			ParentSettings = null;
 			SourceName = String.Empty;
 			EngineName = String.Empty;
+			VerseName = String.Empty;
 			MetaData = new UhtMetaData(outer, outer.Session.Config);
 			Outer = outer;
 			LineNumber = lineNumber;
@@ -846,6 +853,7 @@ namespace EpicGames.UHT.Types
 			ParentSettings = null;
 			SourceName = property.SourceName;
 			EngineName = property.EngineName;
+			VerseName = property.VerseName;
 			MetaData = property.MetaData;
 			Outer = property.Outer;
 			LineNumber = property.LineNumber;
@@ -941,6 +949,11 @@ namespace EpicGames.UHT.Types
 		protected virtual UhtPGetArgumentType PGetTypeArgument => UhtPGetArgumentType.None;
 		#endregion
 
+		/// <summary>
+		/// Verse name of the property 
+		/// </summary>
+		public string VerseName { get; set; } = String.Empty;
+
 		/// <summary>
 		/// Property category
 		/// </summary>
@@ -1031,6 +1044,12 @@ namespace EpicGames.UHT.Types
 		[JsonIgnore]
 		public override bool Deprecated => PropertyFlags.HasAnyFlags(EPropertyFlags.Deprecated);
 
+		/// <summary>
+		/// Return the engine name without and 'b' prefixes
+		/// </summary>
+		[JsonIgnore]
+		public virtual string StrippedEngineName => EngineName;
+
 		///<inheritdoc/>
 		[JsonIgnore]
 		protected override UhtSpecifierValidatorTable? SpecifierValidatorTable
@@ -1107,6 +1126,7 @@ namespace EpicGames.UHT.Types
 			{
 				EngineName = propertySettings.EngineName;
 			}
+			VerseName = propertySettings.VerseName;
 			PropertyCategory = propertySettings.PropertyCategory;
 			PropertyFlags = propertySettings.PropertyFlags;
 			DisallowPropertyFlags = propertySettings.DisallowPropertyFlags;
@@ -1494,6 +1514,27 @@ namespace EpicGames.UHT.Types
 		public virtual void AppendObjectHashes(StringBuilder builder, int startingLength, IUhtPropertyMemberContext context)
 		{
 		}
+
+		/// <summary>
+		/// Fetch the mangled name for a property
+		/// </summary>
+		/// <returns>True if the name needed to be mangled, false if not.</returns>
+		public (bool WasMangled, string Result) GetMangledEngineName()
+		{
+			if (PropertyFlags.HasAnyFlags(EPropertyFlags.ReturnParm))
+			{
+				return (false, EngineName);
+			}
+			if (!String.IsNullOrEmpty(VerseName))
+            {
+				return (false, VerseName);
+            }
+			if (Outer is UhtField fieldObj && fieldObj.IsVerseField)
+			{
+				return VerseNameMangling.MangleCasedName(String.IsNullOrEmpty(VerseName) ? StrippedEngineName : VerseName);
+			}
+			return (false, EngineName);
+		}
 		#endregion
 
 		#region Parsing support
diff --git a/Engine/Source/Programs/Shared/EpicGames.UHT/Types/UhtScriptStruct.cs b/Engine/Source/Programs/Shared/EpicGames.UHT/Types/UhtScriptStruct.cs
--- a/Engine/Source/Programs/Shared/EpicGames.UHT/Types/UhtScriptStruct.cs
+++ b/Engine/Source/Programs/Shared/EpicGames.UHT/Types/UhtScriptStruct.cs
@@ -686,7 +686,7 @@ namespace EpicGames.UHT.Types
 
 		/// <inheritdoc/>
 		[JsonIgnore]
-		public override string EngineNamePrefix => Session.Config!.IsStructWithTPrefix(EngineName) ? "T" : "F";
+		public override string EngineNamePrefix => IsVerseField ? "" : (Session.Config!.IsStructWithTPrefix(EngineName) ? "T" : "F");
 
 		///<inheritdoc/>
 		[JsonIgnore]
@@ -848,10 +848,13 @@ namespace EpicGames.UHT.Types
 			}
 
 			// Validate the engine name
-			string expectedName = $"{EngineNamePrefix}{EngineName}";
-			if (SourceName != expectedName)
+			if (!IsVerseField)
 			{
-				this.LogError($"Struct '{SourceName}' has an invalid Unreal prefix, expecting '{expectedName}");
+				string expectedName = $"{EngineNamePrefix}{EngineName}";
+				if (SourceName != expectedName)
+				{
+					this.LogError($"Struct '{SourceName}' has an invalid Unreal prefix, expecting '{expectedName}");
+				}
 			}
 
 			// Validate RigVM
diff --git a/Engine/Source/Programs/Shared/EpicGames.UHT/Utils/UhtUtilities.cs b/Engine/Source/Programs/Shared/EpicGames.UHT/Utils/UhtUtilities.cs
--- a/Engine/Source/Programs/Shared/EpicGames.UHT/Utils/UhtUtilities.cs
+++ b/Engine/Source/Programs/Shared/EpicGames.UHT/Utils/UhtUtilities.cs
@@ -1,6 +1,7 @@
 // Copyright Epic Games, Inc. All Rights Reserved.
 
 using System;
+using System.Buffers.Binary;
 using System.Collections.Generic;
 using System.Linq;
 using System.Text;
@@ -303,4 +304,100 @@ namespace EpicGames.UHT.Utils
 			}
 		}
 	}
+
+	/// <summary>
+	/// Helper structure used to compute UE style CRC values
+	/// </summary>
+	internal readonly struct Crc
+	{
+		private static readonly uint[] s_table = {
+			0x00000000, 0x77073096, 0xee0e612c, 0x990951ba, 0x076dc419, 0x706af48f, 0xe963a535, 0x9e6495a3, 0x0edb8832, 0x79dcb8a4, 0xe0d5e91e, 0x97d2d988, 0x09b64c2b, 0x7eb17cbd, 0xe7b82d07, 0x90bf1d91,
+			0x1db71064, 0x6ab020f2, 0xf3b97148, 0x84be41de, 0x1adad47d, 0x6ddde4eb, 0xf4d4b551, 0x83d385c7, 0x136c9856, 0x646ba8c0, 0xfd62f97a, 0x8a65c9ec, 0x14015c4f, 0x63066cd9, 0xfa0f3d63, 0x8d080df5,
+			0x3b6e20c8, 0x4c69105e, 0xd56041e4, 0xa2677172, 0x3c03e4d1, 0x4b04d447, 0xd20d85fd, 0xa50ab56b, 0x35b5a8fa, 0x42b2986c, 0xdbbbc9d6, 0xacbcf940, 0x32d86ce3, 0x45df5c75, 0xdcd60dcf, 0xabd13d59,
+			0x26d930ac, 0x51de003a, 0xc8d75180, 0xbfd06116, 0x21b4f4b5, 0x56b3c423, 0xcfba9599, 0xb8bda50f, 0x2802b89e, 0x5f058808, 0xc60cd9b2, 0xb10be924, 0x2f6f7c87, 0x58684c11, 0xc1611dab, 0xb6662d3d,
+			0x76dc4190, 0x01db7106, 0x98d220bc, 0xefd5102a, 0x71b18589, 0x06b6b51f, 0x9fbfe4a5, 0xe8b8d433, 0x7807c9a2, 0x0f00f934, 0x9609a88e, 0xe10e9818, 0x7f6a0dbb, 0x086d3d2d, 0x91646c97, 0xe6635c01,
+			0x6b6b51f4, 0x1c6c6162, 0x856530d8, 0xf262004e, 0x6c0695ed, 0x1b01a57b, 0x8208f4c1, 0xf50fc457, 0x65b0d9c6, 0x12b7e950, 0x8bbeb8ea, 0xfcb9887c, 0x62dd1ddf, 0x15da2d49, 0x8cd37cf3, 0xfbd44c65,
+			0x4db26158, 0x3ab551ce, 0xa3bc0074, 0xd4bb30e2, 0x4adfa541, 0x3dd895d7, 0xa4d1c46d, 0xd3d6f4fb, 0x4369e96a, 0x346ed9fc, 0xad678846, 0xda60b8d0, 0x44042d73, 0x33031de5, 0xaa0a4c5f, 0xdd0d7cc9,
+			0x5005713c, 0x270241aa, 0xbe0b1010, 0xc90c2086, 0x5768b525, 0x206f85b3, 0xb966d409, 0xce61e49f, 0x5edef90e, 0x29d9c998, 0xb0d09822, 0xc7d7a8b4, 0x59b33d17, 0x2eb40d81, 0xb7bd5c3b, 0xc0ba6cad,
+			0xedb88320, 0x9abfb3b6, 0x03b6e20c, 0x74b1d29a, 0xead54739, 0x9dd277af, 0x04db2615, 0x73dc1683, 0xe3630b12, 0x94643b84, 0x0d6d6a3e, 0x7a6a5aa8, 0xe40ecf0b, 0x9309ff9d, 0x0a00ae27, 0x7d079eb1,
+			0xf00f9344, 0x8708a3d2, 0x1e01f268, 0x6906c2fe, 0xf762575d, 0x806567cb, 0x196c3671, 0x6e6b06e7, 0xfed41b76, 0x89d32be0, 0x10da7a5a, 0x67dd4acc, 0xf9b9df6f, 0x8ebeeff9, 0x17b7be43, 0x60b08ed5,
+			0xd6d6a3e8, 0xa1d1937e, 0x38d8c2c4, 0x4fdff252, 0xd1bb67f1, 0xa6bc5767, 0x3fb506dd, 0x48b2364b, 0xd80d2bda, 0xaf0a1b4c, 0x36034af6, 0x41047a60, 0xdf60efc3, 0xa867df55, 0x316e8eef, 0x4669be79,
+			0xcb61b38c, 0xbc66831a, 0x256fd2a0, 0x5268e236, 0xcc0c7795, 0xbb0b4703, 0x220216b9, 0x5505262f, 0xc5ba3bbe, 0xb2bd0b28, 0x2bb45a92, 0x5cb36a04, 0xc2d7ffa7, 0xb5d0cf31, 0x2cd99e8b, 0x5bdeae1d,
+			0x9b64c2b0, 0xec63f226, 0x756aa39c, 0x026d930a, 0x9c0906a9, 0xeb0e363f, 0x72076785, 0x05005713, 0x95bf4a82, 0xe2b87a14, 0x7bb12bae, 0x0cb61b38, 0x92d28e9b, 0xe5d5be0d, 0x7cdcefb7, 0x0bdbdf21,
+			0x86d3d2d4, 0xf1d4e242, 0x68ddb3f8, 0x1fda836e, 0x81be16cd, 0xf6b9265b, 0x6fb077e1, 0x18b74777, 0x88085ae6, 0xff0f6a70, 0x66063bca, 0x11010b5c, 0x8f659eff, 0xf862ae69, 0x616bffd3, 0x166ccf45,
+			0xa00ae278, 0xd70dd2ee, 0x4e048354, 0x3903b3c2, 0xa7672661, 0xd06016f7, 0x4969474d, 0x3e6e77db, 0xaed16a4a, 0xd9d65adc, 0x40df0b66, 0x37d83bf0, 0xa9bcae53, 0xdebb9ec5, 0x47b2cf7f, 0x30b5ffe9,
+			0xbdbdf21c, 0xcabac28a, 0x53b39330, 0x24b4a3a6, 0xbad03605, 0xcdd70693, 0x54de5729, 0x23d967bf, 0xb3667a2e, 0xc4614ab8, 0x5d681b02, 0x2a6f2b94, 0xb40bbe37, 0xc30c8ea1, 0x5a05df1b, 0x2d02ef8d
+		};
+
+		/// <summary>
+		/// Compute a crc from the given string and an initial CRC value
+		/// </summary>
+		/// <param name="text">Text used to generate CRC</param>
+		/// <param name="crc">Initial CRC values</param>
+		/// <returns>Compute CRC value</returns>
+		public static uint Compute(StringView text, uint crc = 0)
+		{
+			crc = ~crc;
+			foreach (char c in text.Span)
+			{
+				crc = (crc >> 8) ^ s_table[(crc ^ c) & 0xFF];
+				crc = (crc >> 8) ^ s_table[(crc ^ (c >> 8)) & 0xFF];
+				crc = (crc >> 8) ^ s_table[(crc) & 0xFF];
+				crc = (crc >> 8) ^ s_table[(crc) & 0xFF];
+			}
+			return ~crc;
+		}
+
+		/// <summary>
+		/// Compute a crc from the given string and an initial CRC value
+		/// </summary>
+		/// <param name="text">Text used to generate CRC</param>
+		/// <param name="crc">Initial CRC values</param>
+		/// <returns>Compute CRC value</returns>
+		public static uint Compute(byte[] text, uint crc = 0)
+		{
+			crc = ~crc;
+			foreach (byte c in text)
+			{
+				crc = (crc >> 8) ^ s_table[(crc ^ c) & 0xFF];
+				crc = (crc >> 8) ^ s_table[(crc) & 0xFF];
+				crc = (crc >> 8) ^ s_table[(crc) & 0xFF];
+				crc = (crc >> 8) ^ s_table[(crc) & 0xFF];
+			}
+			return ~crc;
+		}
+	}
+
+	/// <summary>
+	/// Helper structure for verse name mangling
+	/// </summary>
+	internal readonly struct VerseNameMangling
+	{
+		private const string MangledPrefix = "__verse_0x";
+		private static readonly string[] s_internalNames = { MangledPrefix, "RetVal", "_RetVal", "$TEMP", "_Self" };
+
+		private static bool ShouldMangleCasedName(string name)
+		{
+			return !s_internalNames.Any(x => name.StartsWith(x, StringComparison.Ordinal));
+		}
+
+		/// <summary>
+		/// If required, generate the mangled name for a cased name
+		/// </summary>
+		/// <param name="name">Name to be mangled</param>
+		/// <returns>True if the name needed to be mangled, false if not.</returns>
+		public static (bool WasMangled, string Result) MangleCasedName(string name)
+		{
+			if (ShouldMangleCasedName(name))
+			{
+				byte[] utf8Name = Encoding.UTF8.GetBytes(name);
+				uint crc = BinaryPrimitives.ReverseEndianness(Crc.Compute(utf8Name));
+				return (true, $"{MangledPrefix}{crc:X8}_{name}");
+			}
+			else
+			{
+				return (false, name);
+			}
+		}
+	}
 }
diff --git a/Engine/Source/Programs/UnrealBuildTool/System/ExternalExecution.cs b/Engine/Source/Programs/UnrealBuildTool/System/ExternalExecution.cs
--- a/Engine/Source/Programs/UnrealBuildTool/System/ExternalExecution.cs
+++ b/Engine/Source/Programs/UnrealBuildTool/System/ExternalExecution.cs
@@ -173,6 +173,11 @@ namespace UnrealBuildTool
 		/// </summary>
 		public bool HasVerse;
 
+		/// <summary>
+		/// Verse plugin name
+		/// </summary>
+		public string VersePluginName;
+
 		public UHTModuleInfo(
 			string ModuleName, 
 			FileReference ModuleRulesFile, 
@@ -185,7 +190,8 @@ namespace UnrealBuildTool
 			ModuleRules.PackageOverrideType OverrideType,
 			string VersePath,
 			VerseScope VerseScope,
-			bool HasVerse)
+			bool HasVerse,
+			string VersePluginName)
 		{
 			this.ModuleName = ModuleName;
 			this.ModuleRulesFile = ModuleRulesFile;
@@ -204,6 +210,7 @@ namespace UnrealBuildTool
 			this.VersePath = VersePath;
 			this.VerseScope = VerseScope;
 			this.HasVerse = HasVerse;
+			this.VersePluginName = VersePluginName;
 		}
 
 		public UHTModuleInfo(BinaryArchiveReader Reader)
@@ -226,6 +233,7 @@ namespace UnrealBuildTool
 			VersePath = Reader.ReadString()!;
 			VerseScope = (VerseScope)Reader.ReadInt();
 			HasVerse = Reader.ReadBool();
+			VersePluginName = Reader.ReadString()!;
 		}
 
 		public void Write(BinaryArchiveWriter Writer)
@@ -248,6 +256,7 @@ namespace UnrealBuildTool
 			Writer.WriteString(VersePath);
 			Writer.WriteInt((int)VerseScope);
 			Writer.WriteBool(HasVerse);
+			Writer.WriteString(VersePluginName);
 		}
 
 		public override string ToString()
@@ -465,7 +474,8 @@ namespace UnrealBuildTool
 						Module.Rules.OverridePackageType,
 						Module.Rules.VersePath ?? "",
 						Module.Rules.VerseScope,
-						Module.bHasVerse);
+						Module.bHasVerse,
+						Module.bHasVerse ? (Module.Rules.Plugin != null ? Module.Rules.Plugin.Name : Module.Name) : string.Empty);
 					ModuleInfoArray[Idx] = Info;
 
 					Queue.Enqueue(() => SetupUObjectModule(Info, ExcludedFolders, MetadataCache, Queue));
@@ -1039,6 +1049,7 @@ namespace UnrealBuildTool
 						VersePath = UObjectModule.VersePath,
 						VerseScope = (UHTVerseScope)Enum.Parse(typeof(UHTVerseScope), UObjectModule.VerseScope.ToString()),
 						HasVerse = UObjectModule.HasVerse,
+						VersePluginName = UObjectModule.VersePluginName,
 					});
 			}
 
