Pin version for transitive dependencies that are out of date

#rnx

diff --git a/Engine/Source/Programs/Shared/EpicGames.Horde/EpicGames.Horde.csproj b/Engine/Source/Programs/Shared/EpicGames.Horde/EpicGames.Horde.csproj
--- a/Engine/Source/Programs/Shared/EpicGames.Horde/EpicGames.Horde.csproj
+++ b/Engine/Source/Programs/Shared/EpicGames.Horde/EpicGames.Horde.csproj
@@ -59,6 +59,8 @@
       <PrivateAssets>all</PrivateAssets>
       <IncludeAssets>runtime; build; native; contentfiles; analyzers; buildtransitive</IncludeAssets>
     </PackageReference>
+    <PackageReference Include="System.Net.Http" Version="4.3.4" />
+    <PackageReference Include="System.Text.RegularExpressions" Version="4.3.1" />
     <PackageReference Include="ZstdSharp.Port" Version="0.8.1" />
   </ItemGroup>
   <ItemGroup>
diff --git a/Engine/Source/Programs/Shared/EpicGames.IoHash/EpicGames.IoHash.csproj b/Engine/Source/Programs/Shared/EpicGames.IoHash/EpicGames.IoHash.csproj
--- a/Engine/Source/Programs/Shared/EpicGames.IoHash/EpicGames.IoHash.csproj
+++ b/Engine/Source/Programs/Shared/EpicGames.IoHash/EpicGames.IoHash.csproj
@@ -37,7 +37,7 @@
     </Compile>
   </ItemGroup>
   <ItemGroup>
-    <PackageReference Include="Blake3" Version="0.5.1" />
+    <PackageReference Include="Blake3" Version="0.6.1" />
     <PackageReference Include="System.Memory" Version="4.5.5" />
   </ItemGroup>
   <ItemGroup>
diff --git a/Engine/Source/Programs/Shared/EpicGames.MsBuild/EpicGames.MsBuild.csproj b/Engine/Source/Programs/Shared/EpicGames.MsBuild/EpicGames.MsBuild.csproj
--- a/Engine/Source/Programs/Shared/EpicGames.MsBuild/EpicGames.MsBuild.csproj
+++ b/Engine/Source/Programs/Shared/EpicGames.MsBuild/EpicGames.MsBuild.csproj
@@ -40,6 +40,7 @@
   <ItemGroup>
     <PackageReference Include="Microsoft.Build" Version="17.1.0" ExcludeAssets="runtime" PrivateAssets="all" />
     <PackageReference Include="Microsoft.Build.Locator" Version="1.4.1" PrivateAssets="all" />
+    <PackageReference Include="System.Drawing.Common" Version="4.7.3" />
   </ItemGroup>
 
   <ItemGroup>
diff --git a/Engine/Source/Programs/Shared/EpicGames.OIDC/EpicGames.OIDC.csproj b/Engine/Source/Programs/Shared/EpicGames.OIDC/EpicGames.OIDC.csproj
--- a/Engine/Source/Programs/Shared/EpicGames.OIDC/EpicGames.OIDC.csproj
+++ b/Engine/Source/Programs/Shared/EpicGames.OIDC/EpicGames.OIDC.csproj
@@ -37,10 +37,8 @@
   
   <ItemGroup>
     <PackageReference Include="IdentityModel.OidcClient" Version="4.0.0" />
-    <PackageReference Include="Microsoft.Extensions.Configuration" Version="6.0.0" />
     <PackageReference Include="Microsoft.Extensions.Configuration.Binder" Version="6.0.0" />
     <PackageReference Include="Microsoft.Extensions.Configuration.Json" Version="6.0.0" />
-    <PackageReference Include="Microsoft.Win32.Registry" Version="5.0.0" />
   </ItemGroup>
 
 </Project>
diff --git a/Engine/Source/Programs/Shared/EpicGames.Oodle/EpicGames.Oodle.csproj b/Engine/Source/Programs/Shared/EpicGames.Oodle/EpicGames.Oodle.csproj
--- a/Engine/Source/Programs/Shared/EpicGames.Oodle/EpicGames.Oodle.csproj
+++ b/Engine/Source/Programs/Shared/EpicGames.Oodle/EpicGames.Oodle.csproj
@@ -36,14 +36,6 @@
       <Link>Properties\MetaData.cs</Link>
     </Compile>
   </ItemGroup>
-  <ItemGroup>
-    <PackageReference Include="Microsoft.CodeAnalysis.CSharp.Workspaces" Version="4.2.0" />
-    <PackageReference Include="Microsoft.Extensions.Caching.Abstractions" Version="6.0.0" />
-    <PackageReference Include="Microsoft.Extensions.Caching.Memory" Version="6.0.1" />
-    <PackageReference Include="Microsoft.Extensions.DependencyInjection" Version="6.0.0" />
-    <PackageReference Include="Microsoft.Extensions.Http" Version="6.0.0" />
-    <PackageReference Include="Microsoft.Extensions.Logging" Version="6.0.0" />
-  </ItemGroup>
   <ItemGroup>
     <None Condition="'$([System.Runtime.InteropServices.RuntimeInformation]::IsOSPlatform($([System.Runtime.InteropServices.OSPlatform]::Linux)))' == 'true' Or Exists('Sdk\2.9.10\linux\lib\liboo2corelinux64.so.9')" Include="Sdk\2.9.10\linux\lib\liboo2corelinux64.so.9" Link="liboo2corelinux64.so.9" CopyToOutputDirectory="PreserveNewest" Visible="false" />
     <None Condition="'$([System.Runtime.InteropServices.RuntimeInformation]::IsOSPlatform($([System.Runtime.InteropServices.OSPlatform]::Linux)))' == 'true' Or Exists('Sdk\2.9.10\linuxarm\lib\liboo2corelinuxarm32.so.9')" Include="Sdk\2.9.10\linuxarm\lib\liboo2corelinuxarm32.so.9" Link="liboo2corelinuxarm32.so.9" CopyToOutputDirectory="PreserveNewest" Visible="false" />
diff --git a/Engine/Source/Programs/Shared/EpicGames.UBA/EpicGames.UBA.csproj b/Engine/Source/Programs/Shared/EpicGames.UBA/EpicGames.UBA.csproj
--- a/Engine/Source/Programs/Shared/EpicGames.UBA/EpicGames.UBA.csproj
+++ b/Engine/Source/Programs/Shared/EpicGames.UBA/EpicGames.UBA.csproj
@@ -48,7 +48,5 @@
   <ItemGroup>
     <PackageReference Include="Microsoft.CSharp" Version="4.7.0" />
     <PackageReference Include="Microsoft.Extensions.Logging" Version="6.0.0" />
-    <PackageReference Include="Microsoft.Win32.Registry" Version="5.0.0" />
-    <PackageReference Include="System.Data.DataSetExtensions" Version="4.5.0" />
   </ItemGroup>
 </Project>
\ No newline at end of file
diff --git a/Engine/Source/Programs/Shared/EpicGames.UHT/EpicGames.UHT.csproj b/Engine/Source/Programs/Shared/EpicGames.UHT/EpicGames.UHT.csproj
--- a/Engine/Source/Programs/Shared/EpicGames.UHT/EpicGames.UHT.csproj
+++ b/Engine/Source/Programs/Shared/EpicGames.UHT/EpicGames.UHT.csproj
@@ -47,8 +47,6 @@
   </ItemGroup>
   <ItemGroup>
     <PackageReference Include="Microsoft.CSharp" Version="4.7.0" />
-    <PackageReference Include="Microsoft.Win32.Registry" Version="5.0.0" />
     <PackageReference Include="System.Collections.Specialized" Version="4.3.0" />
-    <PackageReference Include="System.Data.DataSetExtensions" Version="4.5.0" />
   </ItemGroup>
 </Project>
\ No newline at end of file
