Update test nuget packages

#rnx

diff --git a/Engine/Source/Programs/Shared/EpicGames.BuildGraph.Tests/EpicGames.BuildGraph.Tests.csproj b/Engine/Source/Programs/Shared/EpicGames.BuildGraph.Tests/EpicGames.BuildGraph.Tests.csproj
--- a/Engine/Source/Programs/Shared/EpicGames.BuildGraph.Tests/EpicGames.BuildGraph.Tests.csproj
+++ b/Engine/Source/Programs/Shared/EpicGames.BuildGraph.Tests/EpicGames.BuildGraph.Tests.csproj
@@ -45,9 +45,9 @@
   </ItemGroup>
 
   <ItemGroup>
-    <PackageReference Include="Microsoft.NET.Test.Sdk" Version="17.3.0" />
-    <PackageReference Include="MSTest.TestAdapter" Version="2.2.10" />
-    <PackageReference Include="MSTest.TestFramework" Version="2.2.10" />
+    <PackageReference Include="Microsoft.NET.Test.Sdk" Version="17.11.0" />
+    <PackageReference Include="MSTest.TestAdapter" Version="3.5.2" />
+    <PackageReference Include="MSTest.TestFramework" Version="3.5.2" />
     <PackageReference Include="coverlet.collector" Version="6.0.2">
       <PrivateAssets>all</PrivateAssets>
       <IncludeAssets>runtime; build; native; contentfiles; analyzers; buildtransitive</IncludeAssets>
diff --git a/Engine/Source/Programs/Shared/EpicGames.Core.Tests/EpicGames.Core.Tests.csproj b/Engine/Source/Programs/Shared/EpicGames.Core.Tests/EpicGames.Core.Tests.csproj
--- a/Engine/Source/Programs/Shared/EpicGames.Core.Tests/EpicGames.Core.Tests.csproj
+++ b/Engine/Source/Programs/Shared/EpicGames.Core.Tests/EpicGames.Core.Tests.csproj
@@ -48,9 +48,9 @@
   </ItemGroup>
 
   <ItemGroup>
-    <PackageReference Include="Microsoft.NET.Test.Sdk" Version="17.3.0" />
-    <PackageReference Include="MSTest.TestAdapter" Version="2.2.10" />
-    <PackageReference Include="MSTest.TestFramework" Version="2.2.10" />
+    <PackageReference Include="Microsoft.NET.Test.Sdk" Version="17.11.0" />
+    <PackageReference Include="MSTest.TestAdapter" Version="3.5.2" />
+    <PackageReference Include="MSTest.TestFramework" Version="3.5.2" />
     <PackageReference Include="coverlet.collector" Version="6.0.2">
       <PrivateAssets>all</PrivateAssets>
       <IncludeAssets>runtime; build; native; contentfiles; analyzers; buildtransitive</IncludeAssets>
diff --git a/Engine/Source/Programs/Shared/EpicGames.Horde.Tests/EpicGames.Horde.Tests.csproj b/Engine/Source/Programs/Shared/EpicGames.Horde.Tests/EpicGames.Horde.Tests.csproj
--- a/Engine/Source/Programs/Shared/EpicGames.Horde.Tests/EpicGames.Horde.Tests.csproj
+++ b/Engine/Source/Programs/Shared/EpicGames.Horde.Tests/EpicGames.Horde.Tests.csproj
@@ -30,10 +30,10 @@
   </PropertyGroup>
 
   <ItemGroup>
-    <PackageReference Include="Microsoft.NET.Test.Sdk" Version="17.3.0" />
+    <PackageReference Include="Microsoft.NET.Test.Sdk" Version="17.11.0" />
     <PackageReference Include="Microsoft.Extensions.Caching.Memory" Version="6.0.1" />
-    <PackageReference Include="MSTest.TestAdapter" Version="2.2.10" />
-    <PackageReference Include="MSTest.TestFramework" Version="2.2.10" />
+    <PackageReference Include="MSTest.TestAdapter" Version="3.5.2" />
+    <PackageReference Include="MSTest.TestFramework" Version="3.5.2" />
     <PackageReference Include="coverlet.collector" Version="6.0.2">
       <PrivateAssets>all</PrivateAssets>
       <IncludeAssets>runtime; build; native; contentfiles; analyzers; buildtransitive</IncludeAssets>
diff --git a/Engine/Source/Programs/Shared/EpicGames.IoHash.Tests/EpicGames.IoHash.Tests.csproj b/Engine/Source/Programs/Shared/EpicGames.IoHash.Tests/EpicGames.IoHash.Tests.csproj
--- a/Engine/Source/Programs/Shared/EpicGames.IoHash.Tests/EpicGames.IoHash.Tests.csproj
+++ b/Engine/Source/Programs/Shared/EpicGames.IoHash.Tests/EpicGames.IoHash.Tests.csproj
@@ -17,9 +17,9 @@
   </PropertyGroup>
 
   <ItemGroup>
-    <PackageReference Include="Microsoft.NET.Test.Sdk" Version="17.3.0" />
-    <PackageReference Include="MSTest.TestAdapter" Version="2.2.10" />
-    <PackageReference Include="MSTest.TestFramework" Version="2.2.10" />
+    <PackageReference Include="Microsoft.NET.Test.Sdk" Version="17.11.0" />
+    <PackageReference Include="MSTest.TestAdapter" Version="3.5.2" />
+    <PackageReference Include="MSTest.TestFramework" Version="3.5.2" />
     <PackageReference Include="coverlet.collector" Version="6.0.2">
       <PrivateAssets>all</PrivateAssets>
       <IncludeAssets>runtime; build; native; contentfiles; analyzers; buildtransitive</IncludeAssets>
diff --git a/Engine/Source/Programs/Shared/EpicGames.Perforce.Fixture/EpicGames.Perforce.Fixture.csproj b/Engine/Source/Programs/Shared/EpicGames.Perforce.Fixture/EpicGames.Perforce.Fixture.csproj
--- a/Engine/Source/Programs/Shared/EpicGames.Perforce.Fixture/EpicGames.Perforce.Fixture.csproj
+++ b/Engine/Source/Programs/Shared/EpicGames.Perforce.Fixture/EpicGames.Perforce.Fixture.csproj
@@ -33,13 +33,14 @@
 
   <ItemGroup>
     <PackageReference Include="Microsoft.Extensions.Logging.Console" Version="8.0.0" />
-    <PackageReference Include="Microsoft.NET.Test.Sdk" Version="17.3.0" />
-    <PackageReference Include="MSTest.TestAdapter" Version="2.2.10" />
-    <PackageReference Include="MSTest.TestFramework" Version="2.2.10" />
+    <PackageReference Include="Microsoft.NET.Test.Sdk" Version="17.11.0" />
+    <PackageReference Include="MSTest.TestAdapter" Version="3.5.2" />
+    <PackageReference Include="MSTest.TestFramework" Version="3.5.2" />
     <PackageReference Include="coverlet.collector" Version="6.0.2">
       <PrivateAssets>all</PrivateAssets>
       <IncludeAssets>runtime; build; native; contentfiles; analyzers; buildtransitive</IncludeAssets>
     </PackageReference>
+    <PackageReference Include="System.Text.Json" Version="8.0.4" />
   </ItemGroup>
 
   <ItemGroup>
diff --git a/Engine/Source/Programs/Shared/EpicGames.Perforce.Managed.Tests/EpicGames.Perforce.Managed.Tests.csproj b/Engine/Source/Programs/Shared/EpicGames.Perforce.Managed.Tests/EpicGames.Perforce.Managed.Tests.csproj
--- a/Engine/Source/Programs/Shared/EpicGames.Perforce.Managed.Tests/EpicGames.Perforce.Managed.Tests.csproj
+++ b/Engine/Source/Programs/Shared/EpicGames.Perforce.Managed.Tests/EpicGames.Perforce.Managed.Tests.csproj
@@ -31,9 +31,9 @@
   </ItemGroup>
 
   <ItemGroup>
-    <PackageReference Include="Microsoft.NET.Test.Sdk" Version="17.3.0" />
-    <PackageReference Include="MSTest.TestAdapter" Version="2.2.10" />
-    <PackageReference Include="MSTest.TestFramework" Version="2.2.10" />
+    <PackageReference Include="Microsoft.NET.Test.Sdk" Version="17.11.0" />
+    <PackageReference Include="MSTest.TestAdapter" Version="3.5.2" />
+    <PackageReference Include="MSTest.TestFramework" Version="3.5.2" />
     <PackageReference Include="coverlet.collector" Version="6.0.2">
       <PrivateAssets>all</PrivateAssets>
       <IncludeAssets>runtime; build; native; contentfiles; analyzers; buildtransitive</IncludeAssets>
diff --git a/Engine/Source/Programs/Shared/EpicGames.Perforce.Tests/EpicGames.Perforce.Tests.csproj b/Engine/Source/Programs/Shared/EpicGames.Perforce.Tests/EpicGames.Perforce.Tests.csproj
--- a/Engine/Source/Programs/Shared/EpicGames.Perforce.Tests/EpicGames.Perforce.Tests.csproj
+++ b/Engine/Source/Programs/Shared/EpicGames.Perforce.Tests/EpicGames.Perforce.Tests.csproj
@@ -30,9 +30,9 @@
   </ItemGroup>
 
   <ItemGroup>
-    <PackageReference Include="Microsoft.NET.Test.Sdk" Version="17.3.0" />
-    <PackageReference Include="MSTest.TestAdapter" Version="2.2.10" />
-    <PackageReference Include="MSTest.TestFramework" Version="2.2.10" />
+    <PackageReference Include="Microsoft.NET.Test.Sdk" Version="17.11.0" />
+    <PackageReference Include="MSTest.TestAdapter" Version="3.5.2" />
+    <PackageReference Include="MSTest.TestFramework" Version="3.5.2" />
     <PackageReference Include="coverlet.collector" Version="6.0.2">
       <PrivateAssets>all</PrivateAssets>
       <IncludeAssets>runtime; build; native; contentfiles; analyzers; buildtransitive</IncludeAssets>
diff --git a/Engine/Source/Programs/Shared/EpicGames.Redis.Tests/EpicGames.Redis.Tests.csproj b/Engine/Source/Programs/Shared/EpicGames.Redis.Tests/EpicGames.Redis.Tests.csproj
--- a/Engine/Source/Programs/Shared/EpicGames.Redis.Tests/EpicGames.Redis.Tests.csproj
+++ b/Engine/Source/Programs/Shared/EpicGames.Redis.Tests/EpicGames.Redis.Tests.csproj
@@ -31,9 +31,9 @@
   </ItemGroup>
 
   <ItemGroup>
-    <PackageReference Include="Microsoft.NET.Test.Sdk" Version="17.8.0" />
-    <PackageReference Include="MSTest.TestAdapter" Version="3.1.1" />
-    <PackageReference Include="MSTest.TestFramework" Version="3.1.1" />
+    <PackageReference Include="Microsoft.NET.Test.Sdk" Version="17.11.0" />
+    <PackageReference Include="MSTest.TestAdapter" Version="3.5.2" />
+    <PackageReference Include="MSTest.TestFramework" Version="3.5.2" />
     <PackageReference Include="coverlet.collector" Version="6.0.2">
       <PrivateAssets>all</PrivateAssets>
       <IncludeAssets>runtime; build; native; contentfiles; analyzers; buildtransitive</IncludeAssets>
diff --git a/Engine/Source/Programs/Shared/EpicGames.Serialization.Tests/EpicGames.Serialization.Tests.csproj b/Engine/Source/Programs/Shared/EpicGames.Serialization.Tests/EpicGames.Serialization.Tests.csproj
--- a/Engine/Source/Programs/Shared/EpicGames.Serialization.Tests/EpicGames.Serialization.Tests.csproj
+++ b/Engine/Source/Programs/Shared/EpicGames.Serialization.Tests/EpicGames.Serialization.Tests.csproj
@@ -33,9 +33,9 @@
   </ItemGroup>
 
   <ItemGroup>
-    <PackageReference Include="Microsoft.NET.Test.Sdk" Version="17.3.0" />
-    <PackageReference Include="MSTest.TestAdapter" Version="2.2.10" />
-    <PackageReference Include="MSTest.TestFramework" Version="2.2.10" />
+    <PackageReference Include="Microsoft.NET.Test.Sdk" Version="17.11.0" />
+    <PackageReference Include="MSTest.TestAdapter" Version="3.5.2" />
+    <PackageReference Include="MSTest.TestFramework" Version="3.5.2" />
     <PackageReference Include="coverlet.collector" Version="6.0.2">
       <PrivateAssets>all</PrivateAssets>
       <IncludeAssets>runtime; build; native; contentfiles; analyzers; buildtransitive</IncludeAssets>
diff --git a/Engine/Source/Programs/Shared/EpicGames.Tracing.Tests/EpicGames.Tracing.Tests.csproj b/Engine/Source/Programs/Shared/EpicGames.Tracing.Tests/EpicGames.Tracing.Tests.csproj
--- a/Engine/Source/Programs/Shared/EpicGames.Tracing.Tests/EpicGames.Tracing.Tests.csproj
+++ b/Engine/Source/Programs/Shared/EpicGames.Tracing.Tests/EpicGames.Tracing.Tests.csproj
@@ -29,9 +29,9 @@
   </ItemGroup>
 
   <ItemGroup>
-    <PackageReference Include="Microsoft.NET.Test.Sdk" Version="16.5.0" />
-    <PackageReference Include="MSTest.TestAdapter" Version="2.1.0" />
-    <PackageReference Include="MSTest.TestFramework" Version="2.1.0" />
+    <PackageReference Include="Microsoft.NET.Test.Sdk" Version="17.11.0" />
+    <PackageReference Include="MSTest.TestAdapter" Version="3.5.2" />
+    <PackageReference Include="MSTest.TestFramework" Version="3.5.2" />
     <PackageReference Include="coverlet.collector" Version="6.0.2">
       <PrivateAssets>all</PrivateAssets>
       <IncludeAssets>runtime; build; native; contentfiles; analyzers; buildtransitive</IncludeAssets>
diff --git a/Engine/Source/Programs/UnrealBuildTool.Tests/UnrealBuildTool.Tests.csproj b/Engine/Source/Programs/UnrealBuildTool.Tests/UnrealBuildTool.Tests.csproj
--- a/Engine/Source/Programs/UnrealBuildTool.Tests/UnrealBuildTool.Tests.csproj
+++ b/Engine/Source/Programs/UnrealBuildTool.Tests/UnrealBuildTool.Tests.csproj
@@ -17,9 +17,9 @@
   </PropertyGroup>  
 
   <ItemGroup>
-    <PackageReference Include="Microsoft.NET.Test.Sdk" Version="17.3.0" />
-    <PackageReference Include="MSTest.TestAdapter" Version="2.2.10" />
-    <PackageReference Include="MSTest.TestFramework" Version="2.2.10" />
+    <PackageReference Include="Microsoft.NET.Test.Sdk" Version="17.11.0" />
+    <PackageReference Include="MSTest.TestAdapter" Version="3.5.2" />
+    <PackageReference Include="MSTest.TestFramework" Version="3.5.2" />
     <PackageReference Include="coverlet.collector" Version="6.0.2">
       <PrivateAssets>all</PrivateAssets>
       <IncludeAssets>runtime; build; native; contentfiles; analyzers; buildtransitive</IncludeAssets>
