Fixed symbol table lookup to properly include the header file when walking up the outer chain.

diff --git a/Engine/Source/Programs/Shared/EpicGames.UHT/Utils/UhtSymbolTable.cs b/Engine/Source/Programs/Shared/EpicGames.UHT/Utils/UhtSymbolTable.cs
--- a/Engine/Source/Programs/Shared/EpicGames.UHT/Utils/UhtSymbolTable.cs
+++ b/Engine/Source/Programs/Shared/EpicGames.UHT/Utils/UhtSymbolTable.cs
@@ -340,9 +340,12 @@ namespace EpicGames.UHT.Utils
 			// Search the chain
 			for (; currentType != null; currentType = currentType.Outer)
 			{
+
+				// When we hit the package, we assume that we want to do a search of the whole header file.  But this
+				// is only done when starting type isn't the package itself
 				if (currentType is UhtPackage)
 				{
-					return null;
+					return startingType != currentType ? FindSymbolChain(startingType.HeaderFile, options, ref lookup) : null;
 				}
 				UhtType? foundType = FindSymbolChain(currentType, options, ref lookup);
 				if (foundType != null)
@@ -387,6 +390,7 @@ namespace EpicGames.UHT.Utils
 		{
 			for (int index = lookup.SymbolIndex; index != 0; index = _symbols[index].NextIndex)
 			{
+				// Match on the same header file, but only top level types
 				if (IsMatch(options, index, lookup.CasedIndex) && _symbols[index].Type.HeaderFile == header && _symbols[index].Type.Outer is UhtPackage)
 				{
 					return _symbols[index].Type;
