Add HordeHTTPClient GetUgsMetadataAsync method

diff --git a/Engine/Source/Programs/Shared/EpicGames.Horde/HordeHttpClient.cs b/Engine/Source/Programs/Shared/EpicGames.Horde/HordeHttpClient.cs
--- a/Engine/Source/Programs/Shared/EpicGames.Horde/HordeHttpClient.cs
+++ b/Engine/Source/Programs/Shared/EpicGames.Horde/HordeHttpClient.cs
@@ -26,6 +26,8 @@ using EpicGames.Horde.Server;
 using EpicGames.Horde.Storage;
 using EpicGames.Horde.Streams;
 using EpicGames.Horde.Tools;
+using EpicGames.Horde.Ugs;
+
 
 #pragma warning disable CA2234
 
@@ -536,6 +538,22 @@ namespace EpicGames.Horde
 
 		#endregion
 
+		#region UGS
+		/// <summary>
+		/// 
+		/// </summary>
+		/// <param name="streamId"></param>
+		/// <param name="commitId"></param>
+		/// <param name="projectId"></param>
+		/// <param name="cancellationToken"></param>
+		/// <returns></returns>
+		public Task<GetUgsMetadataResponse> GetUgsMetadataAsync(StreamId streamId, CommitId commitId, ProjectId projectId, CancellationToken cancellationToken = default)
+		{
+			string perforceStream = $"//{streamId.ToString().Replace('-', '/')}";
+			return GetAsync<GetUgsMetadataResponse>(_httpClient, $"/ugs/api/metadata?stream={perforceStream}&change={commitId.GetPerforceChange()}&project={projectId}", cancellationToken);
+		}
+		#endregion
+
 		#region Utility Methods
 
 		/// <summary>
diff --git a/Engine/Source/Programs/Shared/EpicGames.Horde/Ugs/UgsMetadata.cs b/Engine/Source/Programs/Shared/EpicGames.Horde/Ugs/UgsMetadata.cs
new file mode 100644
--- /dev/null
+++ b/Engine/Source/Programs/Shared/EpicGames.Horde/Ugs/UgsMetadata.cs
@@ -0,0 +1,93 @@
+// Copyright Epic Games, Inc. All Rights Reserved.
+
+using System;
+using System.Collections.Generic;
+
+#pragma warning disable CA2227
+
+namespace EpicGames.Horde.Ugs
+{
+	/// <summary>
+	/// Information about a Badge
+	/// </summary>
+	public class UgsBadgeMetadata
+	{
+		/// <summary>
+		/// Name of the Badge
+		/// </summary>
+		public string Name { get; set; } = "";
+
+		/// <summary>
+		/// URL of the Badge setter
+		/// </summary>
+		public Uri? Url { get; set; } = null;
+
+		/// <summary>
+		/// Status state of the Badge
+		/// </summary>
+		public string State { get; set; } = "";
+	}
+
+	/// <summary>
+	/// Information about a User from UGS
+	/// </summary>
+	public class UgsUserMetadata
+	{
+		/// <summary>
+		/// User name
+		/// </summary>
+		public string User { get; set; } = "";
+
+		/// <summary>
+		/// Time user synced change in UGS
+		/// </summary>
+		public long SyncTime { get; set; }
+
+		/// <summary>
+		/// User vote state
+		/// </summary>
+		public string Vote { get; set; } = "";
+	}
+
+	/// <summary>
+	/// Information about UGS Metadata
+	/// </summary>
+	public class UgsMetadataItem
+	{
+		/// <summary>
+		/// Change number
+		/// </summary>
+		public int Change { get; set; }
+
+		/// <summary>
+		/// Project name
+		/// </summary>
+		public string Project { get; set; } = "";
+
+		/// <summary>
+		/// Badges associated with Change
+		/// </summary>
+		public List<UgsBadgeMetadata> Badges { get; set; } = new();
+
+		/// <summary>
+		/// Users of this change
+		/// </summary>
+		public List<UgsUserMetadata> Users { get; set; } = new();
+	}
+
+	/// <summary>
+	/// Information about a collection of Metadata
+	/// </summary>
+	public class GetUgsMetadataResponse
+	{
+		/// <summary>
+		/// Sequence number identifier
+		/// </summary>
+		public long SequenceNumber { get; set; }
+
+		/// <summary>
+		/// Metadata items
+		/// </summary>
+		public List<UgsMetadataItem> Items { get; set; } = new();
+	}
+}
