additional net8.0 fixes

#rnx

diff --git a/Engine/Source/Programs/Shared/EpicGames.Core/AsyncQueue.cs b/Engine/Source/Programs/Shared/EpicGames.Core/AsyncQueue.cs
--- a/Engine/Source/Programs/Shared/EpicGames.Core/AsyncQueue.cs
+++ b/Engine/Source/Programs/Shared/EpicGames.Core/AsyncQueue.cs
@@ -42,7 +42,12 @@ namespace EpicGames.Core
 		/// <inheritdoc/>
 		public async ValueTask DisposeAsync()
 		{
+#if NET8_0_OR_GREATER
+			await _cancellationSource.CancelAsync();
+#else
 			_cancellationSource.Cancel();
+#endif
+
 			await StopAsync(CancellationToken.None).ConfigureAwait(false);
 
 			_cancellationSource.Dispose();
diff --git a/Engine/Source/Programs/Shared/EpicGames.Core/CommandLineArguments.cs b/Engine/Source/Programs/Shared/EpicGames.Core/CommandLineArguments.cs
--- a/Engine/Source/Programs/Shared/EpicGames.Core/CommandLineArguments.cs
+++ b/Engine/Source/Programs/Shared/EpicGames.Core/CommandLineArguments.cs
@@ -421,7 +421,7 @@ namespace EpicGames.Core
 		/// <param name="prefix">The argument prefix (eg. "-Foo="). Must end with an '=' character.</param>
 		/// <param name="defaultValue">Default value for the argument</param>
 		/// <returns>Value of the argument</returns>
-		[return: NotNullIfNotNull("defaultValue")]
+		[return: NotNullIfNotNull(nameof(defaultValue))]
 		public string? GetStringOrDefault(string prefix, string? defaultValue)
 		{
 			string? value;
@@ -454,7 +454,7 @@ namespace EpicGames.Core
 		/// <param name="prefix">The argument prefix (eg. "-Foo="). Must end with an '=' character.</param>
 		/// <param name="defaultValue">Default value for the argument</param>
 		/// <returns>Value of the argument</returns>
-		[return: NotNullIfNotNull("defaultValue")]
+		[return: NotNullIfNotNull(nameof(defaultValue))]
 		public FileReference? GetFileReferenceOrDefault(string prefix, FileReference? defaultValue)
 		{
 			FileReference? value;
@@ -471,7 +471,7 @@ namespace EpicGames.Core
 		/// <param name="prefix">The argument prefix (eg. "-Foo="). Must end with an '=' character.</param>
 		/// <param name="defaultValue">Default value for the argument</param>
 		/// <returns>Value of the argument</returns>
-		[return: NotNullIfNotNull("defaultValue")]
+		[return: NotNullIfNotNull(nameof(defaultValue))]
 		public DirectoryReference? GetDirectoryReferenceOrDefault(string prefix, DirectoryReference? defaultValue)
 		{
 			DirectoryReference? value;
diff --git a/Engine/Source/Programs/Shared/EpicGames.Core/CsProjectInfo.cs b/Engine/Source/Programs/Shared/EpicGames.Core/CsProjectInfo.cs
--- a/Engine/Source/Programs/Shared/EpicGames.Core/CsProjectInfo.cs
+++ b/Engine/Source/Programs/Shared/EpicGames.Core/CsProjectInfo.cs
@@ -289,6 +289,27 @@ namespace EpicGames.Core
 			}
 		}
 
+		/// <summary>
+		/// Determines if a TargetFramework is a .NET core framework
+		/// </summary>
+		/// <returns>True if the TargetFramework is a .NET core framework</returns>
+		static bool IsDotNETCoreFramework(string targetFramework)
+		{
+			if (targetFramework.ToLower().Contains("netstandard", StringComparison.Ordinal) || targetFramework.ToLower().Contains("netcoreapp", StringComparison.Ordinal))
+			{
+				return true;
+			}
+			else if (targetFramework.StartsWith("net", StringComparison.OrdinalIgnoreCase))
+			{
+				string[] versionSplit = targetFramework.Substring(3).Split('.');
+				if (versionSplit.Length >= 1 && Int32.TryParse(versionSplit[0], out int majorVersion))
+				{
+					return majorVersion >= 5;
+				}
+			}
+			return false;
+		}
+
 		/// <summary>
 		/// Determines if this project is a .NET core project
 		/// </summary>
@@ -297,18 +318,12 @@ namespace EpicGames.Core
 		{
 			if (Properties.TryGetValue("TargetFramework", out string? targetFramework))
 			{
-				if (targetFramework.ToLower().Contains("netstandard", StringComparison.Ordinal) || targetFramework.ToLower().Contains("netcoreapp", StringComparison.Ordinal))
-				{
-					return true;
-				}
-				else if (targetFramework.StartsWith("net", StringComparison.OrdinalIgnoreCase))
-				{
-					string[] versionSplit = targetFramework.Substring(3).Split('.');
-					if (versionSplit.Length >= 1 && Int32.TryParse(versionSplit[0], out int majorVersion))
-					{
-						return majorVersion >= 5;
-					}
-				}
+				return IsDotNETCoreFramework(targetFramework);
+			}
+
+			if (Properties.TryGetValue("TargetFrameworks", out string? targetFrameworks))
+			{
+				return targetFramework?.Split(';', StringSplitOptions.RemoveEmptyEntries).Any(x => IsDotNETCoreFramework(x.Trim())) ?? false;
 			}
 
 			return false;
diff --git a/Engine/Source/Programs/Shared/EpicGames.Core/DirectoryReference.cs b/Engine/Source/Programs/Shared/EpicGames.Core/DirectoryReference.cs
--- a/Engine/Source/Programs/Shared/EpicGames.Core/DirectoryReference.cs
+++ b/Engine/Source/Programs/Shared/EpicGames.Core/DirectoryReference.cs
@@ -234,7 +234,7 @@ namespace EpicGames.Core
 		/// </summary>
 		/// <param name="file">The file to create a directory reference for</param>
 		/// <returns>The directory containing the file  </returns>
-		[return: NotNullIfNotNull("file")]
+		[return: NotNullIfNotNull(nameof(file))]
 		public static DirectoryReference? FromFile(FileReference? file) => file?.Directory;
 
 		/// <summary>
diff --git a/Engine/Source/Programs/Shared/EpicGames.Core/FileReference.cs b/Engine/Source/Programs/Shared/EpicGames.Core/FileReference.cs
--- a/Engine/Source/Programs/Shared/EpicGames.Core/FileReference.cs
+++ b/Engine/Source/Programs/Shared/EpicGames.Core/FileReference.cs
@@ -69,7 +69,7 @@ namespace EpicGames.Core
 		/// </summary>
 		/// <param name="fileName">FileName for the string</param>
 		/// <returns>Returns a FileReference representing the given string, or null.</returns>
-		[return: NotNullIfNotNull("fileName")]
+		[return: NotNullIfNotNull(nameof(fileName))]
 		public static FileReference? FromString(string? fileName) => String.IsNullOrEmpty(fileName) ? null : new FileReference(fileName);
 
 		/// <summary>
@@ -329,7 +329,7 @@ namespace EpicGames.Core
 		/// </summary>
 		/// <param name="location">Location of the file</param>
 		/// <param name="mode">Mode to use when opening the file</param>
-		/// <returns>New filestream for the given file</returns>
+		/// <returns>New FileStream for the given file</returns>
 		public static FileStream Open(FileReference location, FileMode mode) => File.Open(location.FullName, mode);
 
 		/// <summary>
@@ -338,7 +338,7 @@ namespace EpicGames.Core
 		/// <param name="location">Location of the file</param>
 		/// <param name="mode">Mode to use when opening the file</param>
 		/// <param name="access">Sharing mode for the new file</param>
-		/// <returns>New filestream for the given file</returns>
+		/// <returns>New FileStream for the given file</returns>
 		public static FileStream Open(FileReference location, FileMode mode, FileAccess access) => File.Open(location.FullName, mode, access);
 
 		/// <summary>
@@ -348,7 +348,7 @@ namespace EpicGames.Core
 		/// <param name="mode">Mode to use when opening the file</param>
 		/// <param name="access">Access mode for the new file</param>
 		/// <param name="share">Sharing mode for the open file</param>
-		/// <returns>New filestream for the given file</returns>
+		/// <returns>New FileStream for the given file</returns>
 		public static FileStream Open(FileReference location, FileMode mode, FileAccess access, FileShare share) => File.Open(location.FullName, mode, access, share);
 
 		/// <summary>
diff --git a/Engine/Source/Programs/Shared/EpicGames.Core/LogEvent.cs b/Engine/Source/Programs/Shared/EpicGames.Core/LogEvent.cs
--- a/Engine/Source/Programs/Shared/EpicGames.Core/LogEvent.cs
+++ b/Engine/Source/Programs/Shared/EpicGames.Core/LogEvent.cs
@@ -653,7 +653,7 @@ namespace EpicGames.Core
 		/// Constructor
 		/// </summary>
 		/// <param name="exception"></param>
-		[return: NotNullIfNotNull("exception")]
+		[return: NotNullIfNotNull(nameof(exception))]
 		public static LogException? FromException(Exception? exception)
 		{
 			LogException? result = null;
diff --git a/Engine/Source/Programs/UnrealBuildTool/Executors/UnrealBuildAccelerator/UBAAgentCoordinatorHorde.cs b/Engine/Source/Programs/UnrealBuildTool/Executors/UnrealBuildAccelerator/UBAAgentCoordinatorHorde.cs
--- a/Engine/Source/Programs/UnrealBuildTool/Executors/UnrealBuildAccelerator/UBAAgentCoordinatorHorde.cs
+++ b/Engine/Source/Programs/UnrealBuildTool/Executors/UnrealBuildAccelerator/UBAAgentCoordinatorHorde.cs
@@ -104,7 +104,11 @@ namespace UnrealBuildTool
 		{
 			if (_workers.Count > 0) // Should handle double-dispose, prevent cancelling twice
 			{
+#if NET8_0_OR_GREATER
+				await _cancellationTokenSource.CancelAsync();
+#else
 				_cancellationTokenSource.Cancel();
+#endif
 
 				for (int idx = _workers.Count - 1; idx >= 0; idx--)
 				{
diff --git a/Engine/Source/Programs/UnrealBuildTool/Modes/IWYUMode.cs b/Engine/Source/Programs/UnrealBuildTool/Modes/IWYUMode.cs
--- a/Engine/Source/Programs/UnrealBuildTool/Modes/IWYUMode.cs
+++ b/Engine/Source/Programs/UnrealBuildTool/Modes/IWYUMode.cs
@@ -412,7 +412,7 @@ namespace UnrealBuildTool
 			}
 			else if (!String.IsNullOrEmpty(PathToUpdateFile))
 			{
-				PathToUpdateList = File.ReadAllLines(PathToUpdateFile);
+				PathToUpdateList = await File.ReadAllLinesAsync(PathToUpdateFile);
 			}
 
 			if (PathToUpdateList.Length > 0)
diff --git a/Engine/Source/Programs/UnrealBuildTool/Modes/InlineGeneratedCpps.cs b/Engine/Source/Programs/UnrealBuildTool/Modes/InlineGeneratedCpps.cs
--- a/Engine/Source/Programs/UnrealBuildTool/Modes/InlineGeneratedCpps.cs
+++ b/Engine/Source/Programs/UnrealBuildTool/Modes/InlineGeneratedCpps.cs
@@ -143,7 +143,7 @@ namespace UnrealBuildTool
 										}
 									}
 
-									List<string> TextLines = System.IO.File.ReadAllLines(InputFile.Location.FullName).ToList();
+									List<string> TextLines = FileReference.ReadAllLines(InputFile.Location).ToList();
 									if (TextLines.Any(Text => InlineReflectionMarkupRegex.IsMatch(Text)))
 									{
 										continue;
diff --git a/Engine/Source/Programs/UnrealBuildTool/Platform/Windows/PVSToolChain.cs b/Engine/Source/Programs/UnrealBuildTool/Platform/Windows/PVSToolChain.cs
--- a/Engine/Source/Programs/UnrealBuildTool/Platform/Windows/PVSToolChain.cs
+++ b/Engine/Source/Programs/UnrealBuildTool/Platform/Windows/PVSToolChain.cs
@@ -434,7 +434,7 @@ namespace UnrealBuildTool
 			{
 				foreach (FileReference InputFile in InputFiles)
 				{
-					string[] Lines = File.ReadAllLines(InputFile.FullName);
+					string[] Lines = FileReference.ReadAllLines(InputFile);
 					for (int LineIdx = 0; LineIdx < Lines.Length; LineIdx++)
 					{
 						string Line = Lines[LineIdx];
