Horde: Capture and re-throw exceptions during extraction of file trees.

#rnx

diff --git a/Engine/Source/Programs/Shared/EpicGames.Horde/Storage/Nodes/DirectoryNode.Extract.cs b/Engine/Source/Programs/Shared/EpicGames.Horde/Storage/Nodes/DirectoryNode.Extract.cs
--- a/Engine/Source/Programs/Shared/EpicGames.Horde/Storage/Nodes/DirectoryNode.Extract.cs
+++ b/Engine/Source/Programs/Shared/EpicGames.Horde/Storage/Nodes/DirectoryNode.Extract.cs
@@ -8,6 +8,7 @@ using System.Diagnostics.CodeAnalysis;
 using System.IO;
 using System.IO.MemoryMappedFiles;
 using System.Linq;
+using System.Runtime.ExceptionServices;
 using System.Threading;
 using System.Threading.Channels;
 using System.Threading.Tasks;
@@ -330,6 +331,8 @@ namespace EpicGames.Horde.Storage.Nodes
 			ExtractStats extractStats = new ExtractStats(progress, frequency);
 			using (CancellationTokenSource cancellationSource = CancellationTokenSource.CreateLinkedTokenSource(cancellationToken))
 			{
+				ExceptionDispatchInfo? exceptionDispatchInfo = null;
+
 				// Helper method to run a background task and set a cancellation source on error
 				async Task RunBackgroundTask(Func<CancellationToken, Task> taskFunc)
 				{
@@ -344,6 +347,13 @@ namespace EpicGames.Horde.Storage.Nodes
 					catch (Exception ex)
 					{
 						logger.LogError(ex, "Error while extracting data: {Message}", ex.Message);
+
+						if (exceptionDispatchInfo == null)
+						{
+							ExceptionDispatchInfo dispatchInfo = ExceptionDispatchInfo.Capture(ex);
+							Interlocked.CompareExchange(ref exceptionDispatchInfo, dispatchInfo, null);
+						}
+
 #if NET8_0_OR_GREATER
 						await cancellationSource.CancelAsync();
 #else
@@ -380,6 +390,13 @@ namespace EpicGames.Horde.Storage.Nodes
 					}
 				}
 				extractStats.Flush();
+
+				cancellationToken.ThrowIfCancellationRequested();
+
+				if (exceptionDispatchInfo != null)
+				{
+					exceptionDispatchInfo.Throw();
+				}
 			}
 		}
 
