Horde: Fix description and creation time not being deserialized as part of artifact responses correctly.

diff --git a/Engine/Source/Programs/Shared/EpicGames.Horde/Artifacts/ArtifactHttpMessages.cs b/Engine/Source/Programs/Shared/EpicGames.Horde/Artifacts/ArtifactHttpMessages.cs
--- a/Engine/Source/Programs/Shared/EpicGames.Horde/Artifacts/ArtifactHttpMessages.cs
+++ b/Engine/Source/Programs/Shared/EpicGames.Horde/Artifacts/ArtifactHttpMessages.cs
@@ -86,7 +86,7 @@ namespace EpicGames.Horde.Artifacts
 		public ArtifactType Type { get; set; }
 
 		/// <inheritdoc cref="IArtifact.Description"/>
-		public string? Description { get; }
+		public string? Description { get; set; }
 
 		/// <inheritdoc cref="IArtifact.StreamId"/>
 		public StreamId StreamId { get; set; }
@@ -123,7 +123,7 @@ namespace EpicGames.Horde.Artifacts
 		public RefName RefName { get; set; }
 
 		/// <inheritdoc cref="IArtifact.CreatedAtUtc"/>
-		public DateTime CreatedAtUtc { get; }
+		public DateTime CreatedAtUtc { get; set; }
 
 		/// <summary>
 		/// Default constructor
