Horde: Replace NuGet package with more explicit one to reduce dependency size

diff --git a/Engine/Source/Programs/Shared/EpicGames.Horde/EpicGames.Horde.csproj b/Engine/Source/Programs/Shared/EpicGames.Horde/EpicGames.Horde.csproj
--- a/Engine/Source/Programs/Shared/EpicGames.Horde/EpicGames.Horde.csproj
+++ b/Engine/Source/Programs/Shared/EpicGames.Horde/EpicGames.Horde.csproj
@@ -47,7 +47,7 @@
       <IncludeAssets>runtime; build; native; contentfiles; analyzers; buildtransitive</IncludeAssets>
     </PackageReference>
     <PackageReference Include="K4os.Compression.LZ4" Version="1.2.16" />
-    <PackageReference Include="Microsoft.CodeAnalysis.CSharp.Workspaces" Version="4.2.0" />
+    <PackageReference Include="System.IO.Pipelines" Version="8.0.0" />
     <PackageReference Include="Microsoft.Data.Sqlite" Version="8.0.0" />
     <PackageReference Include="Microsoft.Extensions.Caching.Abstractions" Version="6.0.0" />
     <PackageReference Include="Microsoft.Extensions.Caching.Memory" Version="6.0.1" />
