Horde: Fix compile warnings

#skipci

diff --git a/Engine/Source/Programs/Horde/HordeCmd/Commands/Vcs/VcsBase.cs b/Engine/Source/Programs/Horde/HordeCmd/Commands/Vcs/VcsBase.cs
--- a/Engine/Source/Programs/Horde/HordeCmd/Commands/Vcs/VcsBase.cs
+++ b/Engine/Source/Programs/Horde/HordeCmd/Commands/Vcs/VcsBase.cs
@@ -4,7 +4,6 @@ using System.Diagnostics;
 using EpicGames.Core;
 using EpicGames.Horde.Storage;
 using EpicGames.Horde.Storage.Nodes;
-using Microsoft.CodeAnalysis;
 using Microsoft.Extensions.Logging;
 
 namespace Horde.Commands.Vcs
diff --git a/Engine/Source/Programs/Shared/EpicGames.Horde/Storage/Bundles/V1/BundleReader.cs b/Engine/Source/Programs/Shared/EpicGames.Horde/Storage/Bundles/V1/BundleReader.cs
--- a/Engine/Source/Programs/Shared/EpicGames.Horde/Storage/Bundles/V1/BundleReader.cs
+++ b/Engine/Source/Programs/Shared/EpicGames.Horde/Storage/Bundles/V1/BundleReader.cs
@@ -8,7 +8,6 @@ using System.Linq;
 using System.Threading;
 using System.Threading.Tasks;
 using EpicGames.Core;
-using Microsoft.CodeAnalysis;
 using Microsoft.Extensions.Logging;
 
 namespace EpicGames.Horde.Storage.Bundles.V1
