Fixed issue in recent UHT changes that prevented errors begin generated from the header file object.

diff --git a/Engine/Source/Programs/Shared/EpicGames.UHT/Types/UhtHeaderFile.cs b/Engine/Source/Programs/Shared/EpicGames.UHT/Types/UhtHeaderFile.cs
--- a/Engine/Source/Programs/Shared/EpicGames.UHT/Types/UhtHeaderFile.cs
+++ b/Engine/Source/Programs/Shared/EpicGames.UHT/Types/UhtHeaderFile.cs
@@ -116,7 +116,7 @@ namespace EpicGames.UHT.Types
 	/// <summary>
 	/// Represents a header file.
 	/// </summary>
-	public class UhtHeaderFile : IUhtMessageSite
+	public class UhtHeaderFile : IUhtMessageSite, IUhtMessageLineNumber
 	{
 		private readonly UhtSimpleMessageSite _messageSite;
 		private readonly UhtSourceFile _sourceFile;
@@ -231,7 +231,12 @@ namespace EpicGames.UHT.Types
 
 		/// <inheritdoc/>
 		[JsonIgnore]
-		public IUhtMessageLineNumber? MessageLineNumber => throw new NotImplementedException();
+		public IUhtMessageLineNumber? MessageLineNumber => this;
+		#endregion
+
+		#region IUhtMessageLineNumber implementation
+		[JsonIgnore]
+		int IUhtMessageLineNumber.MessageLineNumber => 1;
 		#endregion
 
 		/// <summary>
