Horde: Adding bundled information to tool responses

diff --git a/Engine/Source/Programs/Horde/Plugins/Tools/HordeServer.Tools/Tools/ITool.cs b/Engine/Source/Programs/Horde/Plugins/Tools/HordeServer.Tools/Tools/ITool.cs
--- a/Engine/Source/Programs/Horde/Plugins/Tools/HordeServer.Tools/Tools/ITool.cs
+++ b/Engine/Source/Programs/Horde/Plugins/Tools/HordeServer.Tools/Tools/ITool.cs
@@ -50,6 +50,11 @@ namespace HordeServer.Tools
 		/// </summary>
 		bool Public { get; }
 
+		/// <summary>
+		/// Whether the tool is bundled with the server
+		/// </summary>
+		bool Bundled { get; }
+
 		/// <summary>
 		/// Whether to show the tool for download in UGS
 		/// </summary>
diff --git a/Engine/Source/Programs/Horde/Plugins/Tools/HordeServer.Tools/Tools/ToolCollection.cs b/Engine/Source/Programs/Horde/Plugins/Tools/HordeServer.Tools/Tools/ToolCollection.cs
--- a/Engine/Source/Programs/Horde/Plugins/Tools/HordeServer.Tools/Tools/ToolCollection.cs
+++ b/Engine/Source/Programs/Horde/Plugins/Tools/HordeServer.Tools/Tools/ToolCollection.cs
@@ -101,6 +101,7 @@ namespace HordeServer.Tools
 			public string? Group => _config.Group;
 			public IReadOnlyList<string>? Platforms => _config.Platforms;
 			public bool Public => _config.Public;
+			public bool Bundled => _config is BundledToolConfig;
 			public bool ShowInUgs => _config.ShowInUgs;
 			public bool ShowInDashboard => _config.ShowInDashboard;
 			public IReadOnlyList<IToolDeployment> Deployments => _deployments;
diff --git a/Engine/Source/Programs/Horde/Plugins/Tools/HordeServer.Tools/Tools/ToolsController.cs b/Engine/Source/Programs/Horde/Plugins/Tools/HordeServer.Tools/Tools/ToolsController.cs
--- a/Engine/Source/Programs/Horde/Plugins/Tools/HordeServer.Tools/Tools/ToolsController.cs
+++ b/Engine/Source/Programs/Horde/Plugins/Tools/HordeServer.Tools/Tools/ToolsController.cs
@@ -204,7 +204,7 @@ namespace HordeServer.Tools
 		static GetToolSummaryResponse CreateGetToolSummaryResponse(ITool tool)
 		{
 			IToolDeployment? deployment = (tool.Deployments.Count == 0) ? null : tool.Deployments[^1];
-			return new GetToolSummaryResponse(tool.Id, tool.Name, tool.Description, tool.Category, tool.Group, tool.Platforms?.ToList(), deployment?.Version, deployment?.Id, deployment?.State, deployment?.Progress, tool.ShowInUgs, tool.ShowInDashboard);
+			return new GetToolSummaryResponse(tool.Id, tool.Name, tool.Description, tool.Category, tool.Group, tool.Platforms?.ToList(), deployment?.Version, deployment?.Id, deployment?.State, deployment?.Progress, tool.Bundled, tool.ShowInUgs, tool.ShowInDashboard);
 		}
 
 		/// <summary>
@@ -248,7 +248,7 @@ namespace HordeServer.Tools
 
 		static GetToolResponse CreateGetToolResponse(ITool tool, List<GetToolDeploymentResponse> deployments)
 		{
-			return new GetToolResponse(tool.Id, tool.Name, tool.Description, tool.Category, tool.Group, tool.Platforms?.ToList(), deployments, tool.Public, tool.ShowInUgs, tool.ShowInDashboard);
+			return new GetToolResponse(tool.Id, tool.Name, tool.Description, tool.Category, tool.Group, tool.Platforms?.ToList(), deployments, tool.Public, tool.Bundled, tool.ShowInUgs, tool.ShowInDashboard);
 		}
 
 		/// <summary>
diff --git a/Engine/Source/Programs/Shared/EpicGames.Horde/Tools/ToolMessages.cs b/Engine/Source/Programs/Shared/EpicGames.Horde/Tools/ToolMessages.cs
--- a/Engine/Source/Programs/Shared/EpicGames.Horde/Tools/ToolMessages.cs
+++ b/Engine/Source/Programs/Shared/EpicGames.Horde/Tools/ToolMessages.cs
@@ -18,9 +18,10 @@ namespace EpicGames.Horde.Tools
 	/// <param name="Platforms">List of platforms that this tool supports, as NET runtime identifiers.</param>
 	/// <param name="Deployments">Current deployments of this tool, sorted by time.</param>
 	/// <param name="Public">Whether this tool should be exposed for download on a public endpoint without authentication</param>
+	/// <param name="Bundled">Whether this tool is bundled with the server</param>
 	/// <param name="ShowInUgs">Whether to show this tool for download inside UGS</param>
 	/// <param name="ShowInDashboard">Whether to show this tool for download on the dashboard</param>
-	public record class GetToolResponse(ToolId Id, string Name, string Description, string? Category, string? Group, List<string>? Platforms, List<GetToolDeploymentResponse> Deployments, bool Public, bool ShowInUgs, bool ShowInDashboard);
+	public record class GetToolResponse(ToolId Id, string Name, string Description, string? Category, string? Group, List<string>? Platforms, List<GetToolDeploymentResponse> Deployments, bool Public, bool Bundled, bool ShowInUgs, bool ShowInDashboard);
 
 	/// <summary>
 	/// Summary for a particular tool.
@@ -35,9 +36,10 @@ namespace EpicGames.Horde.Tools
 	/// <param name="DeploymentId">Identifier for the current deployment</param>
 	/// <param name="DeploymentState">Current state of the deployment</param>
 	/// <param name="DeploymentProgress">Current progress of the deployment</param>
+	/// <param name="Bundled">Whether this tool is bundled with the server</param>
 	/// <param name="ShowInUgs">Whether to show this tool for download inside UGS</param>
 	/// <param name="ShowInDashboard">Whether to show this tool for download on the dashboard</param>
-	public record class GetToolSummaryResponse(ToolId Id, string Name, string Description, string? Category, string? Group, List<string>? Platforms, string? Version, ToolDeploymentId? DeploymentId, ToolDeploymentState? DeploymentState, double? DeploymentProgress, bool ShowInUgs, bool ShowInDashboard);
+	public record class GetToolSummaryResponse(ToolId Id, string Name, string Description, string? Category, string? Group, List<string>? Platforms, string? Version, ToolDeploymentId? DeploymentId, ToolDeploymentState? DeploymentState, double? DeploymentProgress, bool Bundled, bool ShowInUgs, bool ShowInDashboard);
 
 	/// <summary>
 	/// Response when querying all tools
