Horde: Guard against socket errors during closing of compute socket

Also ensure CloseAsync can't be invoked twice.

diff --git a/Engine/Source/Programs/Shared/EpicGames.Horde/Compute/ComputeSocket.cs b/Engine/Source/Programs/Shared/EpicGames.Horde/Compute/ComputeSocket.cs
--- a/Engine/Source/Programs/Shared/EpicGames.Horde/Compute/ComputeSocket.cs
+++ b/Engine/Source/Programs/Shared/EpicGames.Horde/Compute/ComputeSocket.cs
@@ -6,6 +6,7 @@ using System.Buffers.Binary;
 using System.Collections.Generic;
 using System.Diagnostics.CodeAnalysis;
 using System.Linq;
+using System.Net.Sockets;
 using System.Threading;
 using System.Threading.Tasks;
 using EpicGames.Core;
@@ -377,12 +378,13 @@ namespace EpicGames.Horde.Compute
 		readonly object _lockObject = new object();
 
 		bool _complete;
+		bool _isClosing;
 
 		readonly ComputeTransport _transport;
 		readonly ComputeProtocol _protocol;
 		readonly ILogger _logger;
 
-		readonly BackgroundTask _recvTask;
+		BackgroundTask? _recvTask;
 		readonly Dictionary<int, RecvBuffer> _recvBuffers = new Dictionary<int, RecvBuffer>();
 
 		readonly SemaphoreSlim _sendSemaphore = new SemaphoreSlim(1, 1);
@@ -414,16 +416,34 @@ namespace EpicGames.Horde.Compute
 		/// </summary>
 		public async ValueTask CloseAsync(CancellationToken cancellationToken)
 		{
-			// Close the transport layer, freeing the remote end to shutdown.
-			await _transport.MarkCompleteAsync(cancellationToken);
-
-			// Close all the buffers
-			await DetachAllBuffersAsync(true, true, cancellationToken);
-
-			// Wait for the reader to stop
-			if (_recvTask != null)
+			lock (_lockObject)
+			{
+				if (_isClosing)
+				{
+					return;
+				}
+				
+				_isClosing = true;
+			}
+			
+			try
+			{
+				// Close the transport layer, freeing the remote end to shutdown.
+				await _transport.MarkCompleteAsync(cancellationToken);
+				
+				// Close all the buffers
+				await DetachAllBuffersAsync(true, true, cancellationToken);
+				
+				// Wait for the reader to stop
+				if (_recvTask != null)
+				{
+					await _recvTask.DisposeAsync();
+					_recvTask = null;
+				}
+			}
+			catch (SocketException se) when (se.SocketErrorCode is SocketError.ConnectionReset or SocketError.ConnectionAborted)
 			{
-				await _recvTask.DisposeAsync();
+				_logger.LogInformation("Socket already closed. Ignoring exception");
 			}
 		}
 
@@ -667,7 +687,7 @@ namespace EpicGames.Horde.Compute
 				_recvBuffers.Add(channelId, new RecvBuffer(recvBuffer.CreateWriter()));
 
 				// Only start the receive task once we have a buffer to receive data, otherwise we discard data from the remote
-				if (_recvTask.Task == null)
+				if (_recvTask is { Task: null })
 				{
 					_recvTask.Start();
 				}
