Horde: Implement the ArtifactHttpCollection.AddAsync method.

#rnx

diff --git a/Engine/Source/Programs/Horde/Plugins/Build/HordeServer.Build/Artifacts/ArtifactsController.cs b/Engine/Source/Programs/Horde/Plugins/Build/HordeServer.Build/Artifacts/ArtifactsController.cs
--- a/Engine/Source/Programs/Horde/Plugins/Build/HordeServer.Build/Artifacts/ArtifactsController.cs
+++ b/Engine/Source/Programs/Horde/Plugins/Build/HordeServer.Build/Artifacts/ArtifactsController.cs
@@ -153,7 +153,7 @@ namespace HordeServer.Artifacts
 			claims.Add(new AclClaimConfig(HordeClaimTypes.WriteNamespace, $"{artifact.NamespaceId}:{artifact.RefName}"));
 
 			string token = await _aclService.IssueBearerTokenAsync(claims, TimeSpan.FromHours(8.0), cancellationToken);
-			return new CreateArtifactResponse(artifact.Id, artifact.NamespaceId, artifact.RefName, prevRefName, token);
+			return new CreateArtifactResponse(artifact.Id, artifact.CommitId, artifact.NamespaceId, artifact.RefName, prevRefName, token);
 		}
 
 		async Task<RefName?> GetPrevRefNameForArtifactAsync(IArtifact artifact, CancellationToken cancellationToken)
diff --git a/Engine/Source/Programs/Shared/EpicGames.Horde/Artifacts/ArtifactHttpCollection.cs b/Engine/Source/Programs/Shared/EpicGames.Horde/Artifacts/ArtifactHttpCollection.cs
--- a/Engine/Source/Programs/Shared/EpicGames.Horde/Artifacts/ArtifactHttpCollection.cs
+++ b/Engine/Source/Programs/Shared/EpicGames.Horde/Artifacts/ArtifactHttpCollection.cs
@@ -2,6 +2,7 @@
 
 using System;
 using System.Collections.Generic;
+using System.Linq;
 using System.Runtime.CompilerServices;
 using System.Threading;
 using System.Threading.Tasks;
@@ -16,28 +17,42 @@ namespace EpicGames.Horde.Artifacts
 		class Artifact : IArtifact
 		{
 			readonly ArtifactHttpCollection _collection;
-			readonly GetArtifactResponse _response;
 
 			public Artifact(ArtifactHttpCollection collection, GetArtifactResponse response)
+				: this(collection, response.Id, response.Name, response.Type, response.Description, response.StreamId, response.CommitId, response.Keys, response.Metadata, response.NamespaceId, response.RefName, response.CreatedAtUtc)
+			{ }
+
+			public Artifact(ArtifactHttpCollection collection, ArtifactId id, ArtifactName name, ArtifactType type, string? description, StreamId streamId, CommitIdWithOrder commitId, IReadOnlyList<string> keys, IReadOnlyList<string> metadata, NamespaceId namespaceId, RefName refName, DateTime createdAtUtc)
 			{
 				_collection = collection;
-				_response = response;
+
+				Id = id;
+				Name = name;
+				Type = type;
+				Description = description;
+				StreamId = streamId;
+				CommitId = commitId;
+				Keys = keys;
+				Metadata = metadata;
+				NamespaceId = namespaceId;
+				RefName = refName;
+				CreatedAtUtc = createdAtUtc;
 			}
 
-			public ArtifactId Id => _response.Id;
-			public ArtifactName Name => _response.Name;
-			public ArtifactType Type => _response.Type;
-			public string? Description => _response.Description;
-			public StreamId StreamId => _response.StreamId;
-			public CommitIdWithOrder CommitId => _response.CommitId;
-			public IReadOnlyList<string> Keys => _response.Keys;
-			public IReadOnlyList<string> Metadata => _response.Metadata;
-			public NamespaceId NamespaceId => _response.NamespaceId;
-			public RefName RefName => _response.RefName;
-			public DateTime CreatedAtUtc => _response.CreatedAtUtc;
+			public ArtifactId Id { get; }
+			public ArtifactName Name { get; }
+			public ArtifactType Type { get; }
+			public string? Description { get; }
+			public StreamId StreamId { get; }
+			public CommitIdWithOrder CommitId { get; }
+			public IReadOnlyList<string> Keys { get; }
+			public IReadOnlyList<string> Metadata { get; }
+			public NamespaceId NamespaceId { get; }
+			public RefName RefName { get; }
+			public DateTime CreatedAtUtc { get; }
 
 			public Task DeleteAsync(CancellationToken cancellationToken)
-				=> _collection.DeleteAsync(_response.Id, cancellationToken);
+				=> _collection.DeleteAsync(Id, cancellationToken);
 		}
 
 		readonly IHordeClient _hordeClient;
@@ -46,9 +61,11 @@ namespace EpicGames.Horde.Artifacts
 			=> _hordeClient = hordeClient;
 
 		/// <inheritdoc/>
-		public Task<IArtifact> AddAsync(ArtifactName name, ArtifactType type, string? description, StreamId streamId, CommitId commitId, IEnumerable<string> keys, IEnumerable<string> metadata, CancellationToken cancellationToken = default)
+		public async Task<IArtifact> AddAsync(ArtifactName name, ArtifactType type, string? description, StreamId streamId, CommitId commitId, IEnumerable<string> keys, IEnumerable<string> metadata, CancellationToken cancellationToken = default)
 		{
-			throw new NotImplementedException();
+			HordeHttpClient hordeHttpClient = _hordeClient.CreateHttpClient();
+			CreateArtifactResponse response = await hordeHttpClient.CreateArtifactAsync(name, type, description, streamId, commitId, keys, metadata, cancellationToken);
+			return new Artifact(this, response.ArtifactId, name, type, description, streamId, CommitIdWithOrder.FromPerforceChange(commitId.GetPerforceChange()), keys.ToList(), metadata.ToList(), response.NamespaceId, response.RefName, DateTime.UtcNow);
 		}
 
 		/// <inheritdoc/>
diff --git a/Engine/Source/Programs/Shared/EpicGames.Horde/Artifacts/ArtifactHttpMessages.cs b/Engine/Source/Programs/Shared/EpicGames.Horde/Artifacts/ArtifactHttpMessages.cs
--- a/Engine/Source/Programs/Shared/EpicGames.Horde/Artifacts/ArtifactHttpMessages.cs
+++ b/Engine/Source/Programs/Shared/EpicGames.Horde/Artifacts/ArtifactHttpMessages.cs
@@ -49,11 +49,12 @@ namespace EpicGames.Horde.Artifacts
 	/// Information about a created artifact
 	/// </summary>
 	/// <param name="ArtifactId">Identifier for the new artifact</param>
+	/// <param name="CommitId">Resolved commit id for the artifact</param>
 	/// <param name="NamespaceId">Namespace that should be written to with artifact data</param>
 	/// <param name="RefName">Ref to write to</param>
 	/// <param name="PrevRefName">Ref for the artifact at the changelist prior to this one. Can be used to deduplicate against.</param>
 	/// <param name="Token">Token which can be used to upload blobs for the artifact, and read blobs from the previous artifact</param>
-	public record CreateArtifactResponse(ArtifactId ArtifactId, NamespaceId NamespaceId, RefName RefName, RefName? PrevRefName, string Token);
+	public record CreateArtifactResponse(ArtifactId ArtifactId, CommitIdWithOrder CommitId, NamespaceId NamespaceId, RefName RefName, RefName? PrevRefName, string Token);
 
 	/// <summary>
 	/// Type of data to download for an artifact
