Horde: Catch compute cancellation exceptions and avoid flagging them as errors

diff --git a/Engine/Source/Programs/Shared/EpicGames.Horde/Compute/AgentManagedProcess.cs b/Engine/Source/Programs/Shared/EpicGames.Horde/Compute/AgentManagedProcess.cs
--- a/Engine/Source/Programs/Shared/EpicGames.Horde/Compute/AgentManagedProcess.cs
+++ b/Engine/Source/Programs/Shared/EpicGames.Horde/Compute/AgentManagedProcess.cs
@@ -72,6 +72,7 @@ namespace EpicGames.Horde.Compute
 				{
 					case AgentMessageType.Exception:
 						ExceptionMessage exception = message.ParseExceptionMessage();
+						ComputeExecutionCancelledException.TryThrow(exception);
 						throw new ComputeException("Error while executing remote process", new ComputeRemoteException(exception));
 					case AgentMessageType.ExecuteOutput:
 						AppendData(message.Data.Span);
diff --git a/Engine/Source/Programs/Shared/EpicGames.Horde/Compute/AgentMessage.cs b/Engine/Source/Programs/Shared/EpicGames.Horde/Compute/AgentMessage.cs
--- a/Engine/Source/Programs/Shared/EpicGames.Horde/Compute/AgentMessage.cs
+++ b/Engine/Source/Programs/Shared/EpicGames.Horde/Compute/AgentMessage.cs
@@ -204,12 +204,31 @@ namespace EpicGames.Horde.Compute
 	/// </summary>
 	public sealed class ComputeExecutionCancelledException : ComputeException
 	{
+		private const string Text = "Compute execution cancelled";
+		
 		/// <summary>
 		/// Constructor
 		/// </summary>
-		public ComputeExecutionCancelledException() : base("Compute execution cancelled")
+		public ComputeExecutionCancelledException() : base(Text)
 		{
 		}
+		
+		private ComputeExecutionCancelledException(Exception innerException) : base(Text, innerException)
+		{
+		}
+		
+		/// <summary>
+		/// Try constructing and throwing if the exception message matches a cancellation exception
+		/// </summary>
+		/// <param name="em">Deserialized exception message</param>
+		/// <exception cref="ComputeExecutionCancelledException">If message matches</exception>
+		public static void TryThrow(ExceptionMessage em)
+		{
+			if (em.Message == Text)
+			{
+				throw new ComputeExecutionCancelledException(new ComputeRemoteException(em));
+			}
+		}
 	}
 
 	/// <summary>
diff --git a/Engine/Source/Programs/UnrealBuildTool/Executors/UnrealBuildAccelerator/UBAAgentCoordinatorHorde.cs b/Engine/Source/Programs/UnrealBuildTool/Executors/UnrealBuildAccelerator/UBAAgentCoordinatorHorde.cs
--- a/Engine/Source/Programs/UnrealBuildTool/Executors/UnrealBuildAccelerator/UBAAgentCoordinatorHorde.cs
+++ b/Engine/Source/Programs/UnrealBuildTool/Executors/UnrealBuildAccelerator/UBAAgentCoordinatorHorde.cs
@@ -601,6 +601,17 @@ namespace UnrealBuildTool
 				self.NumLogicalCores = 0;
 				UpdateHordeStatus(null);
 			}
+			catch (ComputeExecutionCancelledException ex) 
+			{
+				// Cancellations are expected to happen due to spot instance interruptions or unscheduled maintenance of agents
+				self.NumLogicalCores = 0;
+				UpdateHordeStatus(null);
+				
+				if (!cancellationToken.IsCancellationRequested)
+				{
+					logger.Log(_strict ? LogLevel.Information : LogLevel.Debug, KnownLogEvents.Systemic_Horde_Compute, ex, "Compute lease cancelled");
+				}
+			}
 			catch (Exception ex)
 			{
 				self.NumLogicalCores = 0;
