Horde: Add a simple http object store, so we can upload tools to a public S3 bucket.

#rnx

diff --git a/Engine/Source/Programs/Shared/EpicGames.Horde/Storage/ObjectStores/HttpObjectStore.cs b/Engine/Source/Programs/Shared/EpicGames.Horde/Storage/ObjectStores/HttpObjectStore.cs
new file mode 100644
--- /dev/null
+++ b/Engine/Source/Programs/Shared/EpicGames.Horde/Storage/ObjectStores/HttpObjectStore.cs
@@ -0,0 +1,106 @@
+// Copyright Epic Games, Inc. All Rights Reserved.
+
+using System;
+using System.IO;
+using System.Net.Http;
+using System.Threading;
+using System.Threading.Tasks;
+using EpicGames.Core;
+
+namespace EpicGames.Horde.Storage.ObjectStores
+{
+	/// <summary>
+	/// Implements an object store utilizing standard HTTP operations
+	/// </summary>
+	public class HttpObjectStore : IObjectStore
+	{
+		readonly Uri _baseUrl;
+		readonly Func<HttpClient> _createHttpClient;
+
+		/// <inheritdoc/>
+		public bool SupportsRedirects => true;
+
+		/// <summary>
+		/// Constructor
+		/// </summary>
+		/// <param name="baseUrl">Base url for the store</param>
+		/// <param name="createHttpClient">Method for creating a new http client</param>
+		public HttpObjectStore(Uri baseUrl, Func<HttpClient> createHttpClient)
+		{
+			_baseUrl = baseUrl;
+			_createHttpClient = createHttpClient;
+		}
+
+		Uri GetUri(ObjectKey key)
+			=> new Uri(_baseUrl, key.ToString());
+
+		/// <inheritdoc/>
+		public async Task DeleteAsync(ObjectKey key, CancellationToken cancellationToken = default)
+		{
+			using HttpClient client = _createHttpClient();
+
+			using HttpResponseMessage response = await client.DeleteAsync(GetUri(key), cancellationToken);
+			response.EnsureSuccessStatusCode();
+		}
+
+		/// <inheritdoc/>
+		public async Task<bool> ExistsAsync(ObjectKey key, CancellationToken cancellationToken = default)
+		{
+			using HttpClient client = _createHttpClient();
+
+			using HttpRequestMessage request = new HttpRequestMessage(HttpMethod.Head, GetUri(key));
+			using HttpResponseMessage response = await client.SendAsync(request, cancellationToken);
+			return response.IsSuccessStatusCode;
+		}
+
+		/// <inheritdoc/>
+		public void GetStats(StorageStats stats)
+		{
+		}
+
+		/// <inheritdoc/>
+		public async Task<Stream> OpenAsync(ObjectKey key, int offset, int? length, CancellationToken cancellationToken = default)
+		{
+			using HttpClient client = _createHttpClient();
+
+			using HttpRequestMessage request = new HttpRequestMessage(HttpMethod.Get, GetUri(key));
+			if (offset > 0 || length != null)
+			{
+				request.Headers.Range = new System.Net.Http.Headers.RangeHeaderValue(offset, length);
+			}
+
+			HttpResponseMessage response = await client.SendAsync(request, cancellationToken);
+			Stream stream = await response.Content.ReadAsStreamAsync(cancellationToken);
+			return stream.WrapOwnership(response);
+		}
+
+		/// <inheritdoc/>
+		public async Task<IReadOnlyMemoryOwner<byte>> ReadAsync(ObjectKey key, int offset, int? length, CancellationToken cancellationToken = default)
+		{
+			using Stream stream = await OpenAsync(key, offset, length, cancellationToken);
+
+			byte[] data = await stream.ReadAllBytesAsync(cancellationToken);
+			return ReadOnlyMemoryOwner.Create(data);
+		}
+
+		/// <inheritdoc/>
+		public ValueTask<Uri?> TryGetReadRedirectAsync(ObjectKey key, CancellationToken cancellationToken = default)
+			=> new ValueTask<Uri?>(GetUri(key));
+
+		/// <inheritdoc/>
+		public ValueTask<Uri?> TryGetWriteRedirectAsync(ObjectKey key, CancellationToken cancellationToken = default)
+			=> new ValueTask<Uri?>(GetUri(key));
+
+		/// <inheritdoc/>
+		public async Task WriteAsync(ObjectKey key, Stream stream, CancellationToken cancellationToken = default)
+		{
+			using HttpClient client = _createHttpClient();
+
+			using HttpRequestMessage request = new HttpRequestMessage(HttpMethod.Put, GetUri(key));
+			request.Content = new StreamContent(stream);
+
+			using HttpResponseMessage response = await client.SendAsync(request, cancellationToken);
+			response.EnsureSuccessStatusCode();
+		}
+	}
+}
