Horde: Add a natively typed secrets collection to the Horde client.

#rnx

diff --git a/Engine/Source/Programs/Horde/Drivers/JobDriver.Tests/Execution/WorkspaceExecutorTest.cs b/Engine/Source/Programs/Horde/Drivers/JobDriver.Tests/Execution/WorkspaceExecutorTest.cs
--- a/Engine/Source/Programs/Horde/Drivers/JobDriver.Tests/Execution/WorkspaceExecutorTest.cs
+++ b/Engine/Source/Programs/Horde/Drivers/JobDriver.Tests/Execution/WorkspaceExecutorTest.cs
@@ -21,6 +21,7 @@ using Microsoft.Extensions.Logging;
 using Microsoft.Extensions.Logging.Abstractions;
 using Microsoft.VisualStudio.TestTools.UnitTesting;
 using EpicGames.Horde.Artifacts;
+using EpicGames.Horde.Secrets;
 
 namespace JobDriver.Tests.Execution
 {
@@ -84,6 +85,7 @@ namespace JobDriver.Tests.Execution
 				=> new Uri("http://fake-horde-server");
 
 			public IArtifactCollection Artifacts => throw new NotImplementedException();
+			public ISecretCollection Secrets => throw new NotImplementedException();
 
 			public IComputeClient CreateComputeClient()
 				=> throw new NotImplementedException();
diff --git a/Engine/Source/Programs/Horde/Plugins/Secrets/HordeServer.Secrets/Secrets/ISecret.cs b/Engine/Source/Programs/Horde/Plugins/Secrets/HordeServer.Secrets/Secrets/ISecret.cs
deleted file mode 100644
diff --git a/Engine/Source/Programs/Horde/Plugins/Secrets/HordeServer.Secrets/Secrets/ISecretCollection.cs b/Engine/Source/Programs/Horde/Plugins/Secrets/HordeServer.Secrets/Secrets/ISecretCollection.cs
deleted file mode 100644
diff --git a/Engine/Source/Programs/Horde/Plugins/Secrets/HordeServer.Secrets/SecretsPlugin.cs b/Engine/Source/Programs/Horde/Plugins/Secrets/HordeServer.Secrets/SecretsPlugin.cs
--- a/Engine/Source/Programs/Horde/Plugins/Secrets/HordeServer.Secrets/SecretsPlugin.cs
+++ b/Engine/Source/Programs/Horde/Plugins/Secrets/HordeServer.Secrets/SecretsPlugin.cs
@@ -1,5 +1,6 @@
 // Copyright Epic Games, Inc. All Rights Reserved.
 
+using EpicGames.Horde.Secrets;
 using HordeServer.Plugins;
 using HordeServer.Secrets;
 using HordeServer.Secrets.Providers;
diff --git a/Engine/Source/Programs/Shared/EpicGames.Horde.Tests/ServerLogPacketBuilderTests.cs b/Engine/Source/Programs/Shared/EpicGames.Horde.Tests/ServerLogPacketBuilderTests.cs
--- a/Engine/Source/Programs/Shared/EpicGames.Horde.Tests/ServerLogPacketBuilderTests.cs
+++ b/Engine/Source/Programs/Shared/EpicGames.Horde.Tests/ServerLogPacketBuilderTests.cs
@@ -11,6 +11,7 @@ using EpicGames.Core;
 using EpicGames.Horde.Artifacts;
 using EpicGames.Horde.Compute;
 using EpicGames.Horde.Logs;
+using EpicGames.Horde.Secrets;
 using EpicGames.Horde.Storage;
 using EpicGames.Horde.Storage.Bundles;
 using Grpc.Core;
@@ -216,6 +217,7 @@ namespace EpicGames.Horde.Tests
 			public Uri ServerUrl => throw new NotImplementedException();
 
 			public IArtifactCollection Artifacts => throw new NotImplementedException();
+			public ISecretCollection Secrets => throw new NotImplementedException();
 
 			public Task<bool> LoginAsync(bool allowLogin, CancellationToken cancellationToken)
 				=> throw new NotImplementedException();
diff --git a/Engine/Source/Programs/Shared/EpicGames.Horde/HordeClient.cs b/Engine/Source/Programs/Shared/EpicGames.Horde/HordeClient.cs
--- a/Engine/Source/Programs/Shared/EpicGames.Horde/HordeClient.cs
+++ b/Engine/Source/Programs/Shared/EpicGames.Horde/HordeClient.cs
@@ -12,6 +12,7 @@ using EpicGames.Horde.Artifacts;
 using EpicGames.Horde.Compute;
 using EpicGames.Horde.Compute.Clients;
 using EpicGames.Horde.Logs;
+using EpicGames.Horde.Secrets;
 using EpicGames.Horde.Server;
 using EpicGames.Horde.Storage;
 using EpicGames.Horde.Storage.Backends;
@@ -49,6 +50,9 @@ namespace EpicGames.Horde
 		/// <inheritdoc/>
 		public IArtifactCollection Artifacts { get; }
 
+		/// <inheritdoc/>
+		public ISecretCollection Secrets { get; }
+
 		/// <summary>
 		/// Constructor
 		/// </summary>
@@ -61,6 +65,7 @@ namespace EpicGames.Horde
 			_logger = _loggerFactory.CreateLogger<HordeClient>();
 
 			Artifacts = new ArtifactHttpCollection(this);
+			Secrets = new SecretHttpCollection(this);
 		}
 
 		/// <inheritdoc/>
diff --git a/Engine/Source/Programs/Shared/EpicGames.Horde/IHordeClient.cs b/Engine/Source/Programs/Shared/EpicGames.Horde/IHordeClient.cs
--- a/Engine/Source/Programs/Shared/EpicGames.Horde/IHordeClient.cs
+++ b/Engine/Source/Programs/Shared/EpicGames.Horde/IHordeClient.cs
@@ -6,6 +6,7 @@ using System.Threading.Tasks;
 using EpicGames.Horde.Artifacts;
 using EpicGames.Horde.Compute;
 using EpicGames.Horde.Logs;
+using EpicGames.Horde.Secrets;
 using EpicGames.Horde.Storage;
 using EpicGames.Horde.Tools;
 using Grpc.Core;
@@ -29,6 +30,11 @@ namespace EpicGames.Horde
 		/// </summary>
 		IArtifactCollection Artifacts { get; }
 
+		/// <summary>
+		/// Accessor for the secret collection
+		/// </summary>
+		ISecretCollection Secrets { get; }
+
 		/// <summary>
 		/// Connect to the Horde server
 		/// </summary>
diff --git a/Engine/Source/Programs/Shared/EpicGames.Horde/Secrets/ISecret.cs b/Engine/Source/Programs/Shared/EpicGames.Horde/Secrets/ISecret.cs
new file mode 100644
--- /dev/null
+++ b/Engine/Source/Programs/Shared/EpicGames.Horde/Secrets/ISecret.cs
@@ -0,0 +1,23 @@
+// Copyright Epic Games, Inc. All Rights Reserved.
+
+using System.Collections.Generic;
+using EpicGames.Horde.Secrets;
+
+namespace EpicGames.Horde.Secrets
+{
+	/// <summary>
+	/// Information about a secret
+	/// </summary>
+	public interface ISecret
+	{
+		/// <summary>
+		/// Identifier for the secret
+		/// </summary>
+		SecretId Id { get; }
+
+		/// <summary>
+		/// The secret values
+		/// </summary>
+		IReadOnlyDictionary<string, string> Data { get; }
+	}
+}
diff --git a/Engine/Source/Programs/Shared/EpicGames.Horde/Secrets/ISecretCollection.cs b/Engine/Source/Programs/Shared/EpicGames.Horde/Secrets/ISecretCollection.cs
new file mode 100644
--- /dev/null
+++ b/Engine/Source/Programs/Shared/EpicGames.Horde/Secrets/ISecretCollection.cs
@@ -0,0 +1,20 @@
+﻿// Copyright Epic Games, Inc. All Rights Reserved.
+
+using System.Threading;
+using System.Threading.Tasks;
+
+namespace EpicGames.Horde.Secrets
+{
+	/// <summary>
+	/// Collection of secrets
+	/// </summary>
+	public interface ISecretCollection
+	{
+		/// <summary>
+		/// Resolve a secret to concrete values
+		/// </summary>
+		/// <param name="secretId">Identifier for the secret</param>
+		/// <param name="cancellationToken">Cancellation token for the operation</param>
+		Task<ISecret?> GetAsync(SecretId secretId, CancellationToken cancellationToken = default);
+	}
+}
diff --git a/Engine/Source/Programs/Shared/EpicGames.Horde/Secrets/SecretHttpCollection.cs b/Engine/Source/Programs/Shared/EpicGames.Horde/Secrets/SecretHttpCollection.cs
new file mode 100644
--- /dev/null
+++ b/Engine/Source/Programs/Shared/EpicGames.Horde/Secrets/SecretHttpCollection.cs
@@ -0,0 +1,41 @@
+﻿// Copyright Epic Games, Inc. All Rights Reserved.
+
+using System;
+using System.Collections.Generic;
+using System.Linq;
+using System.Text;
+using System.Threading;
+using System.Threading.Tasks;
+
+namespace EpicGames.Horde.Secrets
+{
+	class SecretHttpCollection : ISecretCollection
+	{
+		class Secret : ISecret
+		{
+			readonly GetSecretResponse _response;
+
+			public SecretId Id => _response.Id;
+			public IReadOnlyDictionary<string, string> Data => _response.Data;
+
+			public Secret(GetSecretResponse response)
+				=> _response = response;
+		}
+
+		readonly IHordeClient _hordeClient;
+
+		/// <summary>
+		/// Constructor
+		/// </summary>
+		public SecretHttpCollection(IHordeClient hordeClient)
+			=> _hordeClient = hordeClient;
+
+		/// <inheritdoc/>
+		public async Task<ISecret?> GetAsync(SecretId secretId, CancellationToken cancellationToken)
+		{
+			HordeHttpClient hordeHttpClient = _hordeClient.CreateHttpClient();
+			GetSecretResponse response = await hordeHttpClient.GetSecretAsync(secretId, cancellationToken);
+			return new Secret(response);
+		}
+	}
+}
