Horde: Add a native collection for projects to IHordeHttpClient.

#rnx

diff --git a/Engine/Source/Programs/Horde/Drivers/JobDriver.Tests/Execution/WorkspaceExecutorTest.cs b/Engine/Source/Programs/Horde/Drivers/JobDriver.Tests/Execution/WorkspaceExecutorTest.cs
--- a/Engine/Source/Programs/Horde/Drivers/JobDriver.Tests/Execution/WorkspaceExecutorTest.cs
+++ b/Engine/Source/Programs/Horde/Drivers/JobDriver.Tests/Execution/WorkspaceExecutorTest.cs
@@ -22,6 +22,7 @@ using Microsoft.Extensions.Logging.Abstractions;
 using Microsoft.VisualStudio.TestTools.UnitTesting;
 using EpicGames.Horde.Artifacts;
 using EpicGames.Horde.Secrets;
+using EpicGames.Horde.Projects;
 
 namespace JobDriver.Tests.Execution
 {
@@ -85,6 +86,7 @@ namespace JobDriver.Tests.Execution
 				=> new Uri("http://fake-horde-server");
 
 			public IArtifactCollection Artifacts => throw new NotImplementedException();
+			public IProjectCollection Projects => throw new NotImplementedException();
 			public ISecretCollection Secrets => throw new NotImplementedException();
 
 			public IComputeClient CreateComputeClient()
diff --git a/Engine/Source/Programs/Shared/EpicGames.Horde.Tests/ServerLogPacketBuilderTests.cs b/Engine/Source/Programs/Shared/EpicGames.Horde.Tests/ServerLogPacketBuilderTests.cs
--- a/Engine/Source/Programs/Shared/EpicGames.Horde.Tests/ServerLogPacketBuilderTests.cs
+++ b/Engine/Source/Programs/Shared/EpicGames.Horde.Tests/ServerLogPacketBuilderTests.cs
@@ -11,6 +11,7 @@ using EpicGames.Core;
 using EpicGames.Horde.Artifacts;
 using EpicGames.Horde.Compute;
 using EpicGames.Horde.Logs;
+using EpicGames.Horde.Projects;
 using EpicGames.Horde.Secrets;
 using EpicGames.Horde.Storage;
 using EpicGames.Horde.Storage.Bundles;
@@ -217,6 +218,7 @@ namespace EpicGames.Horde.Tests
 			public Uri ServerUrl => throw new NotImplementedException();
 
 			public IArtifactCollection Artifacts => throw new NotImplementedException();
+			public IProjectCollection Projects => throw new NotImplementedException();
 			public ISecretCollection Secrets => throw new NotImplementedException();
 
 			public Task<bool> LoginAsync(bool allowLogin, CancellationToken cancellationToken)
diff --git a/Engine/Source/Programs/Shared/EpicGames.Horde/Artifacts/ArtifactHttpCollection.cs b/Engine/Source/Programs/Shared/EpicGames.Horde/Artifacts/ArtifactHttpCollection.cs
--- a/Engine/Source/Programs/Shared/EpicGames.Horde/Artifacts/ArtifactHttpCollection.cs
+++ b/Engine/Source/Programs/Shared/EpicGames.Horde/Artifacts/ArtifactHttpCollection.cs
@@ -2,6 +2,7 @@
 
 using System;
 using System.Collections.Generic;
+using System.Diagnostics;
 using System.Linq;
 using System.Runtime.CompilerServices;
 using System.Threading;
@@ -15,6 +16,7 @@ namespace EpicGames.Horde.Artifacts
 {
 	class ArtifactHttpCollection : IArtifactCollection
 	{
+		[DebuggerDisplay("{Id}")]
 		class Artifact : IArtifact
 		{
 			readonly ArtifactHttpCollection _collection;
diff --git a/Engine/Source/Programs/Shared/EpicGames.Horde/HordeClient.cs b/Engine/Source/Programs/Shared/EpicGames.Horde/HordeClient.cs
--- a/Engine/Source/Programs/Shared/EpicGames.Horde/HordeClient.cs
+++ b/Engine/Source/Programs/Shared/EpicGames.Horde/HordeClient.cs
@@ -12,6 +12,7 @@ using EpicGames.Horde.Artifacts;
 using EpicGames.Horde.Compute;
 using EpicGames.Horde.Compute.Clients;
 using EpicGames.Horde.Logs;
+using EpicGames.Horde.Projects;
 using EpicGames.Horde.Secrets;
 using EpicGames.Horde.Server;
 using EpicGames.Horde.Storage;
@@ -50,6 +51,9 @@ namespace EpicGames.Horde
 		/// <inheritdoc/>
 		public IArtifactCollection Artifacts { get; }
 
+		/// <inheritdoc/>
+		public IProjectCollection Projects { get; }
+
 		/// <inheritdoc/>
 		public ISecretCollection Secrets { get; }
 
@@ -65,6 +69,7 @@ namespace EpicGames.Horde
 			_logger = _loggerFactory.CreateLogger<HordeClient>();
 
 			Artifacts = new ArtifactHttpCollection(this);
+			Projects = new ProjectHttpCollection(this);
 			Secrets = new SecretHttpCollection(this);
 		}
 
diff --git a/Engine/Source/Programs/Shared/EpicGames.Horde/IHordeClient.cs b/Engine/Source/Programs/Shared/EpicGames.Horde/IHordeClient.cs
--- a/Engine/Source/Programs/Shared/EpicGames.Horde/IHordeClient.cs
+++ b/Engine/Source/Programs/Shared/EpicGames.Horde/IHordeClient.cs
@@ -6,6 +6,7 @@ using System.Threading.Tasks;
 using EpicGames.Horde.Artifacts;
 using EpicGames.Horde.Compute;
 using EpicGames.Horde.Logs;
+using EpicGames.Horde.Projects;
 using EpicGames.Horde.Secrets;
 using EpicGames.Horde.Storage;
 using EpicGames.Horde.Tools;
@@ -30,6 +31,11 @@ namespace EpicGames.Horde
 		/// </summary>
 		IArtifactCollection Artifacts { get; }
 
+		/// <summary>
+		/// Accessor for the project collection
+		/// </summary>
+		IProjectCollection Projects { get; }
+
 		/// <summary>
 		/// Accessor for the secret collection
 		/// </summary>
diff --git a/Engine/Source/Programs/Shared/EpicGames.Horde/Projects/IProject.cs b/Engine/Source/Programs/Shared/EpicGames.Horde/Projects/IProject.cs
new file mode 100644
--- /dev/null
+++ b/Engine/Source/Programs/Shared/EpicGames.Horde/Projects/IProject.cs
@@ -0,0 +1,53 @@
+// Copyright Epic Games, Inc. All Rights Reserved.
+
+using System;
+using System.Collections.Generic;
+using System.Linq;
+using System.Text;
+using System.Threading.Tasks;
+using EpicGames.Horde.Streams;
+
+namespace EpicGames.Horde.Projects
+{
+	/// <summary>
+	/// Handle to a Horde project
+	/// </summary>
+	public interface IProject
+	{
+		/// <summary>
+		/// Unique id of the project
+		/// </summary>
+		ProjectId Id { get; }
+
+		/// <summary>
+		/// Name of the project
+		/// </summary>
+		string Name { get; }
+
+		/// <summary>
+		/// Order to display this project on the dashboard
+		/// </summary>
+		int Order { get; }
+
+		/// <summary>
+		/// List of streams that are in this project
+		/// </summary>
+		IReadOnlyList<IProjectStream> Streams { get; }
+	}
+
+	/// <summary>
+	/// Describes a stream within a project
+	/// </summary>
+	public interface IProjectStream
+	{
+		/// <summary>
+		/// The stream id
+		/// </summary>
+		StreamId Id { get; }
+
+		/// <summary>
+		/// The stream name
+		/// </summary>
+		string Name { get; }
+	}
+}
diff --git a/Engine/Source/Programs/Shared/EpicGames.Horde/Projects/IProjectCollection.cs b/Engine/Source/Programs/Shared/EpicGames.Horde/Projects/IProjectCollection.cs
new file mode 100644
--- /dev/null
+++ b/Engine/Source/Programs/Shared/EpicGames.Horde/Projects/IProjectCollection.cs
@@ -0,0 +1,32 @@
+// Copyright Epic Games, Inc. All Rights Reserved.
+
+using System;
+using System.Collections.Generic;
+using System.Linq;
+using System.Text;
+using System.Threading;
+using System.Threading.Tasks;
+
+namespace EpicGames.Horde.Projects
+{
+	/// <summary>
+	/// Collection of projects
+	/// </summary>
+	public interface IProjectCollection
+	{
+		/// <summary>
+		/// Retrieve information about a specific project
+		/// </summary>
+		/// <param name="projectId">Id of the project to get information about</param>
+		/// <param name="cancellationToken">Cancellation token for the operation</param>
+		/// <returns>Information about the requested project</returns>
+		Task<IProject> GetAsync(ProjectId projectId, CancellationToken cancellationToken = default);
+
+		/// <summary>
+		/// Query all the projects
+		/// </summary>
+		/// <param name="cancellationToken">Cancellation token for the operation</param>
+		/// <returns>Information about all the projects</returns>
+		Task<IReadOnlyList<IProject>> GetAllAsync(CancellationToken cancellationToken = default);
+	}
+}
diff --git a/Engine/Source/Programs/Shared/EpicGames.Horde/Projects/ProjectHttpCollection.cs b/Engine/Source/Programs/Shared/EpicGames.Horde/Projects/ProjectHttpCollection.cs
new file mode 100644
--- /dev/null
+++ b/Engine/Source/Programs/Shared/EpicGames.Horde/Projects/ProjectHttpCollection.cs
@@ -0,0 +1,67 @@
+// Copyright Epic Games, Inc. All Rights Reserved.
+
+using System;
+using System.Collections.Generic;
+using System.Diagnostics;
+using System.Threading;
+using System.Threading.Tasks;
+using EpicGames.Horde.Streams;
+
+namespace EpicGames.Horde.Projects
+{
+	class ProjectHttpCollection : IProjectCollection
+	{
+		[DebuggerDisplay("{Id}")]
+		class Project : IProject
+		{
+			readonly GetProjectResponse _response;
+			readonly IReadOnlyList<ProjectStream> _projectStreams;
+
+			public ProjectId Id => _response.Id;
+			public string Name => _response.Name;
+			public int Order => _response.Order;
+			public IReadOnlyList<IProjectStream> Streams => _projectStreams;
+
+			public Project(GetProjectResponse response)
+			{
+				_response = response;
+				_projectStreams = (IReadOnlyList<ProjectStream>?)response.Streams?.ConvertAll(x => new ProjectStream(x)) ?? Array.Empty<ProjectStream>();
+			}
+		}
+
+		class ProjectStream : IProjectStream
+		{
+			readonly GetProjectStreamResponse _response;
+
+			public StreamId Id => new StreamId(_response.Id);
+			public string Name => _response.Name;
+
+			public ProjectStream(GetProjectStreamResponse response)
+				=> _response = response;
+		}
+
+		readonly IHordeClient _hordeClient;
+
+		/// <summary>
+		/// Constructor
+		/// </summary>
+		public ProjectHttpCollection(IHordeClient hordeClient)
+			=> _hordeClient = hordeClient;
+
+		/// <inheritdoc/>
+		public async Task<IProject> GetAsync(ProjectId projectId, CancellationToken cancellationToken)
+		{
+			HordeHttpClient hordeHttpClient = _hordeClient.CreateHttpClient();
+			GetProjectResponse? response = await hordeHttpClient.GetProjectAsync(projectId, cancellationToken);
+			return new Project(response);
+		}
+
+		/// <inheritdoc/>
+		public async Task<IReadOnlyList<IProject>> GetAllAsync(CancellationToken cancellationToken)
+		{
+			HordeHttpClient hordeHttpClient = _hordeClient.CreateHttpClient();
+			List<GetProjectResponse> responses = await hordeHttpClient.GetProjectsAsync(cancellationToken: cancellationToken);
+			return responses.ConvertAll(x => new Project(x));
+		}
+	}
+}
