Horde: Allow creating zip streams from unhashed blob refs.

#rnx

diff --git a/Engine/Source/Programs/Shared/EpicGames.Horde/Storage/Nodes/DirectoryNode.ZipStream.cs b/Engine/Source/Programs/Shared/EpicGames.Horde/Storage/Nodes/DirectoryNode.ZipStream.cs
--- a/Engine/Source/Programs/Shared/EpicGames.Horde/Storage/Nodes/DirectoryNode.ZipStream.cs
+++ b/Engine/Source/Programs/Shared/EpicGames.Horde/Storage/Nodes/DirectoryNode.ZipStream.cs
@@ -45,7 +45,7 @@ namespace EpicGames.Horde.Storage.Nodes
 		/// <param name="node">Root node to copy from</param>
 		/// <param name="filter">Filter for files to include in the zip</param>
 		/// <param name="logger">Optional logger for debug tracing</param>
-		public DirectoryNodeZipStream(IHashedBlobRef<DirectoryNode> node, FileFilter? filter, ILogger? logger)
+		public DirectoryNodeZipStream(IBlobRef<DirectoryNode> node, FileFilter? filter, ILogger? logger)
 		{
 			_pipe = new Pipe();
 			_backgroundTask = BackgroundTask.StartNew(ctx => CopyToPipeAsync(node, filter, _pipe.Writer, logger, ctx));
@@ -109,14 +109,14 @@ namespace EpicGames.Horde.Storage.Nodes
 			return length;
 		}
 
-		static async Task CopyToPipeAsync(IHashedBlobRef<DirectoryNode> node, FileFilter? filter, PipeWriter writer, ILogger? logger, CancellationToken cancellationToken)
+		static async Task CopyToPipeAsync(IBlobRef<DirectoryNode> node, FileFilter? filter, PipeWriter writer, ILogger? logger, CancellationToken cancellationToken)
 		{
 			using Stream outputStream = writer.AsStream();
 			using ZipArchive archive = new ZipArchive(outputStream, ZipArchiveMode.Create);
 			await CopyFilesAsync(node, "", filter, archive, logger, cancellationToken);
 		}
 
-		static async Task CopyFilesAsync(IHashedBlobRef<DirectoryNode> directoryRef, string prefix, FileFilter? filter, ZipArchive archive, ILogger? logger, CancellationToken cancellationToken)
+		static async Task CopyFilesAsync(IBlobRef<DirectoryNode> directoryRef, string prefix, FileFilter? filter, ZipArchive archive, ILogger? logger, CancellationToken cancellationToken)
 		{
 			DirectoryNode directory = await directoryRef.ReadBlobAsync(cancellationToken);
 
@@ -192,7 +192,7 @@ namespace EpicGames.Horde.Storage.Nodes
 		/// <param name="filter">Filter for files to include in the zip</param>
 		/// <param name="logger">Logger for diagnostic output</param>
 		/// <returns>Stream containing zipped archive data</returns>
-		public static Stream AsZipStream(this IHashedBlobRef<DirectoryNode> directoryRef, FileFilter? filter = null, ILogger? logger = null)
+		public static Stream AsZipStream(this IBlobRef<DirectoryNode> directoryRef, FileFilter? filter = null, ILogger? logger = null)
 			=> new DirectoryNodeZipStream(directoryRef, filter, logger);
 	}
 }
