Horde: Add a native interface for the tools collection to IHordeClient.

#rnx

diff --git a/Engine/Source/Programs/Horde/Drivers/JobDriver.Tests/Execution/WorkspaceExecutorTest.cs b/Engine/Source/Programs/Horde/Drivers/JobDriver.Tests/Execution/WorkspaceExecutorTest.cs
--- a/Engine/Source/Programs/Horde/Drivers/JobDriver.Tests/Execution/WorkspaceExecutorTest.cs
+++ b/Engine/Source/Programs/Horde/Drivers/JobDriver.Tests/Execution/WorkspaceExecutorTest.cs
@@ -23,6 +23,7 @@ using Microsoft.VisualStudio.TestTools.UnitTesting;
 using EpicGames.Horde.Artifacts;
 using EpicGames.Horde.Secrets;
 using EpicGames.Horde.Projects;
+using EpicGames.Horde.Tools;
 
 namespace JobDriver.Tests.Execution
 {
@@ -88,6 +89,7 @@ namespace JobDriver.Tests.Execution
 			public IArtifactCollection Artifacts => throw new NotImplementedException();
 			public IProjectCollection Projects => throw new NotImplementedException();
 			public ISecretCollection Secrets => throw new NotImplementedException();
+			public IToolCollection Tools => throw new NotImplementedException();
 
 			public IComputeClient CreateComputeClient()
 				=> throw new NotImplementedException();
diff --git a/Engine/Source/Programs/Horde/Plugins/Build/HordeServer.Build.Tests/BuildTestSetup.cs b/Engine/Source/Programs/Horde/Plugins/Build/HordeServer.Build.Tests/BuildTestSetup.cs
--- a/Engine/Source/Programs/Horde/Plugins/Build/HordeServer.Build.Tests/BuildTestSetup.cs
+++ b/Engine/Source/Programs/Horde/Plugins/Build/HordeServer.Build.Tests/BuildTestSetup.cs
@@ -7,6 +7,7 @@ using Amazon.EC2;
 using EpicGames.Horde.Agents;
 using EpicGames.Horde.Agents.Pools;
 using EpicGames.Horde.Artifacts;
+using EpicGames.Horde.Tools;
 using HordeCommon.Rpc.Messages;
 using HordeServer.Acls;
 using HordeServer.Agents;
diff --git a/Engine/Source/Programs/Horde/Plugins/Compute/HordeServer.Compute/Agents/AgentVersionProvider.cs b/Engine/Source/Programs/Horde/Plugins/Compute/HordeServer.Compute/Agents/AgentVersionProvider.cs
--- a/Engine/Source/Programs/Horde/Plugins/Compute/HordeServer.Compute/Agents/AgentVersionProvider.cs
+++ b/Engine/Source/Programs/Horde/Plugins/Compute/HordeServer.Compute/Agents/AgentVersionProvider.cs
@@ -1,8 +1,8 @@
 // Copyright Epic Games, Inc. All Rights Reserved.
 
 using EpicGames.Core;
+using EpicGames.Horde.Tools;
 using HordeServer.Server;
-using HordeServer.Tools;
 
 namespace HordeServer.Agents
 {
diff --git a/Engine/Source/Programs/Horde/Plugins/Compute/HordeServer.Compute/Agents/Software/AgentSoftwareController.cs b/Engine/Source/Programs/Horde/Plugins/Compute/HordeServer.Compute/Agents/Software/AgentSoftwareController.cs
--- a/Engine/Source/Programs/Horde/Plugins/Compute/HordeServer.Compute/Agents/Software/AgentSoftwareController.cs
+++ b/Engine/Source/Programs/Horde/Plugins/Compute/HordeServer.Compute/Agents/Software/AgentSoftwareController.cs
@@ -1,7 +1,7 @@
 // Copyright Epic Games, Inc. All Rights Reserved.
 
 using EpicGames.Core;
-using HordeServer.Tools;
+using EpicGames.Horde.Tools;
 using Microsoft.AspNetCore.Authorization;
 using Microsoft.AspNetCore.Mvc;
 using Microsoft.Extensions.Options;
diff --git a/Engine/Source/Programs/Horde/Plugins/Compute/HordeServer.Compute/Tasks/UpgradeTaskSource.cs b/Engine/Source/Programs/Horde/Plugins/Compute/HordeServer.Compute/Tasks/UpgradeTaskSource.cs
--- a/Engine/Source/Programs/Horde/Plugins/Compute/HordeServer.Compute/Tasks/UpgradeTaskSource.cs
+++ b/Engine/Source/Programs/Horde/Plugins/Compute/HordeServer.Compute/Tasks/UpgradeTaskSource.cs
@@ -9,7 +9,6 @@ using EpicGames.Horde.Tools;
 using HordeCommon.Rpc.Tasks;
 using HordeServer.Agents;
 using HordeServer.Logs;
-using HordeServer.Tools;
 using HordeServer.Utilities;
 using Microsoft.Extensions.Options;
 
diff --git a/Engine/Source/Programs/Horde/Plugins/Tools/HordeServer.Tools/Tools/ITool.cs b/Engine/Source/Programs/Horde/Plugins/Tools/HordeServer.Tools/Tools/ITool.cs
deleted file mode 100644
diff --git a/Engine/Source/Programs/Horde/Plugins/Tools/HordeServer.Tools/Tools/IToolCollection.cs b/Engine/Source/Programs/Horde/Plugins/Tools/HordeServer.Tools/Tools/IToolCollection.cs
deleted file mode 100644
diff --git a/Engine/Source/Programs/Horde/Plugins/Tools/HordeServer.Tools/Tools/ToolConfig.cs b/Engine/Source/Programs/Horde/Plugins/Tools/HordeServer.Tools/Tools/ToolConfig.cs
--- a/Engine/Source/Programs/Horde/Plugins/Tools/HordeServer.Tools/Tools/ToolConfig.cs
+++ b/Engine/Source/Programs/Horde/Plugins/Tools/HordeServer.Tools/Tools/ToolConfig.cs
@@ -141,21 +141,4 @@ namespace HordeServer.Tools
 			Public = true;
 		}
 	}
-
-	/// <summary>
-	/// Options for a new deployment
-	/// </summary>
-	public class ToolDeploymentConfig
-	{
-		/// <inheritdoc cref="IToolDeployment.Version"/>
-		public string Version { get; set; } = "Unknown";
-
-		/// <inheritdoc cref="IToolDeployment.Duration"/>
-		public TimeSpan Duration { get; set; }
-
-		/// <summary>
-		/// Whether to create the deployment in a paused state
-		/// </summary>
-		public bool CreatePaused { get; set; }
-	}
 }
diff --git a/Engine/Source/Programs/Horde/Plugins/Tools/HordeServer.Tools/ToolsPlugin.cs b/Engine/Source/Programs/Horde/Plugins/Tools/HordeServer.Tools/ToolsPlugin.cs
--- a/Engine/Source/Programs/Horde/Plugins/Tools/HordeServer.Tools/ToolsPlugin.cs
+++ b/Engine/Source/Programs/Horde/Plugins/Tools/HordeServer.Tools/ToolsPlugin.cs
@@ -2,6 +2,7 @@
 
 using System.Reflection;
 using EpicGames.Core;
+using EpicGames.Horde.Tools;
 using HordeServer.Acls;
 using HordeServer.Plugins;
 using HordeServer.Tools;
diff --git a/Engine/Source/Programs/Shared/EpicGames.Horde.Tests/ServerLogPacketBuilderTests.cs b/Engine/Source/Programs/Shared/EpicGames.Horde.Tests/ServerLogPacketBuilderTests.cs
--- a/Engine/Source/Programs/Shared/EpicGames.Horde.Tests/ServerLogPacketBuilderTests.cs
+++ b/Engine/Source/Programs/Shared/EpicGames.Horde.Tests/ServerLogPacketBuilderTests.cs
@@ -15,6 +15,7 @@ using EpicGames.Horde.Projects;
 using EpicGames.Horde.Secrets;
 using EpicGames.Horde.Storage;
 using EpicGames.Horde.Storage.Bundles;
+using EpicGames.Horde.Tools;
 using Grpc.Core;
 using Grpc.Net.Client;
 using Horde.Common.Rpc;
@@ -220,6 +221,7 @@ namespace EpicGames.Horde.Tests
 			public IArtifactCollection Artifacts => throw new NotImplementedException();
 			public IProjectCollection Projects => throw new NotImplementedException();
 			public ISecretCollection Secrets => throw new NotImplementedException();
+			public IToolCollection Tools => throw new NotImplementedException();
 
 			public Task<bool> LoginAsync(bool allowLogin, CancellationToken cancellationToken)
 				=> throw new NotImplementedException();
diff --git a/Engine/Source/Programs/Shared/EpicGames.Horde/HordeClient.cs b/Engine/Source/Programs/Shared/EpicGames.Horde/HordeClient.cs
--- a/Engine/Source/Programs/Shared/EpicGames.Horde/HordeClient.cs
+++ b/Engine/Source/Programs/Shared/EpicGames.Horde/HordeClient.cs
@@ -18,6 +18,7 @@ using EpicGames.Horde.Server;
 using EpicGames.Horde.Storage;
 using EpicGames.Horde.Storage.Backends;
 using EpicGames.Horde.Storage.Bundles;
+using EpicGames.Horde.Tools;
 using Grpc.Core;
 using Grpc.Net.Client;
 using Grpc.Net.Client.Configuration;
@@ -57,6 +58,9 @@ namespace EpicGames.Horde
 		/// <inheritdoc/>
 		public ISecretCollection Secrets { get; }
 
+		/// <inheritdoc/>
+		public IToolCollection Tools { get; }
+
 		/// <summary>
 		/// Constructor
 		/// </summary>
@@ -71,6 +75,7 @@ namespace EpicGames.Horde
 			Artifacts = new ArtifactHttpCollection(this);
 			Projects = new ProjectHttpCollection(this);
 			Secrets = new SecretHttpCollection(this);
+			Tools = new ToolHttpCollection(this);
 		}
 
 		/// <inheritdoc/>
diff --git a/Engine/Source/Programs/Shared/EpicGames.Horde/IHordeClient.cs b/Engine/Source/Programs/Shared/EpicGames.Horde/IHordeClient.cs
--- a/Engine/Source/Programs/Shared/EpicGames.Horde/IHordeClient.cs
+++ b/Engine/Source/Programs/Shared/EpicGames.Horde/IHordeClient.cs
@@ -41,6 +41,11 @@ namespace EpicGames.Horde
 		/// </summary>
 		ISecretCollection Secrets { get; }
 
+		/// <summary>
+		/// Accessor for the tools collection
+		/// </summary>
+		IToolCollection Tools { get; }
+
 		/// <summary>
 		/// Connect to the Horde server
 		/// </summary>
diff --git a/Engine/Source/Programs/Shared/EpicGames.Horde/Projects/IProject.cs b/Engine/Source/Programs/Shared/EpicGames.Horde/Projects/IProject.cs
--- a/Engine/Source/Programs/Shared/EpicGames.Horde/Projects/IProject.cs
+++ b/Engine/Source/Programs/Shared/EpicGames.Horde/Projects/IProject.cs
@@ -1,10 +1,6 @@
 // Copyright Epic Games, Inc. All Rights Reserved.
 
-using System;
 using System.Collections.Generic;
-using System.Linq;
-using System.Text;
-using System.Threading.Tasks;
 using EpicGames.Horde.Streams;
 
 namespace EpicGames.Horde.Projects
diff --git a/Engine/Source/Programs/Shared/EpicGames.Horde/Projects/IProjectCollection.cs b/Engine/Source/Programs/Shared/EpicGames.Horde/Projects/IProjectCollection.cs
--- a/Engine/Source/Programs/Shared/EpicGames.Horde/Projects/IProjectCollection.cs
+++ b/Engine/Source/Programs/Shared/EpicGames.Horde/Projects/IProjectCollection.cs
@@ -1,9 +1,6 @@
 // Copyright Epic Games, Inc. All Rights Reserved.
 
-using System;
 using System.Collections.Generic;
-using System.Linq;
-using System.Text;
 using System.Threading;
 using System.Threading.Tasks;
 
diff --git a/Engine/Source/Programs/Shared/EpicGames.Horde/Secrets/ISecret.cs b/Engine/Source/Programs/Shared/EpicGames.Horde/Secrets/ISecret.cs
--- a/Engine/Source/Programs/Shared/EpicGames.Horde/Secrets/ISecret.cs
+++ b/Engine/Source/Programs/Shared/EpicGames.Horde/Secrets/ISecret.cs
@@ -1,7 +1,6 @@
 // Copyright Epic Games, Inc. All Rights Reserved.
 
 using System.Collections.Generic;
-using EpicGames.Horde.Secrets;
 
 namespace EpicGames.Horde.Secrets
 {
diff --git a/Engine/Source/Programs/Shared/EpicGames.Horde/Secrets/SecretHttpCollection.cs b/Engine/Source/Programs/Shared/EpicGames.Horde/Secrets/SecretHttpCollection.cs
--- a/Engine/Source/Programs/Shared/EpicGames.Horde/Secrets/SecretHttpCollection.cs
+++ b/Engine/Source/Programs/Shared/EpicGames.Horde/Secrets/SecretHttpCollection.cs
@@ -1,9 +1,6 @@
 ﻿// Copyright Epic Games, Inc. All Rights Reserved.
 
-using System;
 using System.Collections.Generic;
-using System.Linq;
-using System.Text;
 using System.Threading;
 using System.Threading.Tasks;
 
diff --git a/Engine/Source/Programs/Shared/EpicGames.Horde/Tools/ITool.cs b/Engine/Source/Programs/Shared/EpicGames.Horde/Tools/ITool.cs
new file mode 100644
--- /dev/null
+++ b/Engine/Source/Programs/Shared/EpicGames.Horde/Tools/ITool.cs
@@ -0,0 +1,240 @@
+// Copyright Epic Games, Inc. All Rights Reserved.
+
+using System;
+using System.Collections.Generic;
+using System.IO;
+using System.Security.Claims;
+using System.Threading;
+using System.Threading.Tasks;
+using EpicGames.Horde.Acls;
+using EpicGames.Horde.Storage;
+
+#pragma warning disable CA1716 // Rename virtual/interface member ITool.Public so that it no longer conflicts with the reserved language keyword 'Public'.
+
+namespace EpicGames.Horde.Tools
+{
+	/// <summary>
+	/// Describes a standalone, external tool hosted and deployed by Horde. Provides basic functionality for performing
+	/// gradual roll-out, versioning, etc...
+	/// </summary>
+	public interface ITool
+	{
+		/// <summary>
+		/// Identifier for the tool
+		/// </summary>
+		ToolId Id { get; }
+
+		/// <summary>
+		/// Name of the tool
+		/// </summary>
+		string Name { get; }
+
+		/// <summary>
+		/// Long-form description of the tool
+		/// </summary>
+		string Description { get; }
+
+		/// <summary>
+		/// Category for the tool on the dashboard
+		/// </summary>
+		string? Category { get; }
+
+		/// <summary>
+		/// Grouping key for merging tool versions together on the dashboard
+		/// </summary>
+		string? Group { get; }
+
+		/// <summary>
+		/// Supported platforms, as a NET runtime identifiers
+		/// </summary>
+		IReadOnlyList<string>? Platforms { get; }
+
+		/// <summary>
+		/// Whether the tool is available to authenticated users
+		/// </summary>
+		bool Public { get; }
+
+		/// <summary>
+		/// Whether the tool is bundled with the server
+		/// </summary>
+		bool Bundled { get; }
+
+		/// <summary>
+		/// Whether to show the tool for download in UGS
+		/// </summary>
+		bool ShowInUgs { get; }
+
+		/// <summary>
+		/// Whether to show the tool for download in the dashboard
+		/// </summary>
+		bool ShowInDashboard { get; }
+
+		/// <summary>
+		/// Current deployments of this tool, sorted by time.
+		/// </summary>
+		IReadOnlyList<IToolDeployment> Deployments { get; }
+
+		/// <summary>
+		/// Authorize a user to perform a particular action
+		/// </summary>
+		/// <param name="action">Action the user is trying to perform</param>
+		/// <param name="principal">Identity of the user trying to perform the action</param>
+		bool Authorize(AclAction action, ClaimsPrincipal principal);
+
+		/// <summary>
+		/// Adds a new deployment to the given tool. The new deployment will replace the current active deployment.
+		/// </summary>
+		/// <param name="options">Options for the new deployment</param>
+		/// <param name="stream">Stream containing the tool data</param>
+		/// <param name="cancellationToken">Cancellation token for the operation</param>
+		/// <returns>Updated tool document, or null if it does not exist</returns>
+		Task<ITool?> CreateDeploymentAsync(ToolDeploymentConfig options, Stream stream, CancellationToken cancellationToken = default);
+
+		/// <summary>
+		/// Adds a new deployment to the given tool. The new deployment will replace the current active deployment.
+		/// </summary>
+		/// <param name="options">Options for the new deployment</param>
+		/// <param name="target">Path to the root node containing the tool data</param>
+		/// <param name="cancellationToken">Cancellation token for the operation</param>
+		/// <returns>Updated tool document, or null if it does not exist</returns>
+		Task<ITool?> CreateDeploymentAsync(ToolDeploymentConfig options, HashedBlobRefValue target, CancellationToken cancellationToken = default);
+
+		/// <summary>
+		/// Gets the storage backend for a particular tool
+		/// </summary>
+		/// <returns>Instance of the backend client</returns>
+		IStorageBackend CreateStorageBackend();
+
+		/// <summary>
+		/// Gets the storage backend for a particular tool
+		/// </summary>
+		/// <returns>Instance of the storage client</returns>
+		IStorageClient CreateStorageClient();
+	}
+
+	/// <summary>
+	/// Deployment of a tool
+	/// </summary>
+	public interface IToolDeployment
+	{
+		/// <summary>
+		/// Identifier for this deployment. A new identifier will be assigned to each created instance, so an identifier corresponds to a unique deployment.
+		/// </summary>
+		ToolDeploymentId Id { get; }
+
+		/// <summary>
+		/// Descriptive version string for this tool revision
+		/// </summary>
+		string Version { get; }
+
+		/// <summary>
+		/// Current state of this deployment
+		/// </summary>
+		ToolDeploymentState State { get; }
+
+		/// <summary>
+		/// Current progress of the deployment
+		/// </summary>
+		double Progress { get; }
+
+		/// <summary>
+		/// Last time at which the progress started. Set to null if the deployment was paused.
+		/// </summary>
+		DateTime? StartedAt { get; }
+
+		/// <summary>
+		/// Length of time over which to make the deployment
+		/// </summary>
+		TimeSpan Duration { get; }
+
+		/// <summary>
+		/// Namespace containing the tool
+		/// </summary>
+		NamespaceId NamespaceId { get; }
+
+		/// <summary>
+		/// Reference to this tool in Horde Storage.
+		/// </summary>
+		RefName RefName { get; }
+
+		/// <summary>
+		/// Updates the state of the current deployment
+		/// </summary>
+		/// <param name="state">New state of the deployment</param>
+		/// <param name="cancellationToken">Cancellation token for the operation</param>
+		/// <returns></returns>
+		Task<IToolDeployment?> UpdateAsync(ToolDeploymentState state, CancellationToken cancellationToken = default);
+
+		/// <summary>
+		/// Opens a stream to the data for a particular deployment
+		/// </summary>
+		/// <param name="cancellationToken">Cancellation token for the operation</param>
+		/// <returns>Stream for the data</returns>
+		Task<Stream> OpenZipStreamAsync(CancellationToken cancellationToken = default);
+	}
+
+	/// <summary>
+	/// Options for a new deployment
+	/// </summary>
+	public class ToolDeploymentConfig
+	{
+		/// <inheritdoc cref="IToolDeployment.Version"/>
+		public string Version { get; set; } = "Unknown";
+
+		/// <inheritdoc cref="IToolDeployment.Duration"/>
+		public TimeSpan Duration { get; set; }
+
+		/// <summary>
+		/// Whether to create the deployment in a paused state
+		/// </summary>
+		public bool CreatePaused { get; set; }
+	}
+
+	/// <summary>
+	/// Extension methods for tools
+	/// </summary>
+	public static class ToolExtensions
+	{
+		/// <summary>
+		/// Gets the current deployment
+		/// </summary>
+		/// <param name="tool">Tool to query</param>
+		/// <param name="phase">Adoption phase for the caller. 0 which the </param>
+		/// <param name="utcNow">Current time</param>
+		/// <returns></returns>
+		public static IToolDeployment? GetCurrentDeployment(this ITool tool, double phase, DateTime utcNow)
+		{
+			int idx = tool.Deployments.Count - 1;
+			for (; idx >= 0; idx--)
+			{
+				if (phase <= tool.Deployments[idx].GetProgressValue(utcNow) || idx == 0)
+				{
+					return tool.Deployments[idx];
+				}
+			}
+			return null;
+		}
+
+		/// <summary>
+		/// Get the progress fraction for a particular deployment and time
+		/// </summary>
+		/// <param name="deployment"></param>
+		/// <param name="utcNow"></param>
+		/// <returns></returns>
+		public static double GetProgressValue(this IToolDeployment deployment, DateTime utcNow)
+		{
+			if (deployment.StartedAt == null)
+			{
+				return deployment.Progress;
+			}
+			else if (deployment.Duration > TimeSpan.Zero)
+			{
+				return Math.Clamp((utcNow - deployment.StartedAt.Value) / deployment.Duration, 0.0, 1.0);
+			}
+			else
+			{
+				return 1.0;
+			}
+		}
+	}
+}
diff --git a/Engine/Source/Programs/Shared/EpicGames.Horde/Tools/IToolCollection.cs b/Engine/Source/Programs/Shared/EpicGames.Horde/Tools/IToolCollection.cs
new file mode 100644
--- /dev/null
+++ b/Engine/Source/Programs/Shared/EpicGames.Horde/Tools/IToolCollection.cs
@@ -0,0 +1,29 @@
+// Copyright Epic Games, Inc. All Rights Reserved.
+
+using System.Collections.Generic;
+using System.Threading;
+using System.Threading.Tasks;
+
+namespace EpicGames.Horde.Tools
+{
+	/// <summary>
+	/// Collection of tools
+	/// </summary>
+	public interface IToolCollection
+	{
+		/// <summary>
+		/// Gets a tool with the given identifier
+		/// </summary>
+		/// <param name="id">The tool identifier</param>
+		/// <param name="cancellationToken">Cancellation token for the operation</param>
+		/// <returns>The requested tool, or null if it does not exist</returns>
+		Task<ITool?> GetAsync(ToolId id, CancellationToken cancellationToken = default);
+
+		/// <summary>
+		/// Gets all the available tools
+		/// </summary>
+		/// <param name="cancellationToken"></param>
+		/// <returns>List of the available tools</returns>
+		Task<IReadOnlyList<ITool>> GetAllAsync(CancellationToken cancellationToken = default);
+	}
+}
diff --git a/Engine/Source/Programs/Shared/EpicGames.Horde/Tools/ToolHttpCollection.cs b/Engine/Source/Programs/Shared/EpicGames.Horde/Tools/ToolHttpCollection.cs
new file mode 100644
--- /dev/null
+++ b/Engine/Source/Programs/Shared/EpicGames.Horde/Tools/ToolHttpCollection.cs
@@ -0,0 +1,130 @@
+﻿// Copyright Epic Games, Inc. All Rights Reserved.
+
+using System;
+using System.Collections.Generic;
+using System.Diagnostics;
+using System.IO;
+using System.Security.Claims;
+using System.Threading;
+using System.Threading.Tasks;
+using EpicGames.Horde.Acls;
+using EpicGames.Horde.Storage;
+
+namespace EpicGames.Horde.Tools
+{
+	class ToolHttpCollection : IToolCollection
+	{
+		[DebuggerDisplay("{Id}")]
+		class Tool : ITool
+		{
+			readonly GetToolResponse _response;
+			readonly List<ToolDeployment> _deployments;
+
+			public ToolId Id => _response.Id;
+
+			string ITool.Name => _response.Name;
+			string ITool.Description => _response.Description;
+			string? ITool.Category => _response.Category;
+			string? ITool.Group => _response.Group;
+			IReadOnlyList<string>? ITool.Platforms => _response.Platforms;
+			bool ITool.Public => _response.Public;
+			bool ITool.Bundled => _response.Bundled;
+			bool ITool.ShowInUgs => _response.ShowInUgs;
+			bool ITool.ShowInDashboard => _response.ShowInDashboard;
+			IReadOnlyList<IToolDeployment> ITool.Deployments => _deployments;
+
+			public Tool(ToolHttpCollection collection, GetToolResponse response)
+			{
+				_response = response;
+				_deployments = response.Deployments.ConvertAll(x => new ToolDeployment(collection, _response.Id, x));
+			}
+
+			public bool Authorize(AclAction action, ClaimsPrincipal principal)
+				=> throw new NotImplementedException();
+
+			public Task<ITool?> CreateDeploymentAsync(ToolDeploymentConfig options, Stream stream, CancellationToken cancellationToken = default)
+				=> throw new NotImplementedException();
+
+			public Task<ITool?> CreateDeploymentAsync(ToolDeploymentConfig options, HashedBlobRefValue target, CancellationToken cancellationToken = default)
+				=> throw new NotImplementedException();
+
+			public IStorageBackend CreateStorageBackend()
+				=> throw new NotImplementedException();
+
+			public IStorageClient CreateStorageClient()
+				=> throw new NotImplementedException();
+		}
+
+		[DebuggerDisplay("{Id}")]
+		class ToolDeployment : IToolDeployment
+		{
+			readonly ToolHttpCollection _collection;
+			readonly ToolId _toolId;
+			readonly GetToolDeploymentResponse _response;
+
+			public ToolDeploymentId Id => _response.Id;
+
+			string IToolDeployment.Version => _response.Version;
+			ToolDeploymentState IToolDeployment.State => _response.State;
+			double IToolDeployment.Progress => _response.Progress;
+			DateTime? IToolDeployment.StartedAt => _response.StartedAt;
+			TimeSpan IToolDeployment.Duration => _response.Duration;
+			NamespaceId IToolDeployment.NamespaceId => throw new NotImplementedException();
+			RefName IToolDeployment.RefName => _response.RefName;
+
+			public ToolDeployment(ToolHttpCollection collection, ToolId toolId, GetToolDeploymentResponse response)
+			{
+				_collection = collection;
+				_toolId = toolId;
+				_response = response;
+			}
+
+			Task<Stream> IToolDeployment.OpenZipStreamAsync(CancellationToken cancellationToken)
+				=> _collection.OpenZipStreamAsync(_toolId, _response.Id, cancellationToken);
+
+			Task<IToolDeployment?> IToolDeployment.UpdateAsync(ToolDeploymentState state, CancellationToken cancellationToken)
+				=> throw new NotImplementedException();
+		}
+
+		readonly IHordeClient _hordeClient;
+
+		/// <summary>
+		/// Constructor
+		/// </summary>
+		public ToolHttpCollection(IHordeClient hordeClient)
+			=> _hordeClient = hordeClient;
+
+		/// <inheritdoc/>
+		public async Task<ITool?> GetAsync(ToolId id, CancellationToken cancellationToken = default)
+		{
+			HordeHttpClient hordeHttpClient = _hordeClient.CreateHttpClient();
+			GetToolResponse response = await hordeHttpClient.GetToolAsync(id, cancellationToken);
+			return new Tool(this, response);
+		}
+
+		/// <inheritdoc/>
+		public async Task<IReadOnlyList<ITool>> GetAllAsync(CancellationToken cancellationToken = default)
+		{
+			HordeHttpClient hordeHttpClient = _hordeClient.CreateHttpClient();
+			GetToolsSummaryResponse responses = await hordeHttpClient.GetToolsAsync(cancellationToken);
+
+			List<ITool> tools = new List<ITool>();
+			foreach (GetToolSummaryResponse response in responses.Tools)
+			{
+				ITool? tool = await GetAsync(response.Id, cancellationToken);
+				if (tool != null)
+				{
+					tools.Add(tool);
+				}
+			}
+
+			return tools;
+		}
+
+		async Task<Stream> OpenZipStreamAsync(ToolId id, ToolDeploymentId deploymentId, CancellationToken cancellationToken = default)
+		{
+			HordeHttpClient hordeHttpClient = _hordeClient.CreateHttpClient();
+			return await hordeHttpClient.GetToolDeploymentZipAsync(id, deploymentId, cancellationToken);
+		}
+	}
+}
