Horde: Add a symbol store plugin. Symbol stores use aliases in the storage system to map symbol store paths onto content streams from existing artifacts, allowing reuse of data already available in artifacts.

Symbols can be tagged with the appropriate metadata by setting the Symbols=true attribute on the CreateArtifact task. Referencing the namespace that the symbols will be uploaded to from the symbol store config will allow accessing them through the api/v1/symbols route.

Hashing for symbols is compatible with symstore.exe, but is handled by a custom implementation in SymStore.cs.

diff --git a/Engine/Source/Programs/AutomationTool/BuildGraph/BgScriptReader.cs b/Engine/Source/Programs/AutomationTool/BuildGraph/BgScriptReader.cs
--- a/Engine/Source/Programs/AutomationTool/BuildGraph/BgScriptReader.cs
+++ b/Engine/Source/Programs/AutomationTool/BuildGraph/BgScriptReader.cs
@@ -1754,28 +1754,16 @@ namespace AutomationTool
 
 		void AddArtifactFromTask(BgScriptElement element, Dictionary<string, string> arguments)
 		{
-			CreateArtifactTaskParameters parameters = new CreateArtifactTaskParameters();
-			foreach (PropertyInfo propertyInfo in typeof(CreateArtifactTaskParameters).GetProperties())
-			{
-				TaskParameterAttribute? attribute = propertyInfo.GetCustomAttribute<TaskParameterAttribute>();
-				if (attribute != null)
-				{
-					if (arguments.TryGetValue(propertyInfo.Name, out string? value))
-					{
-						propertyInfo.SetValue(parameters, value);
-					}
-					else if (!attribute.Optional)
-					{
-						LogError(element, $"Missing '{propertyInfo.Name}' property for CreateArtifact task.");
-						return;
-					}
-				}
-			}
+			arguments.TryGetValue(nameof(CreateArtifactTaskParameters.Name), out string? name);
+			arguments.TryGetValue(nameof(CreateArtifactTaskParameters.Type), out string? type);
+			arguments.TryGetValue(nameof(CreateArtifactTaskParameters.Description), out string? description);
+			arguments.TryGetValue(nameof(CreateArtifactTaskParameters.Keys), out string? keys);
+			arguments.TryGetValue(nameof(CreateArtifactTaskParameters.Metadata), out string? metadata);
 
-			string[] keys = (parameters.Keys ?? String.Empty).Split(';', StringSplitOptions.RemoveEmptyEntries | StringSplitOptions.TrimEntries);
-			string[] metadata = (parameters.Metadata ?? String.Empty).Split(';', StringSplitOptions.RemoveEmptyEntries | StringSplitOptions.TrimEntries);
+			string[] keysArray = (keys ?? String.Empty).Split(';', StringSplitOptions.RemoveEmptyEntries | StringSplitOptions.TrimEntries);
+			string[] metadataArray = (metadata ?? String.Empty).Split(';', StringSplitOptions.RemoveEmptyEntries | StringSplitOptions.TrimEntries);
 
-			BgArtifactDef artifact = new BgArtifactDef(parameters.Name, parameters.Type, parameters.Description, null, _enclosingNode!.Name, null, keys, metadata);
+			BgArtifactDef artifact = new BgArtifactDef(name ?? String.Empty, type, description, null, _enclosingNode!.Name, null, keysArray, metadataArray);
 			_graph.Artifacts.Add(artifact);
 		}
 
diff --git a/Engine/Source/Programs/AutomationTool/BuildGraph/Tasks/CreateArtifactTask.cs b/Engine/Source/Programs/AutomationTool/BuildGraph/Tasks/CreateArtifactTask.cs
--- a/Engine/Source/Programs/AutomationTool/BuildGraph/Tasks/CreateArtifactTask.cs
+++ b/Engine/Source/Programs/AutomationTool/BuildGraph/Tasks/CreateArtifactTask.cs
@@ -4,6 +4,7 @@ using System;
 using System.Collections.Generic;
 using System.Diagnostics;
 using System.Linq;
+using System.Threading;
 using System.Threading.Tasks;
 using System.Xml;
 using EpicGames.Core;
@@ -13,6 +14,7 @@ using EpicGames.Horde.Commits;
 using EpicGames.Horde.Storage;
 using EpicGames.Horde.Storage.Nodes;
 using EpicGames.Horde.Streams;
+using EpicGames.Horde.Symbols;
 using Microsoft.Extensions.DependencyInjection;
 using Microsoft.Extensions.Logging;
 
@@ -78,6 +80,12 @@ namespace AutomationTool.Tasks
 		/// </summary>
 		[TaskParameter(Optional = true)]
 		public string? Metadata { get; set; }
+
+		/// <summary>
+		/// Whether to add aliases for symbol files
+		/// </summary>
+		[TaskParameter(Optional = true)]
+		public bool Symbols { get; set; }
 	}
 
 	/// <summary>
@@ -194,6 +202,69 @@ namespace AutomationTool.Tasks
 			await storage.WriteRefAsync(artifact.RefName, rootRef);
 
 			Logger.LogInformation("Uploaded artifact {ArtifactId} in {Time:n1}s", artifact.Id, timer.Elapsed.TotalSeconds);
+
+			// Tag any uploaded symbols
+			if (_parameters.Symbols)
+			{
+				foreach (FileReference file in files)
+				{
+					string? hash = await SymStore.GetHashAsync(file, Logger, default);
+					if (hash != null)
+					{
+						string path = file.MakeRelativeTo(baseDir);
+
+						FileEntry? fileEntry = await FindFileAsync(rootRef, path);
+						if (fileEntry == null)
+						{
+							Logger.LogWarning("Unable to find file {Path} in uploaded data.", path);
+						}
+						else
+						{
+							string fileName = file.GetFileName().ToUpperInvariant();
+							string alias = $"sym:{fileName}/{hash}/{fileName}";
+							Logger.LogInformation("Adding symbol alias: {Alias}", alias);
+							await storage.AddAliasAsync(alias, fileEntry.Target.Handle);
+						}
+					}
+				}
+			}
+		}
+
+		static async Task<FileEntry?> FindFileAsync(IBlobRef<DirectoryNode> rootDir, string path, CancellationToken cancellationToken = default)
+		{
+			string[] fragments = path.Split('/', '\\');
+			if (fragments.Length == 0)
+			{
+				return null;
+			}
+
+			IBlobRef<DirectoryNode> directoryRef = rootDir;
+			for (int idx = 0;; idx++)
+			{
+				DirectoryNode directory = await directoryRef.ReadBlobAsync(cancellationToken);
+				if (idx + 1 < fragments.Length)
+				{
+					if (directory.TryGetDirectoryEntry(fragments[idx], out DirectoryEntry? directoryEntry))
+					{
+						directoryRef = directoryEntry.Handle;
+					}
+					else
+					{
+						return null;
+					}
+				}
+				else
+				{
+					if (directory.TryGetFileEntry(fragments[idx], out FileEntry? fileEntry))
+					{
+						return fileEntry;
+					}
+					else
+					{
+						return null;
+					}
+				}
+			}
 		}
 
 		/// <inheritdoc/>
diff --git a/Engine/Source/Programs/Horde/Plugins/Symbols/HordeServer.Symbols/HordeServer.Symbols.csproj b/Engine/Source/Programs/Horde/Plugins/Symbols/HordeServer.Symbols/HordeServer.Symbols.csproj
new file mode 100644
--- /dev/null
+++ b/Engine/Source/Programs/Horde/Plugins/Symbols/HordeServer.Symbols/HordeServer.Symbols.csproj
@@ -0,0 +1,36 @@
+<Project Sdk="Microsoft.NET.Sdk">
+
+  <PropertyGroup>
+    <TargetFramework>net8.0</TargetFramework>
+    <ImplicitUsings>enable</ImplicitUsings>
+    <Nullable>enable</Nullable>
+    <RootNamespace>HordeServer</RootNamespace>
+    <RunAnalyzersDuringBuild>False</RunAnalyzersDuringBuild>
+    <Configurations>Debug;Release;Analyze</Configurations>
+  </PropertyGroup>
+
+  <PropertyGroup Condition=" '$(Configuration)|$(Platform)' == 'Analyze|AnyCPU' ">
+    <RunAnalyzersDuringBuild>True</RunAnalyzersDuringBuild>
+    <TreatWarningsAsErrors>false</TreatWarningsAsErrors>
+    <EnforceCodeStyleInBuild>True</EnforceCodeStyleInBuild>
+    <GenerateDocumentationFile>True</GenerateDocumentationFile>
+    <DocumentationFile></DocumentationFile>
+  </PropertyGroup>
+
+  <ItemGroup>
+    <FrameworkReference Include="Microsoft.AspNetCore.App" />
+  </ItemGroup>
+
+  <ItemGroup>
+  </ItemGroup>
+
+  <ItemGroup>
+    <ProjectReference Include="..\..\..\..\Shared\EpicGames.Core\EpicGames.Core.csproj" />
+    <ProjectReference Include="..\..\..\..\Shared\EpicGames.Horde\EpicGames.Horde.csproj" />
+    <ProjectReference Include="..\..\..\..\Shared\EpicGames.IoHash\EpicGames.IoHash.csproj" />
+    <ProjectReference Include="..\..\..\..\Shared\EpicGames.Redis\EpicGames.Redis.csproj" />
+    <ProjectReference Include="..\..\Storage\HordeServer.Storage\HordeServer.Storage.csproj" />
+    <ProjectReference Include="..\..\..\HordeServer.Shared\HordeServer.Shared.csproj" />
+  </ItemGroup>
+
+</Project>
diff --git a/Engine/Source/Programs/Horde/Plugins/Symbols/HordeServer.Symbols/Properties/AssemblyInfo.cs b/Engine/Source/Programs/Horde/Plugins/Symbols/HordeServer.Symbols/Properties/AssemblyInfo.cs
new file mode 100644
--- /dev/null
+++ b/Engine/Source/Programs/Horde/Plugins/Symbols/HordeServer.Symbols/Properties/AssemblyInfo.cs
@@ -0,0 +1,7 @@
+// Copyright Epic Games, Inc. All Rights Reserved.
+
+using System.Resources;
+using System.Runtime.CompilerServices;
+
+[assembly: InternalsVisibleTo("HordeServer.Tools.Tests")]
+[assembly: NeutralResourcesLanguage("en-US")]
diff --git a/Engine/Source/Programs/Horde/Plugins/Symbols/HordeServer.Symbols/Symbols/SymbolStoreAclAction.cs b/Engine/Source/Programs/Horde/Plugins/Symbols/HordeServer.Symbols/Symbols/SymbolStoreAclAction.cs
new file mode 100644
--- /dev/null
+++ b/Engine/Source/Programs/Horde/Plugins/Symbols/HordeServer.Symbols/Symbols/SymbolStoreAclAction.cs
@@ -0,0 +1,17 @@
+// Copyright Epic Games, Inc. All Rights Reserved.
+
+using EpicGames.Horde.Acls;
+
+namespace HordeServer.Symbols
+{
+	/// <summary>
+	/// ACL actions which apply to symbols
+	/// </summary>
+	public static class SymbolStoreAclAction
+	{
+		/// <summary>
+		/// Ability to download symbols
+		/// </summary>
+		public static AclAction ReadSymbols { get; } = new AclAction("ReadSymbols");
+	}
+}
diff --git a/Engine/Source/Programs/Horde/Plugins/Symbols/HordeServer.Symbols/Symbols/SymbolStoreConfig.cs b/Engine/Source/Programs/Horde/Plugins/Symbols/HordeServer.Symbols/Symbols/SymbolStoreConfig.cs
new file mode 100644
--- /dev/null
+++ b/Engine/Source/Programs/Horde/Plugins/Symbols/HordeServer.Symbols/Symbols/SymbolStoreConfig.cs
@@ -0,0 +1,54 @@
+// Copyright Epic Games, Inc. All Rights Reserved.
+
+using System.Security.Claims;
+using EpicGames.Horde.Acls;
+using EpicGames.Horde.Storage;
+using HordeServer.Acls;
+
+namespace HordeServer.Symbols
+{
+	/// <summary>
+	/// Configuration for a symbol store
+	/// </summary>
+	public class SymbolStoreConfig
+	{
+		/// <summary>
+		/// Identifier for this store
+		/// </summary>
+		public SymbolStoreId Id { get; set; }
+
+		/// <summary>
+		/// Configuration for the symbol store backend
+		/// </summary>
+		public NamespaceId NamespaceId { get; set; }
+
+		/// <summary>
+		/// Whether to make this store available without auth
+		/// </summary>
+		public bool Public { get; set; }
+
+		/// <summary>
+		/// Access to the symbol store
+		/// </summary>
+		public AclConfig Acl { get; set; } = new AclConfig();
+
+		/// <summary>
+		/// Authorize a user to perform the given action
+		/// </summary>
+		public bool Authorize(AclAction action, ClaimsPrincipal user)
+		{
+			if (Public)
+			{
+				return true;
+			}
+			else if (user.Identity == null || !user.Identity.IsAuthenticated)
+			{
+				return false;
+			}
+			else
+			{
+				return Acl.Authorize(action, user);
+			}
+		}
+	}
+}
diff --git a/Engine/Source/Programs/Horde/Plugins/Symbols/HordeServer.Symbols/Symbols/SymbolStoreId.cs b/Engine/Source/Programs/Horde/Plugins/Symbols/HordeServer.Symbols/Symbols/SymbolStoreId.cs
new file mode 100644
--- /dev/null
+++ b/Engine/Source/Programs/Horde/Plugins/Symbols/HordeServer.Symbols/Symbols/SymbolStoreId.cs
@@ -0,0 +1,45 @@
+// Copyright Epic Games, Inc. All Rights Reserved.
+
+using System.ComponentModel;
+using EpicGames.Core;
+using EpicGames.Horde;
+using EpicGames.Serialization;
+
+namespace HordeServer.Symbols
+{
+	/// <summary>
+	/// Identifier for a symbol store
+	/// </summary>
+	/// <param name="Id">Id to construct from</param>
+	[JsonSchemaString]
+	[TypeConverter(typeof(StringIdTypeConverter<SymbolStoreId, SymbolStoreIdConverter>))]
+	[StringIdConverter(typeof(SymbolStoreIdConverter))]
+	[CbConverter(typeof(StringIdCbConverter<SymbolStoreId, SymbolStoreIdConverter>))]
+	public record struct SymbolStoreId(StringId Id)
+	{
+		/// <summary>
+		/// Constructor
+		/// </summary>
+		public SymbolStoreId(string id) : this(new StringId(id))
+		{
+		}
+
+		/// <inheritdoc cref="StringId.IsEmpty"/>
+		public bool IsEmpty => Id.IsEmpty;
+
+		/// <inheritdoc/>
+		public override string ToString() => Id.ToString();
+	}
+
+	/// <summary>
+	/// Converter to and from <see cref="StringId"/> instances.
+	/// </summary>
+	class SymbolStoreIdConverter : StringIdConverter<SymbolStoreId>
+	{
+		/// <inheritdoc/>
+		public override SymbolStoreId FromStringId(StringId id) => new SymbolStoreId(id);
+
+		/// <inheritdoc/>
+		public override StringId ToStringId(SymbolStoreId value) => value.Id;
+	}
+}
diff --git a/Engine/Source/Programs/Horde/Plugins/Symbols/HordeServer.Symbols/Symbols/SymbolsController.cs b/Engine/Source/Programs/Horde/Plugins/Symbols/HordeServer.Symbols/Symbols/SymbolsController.cs
new file mode 100644
--- /dev/null
+++ b/Engine/Source/Programs/Horde/Plugins/Symbols/HordeServer.Symbols/Symbols/SymbolsController.cs
@@ -0,0 +1,71 @@
+// Copyright Epic Games, Inc. All Rights Reserved.
+
+using System.Net;
+using EpicGames.Horde.Storage;
+using EpicGames.Horde.Storage.Nodes;
+using HordeServer.Storage;
+using HordeServer.Symbols;
+using HordeServer.Utilities;
+using Microsoft.AspNetCore.Mvc;
+using Microsoft.Extensions.Options;
+
+namespace HordeServer.Tools
+{
+	/// <summary>
+	/// Controller for the /api/v1/tools endpoint
+	/// </summary>
+	[ApiController]
+	[TryAuthorize]
+	public class SymbolsController : HordeControllerBase
+	{
+		readonly StorageService _storageService;
+		readonly IOptionsSnapshot<SymbolsConfig> _symbolsConfig;
+
+		/// <summary>
+		/// Constructor
+		/// </summary>
+		public SymbolsController(StorageService storageService, IOptionsSnapshot<SymbolsConfig> symbolsConfig)
+		{
+			_storageService = storageService;
+			_symbolsConfig = symbolsConfig;
+		}
+
+		/// <summary>
+		/// Reads a file from the symbol store. 
+		/// </summary>
+		/// <param name="storeId">Id of the symbol store</param>
+		/// <param name="path">Path for the file to retrieve</param>
+		/// <param name="cancellationToken">Cancellation token for the operation</param>
+		[HttpGet]
+		[Route("/api/v1/storage/{storeId}/{*path}")]
+		public async Task<ActionResult> ReadFileAsync(SymbolStoreId storeId, string path, CancellationToken cancellationToken)
+		{
+			SymbolStoreConfig? symbolStoreConfig;
+			if (!_symbolsConfig.Value.TryGetStore(storeId, out symbolStoreConfig))
+			{
+				return NotFound(storeId);
+			}
+			if (!symbolStoreConfig.Authorize(SymbolStoreAclAction.ReadSymbols, User))
+			{
+				return Forbid(SymbolStoreAclAction.ReadSymbols);//, storeId);
+			}
+
+			IStorageClient storageClient = _storageService.CreateClient(symbolStoreConfig.NamespaceId);
+			BlobAlias? alias = await storageClient.FindAliasAsync($"sym:{path}", cancellationToken);
+			if (alias == null)
+			{
+				return NotFound();
+			}
+
+			IBlobRef<ChunkedDataNode> blobRef = alias.Target.ForType<ChunkedDataNode>();
+
+			Response.StatusCode = (int)HttpStatusCode.OK;
+
+			await Response.StartAsync(cancellationToken);
+			await ChunkedDataNode.CopyToStreamAsync(blobRef, Response.Body, cancellationToken);
+			await Response.CompleteAsync();
+
+			return NoContent();
+		}
+	}
+}
diff --git a/Engine/Source/Programs/Horde/Plugins/Symbols/HordeServer.Symbols/SymbolsConfig.cs b/Engine/Source/Programs/Horde/Plugins/Symbols/HordeServer.Symbols/SymbolsConfig.cs
new file mode 100644
--- /dev/null
+++ b/Engine/Source/Programs/Horde/Plugins/Symbols/HordeServer.Symbols/SymbolsConfig.cs
@@ -0,0 +1,44 @@
+// Copyright Epic Games, Inc. All Rights Reserved.
+
+using System.Diagnostics.CodeAnalysis;
+using HordeServer.Acls;
+using HordeServer.Plugins;
+
+#pragma warning disable CA2227 // Change x to be read-only by removing the property setter
+
+namespace HordeServer.Symbols
+{
+	/// <summary>
+	/// Configuration for the tools system
+	/// </summary>
+	public class SymbolsConfig : IPluginConfig
+	{
+		/// <summary>
+		/// List of symbol stores
+		/// </summary>
+		public List<SymbolStoreConfig> Stores { get; set; } = new List<SymbolStoreConfig>();
+
+		readonly Dictionary<SymbolStoreId, SymbolStoreConfig> _storeLookup = new Dictionary<SymbolStoreId, SymbolStoreConfig>();
+
+		/// <inheritdoc/>
+		public void PostLoad(PluginConfigOptions configOptions)
+		{
+			foreach (SymbolStoreConfig store in Stores)
+			{
+				store.Acl.PostLoad(configOptions.ParentAcl, $"symbol-store:{store.Id}", AclConfig.GetActions([typeof(SymbolStoreAclAction)]));
+			}
+
+			_storeLookup.Clear();
+			foreach (SymbolStoreConfig store in Stores)
+			{
+				_storeLookup.Add(store.Id, store);
+			}
+		}
+
+		/// <summary>
+		/// Finds a store by id
+		/// </summary>
+		public bool TryGetStore(SymbolStoreId storeId, [NotNullWhen(true)] out SymbolStoreConfig? store)
+			=> _storeLookup.TryGetValue(storeId, out store);
+	}
+}
diff --git a/Engine/Source/Programs/Horde/Plugins/Symbols/HordeServer.Symbols/SymbolsPlugin.cs b/Engine/Source/Programs/Horde/Plugins/Symbols/HordeServer.Symbols/SymbolsPlugin.cs
new file mode 100644
--- /dev/null
+++ b/Engine/Source/Programs/Horde/Plugins/Symbols/HordeServer.Symbols/SymbolsPlugin.cs
@@ -0,0 +1,24 @@
+// Copyright Epic Games, Inc. All Rights Reserved.
+
+using HordeServer.Plugins;
+using HordeServer.Symbols;
+using Microsoft.AspNetCore.Builder;
+using Microsoft.Extensions.DependencyInjection;
+
+namespace HordeServer
+{
+	/// <summary>
+	/// Entry point for the symbols plugin
+	/// </summary>
+	[Plugin("Symbols", GlobalConfigType = typeof(SymbolsConfig))]
+	public class SymbolsPlugin : IPluginStartup
+	{
+		/// <inheritdoc/>
+		public void Configure(IApplicationBuilder app)
+		{ }
+
+		/// <inheritdoc/>
+		public void ConfigureServices(IServiceCollection services)
+		{ }
+	}
+}
diff --git a/Engine/Source/Programs/Shared/EpicGames.Core/MemoryMappedView.cs b/Engine/Source/Programs/Shared/EpicGames.Core/MemoryMappedView.cs
--- a/Engine/Source/Programs/Shared/EpicGames.Core/MemoryMappedView.cs
+++ b/Engine/Source/Programs/Shared/EpicGames.Core/MemoryMappedView.cs
@@ -2,6 +2,7 @@
 
 using System;
 using System.Buffers;
+using System.ComponentModel;
 using System.Diagnostics;
 using System.IO.MemoryMappedFiles;
 using System.Runtime.InteropServices;
@@ -38,14 +39,17 @@ namespace EpicGames.Core
 			protected override void Dispose(bool disposing) { }
 		}
 
-		const uint FILE_MAP_ALL_ACCESS = 0x000F001F;
+		const uint FILE_MAP_COPY = 0x0001;
+		const uint FILE_MAP_WRITE = 0x0002;
+		const uint FILE_MAP_READ = 0x0004;
+		const uint FILE_MAP_EXECUTE = 0x0020;
 
 		[DllImport("kernel32.dll")]
 		static extern unsafe void* MapViewOfFile(SafeMemoryMappedFileHandle handle, uint desiredAccess, uint fileOffsetHigh, uint fileOffsetLow, long numberOfBytes);
 
 		[DllImport("kernel32.dll")]
 		static extern unsafe bool UnmapViewOfFile(void* ptr);
-		
+
 		byte* _data;
 		Action? _disposeMethod;
 
@@ -55,12 +59,28 @@ namespace EpicGames.Core
 		/// <param name="memoryMappedFile">Handle of the file to map into memory</param>
 		/// <param name="offset">Offset within the file to map</param>
 		/// <param name="length">Length of the region to map</param>
-		public MemoryMappedView(MemoryMappedFile memoryMappedFile, long offset, long length)
+		/// <param name="access">How to open the file as read only</param>
+		public MemoryMappedView(MemoryMappedFile memoryMappedFile, long offset, long length, MemoryMappedFileAccess access = MemoryMappedFileAccess.ReadWrite)
 		{
 			if (RuntimeInformation.IsOSPlatform(OSPlatform.Windows))
 			{
-				_data = (byte*)MapViewOfFile(memoryMappedFile.SafeMemoryMappedFileHandle, FILE_MAP_ALL_ACCESS, (uint)(offset >> 32), (uint)offset, length);
-				Trace.Assert(_data != null);
+				uint accessFlags = access switch
+				{
+					MemoryMappedFileAccess.Read => FILE_MAP_READ,
+					MemoryMappedFileAccess.ReadWrite => FILE_MAP_READ | FILE_MAP_WRITE,
+					MemoryMappedFileAccess.Write => FILE_MAP_WRITE,
+					MemoryMappedFileAccess.CopyOnWrite => FILE_MAP_READ | FILE_MAP_WRITE | FILE_MAP_COPY,
+					MemoryMappedFileAccess.ReadExecute => FILE_MAP_READ | FILE_MAP_EXECUTE,
+					MemoryMappedFileAccess.ReadWriteExecute => FILE_MAP_READ | FILE_MAP_WRITE | FILE_MAP_EXECUTE,
+					_ => FILE_MAP_READ | FILE_MAP_WRITE
+				};
+
+				_data = (byte*)MapViewOfFile(memoryMappedFile.SafeMemoryMappedFileHandle, accessFlags, (uint)(offset >> 32), (uint)offset, length);
+				if (_data == null)
+				{
+					int status = Marshal.GetLastWin32Error();
+					throw new Win32Exception(status, $"Unable to create memory mapped file (GetLastError() == {status})");
+				}
 
 				_disposeMethod = () =>
 				{
@@ -69,7 +89,7 @@ namespace EpicGames.Core
 			}
 			else
 			{
-				MemoryMappedViewAccessor memoryMappedViewAccessor = memoryMappedFile.CreateViewAccessor(offset, length);
+				MemoryMappedViewAccessor memoryMappedViewAccessor = memoryMappedFile.CreateViewAccessor(offset, length, access);
 				memoryMappedViewAccessor.SafeMemoryMappedViewHandle.AcquirePointer(ref _data);
 				Trace.Assert(_data != null);
 
diff --git a/Engine/Source/Programs/Shared/EpicGames.Horde/Symbols/SymStore.cs b/Engine/Source/Programs/Shared/EpicGames.Horde/Symbols/SymStore.cs
new file mode 100644
--- /dev/null
+++ b/Engine/Source/Programs/Shared/EpicGames.Horde/Symbols/SymStore.cs
@@ -0,0 +1,273 @@
+// Copyright Epic Games, Inc. All Rights Reserved.
+
+using System;
+using System.Buffers.Binary;
+using System.Diagnostics;
+using System.IO;
+using System.IO.MemoryMappedFiles;
+using System.Runtime.InteropServices;
+using System.Text;
+using System.Threading;
+using System.Threading.Tasks;
+using EpicGames.Core;
+using Microsoft.Extensions.Logging;
+
+namespace EpicGames.Horde.Symbols
+{
+	/// <summary>
+	/// Provides functionality for hashing files to add to a Microsoft Symbol store.
+	/// 
+	/// * For PE-format files (EXE, DLL), this consists of the 
+	/// * For PDB files, this consists of the PDB GUID followed by its age (the number of times it has been written).
+	/// </summary>
+	public static class SymStore
+	{
+		/// <summary>
+		/// Gets the hash of a file.
+		/// </summary>
+		/// <param name="file">File to hash</param>
+		/// <param name="logger">Logger for diagnostic messages</param>
+		/// <param name="cancellationToken">Cancellation token for the operation</param>
+		/// <returns>Hash of the file, or null if it cannot be parsed.</returns>
+		public static Task<string?> GetHashAsync(FileReference file, ILogger logger, CancellationToken cancellationToken = default)
+		{
+			try
+			{
+				if (file.HasExtension(".exe") || file.HasExtension(".dll"))
+				{
+					return GetExeHashAsync(file, cancellationToken);
+				}
+				else if (file.HasExtension(".pdb"))
+				{
+					return Task.FromResult<string?>(GetPdbHash(file));
+				}
+				else
+				{
+					return Task.FromResult<string?>(null);
+				}
+			}
+			catch (Exception ex)
+			{
+				logger.LogError(ex, "Unable to get hash for {File}: {Message}", file, ex.Message);
+				return Task.FromResult<string?>(null);
+			}
+		}
+
+		#region PE files
+
+		/// <summary>
+		/// Gets the hash for a Windows portable executable file, consisting of concatenated hex values for the timestamp and image size fields in the header.
+		/// Structures referenced below (IMAGE_DOS_HEADER, IMAGE_NT_HEADERS) are defined in Windows headers, but are parsed via offsets into the executable
+		/// to reduce messy marshalling in C#.
+		/// </summary>
+		static async Task<string?> GetExeHashAsync(FileReference file, CancellationToken cancellationToken)
+		{
+			using FileStream stream = FileReference.Open(file, FileMode.Open, FileAccess.Read, FileShare.Read | FileShare.Delete);
+
+			// Read the IMAGE_DOS_HEADER structure.
+			const int SizeOf_IMAGE_DOS_HEADER = 64;
+			const int OffsetOf_IMAGE_DOS_HEADER_Magic = 0;
+			const int OffsetOf_IMAGE_DOS_HEADER_LfaNew = 0x3c;
+
+			byte[] dosHeader = new byte[SizeOf_IMAGE_DOS_HEADER];
+			if (await stream.ReadAsync(dosHeader, 0, dosHeader.Length, cancellationToken) != dosHeader.Length)
+			{
+				return null;
+			}
+
+			ushort magic = BinaryPrimitives.ReadUInt16LittleEndian(dosHeader.AsSpan(OffsetOf_IMAGE_DOS_HEADER_Magic));
+			if (magic != (ushort)('M' | ('Z' << 8)))
+			{
+				return null;
+			}
+
+			uint lfaNew = BinaryPrimitives.ReadUInt32LittleEndian(dosHeader.AsSpan(OffsetOf_IMAGE_DOS_HEADER_LfaNew));
+			stream.Seek(lfaNew, SeekOrigin.Begin);
+
+			// Read the first part of the IMAGE_NT_HEADERS structure (just including the signature and file header)
+			const int SizeOf_IMAGE_NT_HEADERS = 84;
+			const int OffsetOf_IMAGE_NT_HEADERS_Signature = 0;
+			const int OffsetOf_IMAGE_NT_HEADERS_TimeDateStamp = 8;
+			const int OffsetOf_IMAGE_NT_HEADERS_SizeOfImage = 80;
+
+			byte[] ntHeader = new byte[SizeOf_IMAGE_NT_HEADERS];
+			if (await stream.ReadAsync(ntHeader, 0, ntHeader.Length, cancellationToken) != ntHeader.Length)
+			{
+				return null;
+			}
+
+			uint signature = BinaryPrimitives.ReadUInt16LittleEndian(ntHeader.AsSpan(OffsetOf_IMAGE_NT_HEADERS_Signature));
+			if (signature != (uint)('P' | ('E' << 8)))
+			{
+				return null;
+			}
+
+			uint timeDateStamp = BinaryPrimitives.ReadUInt32LittleEndian(ntHeader.AsSpan(OffsetOf_IMAGE_NT_HEADERS_TimeDateStamp));
+			uint sizeOfImage = BinaryPrimitives.ReadUInt32LittleEndian(ntHeader.AsSpan(OffsetOf_IMAGE_NT_HEADERS_SizeOfImage));
+
+			// Return the hash value
+			return $"{timeDateStamp:X}{sizeOfImage:X}";
+		}
+
+		#endregion
+		#region PDB files
+
+		/// <summary>
+		/// Parse the hash value from a PDB file, consisting of concatenated hex values for the PDB GUID and age (the number of times it has been written).
+		/// 
+		/// While the fields to obtain are at fairly straightforward offsets in particular data structures, PDBs are internally structured as an MSF file
+		/// consisting of multiple data streams with non-contiguous pages. There are a few references for parsing this format:
+		/// 
+		/// * LLVM documentation: https://llvm.org/docs/PDB/MsfFile.html
+		/// * Microsoft's PDB source code: [REDACTED_PROFILE_URL]/microsoft-pdb (particularly PDB/msf/msf.cpp)
+		/// 
+		/// This function only handles the BigMSF format.
+		/// </summary>
+		static string? GetPdbHash(FileReference file)
+		{
+			using FileStream stream = FileReference.Open(file, FileMode.Open, FileAccess.Read, FileShare.Read | FileShare.Delete);
+			using MemoryMappedFile memoryMappedFile = MemoryMappedFile.CreateFromFile(stream, null, stream.Length, MemoryMappedFileAccess.Read, HandleInheritability.None, false);
+			using MemoryMappedView memoryMappedView = new MemoryMappedView(memoryMappedFile, 0, stream.Length, MemoryMappedFileAccess.Read);
+
+			// Create the MSF file
+			MsfFile msfFile = new MsfFile(memoryMappedView);
+
+			// Read the guid from the PDB stream
+			const int PdbStream = 1;
+
+			ReadOnlySpan<byte> firstPdbPage = msfFile.GetStreamData(PdbStream, 0);
+			Guid guid = new Guid(firstPdbPage.Slice(12, 16));
+			string result = guid.ToString("N").ToUpperInvariant();
+
+			// Read the age from the DBI stream
+			const int DbiStream = 3;
+
+			ReadOnlySpan<byte> firstDbiPage = msfFile.GetStreamData(DbiStream, 0);
+			if (firstDbiPage.Length > 0)
+			{
+				int age = BinaryPrimitives.ReadInt32LittleEndian(firstDbiPage.Slice(8));
+				result += $"{age:X}";
+			}
+
+			return result;
+		}
+
+		class MsfFile
+		{
+			static readonly byte[] s_msfSignature = Encoding.ASCII.GetBytes("Microsoft C/C++ MSF 7.00\r\n\x1a\x44\x53\0\0\0");
+
+			record struct MsfStream(int Length, int DirectoryPagesIdx);
+
+			readonly MemoryMappedView _memoryMappedView;
+			readonly int _numPages;
+			readonly int _pageSize;
+			readonly int[] _streamDirectoryPages;
+			readonly MsfStream[] _streams;
+
+			public MsfFile(MemoryMappedView memoryMappedView)
+			{
+				_memoryMappedView = memoryMappedView;
+
+				// Read the file header (the 'superblock' in LLVM docs, BIGMSF_HDR in Microsoft source)
+				const int SizeOf_BIGMSF_HDR = 0x20 + (5 * 4);
+				const int OffsetOf_BIGMSF_HDR_Magic = 0;
+				const int OffsetOf_BIGMSF_HDR_PageSize = 32;
+				const int OffsetOf_BIGMSF_HDR_NumPages = 40;
+				const int OffsetOf_BIGMSF_HDR_NumDirectoryBytes = 44;
+
+				Span<byte> header = memoryMappedView.GetMemory(0, SizeOf_BIGMSF_HDR).Span.Slice(0, SizeOf_BIGMSF_HDR);
+				if (!header.Slice(OffsetOf_BIGMSF_HDR_Magic, s_msfSignature.Length).SequenceEqual(s_msfSignature))
+				{
+					throw new InvalidDataException($"Invalid signature for PDB file");
+				}
+
+				_pageSize = BinaryPrimitives.ReadInt32LittleEndian(header.Slice(OffsetOf_BIGMSF_HDR_PageSize));
+				_numPages = BinaryPrimitives.ReadInt32LittleEndian(header.Slice(OffsetOf_BIGMSF_HDR_NumPages));
+
+				// The stream directory is distributed across multiple pages, and the indexes for the pages of the stream directory are also split across multiple pages.
+				// We can use the stream directory size to figure out how many pages are in each level of indirection.
+				int streamDirectorySize = BinaryPrimitives.ReadInt32LittleEndian(header.Slice(OffsetOf_BIGMSF_HDR_NumDirectoryBytes));
+
+				// The list of stream directory macro pages is immediately after the header
+				int numStreamDirectoryPages = GetPageCount(streamDirectorySize);
+				int numStreamDirectoryMacroPages = GetPageCount(numStreamDirectoryPages * sizeof(int));
+				Span<int> streamDirectoryMacroPages = MemoryMarshal.Cast<byte, int>(memoryMappedView.GetMemory(SizeOf_BIGMSF_HDR, numStreamDirectoryMacroPages * sizeof(int)).Span);
+
+				// Create a flat list of all the stream directory pages to make it easier to seek within it
+				_streamDirectoryPages = new int[numStreamDirectoryPages];
+
+				int pageIdx = 0;
+				for (int macroPageIdx = 0; macroPageIdx < numStreamDirectoryMacroPages; macroPageIdx++)
+				{
+					ReadOnlySpan<byte> macroPage = GetPage(streamDirectoryMacroPages[macroPageIdx]);
+					for (; macroPage.Length > 0 && pageIdx < _streamDirectoryPages.Length; pageIdx++)
+					{
+						_streamDirectoryPages[pageIdx] = BinaryPrimitives.ReadInt32LittleEndian(macroPage);
+						macroPage = macroPage.Slice(sizeof(int));
+					}
+				}
+
+				// Parse the header for each stream from the stream directory
+				int directoryPage = 0;
+				ReadOnlySpan<int> directoryData = GetDirectoryPage(directoryPage);
+
+				int numStreams = directoryData[0];
+				directoryData = directoryData.Slice(1);
+
+				_streams = new MsfStream[numStreams];
+
+				int nextDirectoryOffset = 1 + numStreams;
+				for (int streamIdx = 0; streamIdx < numStreams; streamIdx++)
+				{
+					if (directoryData.Length == 0)
+					{
+						directoryPage++;
+						directoryData = GetDirectoryPage(directoryPage);
+					}
+
+					int streamSize = directoryData[0];
+					_streams[streamIdx] = new MsfStream(streamSize, nextDirectoryOffset);
+					nextDirectoryOffset += GetPageCount(streamSize);
+					directoryData = directoryData.Slice(1);
+				}
+			}
+
+			int GetPageCount(int numBytes)
+				=> (numBytes + (_pageSize - 1)) / _pageSize;
+
+			public ReadOnlySpan<byte> GetStreamData(int streamIdx, int offset)
+			{
+				int streamPageIdx = offset / _pageSize;
+				int streamPagePos = offset % _pageSize;
+				return GetStreamPage(streamIdx, streamPageIdx).Slice(streamPagePos);
+			}
+
+			public ReadOnlySpan<int> GetStreamDirectoryData(int offset)
+			{
+				int pageIdx = (int)(offset / (_pageSize / sizeof(int)));
+				int pagePos = (int)(offset % (_pageSize / sizeof(int)));
+				return GetDirectoryPage(pageIdx).Slice(pagePos);
+			}
+
+			public Span<byte> GetPage(int pageIdx)
+			{
+				Debug.Assert(pageIdx < _numPages);
+				return _memoryMappedView.GetMemory(pageIdx * _pageSize, _pageSize).Span;
+			}
+
+			ReadOnlySpan<byte> GetStreamPage(int streamIdx, int pageIdx)
+			{
+				ReadOnlySpan<int> pages = GetStreamDirectoryData(_streams[streamIdx].DirectoryPagesIdx + pageIdx);
+				return GetPage(pages[0]);
+			}
+
+			ReadOnlySpan<int> GetDirectoryPage(int directoryPageIdx)
+			{
+				int pageIdx = _streamDirectoryPages[directoryPageIdx];
+				return MemoryMarshal.Cast<byte, int>(GetPage(pageIdx));
+			}
+		}
+
+		#endregion
+	}
+}
