Horde: Remove unnecessary using directive.

#rnx

diff --git a/Engine/Source/Programs/Shared/EpicGames.Horde/Symbols/SymbolStoreId.cs b/Engine/Source/Programs/Shared/EpicGames.Horde/Symbols/SymbolStoreId.cs
--- a/Engine/Source/Programs/Shared/EpicGames.Horde/Symbols/SymbolStoreId.cs
+++ b/Engine/Source/Programs/Shared/EpicGames.Horde/Symbols/SymbolStoreId.cs
@@ -2,7 +2,6 @@
 
 using System.ComponentModel;
 using EpicGames.Core;
-using EpicGames.Horde;
 using EpicGames.Serialization;
 
 namespace EpicGames.Horde.Symbols
