EpicGames.Core: Replace obsolete BinaryFormatter with compatible DataContractSerializer

#rnx

diff --git a/Engine/Source/Programs/Shared/EpicGames.Core/BinaryFormatterUtils.cs b/Engine/Source/Programs/Shared/EpicGames.Core/BinaryFormatterUtils.cs
--- a/Engine/Source/Programs/Shared/EpicGames.Core/BinaryFormatterUtils.cs
+++ b/Engine/Source/Programs/Shared/EpicGames.Core/BinaryFormatterUtils.cs
@@ -1,33 +1,29 @@
 // Copyright Epic Games, Inc. All Rights Reserved.
 
 using System.IO;
-using System.Runtime.Serialization.Formatters.Binary;
-
-#pragma warning disable SYSLIB0011
-#pragma warning disable CA2300 // Do not use insecure deserializer BinaryFormatter
-#pragma warning disable CA2301 // Do not use insecure deserializer BinaryFormatter
+using System.Runtime.Serialization;
 
 namespace EpicGames.Core
 {
 	/// <summary>
-	/// Utility functions for serializing using the BinaryFormatter
+	/// Utility functions for serializing using the DataContractSerializer xml (previously BinaryFormatter which is no longer supported)
 	/// </summary>
 	public static class BinaryFormatterUtils
 	{
 		/// <summary>
-		/// Load an object from a file on disk, using the binary formatter
+		/// Load an object from a file on disk, using DataContractSerializer
 		/// </summary>
 		/// <param name="location">File to read from</param>
 		/// <returns>Instance of the object that was read from disk</returns>
 		public static T Load<T>(FileReference location)
 		{
 			using FileStream stream = new FileStream(location.FullName, FileMode.Open, FileAccess.Read);
-			BinaryFormatter formatter = new BinaryFormatter();
-			return (T)formatter.Deserialize(stream);
+			DataContractSerializer serializer = new DataContractSerializer(typeof(T));
+			return (T?)serializer.ReadObject(stream) ?? throw new SerializationException($"Unable to deserialize {location} as {typeof(T)}");
 		}
 
 		/// <summary>
-		/// Saves a file to disk, using the binary formatter
+		/// Saves a file to disk, using DataContractSerializer
 		/// </summary>
 		/// <param name="location">File to write to</param>
 		/// <param name="obj">Object to serialize</param>
@@ -35,12 +31,12 @@ namespace EpicGames.Core
 		{
 			DirectoryReference.CreateDirectory(location.Directory);
 			using FileStream stream = new FileStream(location.FullName, FileMode.Create, FileAccess.Write);
-			BinaryFormatter formatter = new BinaryFormatter();
-			formatter.Serialize(stream, obj);
+			DataContractSerializer serializer = new DataContractSerializer(obj.GetType());
+			serializer.WriteObject(stream, obj);
 		}
 
 		/// <summary>
-		/// Saves a file to disk using the binary formatter, without updating the timestamp if it hasn't changed
+		/// Saves a file to disk using DataContractSerializer, without updating the timestamp if it hasn't changed
 		/// </summary>
 		/// <param name="location">File to write to</param>
 		/// <param name="obj">Object to serialize</param>
@@ -49,8 +45,8 @@ namespace EpicGames.Core
 			byte[] contents;
 			using (MemoryStream stream = new MemoryStream())
 			{
-				BinaryFormatter formatter = new BinaryFormatter();
-				formatter.Serialize(stream, obj);
+				DataContractSerializer serializer = new DataContractSerializer(obj.GetType());
+				serializer.WriteObject(stream, obj);
 				contents = stream.ToArray();
 			}
 
