Horde: Static analysis fix.

#rnx

diff --git a/Engine/Source/Programs/Shared/EpicGames.Horde/Compute/Requirements.cs b/Engine/Source/Programs/Shared/EpicGames.Horde/Compute/Requirements.cs
--- a/Engine/Source/Programs/Shared/EpicGames.Horde/Compute/Requirements.cs
+++ b/Engine/Source/Programs/Shared/EpicGames.Horde/Compute/Requirements.cs
@@ -3,7 +3,8 @@
 using System;
 using System.Collections.Generic;
 using EpicGames.Horde.Common;
-using EpicGames.Serialization;
+
+#pragma warning disable CA2227
 
 namespace EpicGames.Horde.Compute
 {
