.net 8 await

#rnx

diff --git a/Engine/Source/Programs/Shared/EpicGames.Horde/Storage/BlobPipeline.cs b/Engine/Source/Programs/Shared/EpicGames.Horde/Storage/BlobPipeline.cs
--- a/Engine/Source/Programs/Shared/EpicGames.Horde/Storage/BlobPipeline.cs
+++ b/Engine/Source/Programs/Shared/EpicGames.Horde/Storage/BlobPipeline.cs
@@ -107,7 +107,11 @@ namespace EpicGames.Horde.Storage
 		{
 			if (_tasks.Count > 0)
 			{
+#if NET8_0_OR_GREATER
+				await _cancellationTokenSource.CancelAsync();
+#else
 				_cancellationTokenSource.Cancel();
+#endif
 				try
 				{
 					await Task.WhenAll(_tasks);
