Horde: Set the max thread count for managed workspace operations to one less than the number of reported CPUs on the machine.

diff --git a/Engine/Source/Programs/Shared/EpicGames.Perforce.Managed/ManagedWorkspace.cs b/Engine/Source/Programs/Shared/EpicGames.Perforce.Managed/ManagedWorkspace.cs
--- a/Engine/Source/Programs/Shared/EpicGames.Perforce.Managed/ManagedWorkspace.cs
+++ b/Engine/Source/Programs/Shared/EpicGames.Perforce.Managed/ManagedWorkspace.cs
@@ -95,13 +95,27 @@ namespace EpicGames.Perforce.Managed
 	/// <param name="PreferNativeClient">Whether to prefer the native p4 client</param>
 	public record ManagedWorkspaceOptions
 	(
-		int NumParallelSyncThreads = 4,
-		int MaxFileConcurrency = 4,
-		long MinScratchSpace = 50L * 1024, // 50 GB
-		bool UseHaveTable = true,
-		bool Partitioned = false,
-		bool PreferNativeClient = false
-	);
+		int NumParallelSyncThreads,
+		int MaxFileConcurrency,
+		long MinScratchSpace,
+		bool UseHaveTable,
+		bool Partitioned,
+		bool PreferNativeClient
+	)
+	{
+		public ManagedWorkspaceOptions(int? numParallelSyncThreads = null, int? maxFileConcurrency = null, long minScratchSpace = 50L * 1024, bool useHaveTable = true, bool partitioned = false, bool preferNativeClient = false)
+			: this(
+				  numParallelSyncThreads ?? GetDefaultThreadCount(4),
+				  maxFileConcurrency ?? GetDefaultThreadCount(4),
+				  minScratchSpace,
+				  useHaveTable,
+				  partitioned,
+				  preferNativeClient)
+		{ }
+
+		static int GetDefaultThreadCount(int defaultValue)
+			=> Math.Min(defaultValue, Math.Max(Environment.ProcessorCount - 1, 1));
+	}
 
 	/// <summary>
 	/// Version number for managed workspace cache files
