Horde: Fix static analysis warning.

#rnx

diff --git a/Engine/Source/Programs/Shared/EpicGames.Perforce.Managed/ManagedWorkspace.cs b/Engine/Source/Programs/Shared/EpicGames.Perforce.Managed/ManagedWorkspace.cs
--- a/Engine/Source/Programs/Shared/EpicGames.Perforce.Managed/ManagedWorkspace.cs
+++ b/Engine/Source/Programs/Shared/EpicGames.Perforce.Managed/ManagedWorkspace.cs
@@ -103,6 +103,7 @@ namespace EpicGames.Perforce.Managed
 		bool PreferNativeClient
 	)
 	{
+		/// <inheritdoc cref="ManagedWorkspaceOptions.ManagedWorkspaceOptions(int, int, long, bool, bool, bool)"/>
 		public ManagedWorkspaceOptions(int? numParallelSyncThreads = null, int? maxFileConcurrency = null, long minScratchSpace = 50L * 1024, bool useHaveTable = true, bool partitioned = false, bool preferNativeClient = false)
 			: this(
 				  numParallelSyncThreads ?? GetDefaultThreadCount(4),
