Horde: Add an accessor for the content of a tool deployment.

#rnx

diff --git a/Engine/Source/Programs/Horde/Plugins/Tools/HordeServer.Tools/Tools/ToolCollection.cs b/Engine/Source/Programs/Horde/Plugins/Tools/HordeServer.Tools/Tools/ToolCollection.cs
--- a/Engine/Source/Programs/Horde/Plugins/Tools/HordeServer.Tools/Tools/ToolCollection.cs
+++ b/Engine/Source/Programs/Horde/Plugins/Tools/HordeServer.Tools/Tools/ToolCollection.cs
@@ -151,6 +151,9 @@ namespace HordeServer.Tools
 			public NamespaceId NamespaceId => _document.NamespaceId;
 			public RefName RefName => _document.RefName;
 
+			public IBlobRef<DirectoryNode> Content
+				=> _collection.Open(_tool.Config, _document);
+
 			public ToolDeployment(Tool tool, ToolCollection collection, ToolDeploymentDocument document, DateTime utcNow)
 			{
 				_tool = tool;
@@ -489,6 +492,12 @@ namespace HordeServer.Tools
 			}
 		}
 
+		IBlobRef<DirectoryNode> Open(ToolConfig tool, ToolDeploymentDocument deployment)
+		{
+			IStorageClient client = CreateStorageClient(tool);
+			return client.CreateBlobRef<DirectoryNode>(deployment.RefName, DateTime.UtcNow - TimeSpan.FromDays(2.0));
+		}
+
 		/// <summary>
 		/// Opens a stream to the data for a particular deployment
 		/// </summary>
diff --git a/Engine/Source/Programs/Shared/EpicGames.Horde/Tools/ITool.cs b/Engine/Source/Programs/Shared/EpicGames.Horde/Tools/ITool.cs
--- a/Engine/Source/Programs/Shared/EpicGames.Horde/Tools/ITool.cs
+++ b/Engine/Source/Programs/Shared/EpicGames.Horde/Tools/ITool.cs
@@ -8,6 +8,7 @@ using System.Threading;
 using System.Threading.Tasks;
 using EpicGames.Horde.Acls;
 using EpicGames.Horde.Storage;
+using EpicGames.Horde.Storage.Nodes;
 
 #pragma warning disable CA1716 // Rename virtual/interface member ITool.Public so that it no longer conflicts with the reserved language keyword 'Public'.
 
@@ -157,6 +158,11 @@ namespace EpicGames.Horde.Tools
 		/// </summary>
 		RefName RefName { get; }
 
+		/// <summary>
+		/// Handle to the tool data
+		/// </summary>
+		IBlobRef<DirectoryNode> Content { get; }
+
 		/// <summary>
 		/// Updates the state of the current deployment
 		/// </summary>
diff --git a/Engine/Source/Programs/Shared/EpicGames.Horde/Tools/ToolHttpCollection.cs b/Engine/Source/Programs/Shared/EpicGames.Horde/Tools/ToolHttpCollection.cs
--- a/Engine/Source/Programs/Shared/EpicGames.Horde/Tools/ToolHttpCollection.cs
+++ b/Engine/Source/Programs/Shared/EpicGames.Horde/Tools/ToolHttpCollection.cs
@@ -9,6 +9,7 @@ using System.Threading;
 using System.Threading.Tasks;
 using EpicGames.Horde.Acls;
 using EpicGames.Horde.Storage;
+using EpicGames.Horde.Storage.Nodes;
 
 namespace EpicGames.Horde.Tools
 {
@@ -72,6 +73,9 @@ namespace EpicGames.Horde.Tools
 			NamespaceId IToolDeployment.NamespaceId => throw new NotImplementedException();
 			RefName IToolDeployment.RefName => _response.RefName;
 
+			IBlobRef<DirectoryNode> IToolDeployment.Content
+				=> _collection.Open(_toolId, _response.Locator);
+
 			public ToolDeployment(ToolHttpCollection collection, ToolId toolId, GetToolDeploymentResponse response)
 			{
 				_collection = collection;
@@ -121,6 +125,12 @@ namespace EpicGames.Horde.Tools
 			return tools;
 		}
 
+		IBlobRef<DirectoryNode> Open(ToolId toolId, BlobLocator locator)
+		{
+			IStorageClient store = _hordeClient.CreateStorageClient(toolId);
+			return store.CreateBlobRef<DirectoryNode>(locator);
+		}
+
 		async Task<Stream> OpenZipStreamAsync(ToolId id, ToolDeploymentId deploymentId, CancellationToken cancellationToken = default)
 		{
 			HordeHttpClient hordeHttpClient = _hordeClient.CreateHttpClient();
