Horde: Configure Json serializers for ToolId and ToolDeploymentId types.

#rnx

diff --git a/Engine/Source/Programs/Shared/EpicGames.Horde/Tools/ToolDeploymentId.cs b/Engine/Source/Programs/Shared/EpicGames.Horde/Tools/ToolDeploymentId.cs
--- a/Engine/Source/Programs/Shared/EpicGames.Horde/Tools/ToolDeploymentId.cs
+++ b/Engine/Source/Programs/Shared/EpicGames.Horde/Tools/ToolDeploymentId.cs
@@ -1,6 +1,7 @@
 // Copyright Epic Games, Inc. All Rights Reserved.
 
 using System.ComponentModel;
+using System.Text.Json.Serialization;
 using EpicGames.Core;
 
 namespace EpicGames.Horde.Tools
@@ -11,6 +12,7 @@ namespace EpicGames.Horde.Tools
 	/// <param name="Id">Identifier for the artifact</param>
 	[LogValueType]
 	[JsonSchemaString]
+	[JsonConverter(typeof(BinaryIdJsonConverter<ToolDeploymentId, ToolDeploymentIdConverter>))]
 	[TypeConverter(typeof(BinaryIdTypeConverter<ToolDeploymentId, ToolDeploymentIdConverter>))]
 	[BinaryIdConverter(typeof(ToolDeploymentIdConverter))]
 	public readonly record struct ToolDeploymentId(BinaryId Id)
diff --git a/Engine/Source/Programs/Shared/EpicGames.Horde/Tools/ToolId.cs b/Engine/Source/Programs/Shared/EpicGames.Horde/Tools/ToolId.cs
--- a/Engine/Source/Programs/Shared/EpicGames.Horde/Tools/ToolId.cs
+++ b/Engine/Source/Programs/Shared/EpicGames.Horde/Tools/ToolId.cs
@@ -1,6 +1,7 @@
 // Copyright Epic Games, Inc. All Rights Reserved.
 
 using System.ComponentModel;
+using System.Text.Json.Serialization;
 using EpicGames.Core;
 using EpicGames.Serialization;
 
@@ -12,6 +13,7 @@ namespace EpicGames.Horde.Tools
 	/// <param name="Id">Id to construct from</param>
 	[LogValueType]
 	[JsonSchemaString]
+	[JsonConverter(typeof(StringIdJsonConverter<ToolId, ToolIdConverter>))]
 	[TypeConverter(typeof(StringIdTypeConverter<ToolId, ToolIdConverter>))]
 	[StringIdConverter(typeof(ToolIdConverter))]
 	[CbConverter(typeof(StringIdCbConverter<ToolId, ToolIdConverter>))]
