Horde: Static analysis fixes.

#rnx

diff --git a/Engine/Source/Programs/Shared/EpicGames.Horde/Compute/Clients/LocalComputeClient.cs b/Engine/Source/Programs/Shared/EpicGames.Horde/Compute/Clients/LocalComputeClient.cs
--- a/Engine/Source/Programs/Shared/EpicGames.Horde/Compute/Clients/LocalComputeClient.cs
+++ b/Engine/Source/Programs/Shared/EpicGames.Horde/Compute/Clients/LocalComputeClient.cs
@@ -1,5 +1,6 @@
 // Copyright Epic Games, Inc. All Rights Reserved.
 
+using System;
 using System.Collections.Generic;
 using System.Net;
 using System.Net.Sockets;
@@ -14,7 +15,7 @@ namespace EpicGames.Horde.Compute.Clients
 	/// <summary>
 	/// Implementation of <see cref="IComputeClient"/> which marshals data over a loopback connection to a method running on a background task in the same process.
 	/// </summary>
-	public sealed class LocalComputeClient : IComputeClient
+	public sealed class LocalComputeClient : IComputeClient, IAsyncDisposable
 	{
 		private static readonly ClusterId s_cluster = new ("_local");
 		
diff --git a/Engine/Source/Programs/Shared/EpicGames.Horde/HordeClient.cs b/Engine/Source/Programs/Shared/EpicGames.Horde/HordeClient.cs
--- a/Engine/Source/Programs/Shared/EpicGames.Horde/HordeClient.cs
+++ b/Engine/Source/Programs/Shared/EpicGames.Horde/HordeClient.cs
@@ -36,7 +36,7 @@ namespace EpicGames.Horde
 	/// <summary>
 	/// Default implementation of <see cref="IHordeClient"/>
 	/// </summary>
-	public abstract class HordeClient : IHordeClient
+	public abstract class HordeClient : IHordeClient, IAsyncDisposable
 	{
 		readonly Uri _serverUrl;
 		readonly BundleCache _bundleCache;
@@ -68,7 +68,7 @@ namespace EpicGames.Horde
 		/// <summary>
 		/// Constructor
 		/// </summary>
-		public HordeClient(Uri serverUrl, BundleCache bundleCache, IOptionsSnapshot<HordeOptions> hordeOptions, ILoggerFactory loggerFactory)
+		protected HordeClient(Uri serverUrl, BundleCache bundleCache, IOptionsSnapshot<HordeOptions> hordeOptions, ILoggerFactory loggerFactory)
 		{
 			_serverUrl = serverUrl;
 			_bundleCache = bundleCache;
@@ -94,6 +94,8 @@ namespace EpicGames.Horde
 				await _grpcChannel.DisposeAsync();
 				_grpcChannel = null;
 			}
+
+			GC.SuppressFinalize(this);
 		}
 
 		/// <summary>
diff --git a/Engine/Source/Programs/Shared/EpicGames.Horde/IHordeClient.cs b/Engine/Source/Programs/Shared/EpicGames.Horde/IHordeClient.cs
--- a/Engine/Source/Programs/Shared/EpicGames.Horde/IHordeClient.cs
+++ b/Engine/Source/Programs/Shared/EpicGames.Horde/IHordeClient.cs
@@ -11,7 +11,6 @@ using EpicGames.Horde.Secrets;
 using EpicGames.Horde.Storage;
 using EpicGames.Horde.Tools;
 using Grpc.Core;
-using Grpc.Net.Client;
 using Microsoft.Extensions.Logging;
 
 namespace EpicGames.Horde
