Horde: Rename Horde http collection implementations for consistency.

#rnx

diff --git a/Engine/Source/Programs/Shared/EpicGames.Horde/Artifacts/ArtifactHttpCollection.cs b/Engine/Source/Programs/Shared/EpicGames.Horde/Artifacts/ArtifactHttpCollection.cs
deleted file mode 100644
diff --git a/Engine/Source/Programs/Shared/EpicGames.Horde/Artifacts/ArtifactHttpMessages.cs b/Engine/Source/Programs/Shared/EpicGames.Horde/Artifacts/ArtifactHttpMessages.cs
deleted file mode 100644
diff --git a/Engine/Source/Programs/Shared/EpicGames.Horde/Artifacts/ArtifactMessages.cs b/Engine/Source/Programs/Shared/EpicGames.Horde/Artifacts/ArtifactMessages.cs
new file mode 100644
--- /dev/null
+++ b/Engine/Source/Programs/Shared/EpicGames.Horde/Artifacts/ArtifactMessages.cs
@@ -0,0 +1,258 @@
+﻿// Copyright Epic Games, Inc. All Rights Reserved.
+
+using System;
+using System.Collections.Generic;
+using System.Text.Json.Serialization;
+using EpicGames.Core;
+using EpicGames.Horde.Commits;
+using EpicGames.Horde.Storage;
+using EpicGames.Horde.Streams;
+
+#pragma warning disable CA2227
+
+namespace EpicGames.Horde.Artifacts
+{
+	/// <summary>
+	/// Creates a new artifact
+	/// </summary>
+	/// <param name="Name">Name of the artifact</param>
+	/// <param name="Type">Additional search keys tagged on the artifact</param>
+	/// <param name="Description">Description for the artifact</param>
+	/// <param name="StreamId">Stream to create the artifact for</param>
+	/// <param name="Keys">Keys used to identify the artifact</param>
+	/// <param name="Metadata">Metadata for the artifact</param>
+	public record CreateArtifactRequest(ArtifactName Name, ArtifactType Type, string? Description, StreamId? StreamId, List<string> Keys, List<string> Metadata)
+	{
+		/// <summary>
+		/// Legacy Perforce changelist number.
+		/// </summary>
+		[Obsolete("Use CommitId instead")]
+		public int Change
+		{
+			get => _change ?? _commitId?.TryGetPerforceChange() ?? -1;
+			set => _change = value;
+		}
+		int? _change;
+
+		/// <summary>
+		/// Commit for the new artifact
+		/// </summary>
+		public CommitId CommitId
+		{
+			get => _commitId ?? CommitId.FromPerforceChange(_change) ?? CommitId.Empty;
+			set => _commitId = value;
+		}
+		CommitId? _commitId;
+	}
+
+	/// <summary>
+	/// Information about a created artifact
+	/// </summary>
+	/// <param name="ArtifactId">Identifier for the new artifact</param>
+	/// <param name="CommitId">Resolved commit id for the artifact</param>
+	/// <param name="NamespaceId">Namespace that should be written to with artifact data</param>
+	/// <param name="RefName">Ref to write to</param>
+	/// <param name="PrevRefName">Ref for the artifact at the changelist prior to this one. Can be used to deduplicate against.</param>
+	/// <param name="Token">Token which can be used to upload blobs for the artifact, and read blobs from the previous artifact</param>
+	public record CreateArtifactResponse(ArtifactId ArtifactId, CommitIdWithOrder CommitId, NamespaceId NamespaceId, RefName RefName, RefName? PrevRefName, string Token);
+
+	/// <summary>
+	/// Type of data to download for an artifact
+	/// </summary>
+	public enum DownloadArtifactFormat
+	{
+		/// <summary>
+		/// Download as a zip file
+		/// </summary>
+		Zip,
+
+		/// <summary>
+		/// Download as a UGS link
+		/// </summary>
+		Ugs
+	}
+
+	/// <summary>
+	/// Describes an artifact
+	/// </summary>
+	public class GetArtifactResponse
+	{
+		/// <inheritdoc cref="IArtifact.Id"/>
+		public ArtifactId Id { get; set; }
+
+		/// <inheritdoc cref="IArtifact.Name"/>
+		public ArtifactName Name { get; set; }
+
+		/// <inheritdoc cref="IArtifact.Type"/>
+		public ArtifactType Type { get; set; }
+
+		/// <inheritdoc cref="IArtifact.Description"/>
+		public string? Description { get; set; }
+
+		/// <inheritdoc cref="IArtifact.StreamId"/>
+		public StreamId StreamId { get; set; }
+
+		/// <summary>
+		/// Change number
+		/// </summary>
+		[Obsolete("Use CommitId instead")]
+		public int Change
+		{
+			get => _change ?? _commitId?.TryGetPerforceChange() ?? -1;
+			set => _change = value;
+		}
+		int? _change;
+
+		/// <inheritdoc cref="IArtifact.CommitId"/>
+		public CommitIdWithOrder CommitId
+		{
+			get => _commitId ?? CommitIdWithOrder.FromPerforceChange(_change) ?? CommitIdWithOrder.Empty;
+			set => _commitId = value;
+		}
+		CommitIdWithOrder? _commitId;
+
+		/// <inheritdoc cref="IArtifact.Keys"/>
+		public IReadOnlyList<string> Keys { get; set; }
+
+		/// <inheritdoc cref="IArtifact.Metadata"/>
+		public IReadOnlyList<string> Metadata { get; set; }
+
+		/// <inheritdoc cref="IArtifact.NamespaceId"/>
+		public NamespaceId NamespaceId { get; set; }
+
+		/// <inheritdoc cref="IArtifact.RefName"/>
+		public RefName RefName { get; set; }
+
+		/// <inheritdoc cref="IArtifact.CreatedAtUtc"/>
+		public DateTime CreatedAtUtc { get; set; }
+
+		/// <summary>
+		/// Default constructor
+		/// </summary>
+		[JsonConstructor]
+		public GetArtifactResponse()
+		{
+			Keys = Array.Empty<string>();
+			Metadata = Array.Empty<string>();
+		}
+
+		/// <summary>
+		/// Constructor
+		/// </summary>
+		/// <param name="artifact"></param>
+		public GetArtifactResponse(IArtifact artifact)
+		{
+			Id = artifact.Id;
+			Name = artifact.Name;
+			Type = artifact.Type;
+			Description = artifact.Description;
+			StreamId = artifact.StreamId;
+			CommitId = artifact.CommitId;
+			Keys = artifact.Keys;
+			Metadata = artifact.Metadata;
+			NamespaceId = artifact.NamespaceId;
+			RefName = artifact.RefName;
+			CreatedAtUtc = artifact.CreatedAtUtc;
+		}
+	}
+
+	/// <summary>
+	/// Result of an artifact search
+	/// </summary>
+	public class FindArtifactsResponse
+	{
+		/// <summary>
+		/// List of artifacts matching the search criteria
+		/// </summary>
+		public List<GetArtifactResponse> Artifacts { get; set; } = new List<GetArtifactResponse>();
+	}
+
+	/// <summary>
+	/// Describes a file within an artifact
+	/// </summary>
+	public class GetArtifactFileEntryResponse
+	{
+		/// <summary>
+		/// Name of this file
+		/// </summary>
+		public string Name { get; }
+
+		/// <summary>
+		/// Length of this entry
+		/// </summary>
+		public long Length { get; }
+
+		/// <summary>
+		/// Hash of the target node
+		/// </summary>
+		public IoHash Hash { get; }
+
+		/// <summary>
+		/// Constructor
+		/// </summary>
+		public GetArtifactFileEntryResponse(string name, long length, IoHash hash)
+		{
+			Name = name;
+			Length = length;
+			Hash = hash;
+		}
+	}
+
+	/// <summary>
+	/// Describes a file within an artifact
+	/// </summary>
+	public class GetArtifactDirectoryEntryResponse : GetArtifactDirectoryResponse
+	{
+		/// <summary>
+		/// Name of this file
+		/// </summary>
+		public string Name { get; }
+
+		/// <summary>
+		/// Length of this entry
+		/// </summary>
+		public long Length { get; }
+
+		/// <summary>
+		/// Hash of the target node
+		/// </summary>
+		public IoHash Hash { get; }
+
+		/// <summary>
+		/// Constructor
+		/// </summary>
+		public GetArtifactDirectoryEntryResponse(string name, long length, IoHash hash)
+		{
+			Name = name;
+			Length = length;
+			Hash = hash;
+		}
+	}
+
+	/// <summary>
+	/// Describes a directory within an artifact
+	/// </summary>
+	public class GetArtifactDirectoryResponse
+	{
+		/// <summary>
+		/// Names of sub-directories
+		/// </summary>
+		public List<GetArtifactDirectoryEntryResponse>? Directories { get; set; }
+
+		/// <summary>
+		/// Files within the directory
+		/// </summary>
+		public List<GetArtifactFileEntryResponse>? Files { get; set; }
+	}
+
+	/// <summary>
+	/// Request to create a zip file with artifact data
+	/// </summary>
+	public class CreateZipRequest
+	{
+		/// <summary>
+		/// Filter lines for the zip. Uses standard <see cref="FileFilter"/> syntax.
+		/// </summary>
+		public List<string> Filter { get; set; } = new List<string>();
+	}
+}
diff --git a/Engine/Source/Programs/Shared/EpicGames.Horde/Artifacts/HttpArtifactCollection.cs b/Engine/Source/Programs/Shared/EpicGames.Horde/Artifacts/HttpArtifactCollection.cs
new file mode 100644
--- /dev/null
+++ b/Engine/Source/Programs/Shared/EpicGames.Horde/Artifacts/HttpArtifactCollection.cs
@@ -0,0 +1,109 @@
+// Copyright Epic Games, Inc. All Rights Reserved.
+
+using System;
+using System.Collections.Generic;
+using System.Diagnostics;
+using System.Linq;
+using System.Runtime.CompilerServices;
+using System.Threading;
+using System.Threading.Tasks;
+using EpicGames.Horde.Commits;
+using EpicGames.Horde.Storage;
+using EpicGames.Horde.Storage.Nodes;
+using EpicGames.Horde.Streams;
+
+namespace EpicGames.Horde.Artifacts
+{
+	class HttpArtifactCollection : IArtifactCollection
+	{
+		[DebuggerDisplay("{Id}")]
+		class Artifact : IArtifact
+		{
+			readonly HttpArtifactCollection _collection;
+
+			public Artifact(HttpArtifactCollection collection, GetArtifactResponse response)
+				: this(collection, response.Id, response.Name, response.Type, response.Description, response.StreamId, response.CommitId, response.Keys, response.Metadata, response.NamespaceId, response.RefName, response.CreatedAtUtc)
+			{ }
+
+			public Artifact(HttpArtifactCollection collection, ArtifactId id, ArtifactName name, ArtifactType type, string? description, StreamId streamId, CommitIdWithOrder commitId, IReadOnlyList<string> keys, IReadOnlyList<string> metadata, NamespaceId namespaceId, RefName refName, DateTime createdAtUtc)
+			{
+				_collection = collection;
+
+				Id = id;
+				Name = name;
+				Type = type;
+				Description = description;
+				StreamId = streamId;
+				CommitId = commitId;
+				Keys = keys;
+				Metadata = metadata;
+				NamespaceId = namespaceId;
+				RefName = refName;
+				CreatedAtUtc = createdAtUtc;
+			}
+
+			public ArtifactId Id { get; }
+			public ArtifactName Name { get; }
+			public ArtifactType Type { get; }
+			public string? Description { get; }
+			public StreamId StreamId { get; }
+			public CommitIdWithOrder CommitId { get; }
+			public IReadOnlyList<string> Keys { get; }
+			public IReadOnlyList<string> Metadata { get; }
+			public NamespaceId NamespaceId { get; }
+			public RefName RefName { get; }
+			public DateTime CreatedAtUtc { get; }
+
+			public IBlobRef<DirectoryNode> Content
+				=> _collection.Open(this);
+
+			public Task DeleteAsync(CancellationToken cancellationToken)
+				=> _collection.DeleteAsync(Id, cancellationToken);
+		}
+
+		readonly IHordeClient _hordeClient;
+
+		public HttpArtifactCollection(IHordeClient hordeClient)
+			=> _hordeClient = hordeClient;
+
+		IBlobRef<DirectoryNode> Open(Artifact artifact)
+		{
+			IStorageClient store = _hordeClient.CreateStorageClient(artifact.Id);
+			return store.CreateBlobRef<DirectoryNode>(new RefName("default"));
+		}
+
+		/// <inheritdoc/>
+		public async Task<IArtifact> AddAsync(ArtifactName name, ArtifactType type, string? description, StreamId streamId, CommitId commitId, IEnumerable<string> keys, IEnumerable<string> metadata, CancellationToken cancellationToken = default)
+		{
+			HordeHttpClient hordeHttpClient = _hordeClient.CreateHttpClient();
+			CreateArtifactResponse response = await hordeHttpClient.CreateArtifactAsync(name, type, description, streamId, commitId, keys, metadata, cancellationToken);
+			return new Artifact(this, response.ArtifactId, name, type, description, streamId, CommitIdWithOrder.FromPerforceChange(commitId.GetPerforceChange()), keys.ToList(), metadata.ToList(), response.NamespaceId, response.RefName, DateTime.UtcNow);
+		}
+
+		/// <inheritdoc/>
+		public async IAsyncEnumerable<IArtifact> FindAsync(StreamId? streamId = null, CommitId? minCommitId = null, CommitId? maxCommitId = null, ArtifactName? name = null, ArtifactType? type = null, IEnumerable<string>? keys = null, int maxResults = 100, [EnumeratorCancellation] CancellationToken cancellationToken = default)
+		{
+			HordeHttpClient hordeHttpClient = _hordeClient.CreateHttpClient();
+
+			List<GetArtifactResponse> responses = await hordeHttpClient.FindArtifactsAsync(streamId, minCommitId, maxCommitId, name, type, keys, maxResults, cancellationToken);
+			foreach (GetArtifactResponse response in responses)
+			{
+				yield return new Artifact(this, response);
+			}
+		}
+
+		async Task DeleteAsync(ArtifactId artifactId, CancellationToken cancellationToken = default)
+		{
+			HordeHttpClient hordeHttpClient = _hordeClient.CreateHttpClient();
+			await hordeHttpClient.DeleteArtifactAsync(artifactId, cancellationToken);
+		}
+
+		/// <inheritdoc/>
+		public async Task<IArtifact?> GetAsync(ArtifactId artifactId, CancellationToken cancellationToken = default)
+		{
+			HordeHttpClient hordeHttpClient = _hordeClient.CreateHttpClient();
+			GetArtifactResponse? response = await hordeHttpClient.GetArtifactAsync(artifactId, cancellationToken);
+			return new Artifact(this, response);
+		}
+	}
+}
diff --git a/Engine/Source/Programs/Shared/EpicGames.Horde/HordeClient.cs b/Engine/Source/Programs/Shared/EpicGames.Horde/HordeClient.cs
--- a/Engine/Source/Programs/Shared/EpicGames.Horde/HordeClient.cs
+++ b/Engine/Source/Programs/Shared/EpicGames.Horde/HordeClient.cs
@@ -78,10 +78,10 @@ namespace EpicGames.Horde
 
 			_serverComputeClient = CreateComputeClient();
 
-			Artifacts = new ArtifactHttpCollection(this);
-			Projects = new ProjectHttpCollection(this);
-			Secrets = new SecretHttpCollection(this);
-			Tools = new ToolHttpCollection(this);
+			Artifacts = new HttpArtifactCollection(this);
+			Projects = new HttpProjectCollection(this);
+			Secrets = new HttpSecretCollection(this);
+			Tools = new HttpToolCollection(this);
 		}
 
 		/// <inheritdoc/>
diff --git a/Engine/Source/Programs/Shared/EpicGames.Horde/Projects/HttpProjectCollection.cs b/Engine/Source/Programs/Shared/EpicGames.Horde/Projects/HttpProjectCollection.cs
new file mode 100644
--- /dev/null
+++ b/Engine/Source/Programs/Shared/EpicGames.Horde/Projects/HttpProjectCollection.cs
@@ -0,0 +1,67 @@
+// Copyright Epic Games, Inc. All Rights Reserved.
+
+using System;
+using System.Collections.Generic;
+using System.Diagnostics;
+using System.Threading;
+using System.Threading.Tasks;
+using EpicGames.Horde.Streams;
+
+namespace EpicGames.Horde.Projects
+{
+	class HttpProjectCollection : IProjectCollection
+	{
+		[DebuggerDisplay("{Id}")]
+		class Project : IProject
+		{
+			readonly GetProjectResponse _response;
+			readonly IReadOnlyList<ProjectStream> _projectStreams;
+
+			public ProjectId Id => _response.Id;
+			public string Name => _response.Name;
+			public int Order => _response.Order;
+			public IReadOnlyList<IProjectStream> Streams => _projectStreams;
+
+			public Project(GetProjectResponse response)
+			{
+				_response = response;
+				_projectStreams = (IReadOnlyList<ProjectStream>?)response.Streams?.ConvertAll(x => new ProjectStream(x)) ?? Array.Empty<ProjectStream>();
+			}
+		}
+
+		class ProjectStream : IProjectStream
+		{
+			readonly GetProjectStreamResponse _response;
+
+			public StreamId Id => new StreamId(_response.Id);
+			public string Name => _response.Name;
+
+			public ProjectStream(GetProjectStreamResponse response)
+				=> _response = response;
+		}
+
+		readonly IHordeClient _hordeClient;
+
+		/// <summary>
+		/// Constructor
+		/// </summary>
+		public HttpProjectCollection(IHordeClient hordeClient)
+			=> _hordeClient = hordeClient;
+
+		/// <inheritdoc/>
+		public async Task<IProject> GetAsync(ProjectId projectId, CancellationToken cancellationToken)
+		{
+			HordeHttpClient hordeHttpClient = _hordeClient.CreateHttpClient();
+			GetProjectResponse? response = await hordeHttpClient.GetProjectAsync(projectId, cancellationToken);
+			return new Project(response);
+		}
+
+		/// <inheritdoc/>
+		public async Task<IReadOnlyList<IProject>> GetAllAsync(CancellationToken cancellationToken)
+		{
+			HordeHttpClient hordeHttpClient = _hordeClient.CreateHttpClient();
+			List<GetProjectResponse> responses = await hordeHttpClient.GetProjectsAsync(cancellationToken: cancellationToken);
+			return responses.ConvertAll(x => new Project(x));
+		}
+	}
+}
diff --git a/Engine/Source/Programs/Shared/EpicGames.Horde/Projects/ProjectHttpCollection.cs b/Engine/Source/Programs/Shared/EpicGames.Horde/Projects/ProjectHttpCollection.cs
deleted file mode 100644
diff --git a/Engine/Source/Programs/Shared/EpicGames.Horde/Secrets/HttpSecretCollection.cs b/Engine/Source/Programs/Shared/EpicGames.Horde/Secrets/HttpSecretCollection.cs
new file mode 100644
--- /dev/null
+++ b/Engine/Source/Programs/Shared/EpicGames.Horde/Secrets/HttpSecretCollection.cs
@@ -0,0 +1,38 @@
+﻿// Copyright Epic Games, Inc. All Rights Reserved.
+
+using System.Collections.Generic;
+using System.Threading;
+using System.Threading.Tasks;
+
+namespace EpicGames.Horde.Secrets
+{
+	class HttpSecretCollection : ISecretCollection
+	{
+		class Secret : ISecret
+		{
+			readonly GetSecretResponse _response;
+
+			public SecretId Id => _response.Id;
+			public IReadOnlyDictionary<string, string> Data => _response.Data;
+
+			public Secret(GetSecretResponse response)
+				=> _response = response;
+		}
+
+		readonly IHordeClient _hordeClient;
+
+		/// <summary>
+		/// Constructor
+		/// </summary>
+		public HttpSecretCollection(IHordeClient hordeClient)
+			=> _hordeClient = hordeClient;
+
+		/// <inheritdoc/>
+		public async Task<ISecret?> GetAsync(SecretId secretId, CancellationToken cancellationToken)
+		{
+			HordeHttpClient hordeHttpClient = _hordeClient.CreateHttpClient();
+			GetSecretResponse response = await hordeHttpClient.GetSecretAsync(secretId, cancellationToken);
+			return new Secret(response);
+		}
+	}
+}
diff --git a/Engine/Source/Programs/Shared/EpicGames.Horde/Secrets/SecretHttpCollection.cs b/Engine/Source/Programs/Shared/EpicGames.Horde/Secrets/SecretHttpCollection.cs
deleted file mode 100644
diff --git a/Engine/Source/Programs/Shared/EpicGames.Horde/Tools/HttpToolCollection.cs b/Engine/Source/Programs/Shared/EpicGames.Horde/Tools/HttpToolCollection.cs
new file mode 100644
--- /dev/null
+++ b/Engine/Source/Programs/Shared/EpicGames.Horde/Tools/HttpToolCollection.cs
@@ -0,0 +1,140 @@
+﻿// Copyright Epic Games, Inc. All Rights Reserved.
+
+using System;
+using System.Collections.Generic;
+using System.Diagnostics;
+using System.IO;
+using System.Security.Claims;
+using System.Threading;
+using System.Threading.Tasks;
+using EpicGames.Horde.Acls;
+using EpicGames.Horde.Storage;
+using EpicGames.Horde.Storage.Nodes;
+
+namespace EpicGames.Horde.Tools
+{
+	class HttpToolCollection : IToolCollection
+	{
+		[DebuggerDisplay("{Id}")]
+		class Tool : ITool
+		{
+			readonly GetToolResponse _response;
+			readonly List<ToolDeployment> _deployments;
+
+			public ToolId Id => _response.Id;
+
+			string ITool.Name => _response.Name;
+			string ITool.Description => _response.Description;
+			string? ITool.Category => _response.Category;
+			string? ITool.Group => _response.Group;
+			IReadOnlyList<string>? ITool.Platforms => _response.Platforms;
+			bool ITool.Public => _response.Public;
+			bool ITool.Bundled => _response.Bundled;
+			bool ITool.ShowInUgs => _response.ShowInUgs;
+			bool ITool.ShowInDashboard => _response.ShowInDashboard;
+			IReadOnlyList<IToolDeployment> ITool.Deployments => _deployments;
+
+			public Tool(HttpToolCollection collection, GetToolResponse response)
+			{
+				_response = response;
+				_deployments = response.Deployments.ConvertAll(x => new ToolDeployment(collection, _response.Id, x));
+			}
+
+			public bool Authorize(AclAction action, ClaimsPrincipal principal)
+				=> throw new NotImplementedException();
+
+			public Task<ITool?> CreateDeploymentAsync(ToolDeploymentConfig options, Stream stream, CancellationToken cancellationToken = default)
+				=> throw new NotImplementedException();
+
+			public Task<ITool?> CreateDeploymentAsync(ToolDeploymentConfig options, HashedBlobRefValue target, CancellationToken cancellationToken = default)
+				=> throw new NotImplementedException();
+
+			public IStorageBackend CreateStorageBackend()
+				=> throw new NotImplementedException();
+
+			public IStorageClient CreateStorageClient()
+				=> throw new NotImplementedException();
+		}
+
+		[DebuggerDisplay("{Id}")]
+		class ToolDeployment : IToolDeployment
+		{
+			readonly HttpToolCollection _collection;
+			readonly ToolId _toolId;
+			readonly GetToolDeploymentResponse _response;
+
+			public ToolDeploymentId Id => _response.Id;
+
+			string IToolDeployment.Version => _response.Version;
+			ToolDeploymentState IToolDeployment.State => _response.State;
+			double IToolDeployment.Progress => _response.Progress;
+			DateTime? IToolDeployment.StartedAt => _response.StartedAt;
+			TimeSpan IToolDeployment.Duration => _response.Duration;
+			NamespaceId IToolDeployment.NamespaceId => throw new NotImplementedException();
+			RefName IToolDeployment.RefName => _response.RefName;
+
+			IBlobRef<DirectoryNode> IToolDeployment.Content
+				=> _collection.Open(_toolId, _response.Locator);
+
+			public ToolDeployment(HttpToolCollection collection, ToolId toolId, GetToolDeploymentResponse response)
+			{
+				_collection = collection;
+				_toolId = toolId;
+				_response = response;
+			}
+
+			Task<Stream> IToolDeployment.OpenZipStreamAsync(CancellationToken cancellationToken)
+				=> _collection.OpenZipStreamAsync(_toolId, _response.Id, cancellationToken);
+
+			Task<IToolDeployment?> IToolDeployment.UpdateAsync(ToolDeploymentState state, CancellationToken cancellationToken)
+				=> throw new NotImplementedException();
+		}
+
+		readonly IHordeClient _hordeClient;
+
+		/// <summary>
+		/// Constructor
+		/// </summary>
+		public HttpToolCollection(IHordeClient hordeClient)
+			=> _hordeClient = hordeClient;
+
+		/// <inheritdoc/>
+		public async Task<ITool?> GetAsync(ToolId id, CancellationToken cancellationToken = default)
+		{
+			HordeHttpClient hordeHttpClient = _hordeClient.CreateHttpClient();
+			GetToolResponse response = await hordeHttpClient.GetToolAsync(id, cancellationToken);
+			return new Tool(this, response);
+		}
+
+		/// <inheritdoc/>
+		public async Task<IReadOnlyList<ITool>> GetAllAsync(CancellationToken cancellationToken = default)
+		{
+			HordeHttpClient hordeHttpClient = _hordeClient.CreateHttpClient();
+			GetToolsSummaryResponse responses = await hordeHttpClient.GetToolsAsync(cancellationToken);
+
+			List<ITool> tools = new List<ITool>();
+			foreach (GetToolSummaryResponse response in responses.Tools)
+			{
+				ITool? tool = await GetAsync(response.Id, cancellationToken);
+				if (tool != null)
+				{
+					tools.Add(tool);
+				}
+			}
+
+			return tools;
+		}
+
+		IBlobRef<DirectoryNode> Open(ToolId toolId, BlobLocator locator)
+		{
+			IStorageClient store = _hordeClient.CreateStorageClient(toolId);
+			return store.CreateBlobRef<DirectoryNode>(locator);
+		}
+
+		async Task<Stream> OpenZipStreamAsync(ToolId id, ToolDeploymentId deploymentId, CancellationToken cancellationToken = default)
+		{
+			HordeHttpClient hordeHttpClient = _hordeClient.CreateHttpClient();
+			return await hordeHttpClient.GetToolDeploymentZipAsync(id, deploymentId, cancellationToken);
+		}
+	}
+}
diff --git a/Engine/Source/Programs/Shared/EpicGames.Horde/Tools/ToolHttpCollection.cs b/Engine/Source/Programs/Shared/EpicGames.Horde/Tools/ToolHttpCollection.cs
deleted file mode 100644
