Horde: Move construction of http resilience pipeline out of HordeClient.

#rnx

diff --git a/Engine/Source/Programs/Shared/EpicGames.Horde/HordeClient.cs b/Engine/Source/Programs/Shared/EpicGames.Horde/HordeClient.cs
--- a/Engine/Source/Programs/Shared/EpicGames.Horde/HordeClient.cs
+++ b/Engine/Source/Programs/Shared/EpicGames.Horde/HordeClient.cs
@@ -1,7 +1,6 @@
 ﻿// Copyright Epic Games, Inc. All Rights Reserved.
 
 using System;
-using System.Net;
 using System.Net.Http;
 using System.Net.Http.Headers;
 using System.Net.Http.Json;
@@ -22,12 +21,8 @@ using EpicGames.Horde.Tools;
 using Grpc.Core;
 using Grpc.Net.Client;
 using Grpc.Net.Client.Configuration;
-using Microsoft.Extensions.Http;
 using Microsoft.Extensions.Logging;
 using Microsoft.Extensions.Options;
-using Polly;
-using Polly.Retry;
-using Polly.Timeout;
 
 #pragma warning disable CA2234 // Use URIs instead of strings
 
@@ -65,6 +60,11 @@ namespace EpicGames.Horde
 		/// <inheritdoc/>
 		public IToolCollection Tools { get; }
 
+		/// <summary>
+		/// Accessor for the logger instance
+		/// </summary>
+		protected ILogger Logger => _logger;
+
 		/// <summary>
 		/// Constructor
 		/// </summary>
@@ -98,74 +98,6 @@ namespace EpicGames.Horde
 			GC.SuppressFinalize(this);
 		}
 
-		/// <summary>
-		/// Creates a new message handler with default resilience policies
-		/// </summary>
-		protected HttpMessageHandler CreateDefaultHttpMessageHandler()
-		{
-			HttpMessageHandler? httpMessageHandler = null;
-			try
-			{
-#pragma warning disable CA2000 // Call dispose on httpMessageHandler (disposed by child handlers)
-				httpMessageHandler = new SocketsHttpHandler();
-				httpMessageHandler = new PolicyHttpMessageHandler(request => CreateDefaultTimeoutRetryPolicy(request)) { InnerHandler = httpMessageHandler };
-				httpMessageHandler = new PolicyHttpMessageHandler(request => CreateDefaultTransientErrorPolicy(request)) { InnerHandler = httpMessageHandler };
-#pragma warning restore CA2000
-				return httpMessageHandler;
-			}
-			catch
-			{
-				httpMessageHandler?.Dispose();
-				throw;
-			}
-		}
-
-		/// <summary>
-		/// Create a default timeout retry policy
-		/// </summary>
-		protected IAsyncPolicy<HttpResponseMessage> CreateDefaultTimeoutRetryPolicy(HttpRequestMessage request)
-		{
-			// Wait 30 seconds for operations to timeout
-			Task OnTimeoutAsync(Context context, TimeSpan timespan, Task timeoutTask)
-			{
-				_logger.LogWarning(KnownLogEvents.Systemic_Horde_Http, "{Method} {Url} timed out after {Time}s.", request.Method, request.RequestUri, (int)timespan.TotalSeconds);
-				return Task.CompletedTask;
-			}
-
-			AsyncTimeoutPolicy<HttpResponseMessage> timeoutPolicy = Policy.TimeoutAsync<HttpResponseMessage>(30, OnTimeoutAsync);
-
-			// Retry twice after a timeout
-			void OnRetry(Exception ex, TimeSpan timespan)
-			{
-				_logger.LogWarning(KnownLogEvents.Systemic_Horde_Http, ex, "{Method} {Url} retrying after {Time}s.", request.Method, request.RequestUri, timespan.TotalSeconds);
-			}
-
-			TimeSpan[] retryTimes = new[] { TimeSpan.FromSeconds(5.0), TimeSpan.FromSeconds(10.0) };
-			AsyncRetryPolicy retryPolicy = Policy.Handle<TimeoutRejectedException>().WaitAndRetryAsync(retryTimes, OnRetry);
-			return retryPolicy.WrapAsync(timeoutPolicy);
-		}
-
-		/// <summary>
-		/// Create a default timeout retry policy
-		/// </summary>
-		protected IAsyncPolicy<HttpResponseMessage> CreateDefaultTransientErrorPolicy(HttpRequestMessage request)
-		{
-			Task OnTimeoutAsync(DelegateResult<HttpResponseMessage> outcome, TimeSpan timespan, int retryAttempt, Context context)
-			{
-				_logger.LogWarning(KnownLogEvents.Systemic_Horde_Http, "{Method} {Url} failed ({Result}). Delaying for {DelayMs}ms (attempt #{RetryNum}).", request.Method, request.RequestUri, outcome.Result?.StatusCode, timespan.TotalMilliseconds, retryAttempt);
-				return Task.CompletedTask;
-			}
-
-			TimeSpan[] retryTimes = new[] { TimeSpan.FromSeconds(1.0), TimeSpan.FromSeconds(5.0), TimeSpan.FromSeconds(10.0), TimeSpan.FromSeconds(30.0), TimeSpan.FromSeconds(30.0) };
-
-			// Policy for transient errors is the same as HttpPolicyExtensions.HandleTransientHttpError(), but excludes HttpStatusCode.ServiceUnavailable (which is used as a response
-			// when allocating compute resources when none are available). This pathway is handled explicitly on the application side.
-			return Policy<HttpResponseMessage>
-				.Handle<HttpRequestException>()
-				.OrResult(x => (x.StatusCode >= HttpStatusCode.InternalServerError && x.StatusCode != HttpStatusCode.ServiceUnavailable) || x.StatusCode == HttpStatusCode.RequestTimeout)
-				.WaitAndRetryAsync(retryTimes, OnTimeoutAsync);
-		}
-
 		/// <inheritdoc/>
 		public abstract Task<bool> LoginAsync(bool allowLogin, CancellationToken cancellationToken);
 
@@ -340,7 +272,7 @@ namespace EpicGames.Horde
 			: base(serverUrl, bundleCache, hordeOptions, loggerFactory)
 		{
 			_accessToken = accessToken;
-			_httpMessageHandler = CreateDefaultHttpMessageHandler();
+			_httpMessageHandler = HordeHttpMessageHandler.Create(Logger);
 		}
 
 		/// <inheritdoc/>
@@ -398,7 +330,7 @@ namespace EpicGames.Horde
 		public HordeClientWithDynamicCredentials(Uri serverUrl, BundleCache bundleCache, IOptionsSnapshot<HordeOptions> hordeOptions, ILoggerFactory loggerFactory)
 			: base(serverUrl, bundleCache, hordeOptions, loggerFactory)
 		{
-			_baseHttpMessageHandler = CreateDefaultHttpMessageHandler();
+			_baseHttpMessageHandler = HordeHttpMessageHandler.Create(Logger);
 
 			_authHandlerState = new HordeHttpAuthHandlerState(_baseHttpMessageHandler, serverUrl, hordeOptions, loggerFactory.CreateLogger<HordeHttpAuthHandlerState>());
 			_authHttpMessageHandler = new HordeHttpAuthHandler(_baseHttpMessageHandler, _authHandlerState, hordeOptions);
diff --git a/Engine/Source/Programs/Shared/EpicGames.Horde/HordeHttpMessageHandler.cs b/Engine/Source/Programs/Shared/EpicGames.Horde/HordeHttpMessageHandler.cs
new file mode 100644
--- /dev/null
+++ b/Engine/Source/Programs/Shared/EpicGames.Horde/HordeHttpMessageHandler.cs
@@ -0,0 +1,89 @@
+// Copyright Epic Games, Inc. All Rights Reserved.
+
+using System;
+using System.Net;
+using System.Net.Http;
+using System.Threading.Tasks;
+using EpicGames.Core;
+using Microsoft.Extensions.Http;
+using Microsoft.Extensions.Logging;
+using Polly;
+using Polly.Retry;
+using Polly.Timeout;
+
+namespace EpicGames.Horde
+{
+	/// <summary>
+	/// Helper class to create <see cref="HttpMessageHandler"/> instances for Horde with default resiliance pipeline.
+	/// </summary>
+	static class HordeHttpMessageHandler
+	{
+		/// <summary>
+		/// Creates a new message handler with default resilience policies
+		/// </summary>
+		public static HttpMessageHandler Create(ILogger logger)
+		{
+			HttpMessageHandler? httpMessageHandler = null;
+			try
+			{
+#pragma warning disable CA2000 // Call dispose on httpMessageHandler (disposed by child handlers)
+				httpMessageHandler = new SocketsHttpHandler();
+				httpMessageHandler = new PolicyHttpMessageHandler(request => CreateDefaultTimeoutRetryPolicy(request, logger)) { InnerHandler = httpMessageHandler };
+				httpMessageHandler = new PolicyHttpMessageHandler(request => CreateDefaultTransientErrorPolicy(request, logger)) { InnerHandler = httpMessageHandler };
+#pragma warning restore CA2000
+				return httpMessageHandler;
+			}
+			catch
+			{
+				httpMessageHandler?.Dispose();
+				throw;
+			}
+		}
+
+		/// <summary>
+		/// Create a default timeout retry policy
+		/// </summary>
+		static IAsyncPolicy<HttpResponseMessage> CreateDefaultTimeoutRetryPolicy(HttpRequestMessage request, ILogger logger)
+		{
+			// Wait 30 seconds for operations to timeout
+			Task OnTimeoutAsync(Context context, TimeSpan timespan, Task timeoutTask)
+			{
+				logger.LogWarning(KnownLogEvents.Systemic_Horde_Http, "{Method} {Url} timed out after {Time}s.", request.Method, request.RequestUri, (int)timespan.TotalSeconds);
+				return Task.CompletedTask;
+			}
+
+			AsyncTimeoutPolicy<HttpResponseMessage> timeoutPolicy = Policy.TimeoutAsync<HttpResponseMessage>(30, OnTimeoutAsync);
+
+			// Retry twice after a timeout
+			void OnRetry(Exception ex, TimeSpan timespan)
+			{
+				logger.LogWarning(KnownLogEvents.Systemic_Horde_Http, ex, "{Method} {Url} retrying after {Time}s.", request.Method, request.RequestUri, timespan.TotalSeconds);
+			}
+
+			TimeSpan[] retryTimes = new[] { TimeSpan.FromSeconds(5.0), TimeSpan.FromSeconds(10.0) };
+			AsyncRetryPolicy retryPolicy = Policy.Handle<TimeoutRejectedException>().WaitAndRetryAsync(retryTimes, OnRetry);
+			return retryPolicy.WrapAsync(timeoutPolicy);
+		}
+
+		/// <summary>
+		/// Create a default timeout retry policy
+		/// </summary>
+		static IAsyncPolicy<HttpResponseMessage> CreateDefaultTransientErrorPolicy(HttpRequestMessage request, ILogger logger)
+		{
+			Task OnTimeoutAsync(DelegateResult<HttpResponseMessage> outcome, TimeSpan timespan, int retryAttempt, Context context)
+			{
+				logger.LogWarning(KnownLogEvents.Systemic_Horde_Http, "{Method} {Url} failed ({Result}). Delaying for {DelayMs}ms (attempt #{RetryNum}).", request.Method, request.RequestUri, outcome.Result?.StatusCode, timespan.TotalMilliseconds, retryAttempt);
+				return Task.CompletedTask;
+			}
+
+			TimeSpan[] retryTimes = new[] { TimeSpan.FromSeconds(1.0), TimeSpan.FromSeconds(5.0), TimeSpan.FromSeconds(10.0), TimeSpan.FromSeconds(30.0), TimeSpan.FromSeconds(30.0) };
+
+			// Policy for transient errors is the same as HttpPolicyExtensions.HandleTransientHttpError(), but excludes HttpStatusCode.ServiceUnavailable (which is used as a response
+			// when allocating compute resources when none are available). This pathway is handled explicitly on the application side.
+			return Policy<HttpResponseMessage>
+				.Handle<HttpRequestException>()
+				.OrResult(x => (x.StatusCode >= HttpStatusCode.InternalServerError && x.StatusCode != HttpStatusCode.ServiceUnavailable) || x.StatusCode == HttpStatusCode.RequestTimeout)
+				.WaitAndRetryAsync(retryTimes, OnTimeoutAsync);
+		}
+	}
+}
