Horde: Lazily construct compute client implementations, to avoid null reference exception when called from constructor.

#rnx

diff --git a/Engine/Source/Programs/Shared/EpicGames.Horde/HordeClient.cs b/Engine/Source/Programs/Shared/EpicGames.Horde/HordeClient.cs
--- a/Engine/Source/Programs/Shared/EpicGames.Horde/HordeClient.cs
+++ b/Engine/Source/Programs/Shared/EpicGames.Horde/HordeClient.cs
@@ -39,7 +39,7 @@ namespace EpicGames.Horde
 		readonly ILoggerFactory _loggerFactory;
 		readonly ILogger _logger;
 
-		readonly ServerComputeClient _serverComputeClient;
+		ServerComputeClient? _serverComputeClient;
 		BackgroundTask<GrpcChannel>? _grpcChannel;
 
 		/// <inheritdoc/>
@@ -49,7 +49,7 @@ namespace EpicGames.Horde
 		public IArtifactCollection Artifacts { get; }
 
 		/// <inheritdoc/>
-		public IComputeClient Compute => _serverComputeClient;
+		public IComputeClient Compute => _serverComputeClient ??= CreateComputeClient();
 
 		/// <inheritdoc/>
 		public IProjectCollection Projects { get; }
@@ -76,8 +76,6 @@ namespace EpicGames.Horde
 			_loggerFactory = loggerFactory;
 			_logger = _loggerFactory.CreateLogger<HordeClient>();
 
-			_serverComputeClient = CreateComputeClient();
-
 			Artifacts = new HttpArtifactCollection(this);
 			Projects = new HttpProjectCollection(this);
 			Secrets = new HttpSecretCollection(this);
@@ -87,7 +85,7 @@ namespace EpicGames.Horde
 		/// <inheritdoc/>
 		public virtual async ValueTask DisposeAsync()
 		{
-			_serverComputeClient.Dispose();
+			_serverComputeClient?.Dispose();
 
 			if (_grpcChannel != null)
 			{
