Horde: Rename IStorageClient -> IStorageNamespace.

Renaming of files (IStorageClient.cs, etc..) will be done in a subsequent change to ease diffs.

In a further subsequent change IStorageClient will be reintroduced as an accessor for the storage system in general, that can be used to obtain IStorageNamespace instances for different namespaces.

#rnx

diff --git a/Engine/Source/Programs/AutomationTool/BuildGraph/Tasks/CreateArtifactTask.cs b/Engine/Source/Programs/AutomationTool/BuildGraph/Tasks/CreateArtifactTask.cs
--- a/Engine/Source/Programs/AutomationTool/BuildGraph/Tasks/CreateArtifactTask.cs
+++ b/Engine/Source/Programs/AutomationTool/BuildGraph/Tasks/CreateArtifactTask.cs
@@ -191,7 +191,7 @@ namespace AutomationTool.Tasks
 			Logger.LogInformation("Creating artifact {ArtifactId} '{ArtifactName}' ({ArtifactType}) with namespace {NamespaceId}, ref {RefName} ({Link})", artifact.Id, name, type, artifact.NamespaceId, artifact.RefName, $"{hordeClient.ServerUrl}/api/v1/storage/{artifact.NamespaceId}/refs/{artifact.RefName}");
 
 			// Upload the files
-			IStorageClient storage = hordeClient.CreateStorageClient(artifact.NamespaceId);
+			IStorageNamespace storage = hordeClient.GetStorageNamespace(artifact.NamespaceId);
 			Stopwatch timer = Stopwatch.StartNew();
 
 			IHashedBlobRef<DirectoryNode> rootRef;
diff --git a/Engine/Source/Programs/AutomationTool/BuildGraph/Tasks/DeployToolTask.cs b/Engine/Source/Programs/AutomationTool/BuildGraph/Tasks/DeployToolTask.cs
--- a/Engine/Source/Programs/AutomationTool/BuildGraph/Tasks/DeployToolTask.cs
+++ b/Engine/Source/Programs/AutomationTool/BuildGraph/Tasks/DeployToolTask.cs
@@ -161,8 +161,8 @@ namespace AutomationTool.Tasks
 
 			IHashedBlobRef handle;
 
-			IStorageClient storageClient = hordeClient.CreateStorageClient(toolId);
-			await using (IBlobWriter blobWriter = storageClient.CreateBlobWriter(serializerOptions: serializerOptions))
+			IStorageNamespace storageNamespace = hordeClient.GetStorageNamespace(toolId);
+			await using (IBlobWriter blobWriter = storageNamespace.CreateBlobWriter(serializerOptions: serializerOptions))
 			{
 				DirectoryNode sandbox = new DirectoryNode();
 				if (_parameters.File != null)
diff --git a/Engine/Source/Programs/Horde/Drivers/JobDriver.Tests/Execution/WorkspaceExecutorTest.cs b/Engine/Source/Programs/Horde/Drivers/JobDriver.Tests/Execution/WorkspaceExecutorTest.cs
--- a/Engine/Source/Programs/Horde/Drivers/JobDriver.Tests/Execution/WorkspaceExecutorTest.cs
+++ b/Engine/Source/Programs/Horde/Drivers/JobDriver.Tests/Execution/WorkspaceExecutorTest.cs
@@ -113,7 +113,7 @@ namespace JobDriver.Tests.Execution
 			public IServerLogger CreateServerLogger(LogId logId, LogLevel minimumLevel = LogLevel.Information)
 				=> throw new NotImplementedException();
 
-			public IStorageClient CreateStorageClient(string relativePath, string? accessToken = null)
+			public IStorageNamespace GetStorageNamespace(string relativePath, string? accessToken = null)
 				=> throw new NotImplementedException();
 
 			public Task<string?> GetAccessTokenAsync(bool interactive, CancellationToken cancellationToken = default)
diff --git a/Engine/Source/Programs/Horde/Drivers/JobDriver.Tests/WorkerServiceTest.cs b/Engine/Source/Programs/Horde/Drivers/JobDriver.Tests/WorkerServiceTest.cs
--- a/Engine/Source/Programs/Horde/Drivers/JobDriver.Tests/WorkerServiceTest.cs
+++ b/Engine/Source/Programs/Horde/Drivers/JobDriver.Tests/WorkerServiceTest.cs
@@ -61,7 +61,7 @@ namespace JobDriver.Tests
 			_serviceCollection.AddSingleton<BundleCache>();
 			_serviceCollection.AddSingleton<StorageBackendCache>();
 			_serviceCollection.AddSingleton<HttpStorageBackendFactory>();
-			_serviceCollection.AddSingleton<HttpStorageClientFactory>();
+			_serviceCollection.AddSingleton<HttpStorageNamespaceFactory>();
 
 			_serviceCollection.Configure<DriverSettings>(settings =>
 			{
diff --git a/Engine/Source/Programs/Horde/Drivers/JobDriver/Execution/JobExecutor.cs b/Engine/Source/Programs/Horde/Drivers/JobDriver/Execution/JobExecutor.cs
--- a/Engine/Source/Programs/Horde/Drivers/JobDriver/Execution/JobExecutor.cs
+++ b/Engine/Source/Programs/Horde/Drivers/JobDriver/Execution/JobExecutor.cs
@@ -747,9 +747,9 @@ namespace JobDriver.Execution
 		const string SetupStepName = "setup";
 		const string BuildGraphTempStorageDir = "BuildGraph";
 
-		IStorageClient CreateStorageClient(NamespaceId namespaceId, string? token)
+		IStorageNamespace GetStorageNamespace(NamespaceId namespaceId, string? token)
 		{
-			return HordeClient.CreateStorageClient(namespaceId, token);
+			return HordeClient.GetStorageNamespace(namespaceId, token);
 		}
 
 		protected virtual async Task<bool> SetupAsync(JobStepInfo step, DirectoryReference workspaceDir, bool? useP4, ILogger logger, CancellationToken cancellationToken)
@@ -823,7 +823,7 @@ namespace JobDriver.Execution
 				logger.LogInformation("Creating output artifact {ArtifactId} '{ArtifactName}' ({ArtifactType}) with ref {RefName} in namespace {NamespaceId}", artifactId, artifactName, artifactType, artifact.RefName, artifact.NamespaceId);
 
 				// Write the data
-				IStorageClient storage = CreateStorageClient(new NamespaceId(artifact.NamespaceId), artifact.Token);
+				IStorageNamespace storage = GetStorageNamespace(new NamespaceId(artifact.NamespaceId), artifact.Token);
 
 				Stopwatch timer = Stopwatch.StartNew();
 
@@ -1080,7 +1080,7 @@ namespace JobDriver.Execution
 
 				logger.LogInformation("Reading preprocessed script from {NamespaceId}:{RefName}", namespaceId, refName);
 
-				IStorageClient storage = CreateStorageClient(namespaceId, artifact.Token);
+				IStorageNamespace storage = GetStorageNamespace(namespaceId, artifact.Token);
 
 				DirectoryNode node = await storage.ReadRefTargetAsync<DirectoryNode>(refName, cancellationToken: cancellationToken);
 				DirectoryNode? buildGraphDir = await node.TryOpenDirectoryAsync(BuildGraphTempStorageDir, cancellationToken: cancellationToken);
@@ -1141,7 +1141,7 @@ namespace JobDriver.Execution
 				ArtifactId artifactId = ArtifactId.Parse(artifact.Id);
 				Logger.LogInformation("Created artifact {ArtifactId} '{ArtifactName}' ({ArtifactType}) with ref {RefName} ({Link})", artifactId, name, type, artifact.RefName, $"{HordeClient.ServerUrl}/api/v1/storage/{artifact.NamespaceId}/refs/{artifact.RefName}");
 
-				IStorageClient storage = CreateStorageClient(new NamespaceId(artifact.NamespaceId), artifact.Token);
+				IStorageNamespace storage = GetStorageNamespace(new NamespaceId(artifact.NamespaceId), artifact.Token);
 
 				IHashedBlobRef<DirectoryNode> rootRef;
 				await using (IBlobWriter blobWriter = storage.CreateBlobWriter(new RefName(artifact.RefName)))
@@ -1343,7 +1343,7 @@ namespace JobDriver.Execution
 				ArtifactId artifactId = ArtifactId.Parse(artifact.Id);
 				logger.LogInformation("Created artifact {ArtifactId} '{ArtifactName}' ({ArtifactType}) with ref {RefName} ({RefUrl})", artifactId, artifactRequest.Name, ArtifactType.StepOutput, artifact.RefName, $"{HordeClient.ServerUrl.ToString().TrimEnd('/')}/api/v1/storage/{artifact.NamespaceId}/refs/{artifact.RefName}");
 
-				IStorageClient storage = CreateStorageClient(new NamespaceId(artifact.NamespaceId), artifact.Token);
+				IStorageNamespace storage = GetStorageNamespace(new NamespaceId(artifact.NamespaceId), artifact.Token);
 
 				// Upload the data
 				Stopwatch timer = Stopwatch.StartNew();
@@ -1423,7 +1423,7 @@ namespace JobDriver.Execution
 					ArtifactId artifactId = ArtifactId.Parse(artifact.Id);
 					logger.LogInformation("Created artifact {ArtifactId} '{ArtifactName}' ({ArtifactType}) with ref {RefName} ({RefUrl})", artifactId, artifactRequest.Name, ArtifactType.StepOutput, artifact.RefName, $"{HordeClient.ServerUrl.ToString().TrimEnd('/')}/api/v1/storage/{artifact.NamespaceId}/refs/{artifact.RefName}");
 
-					IStorageClient storage = CreateStorageClient(new NamespaceId(artifact.NamespaceId), artifact.Token);
+					IStorageNamespace storage = GetStorageNamespace(new NamespaceId(artifact.NamespaceId), artifact.Token);
 
 					// Upload the data
 					Stopwatch timer = Stopwatch.StartNew();
diff --git a/Engine/Source/Programs/Horde/Drivers/JobDriver/Utility/TempStorage.cs b/Engine/Source/Programs/Horde/Drivers/JobDriver/Utility/TempStorage.cs
--- a/Engine/Source/Programs/Horde/Drivers/JobDriver/Utility/TempStorage.cs
+++ b/Engine/Source/Programs/Horde/Drivers/JobDriver/Utility/TempStorage.cs
@@ -632,8 +632,8 @@ namespace JobDriver.Utility
 
 				logger.LogInformation("Reading tag \"{NodeName}\":\"{TagName}\" from temp storage (artifact: {ArtifactId} '{ArtifactName}' ({ArtifactType}), ns: {NamespaceId}, ref: {RefName}, localFile: {LocalFile})", nodeName, tagName, artifact.Id, artifactName, artifactType, namespaceId, refName, localFileListLocation);
 
-				IStorageClient storageClient = hordeClient.CreateStorageClient(namespaceId, artifact.Token);
-				DirectoryNode node = await storageClient.ReadRefTargetAsync<DirectoryNode>(artifact.RefName, cancellationToken: cancellationToken);
+				IStorageNamespace storageNamespace = hordeClient.GetStorageNamespace(namespaceId, artifact.Token);
+				DirectoryNode node = await storageNamespace.ReadRefTargetAsync<DirectoryNode>(artifact.RefName, cancellationToken: cancellationToken);
 
 				FileEntry fileEntry = node.GetFileEntry(localFileListLocation.GetFileName());
 				DirectoryReference.CreateDirectory(localFileListLocation.Directory);
@@ -770,8 +770,8 @@ namespace JobDriver.Utility
 
 				logger.LogInformation("Reading block \"{NodeName}\":\"{BlockName}\" from temp storage (artifact: {ArtifactId} '{ArtifactName}' ({ArtifactType}), ns: {NamespaceId}, ref: {RefName}, local: {LocalFile}, blockdir: {BlockDir})", nodeName, blockName, artifact.Id, artifactName, artifactType, namespaceId, refName, localManifestFile, blockDirectoryName);
 
-				IStorageClient storageClient = hordeClient.CreateStorageClient(namespaceId, artifact.Token);
-				DirectoryNode node = await storageClient.ReadRefTargetAsync<DirectoryNode>(refName, cancellationToken: cancellationToken);
+				IStorageNamespace storageNamespace = hordeClient.GetStorageNamespace(namespaceId, artifact.Token);
+				DirectoryNode node = await storageNamespace.ReadRefTargetAsync<DirectoryNode>(refName, cancellationToken: cancellationToken);
 
 				DirectoryEntry? rootDirEntry;
 				if (!node.TryGetDirectoryEntry(blockDirectoryName, out rootDirEntry))
@@ -779,13 +779,13 @@ namespace JobDriver.Utility
 					throw new TempStorageException($"Missing block \"{blockName}\" from node \"{nodeName}\"");
 				}
 
-				StorageStats initialStats = storageClient.GetStats();
+				StorageStats initialStats = storageNamespace.GetStats();
 				Stopwatch timer = Stopwatch.StartNew();
 
 				// Add all the files and flush the ref
 				await rootDirEntry.Handle.ExtractAsync(rootDir.ToDirectoryInfo(), new ExtractStatsLogger(logger), logger, cancellationToken);
 
-				StorageStats deltaStats = StorageStats.GetDelta(initialStats, storageClient.GetStats());
+				StorageStats deltaStats = StorageStats.GetDelta(initialStats, storageNamespace.GetStats());
 				logger.LogInformation("{Stats}", $"Elapsed: {(int)timer.Elapsed.TotalSeconds}s, {String.Join(", ", deltaStats.Values.Select(x => $"{x.Item1}: {x.Item2:n0}"))}");
 
 				// Read the manifest in
diff --git a/Engine/Source/Programs/Horde/HordeCmd/Commands/Artifacts/ArtifactDownload.cs b/Engine/Source/Programs/Horde/HordeCmd/Commands/Artifacts/ArtifactDownload.cs
--- a/Engine/Source/Programs/Horde/HordeCmd/Commands/Artifacts/ArtifactDownload.cs
+++ b/Engine/Source/Programs/Horde/HordeCmd/Commands/Artifacts/ArtifactDownload.cs
@@ -56,7 +56,7 @@ namespace Horde.Commands.Artifacts
 
 			Stopwatch timer = Stopwatch.StartNew();
 
-			IStorageClient store = _hordeClient.CreateStorageClient(artifact.Id);
+			IStorageNamespace store = _hordeClient.GetStorageNamespace(artifact.Id);
 			IBlobRef<DirectoryNode> handle = await store.ReadRefAsync<DirectoryNode>(new RefName("default"));
 			await handle.ExtractAsync(OutputDir.ToDirectoryInfo(), new ExtractStatsLogger(logger), logger, CancellationToken.None);
 
diff --git a/Engine/Source/Programs/Horde/HordeCmd/Commands/Bundles/BundleCreate.cs b/Engine/Source/Programs/Horde/HordeCmd/Commands/Bundles/BundleCreate.cs
--- a/Engine/Source/Programs/Horde/HordeCmd/Commands/Bundles/BundleCreate.cs
+++ b/Engine/Source/Programs/Horde/HordeCmd/Commands/Bundles/BundleCreate.cs
@@ -34,8 +34,8 @@ namespace Horde.Commands.Bundles
 		[Description("Clean the output folder before writing any data")]
 		public bool CleanOutput { get; set; }
 
-		public BundleCreate(HttpStorageClientFactory storageClientFactory, BundleCache bundleCache)
-			: base(storageClientFactory, bundleCache)
+		public BundleCreate(HttpStorageNamespaceFactory storageNamespaceFactory, BundleCache bundleCache)
+			: base(storageNamespaceFactory, bundleCache)
 		{
 		}
 
@@ -50,12 +50,12 @@ namespace Horde.Commands.Bundles
 			if (File != null)
 			{
 				using MemoryMappedFileCache memoryMappedFileCache = new MemoryMappedFileCache();
-				IStorageClient store = BundleStorageClient.CreateFromDirectory(File.Directory, BundleCache, memoryMappedFileCache, logger);
+				IStorageNamespace store = BundleStorageNamespace.CreateFromDirectory(File.Directory, BundleCache, memoryMappedFileCache, logger);
 				return await ExecuteInternalAsync(store, logger);
 			}
 			else if (!String.IsNullOrEmpty(Ref))
 			{
-				IStorageClient store = CreateStorageClient();
+				IStorageNamespace store = GetStorageNamespace();
 				return await ExecuteInternalAsync(store, logger);
 			}
 			else
@@ -64,7 +64,7 @@ namespace Horde.Commands.Bundles
 			}
 		}
 
-		async Task<int> ExecuteInternalAsync(IStorageClient store, ILogger logger)
+		async Task<int> ExecuteInternalAsync(IStorageNamespace store, ILogger logger)
 		{
 			// Gather the input files
 			DirectoryReference baseDir;
diff --git a/Engine/Source/Programs/Horde/HordeCmd/Commands/Bundles/BundleExtract.cs b/Engine/Source/Programs/Horde/HordeCmd/Commands/Bundles/BundleExtract.cs
--- a/Engine/Source/Programs/Horde/HordeCmd/Commands/Bundles/BundleExtract.cs
+++ b/Engine/Source/Programs/Horde/HordeCmd/Commands/Bundles/BundleExtract.cs
@@ -20,7 +20,7 @@ namespace Horde.Commands.Bundles
 		public FileReference? File { get; set; }
 
 		[CommandLine("-Ref=")]
-		[Description("Name of a ref to read from the default storage client. -File=..., -Ref=..., or -Node=... must be specified.")]
+		[Description("Name of a ref to read from the default storage namespace. -File=..., -Ref=..., or -Node=... must be specified.")]
 		public string? Ref { get; set; }
 
 		[CommandLine("-Node=")]
@@ -39,8 +39,8 @@ namespace Horde.Commands.Bundles
 		[Description("If set, deletes the contents of the output directory before extraction.")]
 		public bool CleanOutput { get; set; }
 
-		public BundleExtract(HttpStorageClientFactory storageClientFactory, BundleCache bundleCache)
-			: base(storageClientFactory, bundleCache)
+		public BundleExtract(HttpStorageNamespaceFactory storageNamespaceFactory, BundleCache bundleCache)
+			: base(storageNamespaceFactory, bundleCache)
 		{
 		}
 
@@ -55,19 +55,19 @@ namespace Horde.Commands.Bundles
 			if (File != null)
 			{
 				using MemoryMappedFileCache memoryMappedFileCache = new MemoryMappedFileCache();
-				IStorageClient store = BundleStorageClient.CreateFromDirectory(File.Directory, BundleCache, memoryMappedFileCache, logger);
+				IStorageNamespace store = BundleStorageNamespace.CreateFromDirectory(File.Directory, BundleCache, memoryMappedFileCache, logger);
 				IBlobRef<DirectoryNode> handle = store.CreateBlobRef<DirectoryNode>(await FileStorageBackend.ReadRefAsync(File));
 				await ExecuteInternalAsync(store, handle, logger);
 			}
 			else if (Ref != null)
 			{
-				IStorageClient store = CreateStorageClient();
+				IStorageNamespace store = GetStorageNamespace();
 				IBlobRef<DirectoryNode> handle = await store.ReadRefAsync<DirectoryNode>(new RefName(Ref));
 				await ExecuteInternalAsync(store, handle, logger);
 			}
 			else if (Node != null)
 			{
-				IStorageClient store = CreateStorageClient();
+				IStorageNamespace store = GetStorageNamespace();
 				IBlobRef<DirectoryNode> handle = store.CreateBlobRef<DirectoryNode>(new BlobLocator(Node));
 				await ExecuteInternalAsync(store, handle, logger);
 			}
@@ -79,7 +79,7 @@ namespace Horde.Commands.Bundles
 			return 0;
 		}
 
-		protected async Task ExecuteInternalAsync(IStorageClient store, IBlobRef<DirectoryNode> handle, ILogger logger)
+		protected async Task ExecuteInternalAsync(IStorageNamespace store, IBlobRef<DirectoryNode> handle, ILogger logger)
 		{
 			Stopwatch timer = Stopwatch.StartNew();
 
diff --git a/Engine/Source/Programs/Horde/HordeCmd/Commands/Compute/ComputeRun.cs b/Engine/Source/Programs/Horde/HordeCmd/Commands/Compute/ComputeRun.cs
--- a/Engine/Source/Programs/Horde/HordeCmd/Commands/Compute/ComputeRun.cs
+++ b/Engine/Source/Programs/Horde/HordeCmd/Commands/Compute/ComputeRun.cs
@@ -131,7 +131,7 @@ namespace Horde.Commands.Compute
 			JsonComputeTask task = JsonSerializer.Deserialize<JsonComputeTask>(data, new JsonSerializerOptions { AllowTrailingCommas = true, PropertyNameCaseInsensitive = true, PropertyNamingPolicy = JsonNamingPolicy.CamelCase })!;
 
 			// Create a sandbox from the data to be uploaded
-			BundleStorageClient storage = BundleStorageClient.CreateInMemory(logger);
+			BundleStorageNamespace storage = BundleStorageNamespace.CreateInMemory(logger);
 			BlobLocator sandbox = await CreateSandboxAsync(TaskFile, storage, cancellationToken);
 
 			// Open a socket and upload the sandbox
@@ -157,7 +157,7 @@ namespace Horde.Commands.Compute
 			return true;
 		}
 
-		static async Task<BlobLocator> CreateSandboxAsync(FileReference taskFile, IStorageClient storage, CancellationToken cancellationToken)
+		static async Task<BlobLocator> CreateSandboxAsync(FileReference taskFile, IStorageNamespace storage, CancellationToken cancellationToken)
 		{
 			await using IBlobWriter writer = storage.CreateBlobWriter();
 
diff --git a/Engine/Source/Programs/Horde/HordeCmd/Commands/StorageCommandBase.cs b/Engine/Source/Programs/Horde/HordeCmd/Commands/StorageCommandBase.cs
--- a/Engine/Source/Programs/Horde/HordeCmd/Commands/StorageCommandBase.cs
+++ b/Engine/Source/Programs/Horde/HordeCmd/Commands/StorageCommandBase.cs
@@ -9,7 +9,7 @@ using EpicGames.Horde.Storage.Clients;
 namespace Horde.Commands
 {
 	/// <summary>
-	/// Base class for commands that require a configured storage client
+	/// Base class for commands that require a configured storage namespace
 	/// </summary>
 	abstract class StorageCommandBase : Command
 	{
@@ -32,14 +32,14 @@ namespace Horde.Commands
 		/// </summary>
 		public BundleCache BundleCache { get; }
 
-		readonly HttpStorageClientFactory _storageClientFactory;
+		readonly HttpStorageNamespaceFactory _storageNamespaceFactory;
 
 		/// <summary>
 		/// Constructor
 		/// </summary>
-		public StorageCommandBase(HttpStorageClientFactory storageClientFactory, BundleCache bundleCache)
+		public StorageCommandBase(HttpStorageNamespaceFactory storageNamespaceFactory, BundleCache bundleCache)
 		{
-			_storageClientFactory = storageClientFactory;
+			_storageNamespaceFactory = storageNamespaceFactory;
 
 			BundleCache = bundleCache;
 		}
@@ -47,15 +47,15 @@ namespace Horde.Commands
 		/// <summary>
 		/// Creates a new client instance
 		/// </summary>
-		public IStorageClient CreateStorageClient()
+		public IStorageNamespace GetStorageNamespace()
 		{
 			if (String.IsNullOrEmpty(Path))
 			{
-				return _storageClientFactory.CreateClient(Namespace);
+				return _storageNamespaceFactory.CreateClient(Namespace);
 			}
 			else
 			{
-				return _storageClientFactory.CreateClientWithPath(Path);
+				return _storageNamespaceFactory.CreateClientWithPath(Path);
 			}
 		}
 	}
diff --git a/Engine/Source/Programs/Horde/HordeCmd/Commands/Tools/ToolUpload.cs b/Engine/Source/Programs/Horde/HordeCmd/Commands/Tools/ToolUpload.cs
--- a/Engine/Source/Programs/Horde/HordeCmd/Commands/Tools/ToolUpload.cs
+++ b/Engine/Source/Programs/Horde/HordeCmd/Commands/Tools/ToolUpload.cs
@@ -27,21 +27,21 @@ namespace Horde.Commands
 		public DirectoryReference InputDir { get; set; } = null!;
 
 		readonly HordeHttpClient _hordeHttpClient;
-		readonly HttpStorageClientFactory _httpStorageClientFactory;
+		readonly HttpStorageNamespaceFactory _httpStorageNamespaceFactory;
 
-		public ToolUpload(HordeHttpClient httpClient, HttpStorageClientFactory httpStorageClientFactory)
+		public ToolUpload(HordeHttpClient httpClient, HttpStorageNamespaceFactory httpStorageNamespaceFactory)
 		{
 			_hordeHttpClient = httpClient;
-			_httpStorageClientFactory = httpStorageClientFactory;
+			_httpStorageNamespaceFactory = httpStorageNamespaceFactory;
 		}
 
 		/// <inheritdoc/>
 		public override async Task<int> ExecuteAsync(ILogger logger)
 		{
-			IStorageClient storageClient = _httpStorageClientFactory.CreateClientWithPath($"api/v1/tools/{ToolId}");
+			IStorageNamespace storageNamespace = _httpStorageNamespaceFactory.CreateClientWithPath($"api/v1/tools/{ToolId}");
 
 			IHashedBlobRef<DirectoryNode> target;
-			await using (IBlobWriter writer = storageClient.CreateBlobWriter())
+			await using (IBlobWriter writer = storageNamespace.CreateBlobWriter())
 			{
 				target = await writer.WriteFilesAsync(InputDir);
 			}
diff --git a/Engine/Source/Programs/Horde/HordeCmd/Commands/Vcs/VcsBase.cs b/Engine/Source/Programs/Horde/HordeCmd/Commands/Vcs/VcsBase.cs
--- a/Engine/Source/Programs/Horde/HordeCmd/Commands/Vcs/VcsBase.cs
+++ b/Engine/Source/Programs/Horde/HordeCmd/Commands/Vcs/VcsBase.cs
@@ -149,14 +149,14 @@ namespace Horde.Commands.Vcs
 		const string DataDir = ".horde";
 		const string StateFileName = "state.dat";
 
-		readonly IStorageClientFactory _storageClientFactory;
+		readonly IStorageNamespaceFactory _storageNamespaceFactory;
 
-		protected VcsBase(IStorageClientFactory storageClientFactory)
+		protected VcsBase(IStorageNamespaceFactory storageNamespaceFactory)
 		{
-			_storageClientFactory = storageClientFactory;
+			_storageNamespaceFactory = storageNamespaceFactory;
 		}
 
-		protected IStorageClient CreateStorageClient() => _storageClientFactory.CreateClient(NamespaceId);
+		protected IStorageNamespace GetStorageNamespace() => _storageNamespaceFactory.GetNamespace(NamespaceId);
 
 		protected static async Task<WorkspaceState> ReadStateAsync(DirectoryReference rootDir)
 		{
@@ -340,9 +340,9 @@ namespace Horde.Commands.Vcs
 			}
 		}
 
-		protected static async Task<CommitNode?> GetCommitAsync(IStorageClient storageClient, RefName branchName, int change = 0)
+		protected static async Task<CommitNode?> GetCommitAsync(IStorageNamespace storageNamespace, RefName branchName, int change = 0)
 		{
-			CommitNode tip = await storageClient.ReadRefTargetAsync<CommitNode>(branchName);
+			CommitNode tip = await storageNamespace.ReadRefTargetAsync<CommitNode>(branchName);
 			if (change != 0)
 			{
 				while (tip.Number != change)
diff --git a/Engine/Source/Programs/Horde/HordeCmd/Commands/Vcs/VcsBranch.cs b/Engine/Source/Programs/Horde/HordeCmd/Commands/Vcs/VcsBranch.cs
--- a/Engine/Source/Programs/Horde/HordeCmd/Commands/Vcs/VcsBranch.cs
+++ b/Engine/Source/Programs/Horde/HordeCmd/Commands/Vcs/VcsBranch.cs
@@ -12,8 +12,8 @@ namespace Horde.Commands.Vcs
 		[CommandLine(Prefix = "-Name=", Required = true)]
 		public string Name { get; set; } = "";
 
-		public VcsBranch(IStorageClientFactory storageClientFactory)
-			: base(storageClientFactory)
+		public VcsBranch(IStorageNamespaceFactory storageNamespaceFactory)
+			: base(storageNamespaceFactory)
 		{
 		}
 
@@ -23,7 +23,7 @@ namespace Horde.Commands.Vcs
 
 			WorkspaceState workspaceState = await ReadStateAsync(rootDir);
 
-			IStorageClient store = CreateStorageClient();
+			IStorageNamespace store = GetStorageNamespace();
 
 			RefName branchName = new RefName(Name);
 			if (await store.RefExistsAsync(branchName))
diff --git a/Engine/Source/Programs/Horde/HordeCmd/Commands/Vcs/VcsCheckout.cs b/Engine/Source/Programs/Horde/HordeCmd/Commands/Vcs/VcsCheckout.cs
--- a/Engine/Source/Programs/Horde/HordeCmd/Commands/Vcs/VcsCheckout.cs
+++ b/Engine/Source/Programs/Horde/HordeCmd/Commands/Vcs/VcsCheckout.cs
@@ -22,8 +22,8 @@ namespace Horde.Commands.Vcs
 		[CommandLine("-Force")]
 		public bool Force { get; set; }
 
-		public VcsCheckout(IStorageClientFactory storageClientFactory)
-			: base(storageClientFactory)
+		public VcsCheckout(IStorageNamespaceFactory storageNamespaceFactory)
+			: base(storageNamespaceFactory)
 		{
 		}
 
@@ -51,7 +51,7 @@ namespace Horde.Commands.Vcs
 
 			RefName branchName = (Branch != null) ? new RefName(Branch) : workspaceState.Branch;
 
-			IStorageClient store = CreateStorageClient();
+			IStorageNamespace store = GetStorageNamespace();
 
 			CommitNode? tip = await GetCommitAsync(store, branchName, Change);
 			if (tip == null)
diff --git a/Engine/Source/Programs/Horde/HordeCmd/Commands/Vcs/VcsClone.cs b/Engine/Source/Programs/Horde/HordeCmd/Commands/Vcs/VcsClone.cs
--- a/Engine/Source/Programs/Horde/HordeCmd/Commands/Vcs/VcsClone.cs
+++ b/Engine/Source/Programs/Horde/HordeCmd/Commands/Vcs/VcsClone.cs
@@ -10,8 +10,8 @@ namespace Horde.Commands.Vcs
 	[Command("vcs", "clone", "Initialize a directory for VCS-like operations", Advertise = false)]
 	class VcsClone : VcsCheckout
 	{
-		public VcsClone(IStorageClientFactory storageClientFactory)
-			: base(storageClientFactory)
+		public VcsClone(IStorageNamespaceFactory storageNamespaceFactory)
+			: base(storageNamespaceFactory)
 		{
 		}
 
@@ -25,9 +25,9 @@ namespace Horde.Commands.Vcs
 
 			RefName branchName = new RefName(Branch ?? "ue5-main");
 
-			IStorageClient storageClient = CreateStorageClient();
+			IStorageNamespace storageNamespace = GetStorageNamespace();
 
-			CommitNode? tip = await GetCommitAsync(storageClient, branchName, Change);
+			CommitNode? tip = await GetCommitAsync(storageNamespace, branchName, Change);
 			if (tip == null)
 			{
 				logger.LogError("Unable to find change {Change}", Change);
diff --git a/Engine/Source/Programs/Horde/HordeCmd/Commands/Vcs/VcsCommit.cs b/Engine/Source/Programs/Horde/HordeCmd/Commands/Vcs/VcsCommit.cs
--- a/Engine/Source/Programs/Horde/HordeCmd/Commands/Vcs/VcsCommit.cs
+++ b/Engine/Source/Programs/Horde/HordeCmd/Commands/Vcs/VcsCommit.cs
@@ -15,8 +15,8 @@ namespace Horde.Commands.Vcs
 		[CommandLine("-Message=")]
 		public string? Message { get; set; }
 
-		public VcsCommit(IStorageClientFactory storageClientFactory)
-			: base(storageClientFactory)
+		public VcsCommit(IStorageNamespaceFactory storageNamespaceFactory)
+			: base(storageNamespaceFactory)
 		{
 		}
 
diff --git a/Engine/Source/Programs/Horde/HordeCmd/Commands/Vcs/VcsInit.cs b/Engine/Source/Programs/Horde/HordeCmd/Commands/Vcs/VcsInit.cs
--- a/Engine/Source/Programs/Horde/HordeCmd/Commands/Vcs/VcsInit.cs
+++ b/Engine/Source/Programs/Horde/HordeCmd/Commands/Vcs/VcsInit.cs
@@ -9,8 +9,8 @@ namespace Horde.Commands.Vcs
 	[Command("vcs", "init", "Initialize a directory for VCS-like operations", Advertise = false)]
 	class VcsInit : VcsBase
 	{
-		public VcsInit(IStorageClientFactory storageClientFactory)
-			: base(storageClientFactory)
+		public VcsInit(IStorageNamespaceFactory storageNamespaceFactory)
+			: base(storageNamespaceFactory)
 		{
 		}
 
diff --git a/Engine/Source/Programs/Horde/HordeCmd/Commands/Vcs/VcsLog.cs b/Engine/Source/Programs/Horde/HordeCmd/Commands/Vcs/VcsLog.cs
--- a/Engine/Source/Programs/Horde/HordeCmd/Commands/Vcs/VcsLog.cs
+++ b/Engine/Source/Programs/Horde/HordeCmd/Commands/Vcs/VcsLog.cs
@@ -14,8 +14,8 @@ namespace Horde.Commands.Vcs
 		[CommandLine("-Count")]
 		public int Count { get; set; } = 20;
 
-		public VcsLog(IStorageClientFactory storageClientFactory)
-			: base(storageClientFactory)
+		public VcsLog(IStorageNamespaceFactory storageNamespaceFactory)
+			: base(storageNamespaceFactory)
 		{
 		}
 
@@ -25,7 +25,7 @@ namespace Horde.Commands.Vcs
 
 			WorkspaceState workspaceState = await ReadStateAsync(rootDir);
 
-			IStorageClient store = CreateStorageClient();
+			IStorageNamespace store = GetStorageNamespace();
 
 			using MemoryCache cache = new MemoryCache(new MemoryCacheOptions());
 
diff --git a/Engine/Source/Programs/Horde/HordeCmd/Commands/Vcs/VcsStatus.cs b/Engine/Source/Programs/Horde/HordeCmd/Commands/Vcs/VcsStatus.cs
--- a/Engine/Source/Programs/Horde/HordeCmd/Commands/Vcs/VcsStatus.cs
+++ b/Engine/Source/Programs/Horde/HordeCmd/Commands/Vcs/VcsStatus.cs
@@ -9,8 +9,8 @@ namespace Horde.Commands.Vcs
 	[Command("vcs", "status", "Find status of local files", Advertise = false)]
 	class VcsStatus : VcsBase
 	{
-		public VcsStatus(IStorageClientFactory storageClientFactory)
-			: base(storageClientFactory)
+		public VcsStatus(IStorageNamespaceFactory storageNamespaceFactory)
+			: base(storageNamespaceFactory)
 		{
 		}
 
diff --git a/Engine/Source/Programs/Horde/HordeCmd/Commands/Workspace/WorkspaceSync.cs b/Engine/Source/Programs/Horde/HordeCmd/Commands/Workspace/WorkspaceSync.cs
--- a/Engine/Source/Programs/Horde/HordeCmd/Commands/Workspace/WorkspaceSync.cs
+++ b/Engine/Source/Programs/Horde/HordeCmd/Commands/Workspace/WorkspaceSync.cs
@@ -41,8 +41,8 @@ namespace Horde.Commands.Workspace
 		[Description("Outputs stats for the extraction operation.")]
 		public bool Stats { get; set; }
 
-		public WorkspaceSync(HttpStorageClientFactory storageClientFactory, BundleCache bundleCache)
-			: base(storageClientFactory, bundleCache)
+		public WorkspaceSync(HttpStorageNamespaceFactory storageNamespaceFactory, BundleCache bundleCache)
+			: base(storageNamespaceFactory, bundleCache)
 		{
 		}
 
@@ -51,19 +51,19 @@ namespace Horde.Commands.Workspace
 			if (File != null)
 			{
 				using MemoryMappedFileCache memoryMappedFileCache = new MemoryMappedFileCache();
-				IStorageClient store = BundleStorageClient.CreateFromDirectory(File.Directory, BundleCache, memoryMappedFileCache, logger);
+				IStorageNamespace store = BundleStorageNamespace.CreateFromDirectory(File.Directory, BundleCache, memoryMappedFileCache, logger);
 				IBlobRef handle = store.CreateBlobRef(await FileStorageBackend.ReadRefAsync(File));
 				return await ExecuteInternalAsync(store, handle, logger);
 			}
 			else if (Ref != null)
 			{
-				IStorageClient store = CreateStorageClient();
+				IStorageNamespace store = GetStorageNamespace();
 				IBlobRef handle = await store.ReadRefAsync(new RefName(Ref));
 				return await ExecuteInternalAsync(store, handle, logger);
 			}
 			else if (Node != null)
 			{
-				IStorageClient store = CreateStorageClient();
+				IStorageNamespace store = GetStorageNamespace();
 				IBlobRef handle = store.CreateBlobRef(new BlobLocator(Node));
 				return await ExecuteInternalAsync(store, handle, logger);
 			}
@@ -73,7 +73,7 @@ namespace Horde.Commands.Workspace
 			}
 		}
 
-		async Task<int> ExecuteInternalAsync(IStorageClient store, IBlobRef handle, ILogger logger)
+		async Task<int> ExecuteInternalAsync(IStorageNamespace store, IBlobRef handle, ILogger logger)
 		{
 			RootDir ??= DirectoryReference.GetCurrentDirectory();
 			CancellationToken cancellationToken = CancellationToken.None;
diff --git a/Engine/Source/Programs/Horde/Plugins/Build/HordeServer.Build.Tests/Artifacts/UnsyncCacheTests.cs b/Engine/Source/Programs/Horde/Plugins/Build/HordeServer.Build.Tests/Artifacts/UnsyncCacheTests.cs
--- a/Engine/Source/Programs/Horde/Plugins/Build/HordeServer.Build.Tests/Artifacts/UnsyncCacheTests.cs
+++ b/Engine/Source/Programs/Horde/Plugins/Build/HordeServer.Build.Tests/Artifacts/UnsyncCacheTests.cs
@@ -21,12 +21,12 @@ namespace HordeServer.Tests.Artifacts
 			NamespaceId ns = new NamespaceId("test");
 			byte[] source = Enumerable.Range(0, 1024 * 1024).Select(x => (byte)(x % 257)).ToArray();
 
-			BundleStorageClient storageClient = BundleStorageClient.CreateInMemory(NullLogger.Instance);
+			BundleStorageNamespace storageNamespace = BundleStorageNamespace.CreateInMemory(NullLogger.Instance);
 
 			Dictionary<IoHash, ReadOnlyMemory<byte>> chunks = new Dictionary<IoHash, ReadOnlyMemory<byte>>();
 
 			IHashedBlobRef<DirectoryNode> directoryRef;
-			await using (IBlobWriter writer = storageClient.CreateBlobWriter())
+			await using (IBlobWriter writer = storageNamespace.CreateBlobWriter())
 			{
 				using ChunkedDataWriter chunkedWriter = new ChunkedDataWriter(writer, new ChunkingOptions());
 				await chunkedWriter.AppendAsync(source, CancellationToken.None);
@@ -40,10 +40,10 @@ namespace HordeServer.Tests.Artifacts
 			}
 
 			RefName refName = new RefName("test");
-			await storageClient.WriteRefAsync(refName, directoryRef);
+			await storageNamespace.WriteRefAsync(refName, directoryRef);
 
 			Mock<IStorageService> factory = new Mock<IStorageService>();
-			factory.Setup(x => x.TryCreateClient(ns)).Returns(storageClient);
+			factory.Setup(x => x.TryGetNamespace(ns)).Returns(storageNamespace);
 
 			Mock<IArtifact> artifact = new Mock<IArtifact>();
 			artifact.SetupGet(x => x.NamespaceId).Returns(ns);
diff --git a/Engine/Source/Programs/Horde/Plugins/Build/HordeServer.Build.Tests/Issues/IssueServiceTests.cs b/Engine/Source/Programs/Horde/Plugins/Build/HordeServer.Build.Tests/Issues/IssueServiceTests.cs
--- a/Engine/Source/Programs/Horde/Plugins/Build/HordeServer.Build.Tests/Issues/IssueServiceTests.cs
+++ b/Engine/Source/Programs/Horde/Plugins/Build/HordeServer.Build.Tests/Issues/IssueServiceTests.cs
@@ -43,16 +43,16 @@ namespace HordeServer.Tests.Issues
 			readonly ILog _log;
 			readonly LogBuilder _builder;
 			readonly List<(LogLevel, ReadOnlyMemory<byte>)> _events = new List<(LogLevel, ReadOnlyMemory<byte>)>();
-			readonly IStorageClient _storageClient;
+			readonly IStorageNamespace _storageNamespace;
 			readonly LoggerScopeCollection _scopeCollection = new LoggerScopeCollection();
 
 			int _lineIndex;
 
-			public TestJsonLogger(ILog log, IStorageClient storageClient)
+			public TestJsonLogger(ILog log, IStorageNamespace storageNamespace)
 			{
 				_log = log;
 				_builder = new LogBuilder(LogFormat.Json, NullLogger.Instance);
-				_storageClient = storageClient;
+				_storageNamespace = storageNamespace;
 			}
 
 			public async ValueTask DisposeAsync()
@@ -65,10 +65,10 @@ namespace HordeServer.Tests.Issues
 					await WriteAsync(level, lineWithNewLine);
 				}
 
-				await using (IBlobWriter writer = _storageClient.CreateBlobWriter())
+				await using (IBlobWriter writer = _storageNamespace.CreateBlobWriter())
 				{
 					IHashedBlobRef<LogNode> handle = await _builder.FlushAsync(writer, true, CancellationToken.None);
-					await _storageClient.WriteRefAsync(_log.RefName, handle);
+					await _storageNamespace.WriteRefAsync(_log.RefName, handle);
 				}
 			}
 
@@ -347,14 +347,14 @@ namespace HordeServer.Tests.Issues
 
 			ILog log = (await LogCollection.GetAsync(logId, CancellationToken.None))!;
 
-			IStorageClient storageClient = StorageService.CreateClient(Namespace.Logs);
-			await using (IBlobWriter writer = storageClient.CreateBlobWriter())
+			IStorageNamespace storageNamespace = StorageService.GetNamespace(Namespace.Logs);
+			await using (IBlobWriter writer = storageNamespace.CreateBlobWriter())
 			{
 				LogBuilder builder = new LogBuilder(LogFormat.Json, NullLogger.Instance);
 				builder.WriteData(data);
 
 				IHashedBlobRef<LogNode> handle = await builder.FlushAsync(writer, true, CancellationToken.None);
-				await storageClient.WriteRefAsync(log!.RefName, handle);
+				await storageNamespace.WriteRefAsync(log!.RefName, handle);
 			}
 
 			await log.AddEventsAsync(new List<NewLogEventData> { new NewLogEventData { LineIndex = 0, LineCount = 1, Severity = severity } }, CancellationToken.None);
@@ -363,10 +363,10 @@ namespace HordeServer.Tests.Issues
 		private async Task<TestJsonLogger> CreateLoggerAsync(IJob job, int batchIdx, int stepIdx)
 		{
 			LogId logId = job.Batches[batchIdx].Steps[stepIdx].LogId!.Value;
-			IStorageClient storageClient = StorageService.CreateClient(Namespace.Logs);
+			IStorageNamespace storageNamespace = StorageService.GetNamespace(Namespace.Logs);
 			ILog? log = await LogCollection.GetAsync(logId)!;
 			Assert.IsNotNull(log);
-			return new TestJsonLogger(log, storageClient);
+			return new TestJsonLogger(log, storageNamespace);
 		}
 
 		private async Task ParseEventsAsync(IJob job, int batchIdx, int stepIdx, string[] lines)
@@ -390,8 +390,8 @@ namespace HordeServer.Tests.Issues
 			LogId logId = job.Batches[batchIdx].Steps[stepIdx].LogId!.Value;
 			ILog log = (await LogCollection.GetAsync(logId))!;
 
-			IStorageClient storageClient = StorageService.CreateClient(Namespace.Logs);
-			await using (TestJsonLogger logger = new TestJsonLogger(log, storageClient))
+			IStorageNamespace storageNamespace = StorageService.GetNamespace(Namespace.Logs);
+			await using (TestJsonLogger logger = new TestJsonLogger(log, storageNamespace))
 			{
 				PerforceMetadataLogger perforceLogger = new PerforceMetadataLogger(logger);
 				perforceLogger.AddClientView(_autoSdkDir, "//depot/CarefullyRedist/...", 12345);
diff --git a/Engine/Source/Programs/Horde/Plugins/Build/HordeServer.Build/Artifacts/ArtifactCollection.cs b/Engine/Source/Programs/Horde/Plugins/Build/HordeServer.Build/Artifacts/ArtifactCollection.cs
--- a/Engine/Source/Programs/Horde/Plugins/Build/HordeServer.Build/Artifacts/ArtifactCollection.cs
+++ b/Engine/Source/Programs/Horde/Plugins/Build/HordeServer.Build/Artifacts/ArtifactCollection.cs
@@ -190,7 +190,7 @@ namespace HordeServer.Artifacts
 
 		IBlobRef<DirectoryNode> Open(Artifact artifact)
 		{
-			IStorageClient client = _storageService.CreateClient(artifact.Document.NamespaceId);
+			IStorageNamespace client = _storageService.GetNamespace(artifact.Document.NamespaceId);
 			return client.CreateBlobRef<DirectoryNode>(artifact.Document.RefName);
 		}
 
@@ -377,8 +377,8 @@ namespace HordeServer.Artifacts
 		{
 			_logger.LogInformation("Deleting {StreamId} artifact {ArtifactId}, ref {RefName} (created {CreateTime})", artifact.StreamId, artifact.Id, artifact.RefName, artifact.CreatedAtUtc);
 
-			IStorageClient storageClient = _storageService.CreateClient(artifact.NamespaceId);
-			await storageClient.DeleteRefAsync(artifact.RefName, cancellationToken);
+			IStorageNamespace storageNamespace = _storageService.GetNamespace(artifact.NamespaceId);
+			await storageNamespace.DeleteRefAsync(artifact.RefName, cancellationToken);
 
 			FilterDefinition<ArtifactDocument> filter = Builders<ArtifactDocument>.Filter.Eq(x => x.Id, artifact.Id);
 			await _artifactCollection.DeleteOneAsync(filter, cancellationToken);
@@ -388,12 +388,12 @@ namespace HordeServer.Artifacts
 		{
 			foreach (IGrouping<NamespaceId, ArtifactDocument> artifactGroup in deleteArtifacts.GroupBy(x => x.NamespaceId))
 			{
-				IStorageClient storageClient = _storageService.CreateClient(artifactGroup.Key);
+				IStorageNamespace storageNamespace = _storageService.GetNamespace(artifactGroup.Key);
 				foreach (ArtifactDocument artifact in artifactGroup)
 				{
 					// Delete the ref allowing the storage service to expire this data
 					_logger.LogInformation("Expiring {StreamId} artifact {ArtifactId}, ref {RefName} (created {CreateTime})", artifact.StreamId, artifact.Id, artifact.RefName, artifact.CreatedAtUtc);
-					await storageClient.DeleteRefAsync(artifact.RefName, cancellationToken);
+					await storageNamespace.DeleteRefAsync(artifact.RefName, cancellationToken);
 				}
 
 				FilterDefinition<ArtifactDocument> filter = Builders<ArtifactDocument>.Filter.In(x => x.Id, artifactGroup.Select(x => x.Id));
diff --git a/Engine/Source/Programs/Horde/Plugins/Build/HordeServer.Build/Artifacts/ArtifactsController.cs b/Engine/Source/Programs/Horde/Plugins/Build/HordeServer.Build/Artifacts/ArtifactsController.cs
--- a/Engine/Source/Programs/Horde/Plugins/Build/HordeServer.Build/Artifacts/ArtifactsController.cs
+++ b/Engine/Source/Programs/Horde/Plugins/Build/HordeServer.Build/Artifacts/ArtifactsController.cs
@@ -71,7 +71,7 @@ namespace HordeServer.Artifacts
 		}
 
 		/// <summary>
-		/// Creates a new artifact. Actual data for the artifact can be uploaded using a storage client pointed to the blobs endpoint.
+		/// Creates a new artifact. Actual data for the artifact can be uploaded using a storage namespace pointed to the blobs endpoint.
 		/// </summary>
 		/// <param name="request">Information about the desired artifact</param>
 		/// <param name="cancellationToken">Cancellation token for the operation</param>
@@ -332,12 +332,12 @@ namespace HordeServer.Artifacts
 				return Forbid(ArtifactAclAction.ReadArtifact, artifact.StreamId);
 			}
 
-			IStorageClient storageClient = _storageService.CreateClient(artifact.NamespaceId);
+			IStorageNamespace storageNamespace = _storageService.GetNamespace(artifact.NamespaceId);
 
 			DirectoryNode directoryNode;
 			try
 			{
-				directoryNode = await storageClient.ReadRefTargetAsync<DirectoryNode>(artifact.RefName, DateTime.UtcNow.AddHours(1.0), cancellationToken: cancellationToken);
+				directoryNode = await storageNamespace.ReadRefTargetAsync<DirectoryNode>(artifact.RefName, DateTime.UtcNow.AddHours(1.0), cancellationToken: cancellationToken);
 			}
 			catch (RefNameNotFoundException)
 			{
@@ -499,8 +499,8 @@ namespace HordeServer.Artifacts
 				return Forbid(ArtifactAclAction.ReadArtifact, artifact.StreamId);
 			}
 
-			IStorageClient storageClient = _storageService.CreateClient(artifact.NamespaceId);
-			DirectoryNode directory = await storageClient.ReadRefTargetAsync<DirectoryNode>(artifact.RefName, DateTime.UtcNow.AddHours(1.0), cancellationToken: cancellationToken);
+			IStorageNamespace storageNamespace = _storageService.GetNamespace(artifact.NamespaceId);
+			DirectoryNode directory = await storageNamespace.ReadRefTargetAsync<DirectoryNode>(artifact.RefName, DateTime.UtcNow.AddHours(1.0), cancellationToken: cancellationToken);
 
 			FileEntry? fileEntry = await directory.GetFileEntryByPathAsync(path, cancellationToken: cancellationToken);
 			if (fileEntry == null)
@@ -594,8 +594,8 @@ namespace HordeServer.Artifacts
 			}
 
 #pragma warning disable CA2000
-			IStorageClient storageClient = _storageService.CreateClient(artifact.NamespaceId);
-			IHashedBlobRef<DirectoryNode> directory = await storageClient.ReadRefAsync<DirectoryNode>(artifact.RefName, DateTime.UtcNow.AddHours(1.0), cancellationToken: cancellationToken);
+			IStorageNamespace storageNamespace = _storageService.GetNamespace(artifact.NamespaceId);
+			IHashedBlobRef<DirectoryNode> directory = await storageNamespace.ReadRefAsync<DirectoryNode>(artifact.RefName, DateTime.UtcNow.AddHours(1.0), cancellationToken: cancellationToken);
 
 			Stream stream = directory.AsZipStream(filter);
 			return new FileStreamResult(stream, "application/zip") { FileDownloadName = $"{artifact.RefName}.zip" };
diff --git a/Engine/Source/Programs/Horde/Plugins/Build/HordeServer.Build/Artifacts/UnsyncCache.cs b/Engine/Source/Programs/Horde/Plugins/Build/HordeServer.Build/Artifacts/UnsyncCache.cs
--- a/Engine/Source/Programs/Horde/Plugins/Build/HordeServer.Build/Artifacts/UnsyncCache.cs
+++ b/Engine/Source/Programs/Horde/Plugins/Build/HordeServer.Build/Artifacts/UnsyncCache.cs
@@ -207,9 +207,9 @@ namespace HordeServer.Artifacts
 		async Task<ArtifactInfo?> ReadArtifactAsync(IArtifact artifact, CancellationToken cancellationToken)
 		{
 			Stopwatch timer = Stopwatch.StartNew();
-			IStorageClient storageClient = _storageService.CreateClient(artifact.NamespaceId);
+			IStorageNamespace storageNamespace = _storageService.GetNamespace(artifact.NamespaceId);
 
-			IHashedBlobRef<DirectoryNode>? target = await storageClient.TryReadRefAsync<DirectoryNode>(artifact.RefName, cancellationToken: cancellationToken);
+			IHashedBlobRef<DirectoryNode>? target = await storageNamespace.TryReadRefAsync<DirectoryNode>(artifact.RefName, cancellationToken: cancellationToken);
 			if (target == null)
 			{
 				return null;
diff --git a/Engine/Source/Programs/Horde/Plugins/Build/HordeServer.Build/Replicators/PerforceReplicator.cs b/Engine/Source/Programs/Horde/Plugins/Build/HordeServer.Build/Replicators/PerforceReplicator.cs
--- a/Engine/Source/Programs/Horde/Plugins/Build/HordeServer.Build/Replicators/PerforceReplicator.cs
+++ b/Engine/Source/Programs/Horde/Plugins/Build/HordeServer.Build/Replicators/PerforceReplicator.cs
@@ -316,7 +316,7 @@ namespace HordeServer.Replicators
 			RefName refName = GetRefName(replicator.Id);
 			RefName incRefName = GetIncrementalRefName(replicator.Id);
 
-			IStorageClient store = _storageService.CreateClient(replicatorConfig.NamespaceId);
+			IStorageNamespace store = _storageService.GetNamespace(replicatorConfig.NamespaceId);
 
 			ICommitCollection commits = _perforceService.GetCommits(streamConfig);
 
@@ -396,7 +396,7 @@ namespace HordeServer.Replicators
 		async Task<IReplicator> WriteInternalAsync(IReplicator replicator, int change, BuildConfig buildConfig, StreamConfig streamConfig, ReplicatorConfig replicatorConfig, BlobSerializerOptions blobOptions, CancellationToken cancellationToken = default)
 		{
 			ChunkingOptions chunkingOptions = new();
-			IStorageClient store = _storageService.CreateClient(replicatorConfig.NamespaceId);
+			IStorageNamespace store = _storageService.GetNamespace(replicatorConfig.NamespaceId);
 
 			// Find the parent node
 			RefName refName = GetRefName(replicator.Id);
diff --git a/Engine/Source/Programs/Horde/Plugins/Compute/HordeServer.Compute.Tests/Logs/LogIndexingTests.cs b/Engine/Source/Programs/Horde/Plugins/Compute/HordeServer.Compute.Tests/Logs/LogIndexingTests.cs
--- a/Engine/Source/Programs/Horde/Plugins/Compute/HordeServer.Compute.Tests/Logs/LogIndexingTests.cs
+++ b/Engine/Source/Programs/Horde/Plugins/Compute/HordeServer.Compute.Tests/Logs/LogIndexingTests.cs
@@ -14,7 +14,7 @@ namespace HordeServer.Tests.Logs
 	class TestLogWriter : IAsyncDisposable
 	{
 		ILog _logFile;
-		readonly IStorageClient _storageClient;
+		readonly IStorageNamespace _storageNamespace;
 		readonly IBlobWriter _blobWriter;
 		readonly LogBuilder _builder;
 		int _lineCount;
@@ -22,8 +22,8 @@ namespace HordeServer.Tests.Logs
 		public TestLogWriter(ILog logFile, StorageService storageService)
 		{
 			_logFile = logFile;
-			_storageClient = storageService.CreateClient(Namespace.Logs);
-			_blobWriter = _storageClient.CreateBlobWriter(logFile.RefName);
+			_storageNamespace = storageService.GetNamespace(Namespace.Logs);
+			_blobWriter = _storageNamespace.CreateBlobWriter(logFile.RefName);
 			_builder = new LogBuilder((logFile.Type == LogType.Text) ? LogFormat.Text : LogFormat.Json, NullLogger.Instance);
 		}
 
@@ -43,7 +43,7 @@ namespace HordeServer.Tests.Logs
 		public async Task<ILog> FlushAsync(bool complete = true)
 		{
 			IHashedBlobRef<LogNode> handle = await _builder.FlushAsync(_blobWriter, complete, CancellationToken.None);
-			await _storageClient.WriteRefAsync(_logFile.RefName, handle);
+			await _storageNamespace.WriteRefAsync(_logFile.RefName, handle);
 			_logFile = await _logFile.UpdateLineCountAsync(_lineCount, complete, CancellationToken.None);
 			return _logFile;
 		}
diff --git a/Engine/Source/Programs/Horde/Plugins/Compute/HordeServer.Compute/Logs/LogCollection.cs b/Engine/Source/Programs/Horde/Plugins/Compute/HordeServer.Compute/Logs/LogCollection.cs
--- a/Engine/Source/Programs/Horde/Plugins/Compute/HordeServer.Compute/Logs/LogCollection.cs
+++ b/Engine/Source/Programs/Horde/Plugins/Compute/HordeServer.Compute/Logs/LogCollection.cs
@@ -306,8 +306,8 @@ namespace HordeServer.Logs
 
 			if (logDocument.UseNewStorageBackend)
 			{
-				IStorageClient storageClient = _storageService.CreateClient(logDocument.NamespaceId);
-				await storageClient.DeleteRefAsync(logDocument.RefName, cancellationToken);
+				IStorageNamespace storageNamespace = _storageService.GetNamespace(logDocument.NamespaceId);
+				await storageNamespace.DeleteRefAsync(logDocument.RefName, cancellationToken);
 			}
 		}
 
@@ -327,12 +327,12 @@ namespace HordeServer.Logs
 		{
 			List<Utf8String> lines = new List<Utf8String>();
 
-			IStorageClient storageClient = _storageService.CreateClient(log.NamespaceId);
+			IStorageNamespace storageNamespace = _storageService.GetNamespace(log.NamespaceId);
 
 			int maxIndex = index + count;
 			bool complete = log.Complete;
 
-			LogNode? root = await storageClient.TryReadRefTargetAsync<LogNode>(log.RefName, cancellationToken: cancellationToken);
+			LogNode? root = await storageNamespace.TryReadRefTargetAsync<LogNode>(log.RefName, cancellationToken: cancellationToken);
 			if (root != null)
 			{
 				int chunkIdx = root.TextChunkRefs.GetChunkForLine(index);
@@ -392,9 +392,9 @@ namespace HordeServer.Logs
 		/// <inheritdoc/>
 		async Task<(int, long)> GetLineOffsetAsync(LogDocument log, int lineIdx, CancellationToken cancellationToken)
 		{
-			IStorageClient storageClient = _storageService.CreateClient(log.NamespaceId);
+			IStorageNamespace storageNamespace = _storageService.GetNamespace(log.NamespaceId);
 
-			LogNode? root = await storageClient.TryReadRefTargetAsync<LogNode>(log.RefName, cancellationToken: cancellationToken);
+			LogNode? root = await storageNamespace.TryReadRefTargetAsync<LogNode>(log.RefName, cancellationToken: cancellationToken);
 			if (root == null)
 			{
 				return (0, 0);
@@ -422,9 +422,9 @@ namespace HordeServer.Logs
 		/// <inheritdoc/>
 		async Task<Stream> OpenRawStreamAsync(LogDocument log, long offset, long length, CancellationToken cancellationToken)
 		{
-			IStorageClient storageClient = _storageService.CreateClient(log.NamespaceId);
+			IStorageNamespace storageNamespace = _storageService.GetNamespace(log.NamespaceId);
 
-			LogNode? root = await storageClient.TryReadRefTargetAsync<LogNode>(log.RefName, cancellationToken: cancellationToken);
+			LogNode? root = await storageNamespace.TryReadRefTargetAsync<LogNode>(log.RefName, cancellationToken: cancellationToken);
 			if (root == null || root.TextChunkRefs.Count == 0)
 			{
 				return new MemoryStream(Array.Empty<byte>(), false);
@@ -480,12 +480,12 @@ namespace HordeServer.Logs
 		async IAsyncEnumerable<int> SearchLogDataInternalNewAsync(LogDocument log, string text, int firstLine, SearchStats searchStats, [EnumeratorCancellation] CancellationToken cancellationToken)
 		{
 			SearchTerm searchText = new SearchTerm(text);
-			IStorageClient storageClient = _storageService.CreateClient(log.NamespaceId);
+			IStorageNamespace storageNamespace = _storageService.GetNamespace(log.NamespaceId);
 
 			// Search the index
 			if (log.LineCount > 0)
 			{
-				LogNode? root = await storageClient.ReadRefTargetAsync<LogNode>(log.RefName, cancellationToken: cancellationToken);
+				LogNode? root = await storageNamespace.ReadRefTargetAsync<LogNode>(log.RefName, cancellationToken: cancellationToken);
 				if (root != null)
 				{
 					LogIndexNode index = await root.IndexRef.ReadBlobAsync(cancellationToken: cancellationToken);
diff --git a/Engine/Source/Programs/Horde/Plugins/Compute/HordeServer.Compute/Logs/LogRpcService.cs b/Engine/Source/Programs/Horde/Plugins/Compute/HordeServer.Compute/Logs/LogRpcService.cs
--- a/Engine/Source/Programs/Horde/Plugins/Compute/HordeServer.Compute/Logs/LogRpcService.cs
+++ b/Engine/Source/Programs/Horde/Plugins/Compute/HordeServer.Compute/Logs/LogRpcService.cs
@@ -51,7 +51,7 @@ namespace HordeServer.Logs
 				throw new StructuredRpcException(StatusCode.PermissionDenied, "Access denied");
 			}
 
-			IStorageClient store = _storageService.CreateClient(Namespace.Logs);
+			IStorageNamespace store = _storageService.GetNamespace(Namespace.Logs);
 
 			_logger.LogInformation("Updating {LogId} to node {RefTarget} (lines: {LineCount}, complete: {Complete})", request.LogId, request.TargetLocator, request.LineCount, request.Complete);
 
diff --git a/Engine/Source/Programs/Horde/Plugins/Ddc/HordeServer.Ddc/Ddc/BlobService.cs b/Engine/Source/Programs/Horde/Plugins/Ddc/HordeServer.Ddc/Ddc/BlobService.cs
--- a/Engine/Source/Programs/Horde/Plugins/Ddc/HordeServer.Ddc/Ddc/BlobService.cs
+++ b/Engine/Source/Programs/Horde/Plugins/Ddc/HordeServer.Ddc/Ddc/BlobService.cs
@@ -11,12 +11,12 @@ namespace HordeServer.Ddc
 	{
 		static readonly BlobType s_rawBlobType = new BlobType(new Guid("{03E6C37B-33C1-491F-8541-D3C401B8B8EF}"), 1);
 
-		readonly IStorageClientFactory _storageClientFactory;
+		readonly IStorageNamespaceFactory _storageNamespaceFactory;
 		readonly Tracer _tracer;
 
-		public BlobService(IStorageClientFactory storageClientFactory, Tracer tracer)
+		public BlobService(IStorageNamespaceFactory storageNamespaceFactory, Tracer tracer)
 		{
-			_storageClientFactory = storageClientFactory;
+			_storageNamespaceFactory = storageNamespaceFactory;
 			_tracer = tracer;
 		}
 
@@ -60,30 +60,30 @@ namespace HordeServer.Ddc
 
 		public async Task<bool> ExistsAsync(NamespaceId ns, BlobId blob, List<string>? storageLayers = null, CancellationToken cancellationToken = default)
 		{
-			IStorageClient storageClient = _storageClientFactory.CreateClient(ns);
-			return await storageClient.FindAliasAsync(GetAlias(blob), cancellationToken) != null;
+			IStorageNamespace storageNamespace = _storageNamespaceFactory.GetNamespace(ns);
+			return await storageNamespace.FindAliasAsync(GetAlias(blob), cancellationToken) != null;
 		}
 
 		public async Task DeleteObjectAsync(NamespaceId ns, BlobId blob, CancellationToken cancellationToken)
 		{
-			IStorageClient storageClient = _storageClientFactory.CreateClient(ns);
+			IStorageNamespace storageNamespace = _storageNamespaceFactory.GetNamespace(ns);
 			string aliasName = GetAlias(blob);
 
-			BlobAlias[] aliases = await storageClient.FindAliasesAsync(aliasName, cancellationToken: cancellationToken);
+			BlobAlias[] aliases = await storageNamespace.FindAliasesAsync(aliasName, cancellationToken: cancellationToken);
 			foreach (BlobAlias alias in aliases)
 			{
-				await storageClient.RemoveAliasAsync(aliasName, alias.Target, cancellationToken);
+				await storageNamespace.RemoveAliasAsync(aliasName, alias.Target, cancellationToken);
 			}
 		}
 
 		public async Task<BlobId[]> FilterOutKnownBlobsAsync(NamespaceId ns, IEnumerable<BlobId> blobIds, CancellationToken cancellationToken)
 		{
-			IStorageClient storageClient = _storageClientFactory.CreateClient(ns);
+			IStorageNamespace storageNamespace = _storageNamespaceFactory.GetNamespace(ns);
 
 			List<BlobId> unknownBlobIds = new List<BlobId>();
 			foreach (BlobId blobId in blobIds)
 			{
-				if (await storageClient.FindAliasAsync(GetAlias(blobId), cancellationToken) == null)
+				if (await storageNamespace.FindAliasAsync(GetAlias(blobId), cancellationToken) == null)
 				{
 					unknownBlobIds.Add(blobId);
 				}
@@ -123,9 +123,9 @@ namespace HordeServer.Ddc
 
 		public async Task<BlobContents> GetObjectAsync(NamespaceId ns, BlobId blob, List<string>? storageLayers, bool supportsRedirectUri, bool allowOndemandReplication = true, CancellationToken cancellationToken = default)
 		{
-			IStorageClient storageClient = _storageClientFactory.CreateClient(ns);
+			IStorageNamespace storageNamespace = _storageNamespaceFactory.GetNamespace(ns);
 
-			BlobAlias? alias = await storageClient.FindAliasAsync(GetAlias(blob), cancellationToken);
+			BlobAlias? alias = await storageNamespace.FindAliasAsync(GetAlias(blob), cancellationToken);
 			if (alias == null)
 			{
 				throw new BlobNotFoundException(ns, blob);
@@ -174,10 +174,10 @@ namespace HordeServer.Ddc
 
 		public async Task<BlobId> PutObjectKnownHashAsync(NamespaceId ns, IBufferedPayload content, BlobId identifier, CancellationToken cancellationToken)
 		{
-			IStorageClient storageClient = _storageClientFactory.CreateClient(ns);
+			IStorageNamespace storageNamespace = _storageNamespaceFactory.GetNamespace(ns);
 
 			IHashedBlobRef blobRef;
-			await using (IBlobWriter writer = storageClient.CreateBlobWriter())
+			await using (IBlobWriter writer = storageNamespace.CreateBlobWriter())
 			{
 				int contentLength = (int)content.Length;
 				Memory<byte> memory = writer.GetMemory(contentLength).Slice(0, contentLength);
@@ -190,7 +190,7 @@ namespace HordeServer.Ddc
 				await writer.FlushAsync(cancellationToken);
 			}
 
-			await storageClient.AddAliasAsync(GetAlias(identifier), blobRef, data: identifier.AsIoHash().ToByteArray(), cancellationToken: cancellationToken);
+			await storageNamespace.AddAliasAsync(GetAlias(identifier), blobRef, data: identifier.AsIoHash().ToByteArray(), cancellationToken: cancellationToken);
 			return identifier;
 		}
 
diff --git a/Engine/Source/Programs/Horde/Plugins/Ddc/HordeServer.Ddc/Ddc/ContentIdStore.cs b/Engine/Source/Programs/Horde/Plugins/Ddc/HordeServer.Ddc/Ddc/ContentIdStore.cs
--- a/Engine/Source/Programs/Horde/Plugins/Ddc/HordeServer.Ddc/Ddc/ContentIdStore.cs
+++ b/Engine/Source/Programs/Horde/Plugins/Ddc/HordeServer.Ddc/Ddc/ContentIdStore.cs
@@ -7,11 +7,11 @@ namespace HordeServer.Ddc
 {
 	class ContentIdStore : IContentIdStore
 	{
-		readonly IStorageClientFactory _storageClientFactory;
+		readonly IStorageNamespaceFactory _storageNamespaceFactory;
 
-		public ContentIdStore(IStorageClientFactory storageService)
+		public ContentIdStore(IStorageNamespaceFactory storageService)
 		{
-			_storageClientFactory = storageService;
+			_storageNamespaceFactory = storageService;
 		}
 
 		static string GetAlias(BlobId blobId) => BlobService.GetAlias(blobId);
@@ -19,12 +19,12 @@ namespace HordeServer.Ddc
 
 		public async Task<BlobId[]?> ResolveAsync(NamespaceId ns, ContentId contentId, bool mustBeContentId = false, CancellationToken cancellationToken = default)
 		{
-			IStorageClient storageClient = _storageClientFactory.CreateClient(ns);
+			IStorageNamespace storageNamespace = _storageNamespaceFactory.GetNamespace(ns);
 
-			BlobAlias? blobAlias = await storageClient.FindAliasAsync(GetAlias(contentId), cancellationToken);
+			BlobAlias? blobAlias = await storageNamespace.FindAliasAsync(GetAlias(contentId), cancellationToken);
 			if (blobAlias == null && !mustBeContentId)
 			{
-				blobAlias = await storageClient.FindAliasAsync(GetAlias(contentId.AsBlobIdentifier()), cancellationToken);
+				blobAlias = await storageNamespace.FindAliasAsync(GetAlias(contentId.AsBlobIdentifier()), cancellationToken);
 			}
 			if (blobAlias == null)
 			{
@@ -37,15 +37,15 @@ namespace HordeServer.Ddc
 
 		public async Task PutAsync(NamespaceId ns, ContentId contentId, BlobId blobId, int contentWeight, CancellationToken cancellationToken)
 		{
-			IStorageClient storageClient = _storageClientFactory.CreateClient(ns);
+			IStorageNamespace storageNamespace = _storageNamespaceFactory.GetNamespace(ns);
 
-			BlobAlias? blobAlias = await storageClient.FindAliasAsync(GetAlias(blobId), cancellationToken);
+			BlobAlias? blobAlias = await storageNamespace.FindAliasAsync(GetAlias(blobId), cancellationToken);
 			if (blobAlias == null)
 			{
 				throw new BlobNotFoundException(ns, blobId);
 			}
 
-			await storageClient.AddAliasAsync(GetAlias(contentId), blobAlias.Target, -contentWeight, blobAlias.Data, cancellationToken: cancellationToken);
+			await storageNamespace.AddAliasAsync(GetAlias(contentId), blobAlias.Target, -contentWeight, blobAlias.Data, cancellationToken: cancellationToken);
 		}
 	}
 }
diff --git a/Engine/Source/Programs/Horde/Plugins/Ddc/HordeServer.Ddc/Ddc/RefService.cs b/Engine/Source/Programs/Horde/Plugins/Ddc/HordeServer.Ddc/Ddc/RefService.cs
--- a/Engine/Source/Programs/Horde/Plugins/Ddc/HordeServer.Ddc/Ddc/RefService.cs
+++ b/Engine/Source/Programs/Horde/Plugins/Ddc/HordeServer.Ddc/Ddc/RefService.cs
@@ -8,13 +8,13 @@ namespace HordeServer.Ddc
 {
 	class RefService : IRefService
 	{
-		readonly IStorageClientFactory _storageClientFactory;
+		readonly IStorageNamespaceFactory _storageNamespaceFactory;
 		readonly IReferenceResolver _referenceResolver;
 		readonly IBlobService _blobService;
 
-		public RefService(IStorageClientFactory storageService, IReferenceResolver referenceResolver, IBlobService blobService)
+		public RefService(IStorageNamespaceFactory storageService, IReferenceResolver referenceResolver, IBlobService blobService)
 		{
-			_storageClientFactory = storageService;
+			_storageNamespaceFactory = storageService;
 			_referenceResolver = referenceResolver;
 			_blobService = blobService;
 		}
@@ -23,21 +23,21 @@ namespace HordeServer.Ddc
 
 		public async Task<bool> DeleteAsync(NamespaceId ns, BucketId bucket, RefId key, CancellationToken cancellationToken)
 		{
-			IStorageClient storageClient = _storageClientFactory.CreateClient(ns);
-			return await storageClient.DeleteRefAsync(GetRefName(bucket, key), cancellationToken);
+			IStorageNamespace storageNamespace = _storageNamespaceFactory.GetNamespace(ns);
+			return await storageNamespace.DeleteRefAsync(GetRefName(bucket, key), cancellationToken);
 		}
 
 		public async Task<bool> ExistsAsync(NamespaceId ns, BucketId bucket, RefId key, CancellationToken cancellationToken)
 		{
-			IStorageClient storageClient = _storageClientFactory.CreateClient(ns);
-			return await storageClient.RefExistsAsync(GetRefName(bucket, key), cancellationToken: cancellationToken);
+			IStorageNamespace storageNamespace = _storageNamespaceFactory.GetNamespace(ns);
+			return await storageNamespace.RefExistsAsync(GetRefName(bucket, key), cancellationToken: cancellationToken);
 		}
 
 		public async Task<(ContentId[], BlobId[])> FinalizeAsync(NamespaceId ns, BucketId bucket, RefId key, BlobId blobHash, CancellationToken cancellationToken)
 		{
-			IStorageClient storageClient = _storageClientFactory.CreateClient(ns);
+			IStorageNamespace storageNamespace = _storageNamespaceFactory.GetNamespace(ns);
 
-			BlobAlias? blobAlias = await storageClient.FindAliasAsync(BlobService.GetAlias(blobHash), cancellationToken);
+			BlobAlias? blobAlias = await storageNamespace.FindAliasAsync(BlobService.GetAlias(blobHash), cancellationToken);
 			if (blobAlias == null)
 			{
 				throw new BlobNotFoundException(ns, blobHash);
@@ -74,19 +74,19 @@ namespace HordeServer.Ddc
 				RefName refName = GetRefName(bucket, key);
 
 				IHashedBlobRef<DdcRefNode> refNodeRef;
-				await using (IBlobWriter writer = storageClient.CreateBlobWriter(refName))
+				await using (IBlobWriter writer = storageNamespace.CreateBlobWriter(refName))
 				{
 					DdcRefNode refNode = new DdcRefNode(blobHash.AsIoHash());
 					refNode.References.Add(HashedBlobRef.Create(blobHash.AsIoHash(), blobHandle));
 					foreach (BlobId referencedBlob in referencedBlobs)
 					{
-						BlobAlias? alias = await storageClient.FindAliasAsync(BlobService.GetAlias(referencedBlob), cancellationToken);
+						BlobAlias? alias = await storageNamespace.FindAliasAsync(BlobService.GetAlias(referencedBlob), cancellationToken);
 						refNode.References.Add(HashedBlobRef.Create(referencedBlob.AsIoHash(), alias!.Target));
 					}
 					refNodeRef = await writer.WriteBlobAsync(refNode, cancellationToken);
 				}
 
-				await storageClient.WriteRefAsync(refName, refNodeRef, cancellationToken: cancellationToken);
+				await storageNamespace.WriteRefAsync(refName, refNodeRef, cancellationToken: cancellationToken);
 			}
 
 			return (missingReferences, missingBlobs);
@@ -125,9 +125,9 @@ namespace HordeServer.Ddc
 
 		public async Task<(RefRecord, BlobContents?)> GetAsync(NamespaceId ns, BucketId bucket, RefId key, string[] fields, bool doLastAccessTracking, CancellationToken cancellationToken)
 		{
-			IStorageClient storageClient = _storageClientFactory.CreateClient(ns);
+			IStorageNamespace storageNamespace = _storageNamespaceFactory.GetNamespace(ns);
 
-			DdcRefNode? node = await storageClient.TryReadRefTargetAsync<DdcRefNode>(GetRefName(bucket, key), cancellationToken: cancellationToken);
+			DdcRefNode? node = await storageNamespace.TryReadRefTargetAsync<DdcRefNode>(GetRefName(bucket, key), cancellationToken: cancellationToken);
 			if (node == null)
 			{
 				throw new RefNotFoundException(ns, bucket, key);
@@ -142,9 +142,9 @@ namespace HordeServer.Ddc
 
 		public async Task<List<BlobId>> GetReferencedBlobsAsync(NamespaceId ns, BucketId bucket, RefId key, CancellationToken cancellationToken)
 		{
-			IStorageClient storageClient = _storageClientFactory.CreateClient(ns);
+			IStorageNamespace storageNamespace = _storageNamespaceFactory.GetNamespace(ns);
 
-			DdcRefNode? node = await storageClient.TryReadRefTargetAsync<DdcRefNode>(GetRefName(bucket, key), cancellationToken: cancellationToken);
+			DdcRefNode? node = await storageNamespace.TryReadRefTargetAsync<DdcRefNode>(GetRefName(bucket, key), cancellationToken: cancellationToken);
 			if (node == null)
 			{
 				throw new RefNotFoundException(ns, bucket, key);
diff --git a/Engine/Source/Programs/Horde/Plugins/Ddc/HordeServer.Ddc/Ddc/RequestHelper.cs b/Engine/Source/Programs/Horde/Plugins/Ddc/HordeServer.Ddc/Ddc/RequestHelper.cs
--- a/Engine/Source/Programs/Horde/Plugins/Ddc/HordeServer.Ddc/Ddc/RequestHelper.cs
+++ b/Engine/Source/Programs/Horde/Plugins/Ddc/HordeServer.Ddc/Ddc/RequestHelper.cs
@@ -12,12 +12,12 @@ namespace HordeServer.Ddc
 {
 	class RequestHelper : IRequestHelper
 	{
-		readonly IStorageClientFactory _storageClientFactory;
+		readonly IStorageNamespaceFactory _storageNamespaceFactory;
 		readonly StorageConfig _storageConfig;
 
-		public RequestHelper(IStorageClientFactory storageClientFactory, IOptionsSnapshot<StorageConfig> storageConfig)
+		public RequestHelper(IStorageNamespaceFactory storageNamespaceFactory, IOptionsSnapshot<StorageConfig> storageConfig)
 		{
-			_storageClientFactory = storageClientFactory;
+			_storageNamespaceFactory = storageNamespaceFactory;
 			_storageConfig = storageConfig.Value;
 		}
 
@@ -28,8 +28,8 @@ namespace HordeServer.Ddc
 
 		public Task<ActionResult?> HasAccessToNamespaceAsync(ClaimsPrincipal user, HttpRequest request, NamespaceId ns, AclAction[] aclActions)
 		{
-			IStorageClient? storageClient = _storageClientFactory.TryCreateClient(ns);
-			if (storageClient == null)
+			IStorageNamespace? storageNamespace = _storageNamespaceFactory.TryGetNamespace(ns);
+			if (storageNamespace == null)
 			{
 				return Task.FromResult<ActionResult?>(new ForbidResult());
 			}
diff --git a/Engine/Source/Programs/Horde/Plugins/Storage/HordeServer.Storage.Tests/Storage/BlobStoreTests.cs b/Engine/Source/Programs/Horde/Plugins/Storage/HordeServer.Storage.Tests/Storage/BlobStoreTests.cs
--- a/Engine/Source/Programs/Horde/Plugins/Storage/HordeServer.Storage.Tests/Storage/BlobStoreTests.cs
+++ b/Engine/Source/Programs/Horde/Plugins/Storage/HordeServer.Storage.Tests/Storage/BlobStoreTests.cs
@@ -21,7 +21,7 @@ namespace HordeServer.Tests.Storage
 			AddPlugin<StoragePlugin>();
 		}
 
-		IStorageClient CreateStorageClient()
+		IStorageNamespace GetStorageNamespace()
 		{
 			StorageConfig storageConfig = new StorageConfig();
 			storageConfig.Backends.Add(new BackendConfig { Id = new BackendId("default-backend"), Type = StorageBackendType.Memory });
@@ -31,7 +31,7 @@ namespace HordeServer.Tests.Storage
 			globalConfig.Plugins.Add(new PluginName("storage"), storageConfig);
 			SetConfig(globalConfig);
 
-			return StorageService.CreateClient(new NamespaceId("default"));
+			return StorageService.GetNamespace(new NamespaceId("default"));
 		}
 
 		static byte[] CreateTestData(int length, int seed)
@@ -57,7 +57,7 @@ namespace HordeServer.Tests.Storage
 			}
 		}
 
-		static async ValueTask<Blob> ReadBlobAsync(IStorageClient store, HashedBlobRefValue locator)
+		static async ValueTask<Blob> ReadBlobAsync(IStorageNamespace store, HashedBlobRefValue locator)
 		{
 			using (BlobData blobData = await store.CreateBlobRef(locator.Hash, locator.Locator).ReadBlobDataAsync())
 			{
@@ -75,7 +75,7 @@ namespace HordeServer.Tests.Storage
 			}
 		}
 
-		static async ValueTask<HashedBlobRefValue> WriteBlobAsync(IStorageClient store, Blob blob)
+		static async ValueTask<HashedBlobRefValue> WriteBlobAsync(IStorageNamespace store, Blob blob)
 		{
 			await using IBlobWriter writer = store.CreateBlobWriter();
 			writer.WriteVariableLengthBytes(blob.Data.Span);
@@ -94,7 +94,7 @@ namespace HordeServer.Tests.Storage
 		[TestMethod]
 		public async Task LeafTestAsync()
 		{
-			IStorageClient store = CreateStorageClient();
+			IStorageNamespace store = GetStorageNamespace();
 
 			byte[] input = CreateTestData(256, 0);
 
@@ -109,7 +109,7 @@ namespace HordeServer.Tests.Storage
 		[TestMethod]
 		public async Task ReferenceTestAsync()
 		{
-			IStorageClient store = CreateStorageClient();
+			IStorageNamespace store = GetStorageNamespace();
 
 			byte[] input1 = CreateTestData(256, 1);
 			HashedBlobRefValue locator1 = await WriteBlobAsync(store, new Blob(input1, Array.Empty<HashedBlobRefValue>()));
@@ -141,7 +141,7 @@ namespace HordeServer.Tests.Storage
 		[TestMethod]
 		public async Task RefExpiryTestAsync()
 		{
-			IStorageClient store = CreateStorageClient();
+			IStorageNamespace store = GetStorageNamespace();
 
 			Blob blob1 = new Blob(new byte[] { 1, 2, 3 }, Array.Empty<HashedBlobRefValue>());
 			HashedBlobRefValue target = await WriteBlobAsync(store, blob1);
@@ -173,7 +173,7 @@ namespace HordeServer.Tests.Storage
 			Assert.AreEqual(default, await TryReadRefValueAsync(store, "test-ref-3"));
 		}
 
-		static async Task<HashedBlobRefValue?> TryReadRefValueAsync(IStorageClient store, RefName name)
+		static async Task<HashedBlobRefValue?> TryReadRefValueAsync(IStorageNamespace store, RefName name)
 		{
 			IHashedBlobRef? handle = await store.TryReadRefAsync(name);
 			if (handle == null)
diff --git a/Engine/Source/Programs/Horde/Plugins/Storage/HordeServer.Storage.Tests/Storage/GcServiceTests.cs b/Engine/Source/Programs/Horde/Plugins/Storage/HordeServer.Storage.Tests/Storage/GcServiceTests.cs
--- a/Engine/Source/Programs/Horde/Plugins/Storage/HordeServer.Storage.Tests/Storage/GcServiceTests.cs
+++ b/Engine/Source/Programs/Horde/Plugins/Storage/HordeServer.Storage.Tests/Storage/GcServiceTests.cs
@@ -32,7 +32,7 @@ namespace HordeServer.Tests.Storage
 			globalConfig.Plugins.Add(new PluginName("storage"), storageConfig);
 			SetConfig(globalConfig);
 
-			IStorageClient store = StorageService.CreateClient(new NamespaceId("default"));
+			IStorageNamespace store = StorageService.GetNamespace(new NamespaceId("default"));
 
 			Random random = new Random(0);
 			IHashedBlobRef[] blobs = await CreateTestDataAsync(store, 30, 50, 30, 5, random);
@@ -61,14 +61,14 @@ namespace HordeServer.Tests.Storage
 			Assert.IsTrue(remaining.All(x => nodePaths.Contains(x)));
 		}
 
-		static async Task<HashSet<BlobLocator>> FindNodesAsync(IStorageClient store, IEnumerable<IBlobRef> roots)
+		static async Task<HashSet<BlobLocator>> FindNodesAsync(IStorageNamespace store, IEnumerable<IBlobRef> roots)
 		{
 			HashSet<BlobLocator> nodes = new HashSet<BlobLocator>();
 			await FindNodesAsync(store, roots, nodes);
 			return nodes;
 		}
 
-		static async Task FindNodesAsync(IStorageClient store, IEnumerable<IBlobRef> roots, HashSet<BlobLocator> nodes)
+		static async Task FindNodesAsync(IStorageNamespace store, IEnumerable<IBlobRef> roots, HashSet<BlobLocator> nodes)
 		{
 			foreach (IBlobRef root in roots)
 			{
@@ -81,7 +81,7 @@ namespace HordeServer.Tests.Storage
 			}
 		}
 
-		static async ValueTask<IHashedBlobRef[]> CreateTestDataAsync(IStorageClient store, int numRoots, int numInterior, int numLeaves, int avgChildren, Random random)
+		static async ValueTask<IHashedBlobRef[]> CreateTestDataAsync(IStorageNamespace store, int numRoots, int numInterior, int numLeaves, int avgChildren, Random random)
 		{
 			int firstRoot = 0;
 			int firstInterior = firstRoot + numRoots;
diff --git a/Engine/Source/Programs/Horde/Plugins/Storage/HordeServer.Storage.Tests/Storage/StorageServiceTests.cs b/Engine/Source/Programs/Horde/Plugins/Storage/HordeServer.Storage.Tests/Storage/StorageServiceTests.cs
--- a/Engine/Source/Programs/Horde/Plugins/Storage/HordeServer.Storage.Tests/Storage/StorageServiceTests.cs
+++ b/Engine/Source/Programs/Horde/Plugins/Storage/HordeServer.Storage.Tests/Storage/StorageServiceTests.cs
@@ -21,7 +21,7 @@ namespace HordeServer.Tests.Storage
 		public async Task BlobCollectionTestAsync()
 		{
 			StorageService storageService = ServiceProvider.GetRequiredService<StorageService>();
-			IStorageClient client = storageService.CreateClient(new NamespaceId("memory"));
+			IStorageNamespace client = storageService.GetNamespace(new NamespaceId("memory"));
 
 			BlobType type1 = new BlobType(Guid.Parse("{11C2D886-4164-3349-D1E9-6F943D2ED10B}"), 0);
 			byte[] data1 = new byte[] { 1, 2, 3 };
diff --git a/Engine/Source/Programs/Horde/Plugins/Storage/HordeServer.Storage/Storage/StorageController.cs b/Engine/Source/Programs/Horde/Plugins/Storage/HordeServer.Storage/Storage/StorageController.cs
--- a/Engine/Source/Programs/Horde/Plugins/Storage/HordeServer.Storage/Storage/StorageController.cs
+++ b/Engine/Source/Programs/Horde/Plugins/Storage/HordeServer.Storage/Storage/StorageController.cs
@@ -332,9 +332,9 @@ namespace HordeServer.Storage
 		/// <summary>
 		/// Reads a ref from storage, without performing namespace access checks.
 		/// </summary>
-		public static async Task<ActionResult<ReadRefResponse>> ReadRefInternalAsync(IStorageClientFactory storageService, string basePath, NamespaceId namespaceId, RefName refName, IHeaderDictionary headers, CancellationToken cancellationToken)
+		public static async Task<ActionResult<ReadRefResponse>> ReadRefInternalAsync(IStorageNamespaceFactory storageService, string basePath, NamespaceId namespaceId, RefName refName, IHeaderDictionary headers, CancellationToken cancellationToken)
 		{
-			IStorageClient client = storageService.CreateClient(namespaceId);
+			IStorageNamespace client = storageService.GetNamespace(namespaceId);
 
 			RefCacheTime cacheTime = new RefCacheTime();
 			foreach (string? entry in headers.CacheControl)
@@ -457,11 +457,11 @@ namespace HordeServer.Storage
 				locator = new BlobLocator(locator, String.Join("&", fragments));
 			}
 
-			IStorageClient storageClient = _storageService.CreateClient(namespaceId);
+			IStorageNamespace storageNamespace = _storageService.GetNamespace(namespaceId);
 
 			object content;
 
-			using BlobData blobData = await storageClient.CreateBlobRef(locator).ReadBlobDataAsync(cancellationToken);
+			using BlobData blobData = await storageNamespace.CreateBlobRef(locator).ReadBlobDataAsync(cancellationToken);
 			if (data)
 			{
 				ReadOnlyMemoryStream stream = new ReadOnlyMemoryStream(blobData.Data.ToArray());
diff --git a/Engine/Source/Programs/Horde/Plugins/Storage/HordeServer.Storage/Storage/StorageService.cs b/Engine/Source/Programs/Horde/Plugins/Storage/HordeServer.Storage/Storage/StorageService.cs
--- a/Engine/Source/Programs/Horde/Plugins/Storage/HordeServer.Storage/Storage/StorageService.cs
+++ b/Engine/Source/Programs/Horde/Plugins/Storage/HordeServer.Storage/Storage/StorageService.cs
@@ -32,7 +32,7 @@ namespace HordeServer.Storage
 	/// <summary>
 	/// Interface for the storage service
 	/// </summary>
-	public interface IStorageService : IStorageClientFactory
+	public interface IStorageService : IStorageNamespaceFactory
 	{
 	}
 
@@ -450,26 +450,26 @@ namespace HordeServer.Storage
 
 		internal static ObjectKey GetObjectKey(BlobLocator locator) => new ObjectKey($"{locator.Path}.blob");
 
-		class StorageClientFactory : IStorageClientFactory
+		class StorageNamespaceFactory : IStorageNamespaceFactory
 		{
 			readonly StorageService _storageService;
 			readonly StorageConfig _storageConfig;
 
-			public StorageClientFactory(StorageService storageService, StorageConfig storageConfig)
+			public StorageNamespaceFactory(StorageService storageService, StorageConfig storageConfig)
 			{
 				_storageService = storageService;
 				_storageConfig = storageConfig;
 			}
 
-			public IStorageClient? TryCreateClient(NamespaceId namespaceId)
+			public IStorageNamespace? TryGetNamespace(NamespaceId namespaceId)
 				=> _storageService.TryCreateClient(_storageConfig, namespaceId);
 		}
 
 		/// <summary>
-		/// Creates a new storage client factory using the current global config value
+		/// Creates a new storage namespace factory using the current global config value
 		/// </summary>
-		public IStorageClientFactory CreateStorageClientFactory(StorageConfig storageConfig)
-			=> new StorageClientFactory(this, storageConfig);
+		public IStorageNamespaceFactory CreateStorageNamespaceFactory(StorageConfig storageConfig)
+			=> new StorageNamespaceFactory(this, storageConfig);
 
 		/// <inheritdoc/>
 		public async Task StartAsync(CancellationToken cancellationToken)
@@ -505,17 +505,17 @@ namespace HordeServer.Storage
 		}
 
 		/// <inheritdoc/>
-		public IStorageClient? TryCreateClient(NamespaceId namespaceId)
+		public IStorageNamespace? TryGetNamespace(NamespaceId namespaceId)
 			=> TryCreateClient(namespaceId, null);
 
 		/// <inheritdoc/>
-		public IStorageClient? TryCreateClient(NamespaceId namespaceId, BundleOptions? bundleOptions = null)
+		public IStorageNamespace? TryCreateClient(NamespaceId namespaceId, BundleOptions? bundleOptions = null)
 			=> TryCreateClient(_storageConfig.CurrentValue, namespaceId, bundleOptions);
 
 		/// <inheritdoc/>
-		public IStorageClient? TryCreateClient(StorageConfig storageConfig, NamespaceId namespaceId, BundleOptions? bundleOptions = null)
+		public IStorageNamespace? TryCreateClient(StorageConfig storageConfig, NamespaceId namespaceId, BundleOptions? bundleOptions = null)
 		{
-#pragma warning disable CA2000 // Call dispose on backend; will be disposed by BundleStorageClient
+#pragma warning disable CA2000 // Call dispose on backend; will be disposed by BundleStorageNamespace
 			IStorageBackend? backend = TryCreateBackend(storageConfig, namespaceId);
 #pragma warning restore CA2000
 			if (backend == null)
@@ -524,7 +524,7 @@ namespace HordeServer.Storage
 			}
 			else
 			{
-				return new BundleStorageClient(backend, _bundleCache, bundleOptions, _logger);
+				return new BundleStorageNamespace(backend, _bundleCache, bundleOptions, _logger);
 			}
 		}
 
@@ -686,7 +686,7 @@ namespace HordeServer.Storage
 			// Get the current state of the storage system
 			State state = CreateState(_storageConfig.CurrentValue);
 
-			Dictionary<NamespaceId, BundleStorageClient> cachedClients = new();
+			Dictionary<NamespaceId, BundleStorageNamespace> cachedClients = new();
 
 			long ingestedCount = 0;
 
@@ -1071,7 +1071,7 @@ namespace HordeServer.Storage
 
 		async Task TickGcForNamespaceAsync(NamespaceInfo namespaceInfo, ObjectId lastImportBlobInfoId, DateTime utcNow, bool deleteObjects, CancellationToken cancellationToken)
 		{
-			IStorageClient client = this.CreateClient(namespaceInfo.Id);
+			IStorageNamespace client = this.GetNamespace(namespaceInfo.Id);
 
 			Stopwatch timer = Stopwatch.StartNew();
 			_logger.LogInformation("Running garbage collection for namespace {NamespaceId}...", namespaceInfo.Id);
diff --git a/Engine/Source/Programs/Horde/Plugins/Storage/HordeServer.Storage/StoragePlugin.cs b/Engine/Source/Programs/Horde/Plugins/Storage/HordeServer.Storage/StoragePlugin.cs
--- a/Engine/Source/Programs/Horde/Plugins/Storage/HordeServer.Storage/StoragePlugin.cs
+++ b/Engine/Source/Programs/Horde/Plugins/Storage/HordeServer.Storage/StoragePlugin.cs
@@ -40,7 +40,7 @@ namespace HordeServer.Storage
 		{
 			services.AddSingleton<StorageService>();
 			services.AddSingleton<IStorageService>(sp => sp.GetRequiredService<StorageService>());
-			services.AddScoped(sp => sp.GetRequiredService<StorageService>().CreateStorageClientFactory(sp.GetRequiredService<IOptionsSnapshot<StorageConfig>>().Value));
+			services.AddScoped(sp => sp.GetRequiredService<StorageService>().CreateStorageNamespaceFactory(sp.GetRequiredService<IOptionsSnapshot<StorageConfig>>().Value));
 
 			services.AddSingleton<IObjectStoreFactory, ObjectStoreFactory>();
 			services.AddSingleton<AwsObjectStoreFactory>();
diff --git a/Engine/Source/Programs/Horde/Plugins/Symbols/HordeServer.Symbols/Symbols/SymbolsController.cs b/Engine/Source/Programs/Horde/Plugins/Symbols/HordeServer.Symbols/Symbols/SymbolsController.cs
--- a/Engine/Source/Programs/Horde/Plugins/Symbols/HordeServer.Symbols/Symbols/SymbolsController.cs
+++ b/Engine/Source/Programs/Horde/Plugins/Symbols/HordeServer.Symbols/Symbols/SymbolsController.cs
@@ -51,8 +51,8 @@ namespace HordeServer.Tools
 				return Forbid(SymbolStoreAclAction.ReadSymbols, storeId);
 			}
 
-			IStorageClient storageClient = _storageService.CreateClient(symbolStoreConfig.NamespaceId);
-			BlobAlias? alias = await storageClient.FindAliasAsync($"sym:{path.ToUpperInvariant()}", cancellationToken);
+			IStorageNamespace storageNamespace = _storageService.GetNamespace(symbolStoreConfig.NamespaceId);
+			BlobAlias? alias = await storageNamespace.FindAliasAsync($"sym:{path.ToUpperInvariant()}", cancellationToken);
 			if (alias == null)
 			{
 				return NotFound();
diff --git a/Engine/Source/Programs/Horde/Plugins/Tools/HordeServer.Tools/Commands/Generate/ToolDataCommand.cs b/Engine/Source/Programs/Horde/Plugins/Tools/HordeServer.Tools/Commands/Generate/ToolDataCommand.cs
--- a/Engine/Source/Programs/Horde/Plugins/Tools/HordeServer.Tools/Commands/Generate/ToolDataCommand.cs
+++ b/Engine/Source/Programs/Horde/Plugins/Tools/HordeServer.Tools/Commands/Generate/ToolDataCommand.cs
@@ -73,7 +73,7 @@ namespace HordeServer.Commands.Generate
 			await using (BundleCache bundleCache = new BundleCache())
 			using (MemoryMappedFileCache memoryMappedFileCache = new MemoryMappedFileCache())
 			{
-				IStorageClient client = BundleStorageClient.CreateFromDirectory(bundleDir, bundleCache, memoryMappedFileCache, logger);
+				IStorageNamespace client = BundleStorageNamespace.CreateFromDirectory(bundleDir, bundleCache, memoryMappedFileCache, logger);
 
 				IHashedBlobRef<DirectoryNode> dirNodeRef;
 				await using (DedupeBlobWriter writer = client.CreateDedupeBlobWriter(refName))
@@ -138,7 +138,7 @@ namespace HordeServer.Commands.Generate
 			return 0;
 		}
 
-		static async Task PopulateCacheAsync(IStorageClient client, DedupeBlobWriter writer, DirectoryReference searchDir, CancellationToken cancellationToken)
+		static async Task PopulateCacheAsync(IStorageNamespace client, DedupeBlobWriter writer, DirectoryReference searchDir, CancellationToken cancellationToken)
 		{
 			foreach (RefName refName in FileStorageBackend.EnumerateRefs(searchDir))
 			{
diff --git a/Engine/Source/Programs/Horde/Plugins/Tools/HordeServer.Tools/Tools/ToolCollection.cs b/Engine/Source/Programs/Horde/Plugins/Tools/HordeServer.Tools/Tools/ToolCollection.cs
--- a/Engine/Source/Programs/Horde/Plugins/Tools/HordeServer.Tools/Tools/ToolCollection.cs
+++ b/Engine/Source/Programs/Horde/Plugins/Tools/HordeServer.Tools/Tools/ToolCollection.cs
@@ -129,11 +129,11 @@ namespace HordeServer.Tools
 				return _collection.CreateToolObject(document);
 			}
 
-			public IStorageBackend CreateStorageBackend()
+			public IStorageBackend GetStorageBackend()
 				=> _collection.CreateStorageBackend(_config);
 
-			public IStorageClient CreateStorageClient()
-				=> _collection.CreateStorageClient(_config);
+			public IStorageNamespace GetStorageNamespace()
+				=> _collection.GetStorageNamespace(_config);
 		}
 
 		class ToolDeployment : IToolDeployment
@@ -332,7 +332,7 @@ namespace HordeServer.Tools
 		{
 			ToolDeploymentId deploymentId = new ToolDeploymentId(BinaryIdUtils.CreateNew());
 
-			IStorageClient client = _storageService.CreateClient(toolConfig.NamespaceId);
+			IStorageNamespace client = _storageService.GetNamespace(toolConfig.NamespaceId);
 
 			IHashedBlobRef<DirectoryNode> nodeRef;
 			await using (IBlobWriter writer = client.CreateBlobWriter($"{tool.Id}/{deploymentId}"))
@@ -349,11 +349,11 @@ namespace HordeServer.Tools
 		{
 			ToolDeploymentId deploymentId = new ToolDeploymentId(BinaryIdUtils.CreateNew());
 
-			IStorageClient client = _storageService.CreateClient(toolConfig.NamespaceId);
+			IStorageNamespace client = _storageService.GetNamespace(toolConfig.NamespaceId);
 			return await CreateDeploymentInternalAsync(tool, toolConfig, deploymentId, options, client, client.CreateBlobRef(target), cancellationToken);
 		}
 
-		async Task<ToolDocument?> CreateDeploymentInternalAsync(ToolDocument tool, ToolConfig toolConfig, ToolDeploymentId deploymentId, ToolDeploymentConfig options, IStorageClient storageClient, IHashedBlobRef content, CancellationToken cancellationToken)
+		async Task<ToolDocument?> CreateDeploymentInternalAsync(ToolDocument tool, ToolConfig toolConfig, ToolDeploymentId deploymentId, ToolDeploymentConfig options, IStorageNamespace storageNamespace, IHashedBlobRef content, CancellationToken cancellationToken)
 		{
 			if (toolConfig is BundledToolConfig)
 			{
@@ -362,7 +362,7 @@ namespace HordeServer.Tools
 
 			// Write a ref for the deployment so the blobs aren't GC'd
 			RefName refName = new RefName($"{tool.Id}/{deploymentId}");
-			await storageClient.WriteRefAsync(refName, content, cancellationToken: cancellationToken);
+			await storageNamespace.WriteRefAsync(refName, content, cancellationToken: cancellationToken);
 
 			// Create the new deployment object
 			ToolDeploymentDocument deployment = new ToolDeploymentDocument(deploymentId, options, toolConfig.NamespaceId, refName);
@@ -407,7 +407,7 @@ namespace HordeServer.Tools
 				}
 
 				ToolDeploymentDocument removeDeployment = tool.Deployments[0];
-				IStorageClient client = _storageService.CreateClient(removeDeployment.NamespaceId);
+				IStorageNamespace client = _storageService.GetNamespace(removeDeployment.NamespaceId);
 				await client.DeleteRefAsync(removeDeployment.RefName, cancellationToken);
 			}
 
@@ -459,24 +459,24 @@ namespace HordeServer.Tools
 		}
 
 		/// <summary>
-		/// Gets the storage client containing data for a particular tool
+		/// Gets the storage namespace containing data for a particular tool
 		/// </summary>
 		/// <param name="toolConfig">Identifier for the tool</param>
-		/// <returns>Storage client for the data</returns>
-		IStorageClient CreateStorageClient(ToolConfig toolConfig)
+		/// <returns>Storage namespace for the data</returns>
+		IStorageNamespace GetStorageNamespace(ToolConfig toolConfig)
 		{
 			if (toolConfig is BundledToolConfig bundledConfig)
 			{
-				return BundleStorageClient.CreateFromDirectory(DirectoryReference.Combine(_serverInfo.AppDir, bundledConfig.DataDir ?? "Tools"), _bundleCache, _memoryMappedFileCache, _logger);
+				return BundleStorageNamespace.CreateFromDirectory(DirectoryReference.Combine(_serverInfo.AppDir, bundledConfig.DataDir ?? "Tools"), _bundleCache, _memoryMappedFileCache, _logger);
 			}
 			else
 			{
-				return _storageService.CreateClient(toolConfig.NamespaceId);
+				return _storageService.GetNamespace(toolConfig.NamespaceId);
 			}
 		}
 
 		/// <summary>
-		/// Gets the storage client containing data for a particular tool
+		/// Gets the storage backend containing data for a particular tool
 		/// </summary>
 		/// <param name="toolConfig">Identifier for the tool</param>
 		/// <returns>Storage client for the data</returns>
@@ -494,7 +494,7 @@ namespace HordeServer.Tools
 
 		IBlobRef<DirectoryNode> Open(ToolConfig tool, ToolDeploymentDocument deployment)
 		{
-			IStorageClient client = CreateStorageClient(tool);
+			IStorageNamespace client = GetStorageNamespace(tool);
 			return client.CreateBlobRef<DirectoryNode>(deployment.RefName, DateTime.UtcNow - TimeSpan.FromDays(2.0));
 		}
 
@@ -508,7 +508,7 @@ namespace HordeServer.Tools
 		async Task<Stream> GetDeploymentZipAsync(ToolConfig tool, ToolDeploymentDocument deployment, CancellationToken cancellationToken)
 		{
 #pragma warning disable CA2000
-			IStorageClient client = CreateStorageClient(tool);
+			IStorageNamespace client = GetStorageNamespace(tool);
 
 			IHashedBlobRef<DirectoryNode> nodeRef = await client.ReadRefAsync<DirectoryNode>(deployment.RefName, DateTime.UtcNow - TimeSpan.FromDays(2.0), cancellationToken: cancellationToken);
 			return nodeRef.AsZipStream();
diff --git a/Engine/Source/Programs/Horde/Plugins/Tools/HordeServer.Tools/Tools/ToolsController.cs b/Engine/Source/Programs/Horde/Plugins/Tools/HordeServer.Tools/Tools/ToolsController.cs
--- a/Engine/Source/Programs/Horde/Plugins/Tools/HordeServer.Tools/Tools/ToolsController.cs
+++ b/Engine/Source/Programs/Horde/Plugins/Tools/HordeServer.Tools/Tools/ToolsController.cs
@@ -61,7 +61,7 @@ namespace HordeServer.Tools
 				request.Prefix = $"{id}/{request.Prefix}";
 			}
 
-			IStorageBackend storageBackend = tool.CreateStorageBackend();
+			IStorageBackend storageBackend = tool.GetStorageBackend();
 			return await StorageController.WriteBlobAsync(storageBackend, request, cancellationToken);
 		}
 
@@ -317,7 +317,7 @@ namespace HordeServer.Tools
 				return Ok(response);
 			}
 
-			IStorageClient client = tool.CreateStorageClient();
+			IStorageNamespace client = tool.GetStorageNamespace();
 			IHashedBlobRef<DirectoryNode> nodeRef = await client.ReadRefAsync<DirectoryNode>(deployment.RefName, DateTime.UtcNow - TimeSpan.FromDays(2.0), cancellationToken: cancellationToken);
 
 			// If we weren't specifically asked for a zip, see if this download is a single file. If it is, allow downloading it directory.
@@ -347,7 +347,7 @@ namespace HordeServer.Tools
 
 		static async Task<GetToolDeploymentResponse> GetDeploymentInfoResponseAsync(ITool tool, IToolDeployment deployment, CancellationToken cancellationToken)
 		{
-			IStorageClient client = tool.CreateStorageClient();
+			IStorageNamespace client = tool.GetStorageNamespace();
 			IBlobRef rootHandle = await client.ReadRefAsync(deployment.RefName, cancellationToken: cancellationToken);
 
 			return new GetToolDeploymentResponse(deployment.Id, deployment.Version, deployment.State, deployment.Progress, deployment.StartedAt, deployment.Duration, deployment.RefName, rootHandle.GetLocator());
@@ -379,7 +379,7 @@ namespace HordeServer.Tools
 				return BadRequest("Invalid blob id for tool");
 			}
 
-			IStorageBackend storageBackend = tool.CreateStorageBackend();
+			IStorageBackend storageBackend = tool.GetStorageBackend();
 			return await StorageController.ReadBlobInternalAsync(storageBackend, locator, Request.Headers, cancellationToken);
 		}
 
diff --git a/Engine/Source/Programs/Horde/Samples/RemoteClient/RemoteClientApp.cs b/Engine/Source/Programs/Horde/Samples/RemoteClient/RemoteClientApp.cs
--- a/Engine/Source/Programs/Horde/Samples/RemoteClient/RemoteClientApp.cs
+++ b/Engine/Source/Programs/Horde/Samples/RemoteClient/RemoteClientApp.cs
@@ -110,7 +110,7 @@ namespace RemoteClient
 				await channel.ForkAsync(BackgroundChannelId, 4 * 1024 * 1024, default);
 
 				// Upload the sandbox to the primary channel.
-				BundleStorageClient storage = BundleStorageClient.CreateInMemory(logger);
+				BundleStorageNamespace storage = BundleStorageNamespace.CreateInMemory(logger);
 
 				await using (IBlobWriter writer = storage.CreateBlobWriter())
 				{
diff --git a/Engine/Source/Programs/Shared/EpicGames.Horde.Tests/BlobPipelineTests.cs b/Engine/Source/Programs/Shared/EpicGames.Horde.Tests/BlobPipelineTests.cs
--- a/Engine/Source/Programs/Shared/EpicGames.Horde.Tests/BlobPipelineTests.cs
+++ b/Engine/Source/Programs/Shared/EpicGames.Horde.Tests/BlobPipelineTests.cs
@@ -20,7 +20,7 @@ namespace EpicGames.Horde.Tests
 		[TestMethod]
 		public async Task TestAsync()
 		{
-			BundleStorageClient client = BundleStorageClient.CreateInMemory(NullLogger.Instance);
+			BundleStorageNamespace client = BundleStorageNamespace.CreateInMemory(NullLogger.Instance);
 
 			List<IHashedBlobRef> handles = new List<IHashedBlobRef>();
 			await using (IBlobWriter writer = client.CreateBlobWriter())
diff --git a/Engine/Source/Programs/Shared/EpicGames.Horde.Tests/BundleStoreTests.cs b/Engine/Source/Programs/Shared/EpicGames.Horde.Tests/BundleStoreTests.cs
--- a/Engine/Source/Programs/Shared/EpicGames.Horde.Tests/BundleStoreTests.cs
+++ b/Engine/Source/Programs/Shared/EpicGames.Horde.Tests/BundleStoreTests.cs
@@ -61,7 +61,7 @@ namespace EpicGames.Horde.Tests
 		static async Task<byte[]> CreateBundleNormalAsync()
 		{
 			MemoryStorageBackend memoryStore = new MemoryStorageBackend();
-			BundleStorageClient store = new BundleStorageClient(memoryStore, BundleCache.None, new BundleOptions { MaxVersion = BundleVersion.ImportHashes, CompressionFormat = BundleCompressionFormat.None }, NullLogger.Instance);
+			BundleStorageNamespace store = new BundleStorageNamespace(memoryStore, BundleCache.None, new BundleOptions { MaxVersion = BundleVersion.ImportHashes, CompressionFormat = BundleCompressionFormat.None }, NullLogger.Instance);
 			await using IBlobWriter writer = store.CreateBlobWriter();
 
 			TextNode node = new TextNode("Hello world");
@@ -101,7 +101,7 @@ namespace EpicGames.Horde.Tests
 			MemoryStorageBackend blobStore = new MemoryStorageBackend();
 
 			BundleOptions bundleOptions = new BundleOptions { MaxBlobSize = 1024 * 1024 };
-			BundleStorageClient bundleStore = new BundleStorageClient(blobStore, BundleCache.None, bundleOptions, NullLogger.Instance);
+			BundleStorageNamespace bundleStore = new BundleStorageNamespace(blobStore, BundleCache.None, bundleOptions, NullLogger.Instance);
 
 			await TestTreeAsync(bundleStore);
 
@@ -115,7 +115,7 @@ namespace EpicGames.Horde.Tests
 			MemoryStorageBackend blobStore = new MemoryStorageBackend();
 
 			BundleOptions bundleOptions = new BundleOptions { MaxBlobSize = 1 };
-			BundleStorageClient bundleStore = new BundleStorageClient(blobStore, BundleCache.None, bundleOptions, NullLogger.Instance);
+			BundleStorageNamespace bundleStore = new BundleStorageNamespace(blobStore, BundleCache.None, bundleOptions, NullLogger.Instance);
 
 			await TestTreeAsync(bundleStore);
 
@@ -155,7 +155,7 @@ namespace EpicGames.Horde.Tests
 			}
 		}
 
-		static async Task TestTreeAsync(BundleStorageClient store)
+		static async Task TestTreeAsync(BundleStorageNamespace store)
 		{
 			// Generate a tree
 			{
@@ -219,7 +219,7 @@ namespace EpicGames.Horde.Tests
 		[TestMethod]
 		public async Task SimpleNodeAsync()
 		{
-			KeyValueStorageClient store = KeyValueStorageClient.CreateInMemory();
+			KeyValueStorageNamespace store = KeyValueStorageNamespace.CreateInMemory();
 
 			RefName refName = new RefName("test");
 
@@ -238,7 +238,7 @@ namespace EpicGames.Horde.Tests
 		[TestMethod]
 		public async Task DirectoryNodesAsync()
 		{
-			KeyValueStorageClient store = KeyValueStorageClient.CreateInMemory();
+			KeyValueStorageNamespace store = KeyValueStorageNamespace.CreateInMemory();
 
 			// Generate a tree
 			{
@@ -285,7 +285,7 @@ namespace EpicGames.Horde.Tests
 		public async Task FileNodesAsync()
 		{
 			using IMemoryCache cache = new MemoryCache(new MemoryCacheOptions());
-			BundleStorageClient store = BundleStorageClient.CreateInMemory(NullLogger.Instance);
+			BundleStorageNamespace store = BundleStorageNamespace.CreateInMemory(NullLogger.Instance);
 			BundleReader reader = new BundleReader(store, BundleCache.None, NullLogger.Instance);
 
 			// Generate a tree
@@ -334,7 +334,7 @@ namespace EpicGames.Horde.Tests
 		public async Task StreamTestAsync()
 		{
 			BundleOptions bundleOptions = new BundleOptions { MaxBlobSize = 1024 };
-			BundleStorageClient store = BundleStorageClient.CreateInMemory(bundleOptions, NullLogger.Instance);
+			BundleStorageNamespace store = BundleStorageNamespace.CreateInMemory(bundleOptions, NullLogger.Instance);
 
 			const int Length = 4096;
 
@@ -374,7 +374,7 @@ namespace EpicGames.Horde.Tests
 		public async Task LargeFileTestAsync()
 		{
 			BundleOptions bundleOptions = new BundleOptions { MaxBlobSize = 1024 };
-			BundleStorageClient store = BundleStorageClient.CreateInMemory(bundleOptions, NullLogger.Instance);
+			BundleStorageNamespace store = BundleStorageNamespace.CreateInMemory(bundleOptions, NullLogger.Instance);
 
 			const int Length = 1024;
 			const int Copies = 4096;
diff --git a/Engine/Source/Programs/Shared/EpicGames.Horde.Tests/BundleTests.cs b/Engine/Source/Programs/Shared/EpicGames.Horde.Tests/BundleTests.cs
--- a/Engine/Source/Programs/Shared/EpicGames.Horde.Tests/BundleTests.cs
+++ b/Engine/Source/Programs/Shared/EpicGames.Horde.Tests/BundleTests.cs
@@ -52,7 +52,7 @@ namespace EpicGames.Horde.Tests
 		[TestMethod]
 		public async Task BasicChunkingTestsAsync()
 		{
-			BundleStorageClient storage = BundleStorageClient.CreateInMemory(NullLogger.Instance);
+			BundleStorageNamespace storage = BundleStorageNamespace.CreateInMemory(NullLogger.Instance);
 
 			RefName refName = new RefName("test");
 			await using IBlobWriter writer = storage.CreateBlobWriter(refName);
@@ -175,7 +175,7 @@ namespace EpicGames.Horde.Tests
 		[TestMethod]
 		public async Task BasicTestDirectoryAsync()
 		{
-			BundleStorageClient store = BundleStorageClient.CreateInMemory(NullLogger.Instance);
+			BundleStorageNamespace store = BundleStorageNamespace.CreateInMemory(NullLogger.Instance);
 
 			IHashedBlobRef<DirectoryNode> rootRef;
 			await using (IBlobWriter writer = store.CreateBlobWriter())
@@ -235,7 +235,7 @@ namespace EpicGames.Horde.Tests
 			BundleOptions bundleOptions = new BundleOptions();
 			bundleOptions.MaxBlobSize = 1;
 
-			BundleStorageClient storage = BundleStorageClient.CreateInMemory(bundleOptions, NullLogger.Instance);
+			BundleStorageNamespace storage = BundleStorageNamespace.CreateInMemory(bundleOptions, NullLogger.Instance);
 
 			await using (IBlobWriter writer = new DedupeBlobWriter(storage.CreateBlobWriter()))
 			{
@@ -263,7 +263,7 @@ namespace EpicGames.Horde.Tests
 			BundleOptions bundleOptions = new BundleOptions();
 			bundleOptions.MaxBlobSize = 1;
 
-			BundleStorageClient storage = BundleStorageClient.CreateInMemory(bundleOptions, NullLogger.Instance);
+			BundleStorageNamespace storage = BundleStorageNamespace.CreateInMemory(bundleOptions, NullLogger.Instance);
 
 			RefName refName = new RefName("ref");
 
diff --git a/Engine/Source/Programs/Shared/EpicGames.Horde.Tests/ChunkingTests.cs b/Engine/Source/Programs/Shared/EpicGames.Horde.Tests/ChunkingTests.cs
--- a/Engine/Source/Programs/Shared/EpicGames.Horde.Tests/ChunkingTests.cs
+++ b/Engine/Source/Programs/Shared/EpicGames.Horde.Tests/ChunkingTests.cs
@@ -52,7 +52,7 @@ namespace EpicGames.Horde.Tests
 		[TestMethod]
 		public async Task EmptyNodeTestAsync()
 		{
-			KeyValueStorageClient store = KeyValueStorageClient.CreateInMemory();
+			KeyValueStorageNamespace store = KeyValueStorageNamespace.CreateInMemory();
 
 			const string RefName = "hello";
 			await using (IBlobWriter writer = store.CreateBlobWriter(RefName))
@@ -97,7 +97,7 @@ namespace EpicGames.Horde.Tests
 		{
 			using MemoryCache cache = new MemoryCache(new MemoryCacheOptions());
 
-			KeyValueStorageClient store = KeyValueStorageClient.CreateInMemory();
+			KeyValueStorageNamespace store = KeyValueStorageNamespace.CreateInMemory();
 
 			await using IBlobWriter writer = store.CreateBlobWriter();
 
@@ -176,7 +176,7 @@ namespace EpicGames.Horde.Tests
 			BlobSerializerOptions serializerOptions = new BlobSerializerOptions();
 			serializerOptions.Converters.Add(new InteriorChunkedDataNodeConverter(HordeApiVersion.Initial)); // Does not include length fields in interior nodes
 
-			KeyValueStorageClient store = KeyValueStorageClient.CreateInMemory();
+			KeyValueStorageNamespace store = KeyValueStorageNamespace.CreateInMemory();
 
 			byte[] data = Encoding.UTF8.GetBytes("hello world");
 
diff --git a/Engine/Source/Programs/Shared/EpicGames.Horde.Tests/ComputeSocketTests.cs b/Engine/Source/Programs/Shared/EpicGames.Horde.Tests/ComputeSocketTests.cs
--- a/Engine/Source/Programs/Shared/EpicGames.Horde.Tests/ComputeSocketTests.cs
+++ b/Engine/Source/Programs/Shared/EpicGames.Horde.Tests/ComputeSocketTests.cs
@@ -175,7 +175,7 @@ namespace EpicGames.Horde.Tests
 						await channel2.CloseAsync(cancellationToken);
 					}
 
-					BundleStorageClient storage = BundleStorageClient.CreateInMemory(NullLogger.Instance);
+					BundleStorageNamespace storage = BundleStorageNamespace.CreateInMemory(NullLogger.Instance);
 					await using (IBlobWriter blobWriter = storage.CreateBlobWriter())
 					{
 						FileReference file = FileReference.Combine(tempDir, "subdir/hello.txt");
diff --git a/Engine/Source/Programs/Shared/EpicGames.Horde.Tests/LogTests.cs b/Engine/Source/Programs/Shared/EpicGames.Horde.Tests/LogTests.cs
--- a/Engine/Source/Programs/Shared/EpicGames.Horde.Tests/LogTests.cs
+++ b/Engine/Source/Programs/Shared/EpicGames.Horde.Tests/LogTests.cs
@@ -40,7 +40,7 @@ namespace EpicGames.Horde.Tests
 		[TestMethod]
 		public async Task IndexTestsAsync()
 		{
-			KeyValueStorageClient store = KeyValueStorageClient.CreateInMemory();
+			KeyValueStorageNamespace store = KeyValueStorageNamespace.CreateInMemory();
 
 			// Write the test data to the log file in blocks
 			LogBuilder builder = new LogBuilder(LogFormat.Text, 1, 1, NullLogger.Instance);
@@ -101,7 +101,7 @@ namespace EpicGames.Horde.Tests
 		[TestMethod]
 		public async Task PartialTokenTestsAsync()
 		{
-			IStorageClient store = KeyValueStorageClient.CreateInMemory();
+			IStorageNamespace store = KeyValueStorageNamespace.CreateInMemory();
 
 			// Generate the test data
 			string[] lines =
diff --git a/Engine/Source/Programs/Shared/EpicGames.Horde.Tests/ServerLogPacketBuilderTests.cs b/Engine/Source/Programs/Shared/EpicGames.Horde.Tests/ServerLogPacketBuilderTests.cs
--- a/Engine/Source/Programs/Shared/EpicGames.Horde.Tests/ServerLogPacketBuilderTests.cs
+++ b/Engine/Source/Programs/Shared/EpicGames.Horde.Tests/ServerLogPacketBuilderTests.cs
@@ -214,7 +214,7 @@ namespace EpicGames.Horde.Tests
 		class FakeHordeClient : IHordeClient
 		{
 			public FakeLogRpcClient LogRpc { get; } = new FakeLogRpcClient();
-			public Dictionary<string, BundleStorageClient> StorageClients { get; } = new Dictionary<string, BundleStorageClient>();
+			public Dictionary<string, BundleStorageNamespace> StorageNamespaces { get; } = new Dictionary<string, BundleStorageNamespace>();
 
 			public Uri ServerUrl => throw new NotImplementedException();
 
@@ -230,20 +230,20 @@ namespace EpicGames.Horde.Tests
 			public HordeHttpClient CreateHttpClient()
 				=> throw new NotImplementedException();
 
-			public IStorageClient CreateStorageClient(string relativePath, string? accessToken = null)
+			public IStorageNamespace GetStorageNamespace(string relativePath, string? accessToken = null)
 			{
-				BundleStorageClient? storageClient;
-				if (!StorageClients.TryGetValue(relativePath, out storageClient))
+				BundleStorageNamespace? storageNamespace;
+				if (!StorageNamespaces.TryGetValue(relativePath, out storageNamespace))
 				{
-					storageClient = BundleStorageClient.CreateInMemory(NullLogger.Instance);
-					StorageClients.Add(relativePath, storageClient);
+					storageNamespace = BundleStorageNamespace.CreateInMemory(NullLogger.Instance);
+					StorageNamespaces.Add(relativePath, storageNamespace);
 				}
-				return storageClient;
+				return storageNamespace;
 			}
 
 			public ValueTask DisposeAsync()
 			{
-				StorageClients.Clear();
+				StorageNamespaces.Clear();
 				return default;
 			}
 
@@ -287,8 +287,8 @@ namespace EpicGames.Horde.Tests
 			}
 
 			// Read the log
-			IStorageClient storageClient = hordeClient.CreateStorageClient(logId);
-			LogNode file = await storageClient.CreateBlobRef(hordeClient.LogRpc.Logs[logId]).ReadBlobAsync<LogNode>();
+			IStorageNamespace storageNamespace = hordeClient.GetStorageNamespace(logId);
+			LogNode file = await storageNamespace.CreateBlobRef(hordeClient.LogRpc.Logs[logId]).ReadBlobAsync<LogNode>();
 
 			// Check the index text
 			List<Utf8String> extractedIndexText = new List<Utf8String>();
diff --git a/Engine/Source/Programs/Shared/EpicGames.Horde.Tests/StorageClientTests.cs b/Engine/Source/Programs/Shared/EpicGames.Horde.Tests/StorageClientTests.cs
--- a/Engine/Source/Programs/Shared/EpicGames.Horde.Tests/StorageClientTests.cs
+++ b/Engine/Source/Programs/Shared/EpicGames.Horde.Tests/StorageClientTests.cs
@@ -13,7 +13,7 @@ using Microsoft.VisualStudio.TestTools.UnitTesting;
 namespace EpicGames.Horde.Tests
 {
 	[TestClass]
-	public class StorageClientTests
+	public class StorageNamespaceTests
 	{
 		[BlobConverter(typeof(TestNodeConverter))]
 		class TestNode
@@ -62,18 +62,18 @@ namespace EpicGames.Horde.Tests
 		[TestMethod]
 		public async Task TestBasicAsync()
 		{
-			KeyValueStorageClient store = KeyValueStorageClient.CreateInMemory();
+			KeyValueStorageNamespace store = KeyValueStorageNamespace.CreateInMemory();
 			await TestBasicAsync(store);
 		}
 
 		[TestMethod]
 		public async Task TestBasicBundleV2Async()
 		{
-			BundleStorageClient storeV2 = BundleStorageClient.CreateInMemory(NullLogger.Instance);
+			BundleStorageNamespace storeV2 = BundleStorageNamespace.CreateInMemory(NullLogger.Instance);
 			await TestBasicAsync(storeV2);
 		}
 
-		static async Task TestBasicAsync(IStorageClient store)
+		static async Task TestBasicAsync(IStorageNamespace store)
 		{
 			IHashedBlobRef<TestNode> nodeRef;
 			await using (IBlobWriter writer = store.CreateBlobWriter())
@@ -91,7 +91,7 @@ namespace EpicGames.Horde.Tests
 		{
 			await using BundleCache cache = new BundleCache();
 
-			BundleStorageClient store = BundleStorageClient.CreateInMemory(NullLogger.Instance);
+			BundleStorageNamespace store = BundleStorageNamespace.CreateInMemory(NullLogger.Instance);
 
 			IHashedBlobRef<TestNode> nodeRef2;
 			await using (IBlobWriter writer = store.CreateBlobWriter())
@@ -116,7 +116,7 @@ namespace EpicGames.Horde.Tests
 			await using BundleCache cache = new BundleCache();
 
 			BundleOptions bundleOptions = new BundleOptions { MinCompressionPacketSize = 100, MaxBlobSize = 1024 * 1024, MaxVersion = BundleVersion.LatestV2 };
-			BundleStorageClient store = BundleStorageClient.CreateInMemory(bundleOptions, NullLogger.Instance);
+			BundleStorageNamespace store = BundleStorageNamespace.CreateInMemory(bundleOptions, NullLogger.Instance);
 
 			await using IBlobWriter writer = store.CreateBlobWriter();
 			IHashedBlobRef<TestNode> nodeRef1 = await writer.WriteBlobAsync(new TestNode(123) { Padding = new byte[1024] });
diff --git a/Engine/Source/Programs/Shared/EpicGames.Horde/Artifacts/HttpArtifactCollection.cs b/Engine/Source/Programs/Shared/EpicGames.Horde/Artifacts/HttpArtifactCollection.cs
--- a/Engine/Source/Programs/Shared/EpicGames.Horde/Artifacts/HttpArtifactCollection.cs
+++ b/Engine/Source/Programs/Shared/EpicGames.Horde/Artifacts/HttpArtifactCollection.cs
@@ -68,7 +68,7 @@ namespace EpicGames.Horde.Artifacts
 
 		IBlobRef<DirectoryNode> Open(Artifact artifact)
 		{
-			IStorageClient store = _hordeClient.CreateStorageClient(artifact.Id);
+			IStorageNamespace store = _hordeClient.GetStorageNamespace(artifact.Id);
 			return store.CreateBlobRef<DirectoryNode>(new RefName("default"));
 		}
 
diff --git a/Engine/Source/Programs/Shared/EpicGames.Horde/Compute/AgentMessage.cs b/Engine/Source/Programs/Shared/EpicGames.Horde/Compute/AgentMessage.cs
--- a/Engine/Source/Programs/Shared/EpicGames.Horde/Compute/AgentMessage.cs
+++ b/Engine/Source/Programs/Shared/EpicGames.Horde/Compute/AgentMessage.cs
@@ -725,7 +725,7 @@ namespace EpicGames.Horde.Compute
 		/// </summary>
 		/// <param name="channel">Channel to write to</param>
 		/// <param name="message">The read request</param>
-		/// <param name="storage">Storage client to retrieve the blob from</param>
+		/// <param name="storage">Storage backend to retrieve the blob from</param>
 		/// <param name="cancellationToken">Cancellation token for the operation</param>
 		public static Task SendBlobDataAsync(this AgentMessageChannel channel, ReadBlobMessage message, IStorageBackend storage, CancellationToken cancellationToken = default)
 		{
@@ -739,7 +739,7 @@ namespace EpicGames.Horde.Compute
 		/// <param name="locator">Locator for the blob to send</param>
 		/// <param name="offset">Starting offset of the data</param>
 		/// <param name="length">Length of the data</param>
-		/// <param name="storage">Storage client to retrieve the blob from</param>
+		/// <param name="storage">Storage backend to retrieve the blob from</param>
 		/// <param name="cancellationToken">Cancellation token for the operation</param>
 		public static async Task SendBlobDataAsync(this AgentMessageChannel channel, BlobLocator locator, int offset, int length, IStorageBackend storage, CancellationToken cancellationToken = default)
 		{
diff --git a/Engine/Source/Programs/Shared/EpicGames.Horde/Compute/AgentMessageHandler.cs b/Engine/Source/Programs/Shared/EpicGames.Horde/Compute/AgentMessageHandler.cs
--- a/Engine/Source/Programs/Shared/EpicGames.Horde/Compute/AgentMessageHandler.cs
+++ b/Engine/Source/Programs/Shared/EpicGames.Horde/Compute/AgentMessageHandler.cs
@@ -191,7 +191,7 @@ namespace EpicGames.Horde.Compute
 			await using BundleCache cache = new BundleCache(new BundleCacheOptions { HeaderCacheSize = 10 * 1024 * 1024, PacketCacheSize = 128 * 1024 * 1024 });
 
 			BundleOptions bundleOptions = ComputeProtocolUtilities.GetBundleOptions(channel.Protocol);
-			BundleStorageClient store = new BundleStorageClient(innerStore, cache, bundleOptions, _logger);
+			BundleStorageNamespace store = new BundleStorageNamespace(innerStore, cache, bundleOptions, _logger);
 
 			IBlobRef<DirectoryNode> directoryRef = store.CreateBlobRef<DirectoryNode>(locator, options);
 
diff --git a/Engine/Source/Programs/Shared/EpicGames.Horde/Compute/AgentStorageClient.cs b/Engine/Source/Programs/Shared/EpicGames.Horde/Compute/AgentStorageClient.cs
--- a/Engine/Source/Programs/Shared/EpicGames.Horde/Compute/AgentStorageClient.cs
+++ b/Engine/Source/Programs/Shared/EpicGames.Horde/Compute/AgentStorageClient.cs
@@ -11,7 +11,7 @@ using EpicGames.Horde.Storage;
 namespace EpicGames.Horde.Compute
 {
 	/// <summary>
-	/// Storage client which can read bundles over a compute channel
+	/// Storage backend which can read bundles over a compute channel
 	/// </summary>
 	public sealed class AgentStorageBackend : IStorageBackend, IDisposable
 	{
diff --git a/Engine/Source/Programs/Shared/EpicGames.Horde/HordeClient.cs b/Engine/Source/Programs/Shared/EpicGames.Horde/HordeClient.cs
--- a/Engine/Source/Programs/Shared/EpicGames.Horde/HordeClient.cs
+++ b/Engine/Source/Programs/Shared/EpicGames.Horde/HordeClient.cs
@@ -220,7 +220,7 @@ namespace EpicGames.Horde
 		}
 
 		/// <inheritdoc/>
-		public IStorageClient CreateStorageClient(string basePath, string? accessToken = null)
+		public IStorageNamespace GetStorageNamespace(string basePath, string? accessToken = null)
 		{
 			HttpClient CreateClient()
 			{
@@ -237,7 +237,7 @@ namespace EpicGames.Horde
 			}
 
 			HttpStorageBackend httpStorageBackend = new HttpStorageBackend(basePath, CreateClient, CreateUnauthenticatedHttpClient, _loggerFactory.CreateLogger<HttpStorageBackend>());
-			return new BundleStorageClient(httpStorageBackend, _bundleCache, _hordeOptions.Bundle, _loggerFactory.CreateLogger<BundleStorageClient>());
+			return new BundleStorageNamespace(httpStorageBackend, _bundleCache, _hordeOptions.Bundle, _loggerFactory.CreateLogger<BundleStorageNamespace>());
 		}
 
 		/// <inheritdoc/>
diff --git a/Engine/Source/Programs/Shared/EpicGames.Horde/HordeExtensions.cs b/Engine/Source/Programs/Shared/EpicGames.Horde/HordeExtensions.cs
--- a/Engine/Source/Programs/Shared/EpicGames.Horde/HordeExtensions.cs
+++ b/Engine/Source/Programs/Shared/EpicGames.Horde/HordeExtensions.cs
@@ -35,7 +35,7 @@ namespace EpicGames.Horde
 			serviceCollection.AddSingleton<BundleCache>(sp => new BundleCache(sp.GetRequiredService<IOptions<HordeOptions>>().Value.BundleCache));
 			serviceCollection.AddSingleton<StorageBackendCache>(CreateBackendCache);
 			serviceCollection.AddSingleton<HttpStorageBackendFactory>();
-			serviceCollection.AddSingleton<HttpStorageClientFactory>();
+			serviceCollection.AddSingleton<HttpStorageNamespaceFactory>();
 			serviceCollection.AddSingleton<IHordeClient>(sp => sp.GetRequiredService<HordeClientFactory>().Create());
 			serviceCollection.AddSingleton<HordeClientFactory>();
 		}
diff --git a/Engine/Source/Programs/Shared/EpicGames.Horde/IHordeClient.cs b/Engine/Source/Programs/Shared/EpicGames.Horde/IHordeClient.cs
--- a/Engine/Source/Programs/Shared/EpicGames.Horde/IHordeClient.cs
+++ b/Engine/Source/Programs/Shared/EpicGames.Horde/IHordeClient.cs
@@ -79,9 +79,9 @@ namespace EpicGames.Horde
 		HordeHttpClient CreateHttpClient();
 
 		/// <summary>
-		/// Creates a storage client for the given base path
+		/// Creates a storage namespace for the given base path
 		/// </summary>
-		IStorageClient CreateStorageClient(string relativePath, string? accessToken = null);
+		IStorageNamespace GetStorageNamespace(string relativePath, string? accessToken = null);
 
 		/// <summary>
 		/// Creates a logger device that writes data to the server
@@ -95,28 +95,28 @@ namespace EpicGames.Horde
 	public static class HordeClientExtensions
 	{
 		/// <summary>
-		/// Creates a storage client for a particular namespace
+		/// Creates a storage namespace for a particular id
 		/// </summary>
-		public static IStorageClient CreateStorageClient(this IHordeClient hordeClient, NamespaceId namespaceId, string? accessToken = null)
-			=> hordeClient.CreateStorageClient($"api/v1/storage/{namespaceId}", accessToken);
+		public static IStorageNamespace GetStorageNamespace(this IHordeClient hordeClient, NamespaceId namespaceId, string? accessToken = null)
+			=> hordeClient.GetStorageNamespace($"api/v1/storage/{namespaceId}", accessToken);
 
 		/// <summary>
-		/// Creates a storage client for a particular artifact
+		/// Creates a storage namespace for a particular artifact
 		/// </summary>
-		public static IStorageClient CreateStorageClient(this IHordeClient hordeClient, ArtifactId artifactId)
-			=> hordeClient.CreateStorageClient($"api/v2/artifacts/{artifactId}");
+		public static IStorageNamespace GetStorageNamespace(this IHordeClient hordeClient, ArtifactId artifactId)
+			=> hordeClient.GetStorageNamespace($"api/v2/artifacts/{artifactId}");
 
 		/// <summary>
-		/// Creates a storage client for a particular log
+		/// Creates a storage namespace for a particular log
 		/// </summary>
-		public static IStorageClient CreateStorageClient(this IHordeClient hordeClient, LogId logId)
-			=> hordeClient.CreateStorageClient($"api/v1/logs/{logId}");
+		public static IStorageNamespace GetStorageNamespace(this IHordeClient hordeClient, LogId logId)
+			=> hordeClient.GetStorageNamespace($"api/v1/logs/{logId}");
 
 		/// <summary>
-		/// Creates a storage client for a particular tool
+		/// Creates a storage namespace for a particular tool
 		/// </summary>
-		public static IStorageClient CreateStorageClient(this IHordeClient hordeClient, ToolId toolId)
-			=> hordeClient.CreateStorageClient($"api/v1/tools/{toolId}");
+		public static IStorageNamespace GetStorageNamespace(this IHordeClient hordeClient, ToolId toolId)
+			=> hordeClient.GetStorageNamespace($"api/v1/tools/{toolId}");
 
 		/// <summary>
 		/// Reads a blob storage ref from a path
@@ -129,8 +129,8 @@ namespace EpicGames.Horde
 				return null;
 			}
 
-			IStorageClient storageClient = hordeClient.CreateStorageClient(response.BasePath);
-			return storageClient.CreateBlobRef(response.Target);
+			IStorageNamespace storageNamespace = hordeClient.GetStorageNamespace(response.BasePath);
+			return storageNamespace.CreateBlobRef(response.Target);
 		}
 
 		/// <summary>
diff --git a/Engine/Source/Programs/Shared/EpicGames.Horde/Logs/ServerLogger.cs b/Engine/Source/Programs/Shared/EpicGames.Horde/Logs/ServerLogger.cs
--- a/Engine/Source/Programs/Shared/EpicGames.Horde/Logs/ServerLogger.cs
+++ b/Engine/Source/Programs/Shared/EpicGames.Horde/Logs/ServerLogger.cs
@@ -27,7 +27,7 @@ namespace EpicGames.Horde.Logs
 		readonly LogLevel _minimumLevel;
 		readonly ILogger _internalLogger;
 		readonly LogBuilder _builder;
-		readonly IStorageClient _store;
+		readonly IStorageNamespace _store;
 		readonly IBlobWriter _writer;
 
 		int _bufferLength;
@@ -54,7 +54,7 @@ namespace EpicGames.Horde.Logs
 			_minimumLevel = minimumLevel;
 			_internalLogger = internalLogger;
 			_builder = new LogBuilder(LogFormat.Json, internalLogger);
-			_store = _hordeClient.CreateStorageClient(logId);
+			_store = _hordeClient.GetStorageNamespace(logId);
 			_writer = _store.CreateBlobWriter();
 
 			_tailTaskStop = new AsyncEvent();
diff --git a/Engine/Source/Programs/Shared/EpicGames.Horde/Storage/Backends/FileStorageBackend.cs b/Engine/Source/Programs/Shared/EpicGames.Horde/Storage/Backends/FileStorageBackend.cs
--- a/Engine/Source/Programs/Shared/EpicGames.Horde/Storage/Backends/FileStorageBackend.cs
+++ b/Engine/Source/Programs/Shared/EpicGames.Horde/Storage/Backends/FileStorageBackend.cs
@@ -112,15 +112,15 @@ namespace EpicGames.Horde.Storage.Backends
 
 		/// <inheritdoc/>
 		public Task AddAliasAsync(string name, BlobLocator locator, int rank = 0, ReadOnlyMemory<byte> data = default, CancellationToken cancellationToken = default)
-			=> throw new NotSupportedException("File storage client does not currently support aliases.");
+			=> throw new NotSupportedException("File storage backend does not currently support aliases.");
 
 		/// <inheritdoc/>
 		public Task RemoveAliasAsync(string name, BlobLocator locator, CancellationToken cancellationToken = default)
-			=> throw new NotSupportedException("File storage client does not currently support aliases.");
+			=> throw new NotSupportedException("File storage backend does not currently support aliases.");
 
 		/// <inheritdoc/>
 		public Task<BlobAliasLocator[]> FindAliasesAsync(string alias, int? maxResults = null, CancellationToken cancellationToken = default)
-			=> throw new NotSupportedException("File storage client does not currently support aliases.");
+			=> throw new NotSupportedException("File storage backend does not currently support aliases.");
 
 		#endregion
 
diff --git a/Engine/Source/Programs/Shared/EpicGames.Horde/Storage/Backends/HttpStorageBackend.cs b/Engine/Source/Programs/Shared/EpicGames.Horde/Storage/Backends/HttpStorageBackend.cs
--- a/Engine/Source/Programs/Shared/EpicGames.Horde/Storage/Backends/HttpStorageBackend.cs
+++ b/Engine/Source/Programs/Shared/EpicGames.Horde/Storage/Backends/HttpStorageBackend.cs
@@ -411,7 +411,7 @@ namespace EpicGames.Horde.Storage.Backends
 		}
 
 		/// <summary>
-		/// Creates a new HTTP storage client
+		/// Creates a new HTTP storage backend
 		/// </summary>
 		/// <param name="basePath">Base path for all requests</param>
 		/// <param name="accessToken">Custom access token to use for requests</param>
@@ -440,7 +440,7 @@ namespace EpicGames.Horde.Storage.Backends
 		}
 
 		/// <summary>
-		/// Creates a new HTTP storage client
+		/// Creates a new HTTP storage backend
 		/// </summary>
 		/// <param name="namespaceId">Namespace to create a client for</param>
 		/// <param name="accessToken">Custom access token to use for requests</param>
diff --git a/Engine/Source/Programs/Shared/EpicGames.Horde/Storage/BlobSerializer.cs b/Engine/Source/Programs/Shared/EpicGames.Horde/Storage/BlobSerializer.cs
--- a/Engine/Source/Programs/Shared/EpicGames.Horde/Storage/BlobSerializer.cs
+++ b/Engine/Source/Programs/Shared/EpicGames.Horde/Storage/BlobSerializer.cs
@@ -269,7 +269,7 @@ namespace EpicGames.Horde.Storage
 		/// <param name="options">Options to control serialization</param>
 		/// <param name="cancellationToken">Cancellation token for the operation</param>
 		/// <returns>Node for the given ref, or null if it does not exist</returns>
-		public static async Task<IHashedBlobRef<TNode>?> TryReadRefAsync<TNode>(this IStorageClient store, RefName name, DateTime cacheTime = default, BlobSerializerOptions? options = null, CancellationToken cancellationToken = default) where TNode : class
+		public static async Task<IHashedBlobRef<TNode>?> TryReadRefAsync<TNode>(this IStorageNamespace store, RefName name, DateTime cacheTime = default, BlobSerializerOptions? options = null, CancellationToken cancellationToken = default) where TNode : class
 		{
 			IHashedBlobRef? refTarget = await store.TryReadRefAsync(name, cacheTime, cancellationToken);
 			if (refTarget == null)
@@ -288,7 +288,7 @@ namespace EpicGames.Horde.Storage
 		/// <param name="options">Options to control serialization</param>
 		/// <param name="cancellationToken">Cancellation token for the operation</param>
 		/// <returns>The blob instance</returns>
-		public static async Task<IHashedBlobRef<TNode>> ReadRefAsync<TNode>(this IStorageClient store, RefName name, DateTime cacheTime = default, BlobSerializerOptions? options = null, CancellationToken cancellationToken = default) where TNode : class
+		public static async Task<IHashedBlobRef<TNode>> ReadRefAsync<TNode>(this IStorageNamespace store, RefName name, DateTime cacheTime = default, BlobSerializerOptions? options = null, CancellationToken cancellationToken = default) where TNode : class
 		{
 			IHashedBlobRef<TNode>? refValue = await store.TryReadRefAsync<TNode>(name, cacheTime, options, cancellationToken);
 			if (refValue == null)
@@ -307,7 +307,7 @@ namespace EpicGames.Horde.Storage
 		/// <param name="options">Options to control serialization</param>
 		/// <param name="cancellationToken">Cancellation token for the operation</param>
 		/// <returns>Node for the given ref, or null if it does not exist</returns>
-		public static async Task<TNode?> TryReadRefTargetAsync<TNode>(this IStorageClient store, RefName name, DateTime cacheTime = default, BlobSerializerOptions? options = null, CancellationToken cancellationToken = default) where TNode : class
+		public static async Task<TNode?> TryReadRefTargetAsync<TNode>(this IStorageNamespace store, RefName name, DateTime cacheTime = default, BlobSerializerOptions? options = null, CancellationToken cancellationToken = default) where TNode : class
 		{
 			IHashedBlobRef<TNode>? refTarget = await store.TryReadRefAsync<TNode>(name, cacheTime, options, cancellationToken);
 			if (refTarget == null)
@@ -326,7 +326,7 @@ namespace EpicGames.Horde.Storage
 		/// <param name="options">Options to control serialization</param>
 		/// <param name="cancellationToken">Cancellation token for the operation</param>
 		/// <returns>The blob instance</returns>
-		public static async Task<TNode> ReadRefTargetAsync<TNode>(this IStorageClient store, RefName name, DateTime cacheTime = default, BlobSerializerOptions? options = null, CancellationToken cancellationToken = default) where TNode : class
+		public static async Task<TNode> ReadRefTargetAsync<TNode>(this IStorageNamespace store, RefName name, DateTime cacheTime = default, BlobSerializerOptions? options = null, CancellationToken cancellationToken = default) where TNode : class
 		{
 			IHashedBlobRef<TNode> blobRef = await ReadRefAsync<TNode>(store, name, cacheTime, options, cancellationToken);
 			return await blobRef.ReadBlobAsync(cancellationToken);
diff --git a/Engine/Source/Programs/Shared/EpicGames.Horde/Storage/Bundles/BundleHandle.cs b/Engine/Source/Programs/Shared/EpicGames.Horde/Storage/Bundles/BundleHandle.cs
--- a/Engine/Source/Programs/Shared/EpicGames.Horde/Storage/Bundles/BundleHandle.cs
+++ b/Engine/Source/Programs/Shared/EpicGames.Horde/Storage/Bundles/BundleHandle.cs
@@ -38,15 +38,15 @@ namespace EpicGames.Horde.Storage.Bundles
 	/// </summary>
 	class FlushedBundleHandle : BundleHandle
 	{
-		readonly BundleStorageClient _storageClient;
+		readonly BundleStorageNamespace _storageNamespace;
 		readonly BlobLocator _locator;
 
 		/// <summary>
 		/// Constructor
 		/// </summary>
-		public FlushedBundleHandle(BundleStorageClient storageClient, BlobLocator locator)
+		public FlushedBundleHandle(BundleStorageNamespace storageNamespace, BlobLocator locator)
 		{
-			_storageClient = storageClient;
+			_storageNamespace = storageNamespace;
 			_locator = locator;
 		}
 
@@ -55,11 +55,11 @@ namespace EpicGames.Horde.Storage.Bundles
 
 		/// <inheritdoc/>
 		public override Task<Stream> OpenAsync(int offset = 0, int? length = null, CancellationToken cancellationToken = default)
-			=> _storageClient.Backend.OpenBlobAsync(_locator, offset, length, cancellationToken);
+			=> _storageNamespace.Backend.OpenBlobAsync(_locator, offset, length, cancellationToken);
 
 		/// <inheritdoc/>
 		public override async ValueTask<IReadOnlyMemoryOwner<byte>> ReadAsync(int offset, int? length, CancellationToken cancellationToken = default)
-			=> await _storageClient.Backend.ReadBlobAsync(_locator, offset, length, cancellationToken);
+			=> await _storageNamespace.Backend.ReadBlobAsync(_locator, offset, length, cancellationToken);
 
 		/// <inheritdoc/>
 		public override bool TryGetLocator(out BlobLocator locator)
diff --git a/Engine/Source/Programs/Shared/EpicGames.Horde/Storage/Bundles/BundleStorageClient.cs b/Engine/Source/Programs/Shared/EpicGames.Horde/Storage/Bundles/BundleStorageClient.cs
--- a/Engine/Source/Programs/Shared/EpicGames.Horde/Storage/Bundles/BundleStorageClient.cs
+++ b/Engine/Source/Programs/Shared/EpicGames.Horde/Storage/Bundles/BundleStorageClient.cs
@@ -16,9 +16,9 @@ using Microsoft.Extensions.Logging;
 namespace EpicGames.Horde.Storage.Bundles
 {
 	/// <summary>
-	/// Base class for an implementation of <see cref="IStorageClient"/>, providing implementations for some common functionality using bundles.
+	/// Base class for an implementation of <see cref="IStorageNamespace"/>, providing implementations for some common functionality using bundles.
 	/// </summary>
-	public sealed class BundleStorageClient : IStorageClient
+	public sealed class BundleStorageNamespace : IStorageNamespace
 	{
 		readonly IStorageBackend _backend;
 		readonly BundleCache _cache;
@@ -51,7 +51,7 @@ namespace EpicGames.Horde.Storage.Bundles
 		/// <summary>
 		/// Constructor
 		/// </summary>
-		public BundleStorageClient(IStorageBackend backend, BundleCache cache, BundleOptions? options, ILogger logger)
+		public BundleStorageNamespace(IStorageBackend backend, BundleCache cache, BundleOptions? options, ILogger logger)
 		{
 			_backend = backend;
 			_cache = cache;
@@ -60,33 +60,33 @@ namespace EpicGames.Horde.Storage.Bundles
 		}
 
 		/// <summary>
-		/// Creates a bundle storage client around a memory client backend
+		/// Creates a bundle storage namespace around a memory backend
 		/// </summary>
-		public static BundleStorageClient CreateInMemory(ILogger logger)
+		public static BundleStorageNamespace CreateInMemory(ILogger logger)
 			=> CreateInMemory(null, logger);
 
 		/// <summary>
-		/// Creates a bundle storage client around a memory client backend
+		/// Creates a bundle storage namespace around a memory backend
 		/// </summary>
-		public static BundleStorageClient CreateInMemory(BundleOptions? options, ILogger logger)
+		public static BundleStorageNamespace CreateInMemory(BundleOptions? options, ILogger logger)
 		{
 			MemoryStorageBackend backend = new MemoryStorageBackend();
-			return new BundleStorageClient(backend, BundleCache.None, options, logger);
+			return new BundleStorageNamespace(backend, BundleCache.None, options, logger);
 		}
 
 		/// <summary>
-		/// Creates a bundle storage client around a directory on the filesystem
+		/// Creates a bundle storage namespace around a directory on the filesystem
 		/// </summary>
-		public static BundleStorageClient CreateFromDirectory(DirectoryReference rootDir, BundleCache cache, MemoryMappedFileCache memoryMappedFileCache, ILogger logger)
+		public static BundleStorageNamespace CreateFromDirectory(DirectoryReference rootDir, BundleCache cache, MemoryMappedFileCache memoryMappedFileCache, ILogger logger)
 			=> CreateFromDirectory(rootDir, cache, memoryMappedFileCache, null, logger);
 
 		/// <summary>
-		/// Creates a bundle storage client around a directory on the filesystem
+		/// Creates a bundle storage namespace around a directory on the filesystem
 		/// </summary>
-		public static BundleStorageClient CreateFromDirectory(DirectoryReference rootDir, BundleCache cache, MemoryMappedFileCache memoryMappedFileCache, BundleOptions? options, ILogger logger)
+		public static BundleStorageNamespace CreateFromDirectory(DirectoryReference rootDir, BundleCache cache, MemoryMappedFileCache memoryMappedFileCache, BundleOptions? options, ILogger logger)
 		{
 			FileStorageBackend backend = new FileStorageBackend(new FileObjectStore(rootDir, memoryMappedFileCache), logger);
-			return new BundleStorageClient(backend, cache, options, logger);
+			return new BundleStorageNamespace(backend, cache, options, logger);
 		}
 
 		/// <summary>
diff --git a/Engine/Source/Programs/Shared/EpicGames.Horde/Storage/Bundles/V1/BundleReader.cs b/Engine/Source/Programs/Shared/EpicGames.Horde/Storage/Bundles/V1/BundleReader.cs
--- a/Engine/Source/Programs/Shared/EpicGames.Horde/Storage/Bundles/V1/BundleReader.cs
+++ b/Engine/Source/Programs/Shared/EpicGames.Horde/Storage/Bundles/V1/BundleReader.cs
@@ -44,7 +44,7 @@ namespace EpicGames.Horde.Storage.Bundles.V1
 	}
 
 	/// <summary>
-	/// Writes nodes from bundles in an <see cref="IStorageClient"/> instance.
+	/// Writes nodes from bundles in an <see cref="IStorageNamespace"/> instance.
 	/// </summary>
 	public class BundleReader
 	{
@@ -125,7 +125,7 @@ namespace EpicGames.Horde.Storage.Bundles.V1
 		// When reader is uncached, use a smaller default fetch size
 		const int DefaultUncachedFetchSize = 1 * 1024 * 1024;
 
-		readonly BundleStorageClient _store;
+		readonly BundleStorageNamespace _store;
 		readonly BundleCache _cache;
 		readonly ILogger _logger;
 
@@ -149,7 +149,7 @@ namespace EpicGames.Horde.Storage.Bundles.V1
 		/// <param name="store"></param>
 		/// <param name="cache">Cache for data</param>
 		/// <param name="logger">Logger for output</param>
-		public BundleReader(BundleStorageClient store, BundleCache cache, ILogger logger)
+		public BundleReader(BundleStorageNamespace store, BundleCache cache, ILogger logger)
 		{
 			_store = store;
 			_cache = cache;
diff --git a/Engine/Source/Programs/Shared/EpicGames.Horde/Storage/Bundles/V1/BundleWriter.cs b/Engine/Source/Programs/Shared/EpicGames.Horde/Storage/Bundles/V1/BundleWriter.cs
--- a/Engine/Source/Programs/Shared/EpicGames.Horde/Storage/Bundles/V1/BundleWriter.cs
+++ b/Engine/Source/Programs/Shared/EpicGames.Horde/Storage/Bundles/V1/BundleWriter.cs
@@ -87,7 +87,7 @@ namespace EpicGames.Horde.Storage.Bundles.V1
 	}
 
 	/// <summary>
-	/// Writes nodes of a tree to an <see cref="IStorageClient"/>, packed into bundles. Each <see cref="BundleWriter"/> instance is single threaded,
+	/// Writes nodes of a tree to an <see cref="IStorageNamespace"/>, packed into bundles. Each <see cref="BundleWriter"/> instance is single threaded,
 	/// but multiple instances may be written to in parallel.
 	/// </summary>
 	public sealed class BundleWriter : BlobWriter
@@ -425,7 +425,7 @@ namespace EpicGames.Horde.Storage.Bundles.V1
 			}
 
 			// Mark the bundle as complete
-			public async Task WriteAsync(BundleStorageClient store, string? basePath, ILogger? traceLogger)
+			public async Task WriteAsync(BundleStorageNamespace store, string? basePath, ILogger? traceLogger)
 			{
 				traceLogger?.LogInformation("Marking bundle {BundleId} as complete ({NumNodes} nodes); adding to write queue.", BundleId, _queue.Count);
 
@@ -584,7 +584,7 @@ namespace EpicGames.Horde.Storage.Bundles.V1
 		class WriteQueue
 		{
 			long _memoryFootprint;
-			readonly BundleStorageClient _store;
+			readonly BundleStorageNamespace _store;
 			readonly string? _basePath;
 			readonly long _maxMemoryFootprint;
 			readonly ILogger? _traceLogger;
@@ -592,7 +592,7 @@ namespace EpicGames.Horde.Storage.Bundles.V1
 			int _refCount;
 			readonly List<Task> _writeTasks = new List<Task>();
 
-			public WriteQueue(BundleStorageClient store, string? basePath, long maxMemoryFootprint, ILogger? traceLogger)
+			public WriteQueue(BundleStorageNamespace store, string? basePath, long maxMemoryFootprint, ILogger? traceLogger)
 			{
 				_store = store;
 				_basePath = basePath;
@@ -667,7 +667,7 @@ namespace EpicGames.Horde.Storage.Bundles.V1
 
 		static readonly BundleOptions s_defaultOptions = new BundleOptions();
 
-		readonly BundleStorageClient _store;
+		readonly BundleStorageNamespace _store;
 		readonly BundleReader _reader;
 		readonly BundleOptions _bundleOptions;
 		readonly string? _basePath;
@@ -689,7 +689,7 @@ namespace EpicGames.Horde.Storage.Bundles.V1
 		/// <param name="bundleOptions">Options for the writer</param>
 		/// <param name="blobSerializerOptions"></param>
 		/// <param name="traceLogger">Optional logger for trace information</param>
-		public BundleWriter(BundleStorageClient store, BundleReader reader, string? basePath, BundleOptions? bundleOptions = null, BlobSerializerOptions? blobSerializerOptions = null, ILogger? traceLogger = null)
+		public BundleWriter(BundleStorageNamespace store, BundleReader reader, string? basePath, BundleOptions? bundleOptions = null, BlobSerializerOptions? blobSerializerOptions = null, ILogger? traceLogger = null)
 			: this(store, reader, basePath, bundleOptions, null, blobSerializerOptions, traceLogger)
 		{
 		}
@@ -707,7 +707,7 @@ namespace EpicGames.Horde.Storage.Bundles.V1
 		/// <summary>
 		/// Internal constructor
 		/// </summary>
-		private BundleWriter(BundleStorageClient store, BundleReader reader, string? basePath, BundleOptions? bundleOptions, WriteQueue? writeQueue, BlobSerializerOptions? blobSerializerOptions, ILogger? traceLogger = null)
+		private BundleWriter(BundleStorageNamespace store, BundleReader reader, string? basePath, BundleOptions? bundleOptions, WriteQueue? writeQueue, BlobSerializerOptions? blobSerializerOptions, ILogger? traceLogger = null)
 			: base(blobSerializerOptions)
 		{
 			_store = store;
diff --git a/Engine/Source/Programs/Shared/EpicGames.Horde/Storage/Bundles/V2/BundleWriter.cs b/Engine/Source/Programs/Shared/EpicGames.Horde/Storage/Bundles/V2/BundleWriter.cs
--- a/Engine/Source/Programs/Shared/EpicGames.Horde/Storage/Bundles/V2/BundleWriter.cs
+++ b/Engine/Source/Programs/Shared/EpicGames.Horde/Storage/Bundles/V2/BundleWriter.cs
@@ -56,7 +56,7 @@ namespace EpicGames.Horde.Storage.Bundles.V2
 		{
 			readonly object _lockObject = new object();
 
-			readonly BundleStorageClient _storageClient;
+			readonly BundleStorageNamespace _storageNamespace;
 			readonly string? _basePath;
 			readonly BundleCache _cache;
 			readonly BundleOptions _bundleOptions;
@@ -81,9 +81,9 @@ namespace EpicGames.Horde.Storage.Bundles.V2
 			/// </summary>
 			public int Length => _encodedPacketWriter?.Length ?? throw new InvalidOperationException("Bundle has been flushed");
 
-			public PendingBundleHandle(BundleStorageClient storageClient, string? basePath, BundleCache cache, BundleOptions bundleOptions)
+			public PendingBundleHandle(BundleStorageNamespace storageNamespace, string? basePath, BundleCache cache, BundleOptions bundleOptions)
 			{
-				_storageClient = storageClient;
+				_storageNamespace = storageNamespace;
 				_basePath = basePath;
 				_cache = cache;
 				_bundleOptions = bundleOptions;
@@ -175,7 +175,7 @@ namespace EpicGames.Horde.Storage.Bundles.V2
 					// Point the packet handle to the encoded data
 					lock (_lockObject)
 					{
-						FlushedPacketHandle flushedPacketHandle = new FlushedPacketHandle(_storageClient, this, packetOffset, packetLength, _cache);
+						FlushedPacketHandle flushedPacketHandle = new FlushedPacketHandle(_storageNamespace, this, packetOffset, packetLength, _cache);
 						_packetHandle.CompletePacket(flushedPacketHandle);
 					}
 				}
@@ -211,8 +211,8 @@ namespace EpicGames.Horde.Storage.Bundles.V2
 				FlushedBundleHandle flushedHandle;
 				using (ReadOnlySequenceStream stream = new ReadOnlySequenceStream(_encodedPacketWriter.AsSequence()))
 				{
-					BlobLocator locator = await _storageClient.Backend.WriteBlobAsync(stream, _bundleImports, _basePath, cancellationToken);
-					flushedHandle = new FlushedBundleHandle(_storageClient, locator);
+					BlobLocator locator = await _storageNamespace.Backend.WriteBlobAsync(stream, _bundleImports, _basePath, cancellationToken);
+					flushedHandle = new FlushedBundleHandle(_storageNamespace, locator);
 				}
 
 				// Release all the intermediate data
@@ -230,7 +230,7 @@ namespace EpicGames.Horde.Storage.Bundles.V2
 				{
 					foreach ((ExportHandle exportHandle, AliasInfo aliasInfo) in pendingExportAliases)
 					{
-						await _storageClient.AddAliasAsync(aliasInfo.Name, exportHandle, aliasInfo.Rank, aliasInfo.Data, cancellationToken);
+						await _storageNamespace.AddAliasAsync(aliasInfo.Name, exportHandle, aliasInfo.Rank, aliasInfo.Data, cancellationToken);
 					}
 				}
 			}
@@ -280,7 +280,7 @@ namespace EpicGames.Horde.Storage.Bundles.V2
 			}
 		}
 
-		readonly BundleStorageClient _storageClient;
+		readonly BundleStorageNamespace _storageNamespace;
 		readonly string? _basePath;
 		readonly BundleOptions _bundleOptions;
 		readonly BundleCache _bundleCache;
@@ -290,15 +290,15 @@ namespace EpicGames.Horde.Storage.Bundles.V2
 		/// <summary>
 		/// 
 		/// </summary>
-		public BundleWriter(BundleStorageClient storageClient, string? basePath, BundleCache bundleCache, BundleOptions bundleOptions, BlobSerializerOptions? blobOptions)
+		public BundleWriter(BundleStorageNamespace storageNamespace, string? basePath, BundleCache bundleCache, BundleOptions bundleOptions, BlobSerializerOptions? blobOptions)
 			: base(blobOptions)
 		{
-			_storageClient = storageClient;
+			_storageNamespace = storageNamespace;
 			_basePath = basePath;
 			_bundleOptions = bundleOptions;
 			_bundleCache = bundleCache;
 
-			_currentBundle = new PendingBundleHandle(storageClient, basePath, bundleCache, _bundleOptions);
+			_currentBundle = new PendingBundleHandle(storageNamespace, basePath, bundleCache, _bundleOptions);
 		}
 
 		/// <inheritdoc/>
@@ -313,12 +313,12 @@ namespace EpicGames.Horde.Storage.Bundles.V2
 		{
 			await _currentBundle.FlushAsync(cancellationToken);
 			_currentBundle.Dispose();
-			_currentBundle = new PendingBundleHandle(_storageClient, _basePath, _bundleCache, _bundleOptions);
+			_currentBundle = new PendingBundleHandle(_storageNamespace, _basePath, _bundleCache, _bundleOptions);
 		}
 
 		/// <inheritdoc/>
 		public override IBlobWriter Fork()
-			=> new BundleWriter(_storageClient, _basePath, _bundleCache, _bundleOptions, Options);
+			=> new BundleWriter(_storageNamespace, _basePath, _bundleCache, _bundleOptions, Options);
 
 		/// <inheritdoc/>
 		public override Memory<byte> GetOutputBuffer(int usedSize, int desiredSize)
diff --git a/Engine/Source/Programs/Shared/EpicGames.Horde/Storage/Bundles/V2/PacketHandle.cs b/Engine/Source/Programs/Shared/EpicGames.Horde/Storage/Bundles/V2/PacketHandle.cs
--- a/Engine/Source/Programs/Shared/EpicGames.Horde/Storage/Bundles/V2/PacketHandle.cs
+++ b/Engine/Source/Programs/Shared/EpicGames.Horde/Storage/Bundles/V2/PacketHandle.cs
@@ -56,7 +56,7 @@ namespace EpicGames.Horde.Storage.Bundles.V2
 	/// </summary>
 	class FlushedPacketHandle : PacketHandle
 	{
-		readonly BundleStorageClient _storageClient;
+		readonly BundleStorageNamespace _storageNamespace;
 		readonly BundleHandle _outer;
 		readonly int _packetOffset;
 		readonly int _packetLength;
@@ -75,9 +75,9 @@ namespace EpicGames.Horde.Storage.Bundles.V2
 		/// <summary>
 		/// Constructor
 		/// </summary>
-		public FlushedPacketHandle(BundleStorageClient storageClient, BundleHandle outer, int packetOffset, int packetLength, BundleCache cache)
+		public FlushedPacketHandle(BundleStorageNamespace storageNamespace, BundleHandle outer, int packetOffset, int packetLength, BundleCache cache)
 		{
-			_storageClient = storageClient;
+			_storageNamespace = storageNamespace;
 
 			_outer = outer;
 			_packetOffset = packetOffset;
@@ -88,9 +88,9 @@ namespace EpicGames.Horde.Storage.Bundles.V2
 		/// <summary>
 		/// Constructor
 		/// </summary>
-		public FlushedPacketHandle(BundleStorageClient storageClient, BundleHandle outer, ReadOnlySpan<byte> fragment, BundleCache cache)
+		public FlushedPacketHandle(BundleStorageNamespace storageNamespace, BundleHandle outer, ReadOnlySpan<byte> fragment, BundleCache cache)
 		{
-			_storageClient = storageClient;
+			_storageNamespace = storageNamespace;
 
 			_outer = outer;
 			_cache = cache;
@@ -219,16 +219,16 @@ namespace EpicGames.Horde.Storage.Bundles.V2
 		{
 			using IReadOnlyMemoryOwner<byte> encodedData = await ReadEncodedPacketAsync(cancellationToken);
 			IRefCountedHandle<Packet> packet = Packet.Decode(encodedData.Memory, _cache.Allocator, cacheKey);
-			Interlocked.Add(ref _storageClient.PacketReaderStats._numDecodedBytesRead, packet.Target.Length);
+			Interlocked.Add(ref _storageNamespace.PacketReaderStats._numDecodedBytesRead, packet.Target.Length);
 #pragma warning disable CA2000
-			PacketReader reader = new PacketReader(_storageClient, _cache, _outer, this, packet.Target, packet);
+			PacketReader reader = new PacketReader(_storageNamespace, _cache, _outer, this, packet.Target, packet);
 			return new Scoped<PacketReader>(reader);
 #pragma warning restore CA2000
 		}
 
 		async ValueTask<IReadOnlyMemoryOwner<byte>> ReadEncodedPacketAsync(CancellationToken cancellationToken)
 		{
-			Interlocked.Add(ref _storageClient.PacketReaderStats._numEncodedBytesRead, _packetLength);
+			Interlocked.Add(ref _storageNamespace.PacketReaderStats._numEncodedBytesRead, _packetLength);
 
 			int minPageIdx = _packetOffset / _cache.BundlePageSize;
 			int maxPageIdx = ((_packetOffset + _packetLength) + (_cache.BundlePageSize - 1)) / _cache.BundlePageSize;
@@ -276,7 +276,7 @@ namespace EpicGames.Horde.Storage.Bundles.V2
 		async Task<Scoped<IReadOnlyMemoryOwner<byte>>> ReadBundlePageInternalAsync(BundlePageCacheKey key, CancellationToken cancellationToken)
 		{
 			IReadOnlyMemoryOwner<byte> owner = await _outer.ReadAsync(key.Index * _cache.BundlePageSize, _cache.BundlePageSize, cancellationToken);
-			Interlocked.Add(ref _storageClient.PacketReaderStats._numBytesRead, owner.Memory.Length);
+			Interlocked.Add(ref _storageNamespace.PacketReaderStats._numBytesRead, owner.Memory.Length);
 			return new Scoped<IReadOnlyMemoryOwner<byte>>(owner);
 		}
 
diff --git a/Engine/Source/Programs/Shared/EpicGames.Horde/Storage/Bundles/V2/PacketReader.cs b/Engine/Source/Programs/Shared/EpicGames.Horde/Storage/Bundles/V2/PacketReader.cs
--- a/Engine/Source/Programs/Shared/EpicGames.Horde/Storage/Bundles/V2/PacketReader.cs
+++ b/Engine/Source/Programs/Shared/EpicGames.Horde/Storage/Bundles/V2/PacketReader.cs
@@ -7,7 +7,7 @@ using EpicGames.Core;
 namespace EpicGames.Horde.Storage.Bundles.V2
 {
 	/// <summary>
-	/// Counters used for tracking operations performed by a BundleStorageClient
+	/// Counters used for tracking operations performed by a BundleStorageNamespace
 	/// </summary>
 	class PacketReaderStats
 	{
@@ -34,7 +34,7 @@ namespace EpicGames.Horde.Storage.Bundles.V2
 	/// </summary>
 	sealed class PacketReader : IDisposable
 	{
-		readonly BundleStorageClient _storageClient;
+		readonly BundleStorageNamespace _storageNamespace;
 		readonly BundleCache _cache;
 		readonly BundleHandle _bundleHandle;
 		readonly FlushedPacketHandle _packetHandle;
@@ -51,15 +51,15 @@ namespace EpicGames.Horde.Storage.Bundles.V2
 		/// <summary>
 		/// Constructor
 		/// </summary>
-		/// <param name="storageClient"></param>
+		/// <param name="storageNamespace"></param>
 		/// <param name="cache"></param>
 		/// <param name="bundleHandle"></param>
 		/// <param name="packetHandle"></param>
 		/// <param name="decodedPacket">Data for the packet</param>
 		/// <param name="memoryOwner">Owner for the packet data</param>
-		public PacketReader(BundleStorageClient storageClient, BundleCache cache, BundleHandle bundleHandle, FlushedPacketHandle packetHandle, Packet decodedPacket, IRefCountedHandle memoryOwner)
+		public PacketReader(BundleStorageNamespace storageNamespace, BundleCache cache, BundleHandle bundleHandle, FlushedPacketHandle packetHandle, Packet decodedPacket, IRefCountedHandle memoryOwner)
 		{
-			_storageClient = storageClient;
+			_storageNamespace = storageNamespace;
 			_cache = cache;
 			_bundleHandle = bundleHandle;
 			_packetHandle = packetHandle;
@@ -167,7 +167,7 @@ namespace EpicGames.Horde.Storage.Bundles.V2
 				int bundleIdx = packetImport.BaseIdx;
 				BundleHandle bundleHandle = GetImportedBundleHandle(bundleIdx);
 
-				packetHandle = new FlushedPacketHandle(_storageClient, bundleHandle, packetImport.Fragment, _cache);
+				packetHandle = new FlushedPacketHandle(_storageNamespace, bundleHandle, packetImport.Fragment, _cache);
 				_cachedImportHandles[packetIdx] = packetHandle;
 			}
 			return packetHandle;
@@ -186,7 +186,7 @@ namespace EpicGames.Horde.Storage.Bundles.V2
 				PacketImport bundleImport = _decodedPacket.GetImport(bundleIdx);
 				Trace.Assert(bundleImport.BaseIdx == -1);
 
-				bundleHandle = new FlushedBundleHandle(_storageClient, new BlobLocator(bundleImport.Fragment.Clone()));
+				bundleHandle = new FlushedBundleHandle(_storageNamespace, new BlobLocator(bundleImport.Fragment.Clone()));
 				_cachedImportHandles[bundleIdx] = bundleHandle;
 			}
 			return bundleHandle;
diff --git a/Engine/Source/Programs/Shared/EpicGames.Horde/Storage/Clients/HttpStorageClient.cs b/Engine/Source/Programs/Shared/EpicGames.Horde/Storage/Clients/HttpStorageClient.cs
--- a/Engine/Source/Programs/Shared/EpicGames.Horde/Storage/Clients/HttpStorageClient.cs
+++ b/Engine/Source/Programs/Shared/EpicGames.Horde/Storage/Clients/HttpStorageClient.cs
@@ -8,19 +8,19 @@ using Microsoft.Extensions.Options;
 namespace EpicGames.Horde.Storage.Clients
 {
 	/// <summary>
-	/// Factory for constructing HttpStorageClient instances
+	/// Factory for constructing storage namespace instances with an HTTP backend
 	/// </summary>
-	public sealed class HttpStorageClientFactory : IStorageClientFactory
+	public sealed class HttpStorageNamespaceFactory : IStorageNamespaceFactory
 	{
 		readonly HttpStorageBackendFactory _backendFactory;
 		readonly BundleCache _bundleCache;
 		readonly IOptionsSnapshot<HordeOptions> _hordeOptions;
-		readonly ILogger<BundleStorageClient> _clientLogger;
+		readonly ILogger<BundleStorageNamespace> _clientLogger;
 
 		/// <summary>
 		/// Constructor
 		/// </summary>
-		public HttpStorageClientFactory(HttpStorageBackendFactory backendFactory, BundleCache bundleCache, IOptionsSnapshot<HordeOptions> hordeOptions, ILogger<BundleStorageClient> clientLogger)
+		public HttpStorageNamespaceFactory(HttpStorageBackendFactory backendFactory, BundleCache bundleCache, IOptionsSnapshot<HordeOptions> hordeOptions, ILogger<BundleStorageNamespace> clientLogger)
 		{
 			_backendFactory = backendFactory;
 			_bundleCache = bundleCache;
@@ -29,26 +29,26 @@ namespace EpicGames.Horde.Storage.Clients
 		}
 
 		/// <summary>
-		/// Creates a new HTTP storage client
+		/// Creates a new HTTP storage namespace
 		/// </summary>
 		/// <param name="basePath">Base path for all requests</param>
 		/// <param name="accessToken">Custom access token to use for requests</param>
 		/// <param name="withBackendCache"></param>
-		public IStorageClient CreateClientWithPath(string basePath, string? accessToken = null, bool withBackendCache = true)
+		public IStorageNamespace CreateClientWithPath(string basePath, string? accessToken = null, bool withBackendCache = true)
 		{
 			IStorageBackend backend = _backendFactory.CreateBackend(basePath, accessToken, withBackendCache);
-			return new BundleStorageClient(backend, _bundleCache, _hordeOptions.Value.Bundle, _clientLogger);
+			return new BundleStorageNamespace(backend, _bundleCache, _hordeOptions.Value.Bundle, _clientLogger);
 		}
 
 		/// <summary>
-		/// Creates a new HTTP storage client
+		/// Creates a new HTTP storage namespace
 		/// </summary>
 		/// <param name="namespaceId">Namespace to create a client for</param>
 		/// <param name="accessToken">Custom access token to use for requests</param>
 		/// <param name="withBackendCache">Whether to enable the backend cache, which caches full bundles to disk</param>
-		public IStorageClient CreateClient(NamespaceId namespaceId, string? accessToken = null, bool withBackendCache = true) => CreateClientWithPath($"api/v1/storage/{namespaceId}", accessToken, withBackendCache);
+		public IStorageNamespace CreateClient(NamespaceId namespaceId, string? accessToken = null, bool withBackendCache = true) => CreateClientWithPath($"api/v1/storage/{namespaceId}", accessToken, withBackendCache);
 
 		/// <inheritdoc/>
-		IStorageClient? IStorageClientFactory.TryCreateClient(NamespaceId namespaceId) => CreateClient(namespaceId);
+		IStorageNamespace? IStorageNamespaceFactory.TryGetNamespace(NamespaceId namespaceId) => CreateClient(namespaceId);
 	}
 }
diff --git a/Engine/Source/Programs/Shared/EpicGames.Horde/Storage/Clients/KeyValueStorageClient.cs b/Engine/Source/Programs/Shared/EpicGames.Horde/Storage/Clients/KeyValueStorageClient.cs
--- a/Engine/Source/Programs/Shared/EpicGames.Horde/Storage/Clients/KeyValueStorageClient.cs
+++ b/Engine/Source/Programs/Shared/EpicGames.Horde/Storage/Clients/KeyValueStorageClient.cs
@@ -12,13 +12,13 @@ using EpicGames.Horde.Storage.Backends;
 namespace EpicGames.Horde.Storage.Clients
 {
 	/// <summary>
-	/// Base class for storage clients that wrap a diirect key/value type store without any merging/splitting.
+	/// Base class for storage namespaces that wrap a diirect key/value type store without any merging/splitting.
 	/// </summary>
-	public sealed class KeyValueStorageClient : IStorageClient
+	public sealed class KeyValueStorageNamespace : IStorageNamespace
 	{
 		class Handle : IBlobRef
 		{
-			readonly KeyValueStorageClient _keyValueStorageClient;
+			readonly KeyValueStorageNamespace _keyValueStorageNamespace;
 			readonly BlobLocator _locator;
 
 			/// <inheritdoc/>
@@ -27,9 +27,9 @@ namespace EpicGames.Horde.Storage.Clients
 			/// <summary>
 			/// Constructor
 			/// </summary>
-			public Handle(KeyValueStorageClient keyValueStorageClient, BlobLocator locator)
+			public Handle(KeyValueStorageNamespace keyValueStorageNamespace, BlobLocator locator)
 			{
-				_keyValueStorageClient = keyValueStorageClient;
+				_keyValueStorageNamespace = keyValueStorageNamespace;
 				_locator = locator;
 			}
 
@@ -37,7 +37,7 @@ namespace EpicGames.Horde.Storage.Clients
 			public ValueTask FlushAsync(CancellationToken cancellationToken = default) => default;
 
 			/// <inheritdoc/>
-			public ValueTask<BlobData> ReadBlobDataAsync(CancellationToken cancellationToken = default) => _keyValueStorageClient.ReadBlobAsync(_locator, cancellationToken);
+			public ValueTask<BlobData> ReadBlobDataAsync(CancellationToken cancellationToken = default) => _keyValueStorageNamespace.ReadBlobAsync(_locator, cancellationToken);
 
 			/// <inheritdoc/>
 			public bool TryGetLocator(out BlobLocator locator)
@@ -55,7 +55,7 @@ namespace EpicGames.Horde.Storage.Clients
 
 		class Writer : BlobWriter
 		{
-			readonly KeyValueStorageClient _outer;
+			readonly KeyValueStorageNamespace _outer;
 			readonly string _basePath;
 			byte[] _data = Array.Empty<byte>();
 			int _offset;
@@ -63,7 +63,7 @@ namespace EpicGames.Horde.Storage.Clients
 			/// <summary>
 			/// Constructor
 			/// </summary>
-			public Writer(KeyValueStorageClient outer, string? basePath, BlobSerializerOptions? options)
+			public Writer(KeyValueStorageNamespace outer, string? basePath, BlobSerializerOptions? options)
 				: base(options)
 			{
 				_outer = outer;
@@ -122,15 +122,15 @@ namespace EpicGames.Horde.Storage.Clients
 		/// <summary>
 		/// Constructor
 		/// </summary>
-		public KeyValueStorageClient(IStorageBackend backend)
+		public KeyValueStorageNamespace(IStorageBackend backend)
 		{
 			_backend = backend;
 		}
 
 		/// <summary>
-		/// Create an in-memory storage client
+		/// Create an in-memory storage namespace
 		/// </summary>
-		public static KeyValueStorageClient CreateInMemory() => new KeyValueStorageClient(new MemoryStorageBackend());
+		public static KeyValueStorageNamespace CreateInMemory() => new KeyValueStorageNamespace(new MemoryStorageBackend());
 
 		#region Blobs
 
diff --git a/Engine/Source/Programs/Shared/EpicGames.Horde/Storage/DedupeBlobWriter.cs b/Engine/Source/Programs/Shared/EpicGames.Horde/Storage/DedupeBlobWriter.cs
--- a/Engine/Source/Programs/Shared/EpicGames.Horde/Storage/DedupeBlobWriter.cs
+++ b/Engine/Source/Programs/Shared/EpicGames.Horde/Storage/DedupeBlobWriter.cs
@@ -192,7 +192,7 @@ namespace EpicGames.Horde.Storage
 		/// </summary>
 		/// <param name="store">The store instance to read from</param>
 		/// <param name="maxKeys">Maximum number of keys to include in the cache</param>
-		public static DedupeBlobWriter CreateDedupeBlobWriter(this IStorageClient store, int maxKeys = DedupeBlobWriter.DefaultMaxKeys)
+		public static DedupeBlobWriter CreateDedupeBlobWriter(this IStorageNamespace store, int maxKeys = DedupeBlobWriter.DefaultMaxKeys)
 		{
 			IBlobWriter writer = store.CreateBlobWriter();
 			return new DedupeBlobWriter(writer, maxKeys);
@@ -204,7 +204,7 @@ namespace EpicGames.Horde.Storage
 		/// <param name="store">The store instance to read from</param>
 		/// <param name="refName">Ref name to use as a base path</param>
 		/// <param name="maxKeys">Maximum number of keys to include in the cache</param>
-		public static DedupeBlobWriter CreateDedupeBlobWriter(this IStorageClient store, RefName refName, int maxKeys = DedupeBlobWriter.DefaultMaxKeys)
+		public static DedupeBlobWriter CreateDedupeBlobWriter(this IStorageNamespace store, RefName refName, int maxKeys = DedupeBlobWriter.DefaultMaxKeys)
 		{
 			IBlobWriter writer = store.CreateBlobWriter(refName.ToString());
 			return new DedupeBlobWriter(writer, maxKeys);
diff --git a/Engine/Source/Programs/Shared/EpicGames.Horde/Storage/IStorageClient.cs b/Engine/Source/Programs/Shared/EpicGames.Horde/Storage/IStorageClient.cs
--- a/Engine/Source/Programs/Shared/EpicGames.Horde/Storage/IStorageClient.cs
+++ b/Engine/Source/Programs/Shared/EpicGames.Horde/Storage/IStorageClient.cs
@@ -21,7 +21,7 @@ namespace EpicGames.Horde.Storage
 	/// <summary>
 	/// Interface for the storage system.
 	/// </summary>
-	public interface IStorageClient
+	public interface IStorageNamespace
 	{
 		#region Blobs
 
@@ -104,35 +104,35 @@ namespace EpicGames.Horde.Storage
 		#endregion
 
 		/// <summary>
-		/// Gets a snapshot of the stats for the storage client.
+		/// Gets a snapshot of the stats for the storage namespace.
 		/// </summary>
 		void GetStats(StorageStats stats);
 	}
 
 	/// <summary>
-	/// Allows creating storage clients for different namespaces
+	/// Allows creating storage namespaces for different namespace ids
 	/// </summary>
-	public interface IStorageClientFactory
+	public interface IStorageNamespaceFactory
 	{
 		/// <summary>
-		/// Creates a storage client for the given namespace
+		/// Creates a storage namespace for the given id
 		/// </summary>
 		/// <param name="namespaceId">Namespace to manipulate</param>
-		/// <returns>Storage client instance. May be null if the namespace does not exist.</returns>
-		IStorageClient? TryCreateClient(NamespaceId namespaceId);
+		/// <returns>Storage namespace instance. May be null if the namespace does not exist.</returns>
+		IStorageNamespace? TryGetNamespace(NamespaceId namespaceId);
 	}
 
 	/// <summary>
-	/// Extension methods for <see cref="IStorageClientFactory"/>
+	/// Extension methods for <see cref="IStorageNamespaceFactory"/>
 	/// </summary>
-	public static class StorageClientFactoryExtensions
+	public static class StorageNamespaceFactoryExtensions
 	{
 		/// <summary>
-		/// Creates a new storage client, throwing an exception if it does not exist
+		/// Creates a new storage namespace, throwing an exception if it does not exist
 		/// </summary>
-		public static IStorageClient CreateClient(this IStorageClientFactory factory, NamespaceId namespaceId)
+		public static IStorageNamespace GetNamespace(this IStorageNamespaceFactory factory, NamespaceId namespaceId)
 		{
-			IStorageClient? client = factory.TryCreateClient(namespaceId);
+			IStorageNamespace? client = factory.TryGetNamespace(namespaceId);
 			if (client == null)
 			{
 				throw new InvalidOperationException($"No namespace '{namespaceId}' is configured");
@@ -226,64 +226,64 @@ namespace EpicGames.Horde.Storage
 	}
 
 	/// <summary>
-	/// Extension methods for <see cref="IStorageClient"/>
+	/// Extension methods for <see cref="IStorageNamespace"/>
 	/// </summary>
-	public static class StorageClientExtensions
+	public static class StorageNamespaceExtensions
 	{
 		#region Blobs
 
 		/// <summary>
 		/// Creates a new blob handle by parsing a locator
 		/// </summary>
-		/// <param name="storageClient">Storage client to operate on</param>
+		/// <param name="storageNamespace">Storage namespace to operate on</param>
 		/// <param name="locator">Path to the blob</param>
 		/// <param name="serializerOptions">Options for deserializing the blob</param>
 		/// <returns>New handle to the blob</returns>
-		public static IBlobRef<T> CreateBlobRef<T>(this IStorageClient storageClient, BlobLocator locator, BlobSerializerOptions? serializerOptions = null)
-			=> storageClient.CreateBlobRef(locator).ForType<T>(serializerOptions);
+		public static IBlobRef<T> CreateBlobRef<T>(this IStorageNamespace storageNamespace, BlobLocator locator, BlobSerializerOptions? serializerOptions = null)
+			=> storageNamespace.CreateBlobRef(locator).ForType<T>(serializerOptions);
 
 		/// <summary>
 		/// Creates a new blob reference from a locator and hash
 		/// </summary>
-		/// <param name="storageClient">Storage client to operate on</param>
+		/// <param name="storageNamespace">Storage namespace to operate on</param>
 		/// <param name="hash">Hash of the target blob</param>
 		/// <param name="locator">Path to the blob</param>
 		/// <returns>New handle to the blob</returns>
-		public static IHashedBlobRef CreateBlobRef(this IStorageClient storageClient, IoHash hash, BlobLocator locator)
-			=> HashedBlobRef.Create(hash, storageClient.CreateBlobRef(locator));
+		public static IHashedBlobRef CreateBlobRef(this IStorageNamespace storageNamespace, IoHash hash, BlobLocator locator)
+			=> HashedBlobRef.Create(hash, storageNamespace.CreateBlobRef(locator));
 
 		/// <summary>
 		/// Creates a new blob reference from a locator and hash
 		/// </summary>
-		/// <param name="storageClient">Storage client to operate on</param>
+		/// <param name="storageNamespace">Storage namespace to operate on</param>
 		/// <param name="hash">Hash of the target blob</param>
 		/// <param name="locator">Path to the blob</param>
 		/// <param name="serializerOptions">Options for deserializing the blob</param>
 		/// <returns>New handle to the blob</returns>
-		public static IHashedBlobRef<T> CreateBlobRef<T>(this IStorageClient storageClient, IoHash hash, BlobLocator locator, BlobSerializerOptions? serializerOptions = null)
-			=> HashedBlobRef.Create<T>(hash, storageClient.CreateBlobRef(locator), serializerOptions);
+		public static IHashedBlobRef<T> CreateBlobRef<T>(this IStorageNamespace storageNamespace, IoHash hash, BlobLocator locator, BlobSerializerOptions? serializerOptions = null)
+			=> HashedBlobRef.Create<T>(hash, storageNamespace.CreateBlobRef(locator), serializerOptions);
 
 		/// <summary>
 		/// Create a blob ref from a RefValue
 		/// </summary>
-		public static IHashedBlobRef CreateBlobRef(this IStorageClient store, HashedBlobRefValue refValue)
-			=> store.CreateBlobRef(refValue.Hash, refValue.Locator);
+		public static IHashedBlobRef CreateBlobRef(this IStorageNamespace storageNamespace, HashedBlobRefValue refValue)
+			=> storageNamespace.CreateBlobRef(refValue.Hash, refValue.Locator);
 
 		/// <summary>
 		/// Create a typed blob ref from a RefValue
 		/// </summary>
-		public static IHashedBlobRef<T> CreateBlobRef<T>(this IStorageClient store, HashedBlobRefValue refValue, BlobSerializerOptions? options)
-			=> store.CreateBlobRef<T>(refValue.Hash, refValue.Locator, options);
+		public static IHashedBlobRef<T> CreateBlobRef<T>(this IStorageNamespace storageNamespace, HashedBlobRefValue refValue, BlobSerializerOptions? options)
+			=> storageNamespace.CreateBlobRef<T>(refValue.Hash, refValue.Locator, options);
 
 		class NamedBlobRef : IBlobRef
 		{
-			readonly IStorageClient _storageClient;
+			readonly IStorageNamespace _storageNamespace;
 			readonly RefName _name;
 			readonly RefCacheTime _cacheTime;
 
-			public NamedBlobRef(IStorageClient storageClient, RefName name, RefCacheTime cacheTime = default)
+			public NamedBlobRef(IStorageNamespace storageNamespace, RefName name, RefCacheTime cacheTime = default)
 			{
-				_storageClient = storageClient;
+				_storageNamespace = storageNamespace;
 				_name = name;
 				_cacheTime = cacheTime;
 			}
@@ -295,7 +295,7 @@ namespace EpicGames.Horde.Storage
 
 			public async ValueTask<BlobData> ReadBlobDataAsync(CancellationToken cancellationToken = default)
 			{
-				IHashedBlobRef handle = await _storageClient.ReadRefAsync(_name, _cacheTime, cancellationToken);
+				IHashedBlobRef handle = await _storageNamespace.ReadRefAsync(_name, _cacheTime, cancellationToken);
 				return await handle.ReadBlobDataAsync(cancellationToken);
 			}
 
@@ -309,30 +309,30 @@ namespace EpicGames.Horde.Storage
 		/// <summary>
 		/// Creates a new blob ref from a ref name
 		/// </summary>
-		/// <param name="storageClient">The store instance to read from</param>
+		/// <param name="storageNamespace">The store instance to read from</param>
 		/// <param name="name">Name of the reference</param>
 		/// <param name="cacheTime">Maximum age for cached responses</param>
 		/// <returns>New handle to the blob</returns>
-		public static IBlobRef CreateBlobRef(this IStorageClient storageClient, RefName name, RefCacheTime cacheTime = default)
-			=> new NamedBlobRef(storageClient, name, cacheTime);
+		public static IBlobRef CreateBlobRef(this IStorageNamespace storageNamespace, RefName name, RefCacheTime cacheTime = default)
+			=> new NamedBlobRef(storageNamespace, name, cacheTime);
 
 		/// <summary>
 		/// Creates a new blob ref from a ref name
 		/// </summary>
-		/// <param name="storageClient">The store instance to read from</param>
+		/// <param name="storageNamespace">The store instance to read from</param>
 		/// <param name="name">Name of the reference</param>
 		/// <param name="cacheTime">Maximum age for cached responses</param>
 		/// <param name="serializerOptions">Options for deserializing the blob</param>
 		/// <returns>New handle to the blob</returns>
-		public static IBlobRef<T> CreateBlobRef<T>(this IStorageClient storageClient, RefName name, RefCacheTime cacheTime = default, BlobSerializerOptions? serializerOptions = null)
-			=> CreateBlobRef(storageClient, name, cacheTime).ForType<T>(serializerOptions);
+		public static IBlobRef<T> CreateBlobRef<T>(this IStorageNamespace storageNamespace, RefName name, RefCacheTime cacheTime = default, BlobSerializerOptions? serializerOptions = null)
+			=> CreateBlobRef(storageNamespace, name, cacheTime).ForType<T>(serializerOptions);
 
 		/// <summary>
 		/// Creates a writer using a refname as a base path
 		/// </summary>
 		/// <param name="store">The store instance to read from</param>
 		/// <param name="refName">Ref name to use as a base path</param>
-		public static IBlobWriter CreateBlobWriter(this IStorageClient store, RefName refName) => store.CreateBlobWriter(refName.ToString());
+		public static IBlobWriter CreateBlobWriter(this IStorageNamespace store, RefName refName) => store.CreateBlobWriter(refName.ToString());
 
 		#endregion
 
@@ -345,7 +345,7 @@ namespace EpicGames.Horde.Storage
 		/// <param name="name">Alias for the blob</param>
 		/// <param name="cancellationToken">Cancellation token for the operation</param>
 		/// <returns>Blobs matching the given handle</returns>
-		public static async Task<BlobAlias?> FindAliasAsync(this IStorageClient store, string name, CancellationToken cancellationToken = default)
+		public static async Task<BlobAlias?> FindAliasAsync(this IStorageNamespace store, string name, CancellationToken cancellationToken = default)
 		{
 			BlobAlias[] aliases = await store.FindAliasesAsync(name, 1, cancellationToken);
 			return aliases.FirstOrDefault();
@@ -363,7 +363,7 @@ namespace EpicGames.Horde.Storage
 		/// <param name="cacheTime">Minimum coherency for any cached value to be returned</param>
 		/// <param name="cancellationToken">Cancellation token for the operation</param>
 		/// <returns>True if the ref exists, false if it did not exist</returns>
-		public static async Task<bool> RefExistsAsync(this IStorageClient store, RefName name, RefCacheTime cacheTime = default, CancellationToken cancellationToken = default)
+		public static async Task<bool> RefExistsAsync(this IStorageNamespace store, RefName name, RefCacheTime cacheTime = default, CancellationToken cancellationToken = default)
 		{
 			IBlobRef? target = await store.TryReadRefAsync(name, cacheTime, cancellationToken);
 			return target != null;
@@ -377,7 +377,7 @@ namespace EpicGames.Horde.Storage
 		/// <param name="cacheTime">Minimum coherency of any cached result</param>
 		/// <param name="cancellationToken">Cancellation token for the operation</param>
 		/// <returns>The ref target</returns>
-		public static async Task<IHashedBlobRef> ReadRefAsync(this IStorageClient store, RefName name, RefCacheTime cacheTime = default, CancellationToken cancellationToken = default)
+		public static async Task<IHashedBlobRef> ReadRefAsync(this IStorageNamespace store, RefName name, RefCacheTime cacheTime = default, CancellationToken cancellationToken = default)
 		{
 			IHashedBlobRef? refTarget = await store.TryReadRefAsync(name, cacheTime, cancellationToken);
 			return refTarget ?? throw new RefNameNotFoundException(name);
@@ -386,9 +386,9 @@ namespace EpicGames.Horde.Storage
 		#endregion
 
 		/// <summary>
-		/// Gets a snapshot of the stats for the storage client.
+		/// Gets a snapshot of the stats for the storage namespace.
 		/// </summary>
-		public static StorageStats GetStats(this IStorageClient store)
+		public static StorageStats GetStats(this IStorageNamespace store)
 		{
 			StorageStats stats = new StorageStats();
 			store.GetStats(stats);
diff --git a/Engine/Source/Programs/Shared/EpicGames.Horde/Tools/HttpToolCollection.cs b/Engine/Source/Programs/Shared/EpicGames.Horde/Tools/HttpToolCollection.cs
--- a/Engine/Source/Programs/Shared/EpicGames.Horde/Tools/HttpToolCollection.cs
+++ b/Engine/Source/Programs/Shared/EpicGames.Horde/Tools/HttpToolCollection.cs
@@ -49,10 +49,10 @@ namespace EpicGames.Horde.Tools
 			public Task<ITool?> CreateDeploymentAsync(ToolDeploymentConfig options, HashedBlobRefValue target, CancellationToken cancellationToken = default)
 				=> throw new NotImplementedException();
 
-			public IStorageBackend CreateStorageBackend()
+			public IStorageBackend GetStorageBackend()
 				=> throw new NotImplementedException();
 
-			public IStorageClient CreateStorageClient()
+			public IStorageNamespace GetStorageNamespace()
 				=> throw new NotImplementedException();
 		}
 
@@ -127,7 +127,7 @@ namespace EpicGames.Horde.Tools
 
 		IBlobRef<DirectoryNode> Open(ToolId toolId, BlobLocator locator)
 		{
-			IStorageClient store = _hordeClient.CreateStorageClient(toolId);
+			IStorageNamespace store = _hordeClient.GetStorageNamespace(toolId);
 			return store.CreateBlobRef<DirectoryNode>(locator);
 		}
 
diff --git a/Engine/Source/Programs/Shared/EpicGames.Horde/Tools/ITool.cs b/Engine/Source/Programs/Shared/EpicGames.Horde/Tools/ITool.cs
--- a/Engine/Source/Programs/Shared/EpicGames.Horde/Tools/ITool.cs
+++ b/Engine/Source/Programs/Shared/EpicGames.Horde/Tools/ITool.cs
@@ -101,16 +101,14 @@ namespace EpicGames.Horde.Tools
 		Task<ITool?> CreateDeploymentAsync(ToolDeploymentConfig options, HashedBlobRefValue target, CancellationToken cancellationToken = default);
 
 		/// <summary>
-		/// Gets the storage backend for a particular tool
+		/// Gets the storage backend for this tool
 		/// </summary>
-		/// <returns>Instance of the backend client</returns>
-		IStorageBackend CreateStorageBackend();
+		IStorageBackend GetStorageBackend();
 
 		/// <summary>
-		/// Gets the storage backend for a particular tool
+		/// Gets the storage namespace for this particular tool
 		/// </summary>
-		/// <returns>Instance of the storage client</returns>
-		IStorageClient CreateStorageClient();
+		IStorageNamespace GetStorageNamespace();
 	}
 
 	/// <summary>
diff --git a/Engine/Source/Programs/UnrealBuildTool/Artifacts/HordeStorageArtifactCache.cs b/Engine/Source/Programs/UnrealBuildTool/Artifacts/HordeStorageArtifactCache.cs
--- a/Engine/Source/Programs/UnrealBuildTool/Artifacts/HordeStorageArtifactCache.cs
+++ b/Engine/Source/Programs/UnrealBuildTool/Artifacts/HordeStorageArtifactCache.cs
@@ -163,7 +163,7 @@ namespace UnrealBuildTool.Artifacts
 		/// <summary>
 		/// Underlying storage object
 		/// </summary>
-		private IStorageClient? _store = null;
+		private IStorageNamespace? _store = null;
 
 		/// <summary>
 		/// Task used to wait on ready state
@@ -210,10 +210,9 @@ namespace UnrealBuildTool.Artifacts
 		/// <summary>
 		/// Create a memory only cache
 		/// </summary>
-		/// <returns>Storage client instance</returns>
 		public static IArtifactCache CreateMemoryCache(ILogger logger)
 		{
-			HordeStorageArtifactCache cache = new(BundleStorageClient.CreateInMemory(logger))
+			HordeStorageArtifactCache cache = new(BundleStorageNamespace.CreateInMemory(logger))
 			{
 				State = ArtifactCacheState.Available
 			};
@@ -227,7 +226,6 @@ namespace UnrealBuildTool.Artifacts
 		/// <param name="memoryMappedFileCache">Cache for memory mapped files</param>
 		/// <param name="logger">Logging object</param>
 		/// <param name="cleanDirectory">If true, clean the directory</param>
-		/// <returns>Storage client instance</returns>
 		public static IArtifactCache CreateFileCache(DirectoryReference directory, MemoryMappedFileCache memoryMappedFileCache, ILogger logger, bool cleanDirectory)
 		{
 			HordeStorageArtifactCache cache = new(null);
@@ -239,7 +237,7 @@ namespace UnrealBuildTool.Artifacts
 		/// Constructor
 		/// </summary>
 		/// <param name="storage">Storage object to use</param>
-		private HordeStorageArtifactCache(IStorageClient? storage)
+		private HordeStorageArtifactCache(IStorageNamespace? storage)
 		{
 			_store = storage;
 			_pendingWrites = new(MaxPendingSize);
@@ -512,7 +510,7 @@ namespace UnrealBuildTool.Artifacts
 				}
 				Directory.CreateDirectory(directory.FullName);
 
-				_store = BundleStorageClient.CreateFromDirectory(directory, BundleCache.None, memoryMappedFileCache, logger);
+				_store = BundleStorageNamespace.CreateFromDirectory(directory, BundleCache.None, memoryMappedFileCache, logger);
 
 				State = ArtifactCacheState.Available;
 				return State;
diff --git a/Engine/Source/Programs/UnrealBuildTool/Executors/UnrealBuildAccelerator/UBAAgentCoordinatorHorde.cs b/Engine/Source/Programs/UnrealBuildTool/Executors/UnrealBuildAccelerator/UBAAgentCoordinatorHorde.cs
--- a/Engine/Source/Programs/UnrealBuildTool/Executors/UnrealBuildAccelerator/UBAAgentCoordinatorHorde.cs
+++ b/Engine/Source/Programs/UnrealBuildTool/Executors/UnrealBuildAccelerator/UBAAgentCoordinatorHorde.cs
@@ -52,7 +52,7 @@ namespace UnrealBuildTool
 		readonly UBAExecutor _owner;
 		internal StatusUpdateAction? _updateStatus;
 
-		readonly BundleStorageClient _storage = BundleStorageClient.CreateInMemory(NullLogger.Instance);
+		readonly BundleStorageNamespace _storage = BundleStorageNamespace.CreateInMemory(NullLogger.Instance);
 		BlobLocator _ubaAgentLocator;
 
 		readonly ServiceProvider _serviceProvider;
diff --git a/Engine/Source/Programs/UnrealGameSync/UnrealGameSync/ArtifactDownload.cs b/Engine/Source/Programs/UnrealGameSync/UnrealGameSync/ArtifactDownload.cs
--- a/Engine/Source/Programs/UnrealGameSync/UnrealGameSync/ArtifactDownload.cs
+++ b/Engine/Source/Programs/UnrealGameSync/UnrealGameSync/ArtifactDownload.cs
@@ -55,11 +55,11 @@ namespace UnrealGameSync
 				await using ServiceProvider serviceProvider = services.BuildServiceProvider();
 
 				IHordeClient hordeClient = serviceProvider.GetRequiredService<IHordeClient>();
-				IStorageClient storageClient = hordeClient.CreateStorageClient(baseUri.AbsolutePath);
+				IStorageNamespace storageNamespace = hordeClient.GetStorageNamespace(baseUri.AbsolutePath);
 
 				progress.Report("Connecting to server...");
 
-				IBlobRef<DirectoryNode>? blobRef = await storageClient.ReadRefAsync<DirectoryNode>(refName, cancellationToken: cancellationToken);
+				IBlobRef<DirectoryNode>? blobRef = await storageNamespace.ReadRefAsync<DirectoryNode>(refName, cancellationToken: cancellationToken);
 
 				progress.Report("Starting...");
 				await blobRef.ExtractAsync(outputDir.ToDirectoryInfo(), new CopyProgressAdapter(progress), TimeSpan.FromSeconds(0.2), serviceProvider.GetRequiredService<ILogger<ArtifactDownload>>(), cancellationToken);
diff --git a/Engine/Source/Programs/UnrealGameSync/UnrealGameSync/Launcher.cs b/Engine/Source/Programs/UnrealGameSync/UnrealGameSync/Launcher.cs
--- a/Engine/Source/Programs/UnrealGameSync/UnrealGameSync/Launcher.cs
+++ b/Engine/Source/Programs/UnrealGameSync/UnrealGameSync/Launcher.cs
@@ -186,8 +186,8 @@ namespace UnrealGameSync
 						await SafeDeleteDirectoryContentsWithRetryAsync(applicationFolder, cancellationToken);
 
 						// Download and extract the data
-						IStorageClient storageClient = hordeClient.CreateStorageClient(toolId);
-						IBlobRef<DirectoryNode> handle = storageClient.CreateBlobRef<DirectoryNode>(deployment.Locator);
+						IStorageNamespace storageNamespace = hordeClient.GetStorageNamespace(toolId);
+						IBlobRef<DirectoryNode> handle = storageNamespace.CreateBlobRef<DirectoryNode>(deployment.Locator);
 
 						DirectoryNode rootNode = await handle.ReadBlobAsync(cancellationToken: cancellationToken);
 						await rootNode.ExtractAsync(new DirectoryInfo(applicationFolder), new ExtractStatsLogger(0, rootNode.Length, logger), TimeSpan.FromSeconds(5.0), logger, cancellationToken);
