Horde: Renaming files to match class names. Reintroducing IStorageClient as an interface to the storage system (was IStorageNamespaceFactory).

#rnx

diff --git a/Engine/Source/Programs/Horde/Drivers/JobDriver.Tests/WorkerServiceTest.cs b/Engine/Source/Programs/Horde/Drivers/JobDriver.Tests/WorkerServiceTest.cs
--- a/Engine/Source/Programs/Horde/Drivers/JobDriver.Tests/WorkerServiceTest.cs
+++ b/Engine/Source/Programs/Horde/Drivers/JobDriver.Tests/WorkerServiceTest.cs
@@ -61,7 +61,7 @@ namespace JobDriver.Tests
 			_serviceCollection.AddSingleton<BundleCache>();
 			_serviceCollection.AddSingleton<StorageBackendCache>();
 			_serviceCollection.AddSingleton<HttpStorageBackendFactory>();
-			_serviceCollection.AddSingleton<HttpStorageNamespaceFactory>();
+			_serviceCollection.AddSingleton<HttpStorageClient>();
 
 			_serviceCollection.Configure<DriverSettings>(settings =>
 			{
diff --git a/Engine/Source/Programs/Horde/HordeCmd/Commands/Bundles/BundleCreate.cs b/Engine/Source/Programs/Horde/HordeCmd/Commands/Bundles/BundleCreate.cs
--- a/Engine/Source/Programs/Horde/HordeCmd/Commands/Bundles/BundleCreate.cs
+++ b/Engine/Source/Programs/Horde/HordeCmd/Commands/Bundles/BundleCreate.cs
@@ -5,7 +5,6 @@ using System.Diagnostics;
 using EpicGames.Core;
 using EpicGames.Horde.Storage;
 using EpicGames.Horde.Storage.Bundles;
-using EpicGames.Horde.Storage.Clients;
 using EpicGames.Horde.Storage.Nodes;
 using Microsoft.Extensions.Logging;
 
@@ -34,8 +33,8 @@ namespace Horde.Commands.Bundles
 		[Description("Clean the output folder before writing any data")]
 		public bool CleanOutput { get; set; }
 
-		public BundleCreate(HttpStorageNamespaceFactory storageNamespaceFactory, BundleCache bundleCache)
-			: base(storageNamespaceFactory, bundleCache)
+		public BundleCreate(HttpStorageClient storageClient, BundleCache bundleCache)
+			: base(storageClient, bundleCache)
 		{
 		}
 
diff --git a/Engine/Source/Programs/Horde/HordeCmd/Commands/Bundles/BundleExtract.cs b/Engine/Source/Programs/Horde/HordeCmd/Commands/Bundles/BundleExtract.cs
--- a/Engine/Source/Programs/Horde/HordeCmd/Commands/Bundles/BundleExtract.cs
+++ b/Engine/Source/Programs/Horde/HordeCmd/Commands/Bundles/BundleExtract.cs
@@ -6,7 +6,6 @@ using EpicGames.Core;
 using EpicGames.Horde.Storage;
 using EpicGames.Horde.Storage.Backends;
 using EpicGames.Horde.Storage.Bundles;
-using EpicGames.Horde.Storage.Clients;
 using EpicGames.Horde.Storage.Nodes;
 using Microsoft.Extensions.Logging;
 
@@ -39,8 +38,8 @@ namespace Horde.Commands.Bundles
 		[Description("If set, deletes the contents of the output directory before extraction.")]
 		public bool CleanOutput { get; set; }
 
-		public BundleExtract(HttpStorageNamespaceFactory storageNamespaceFactory, BundleCache bundleCache)
-			: base(storageNamespaceFactory, bundleCache)
+		public BundleExtract(HttpStorageClient storageClient, BundleCache bundleCache)
+			: base(storageClient, bundleCache)
 		{
 		}
 
diff --git a/Engine/Source/Programs/Horde/HordeCmd/Commands/StorageCommandBase.cs b/Engine/Source/Programs/Horde/HordeCmd/Commands/StorageCommandBase.cs
--- a/Engine/Source/Programs/Horde/HordeCmd/Commands/StorageCommandBase.cs
+++ b/Engine/Source/Programs/Horde/HordeCmd/Commands/StorageCommandBase.cs
@@ -4,7 +4,6 @@ using System.ComponentModel;
 using EpicGames.Core;
 using EpicGames.Horde.Storage;
 using EpicGames.Horde.Storage.Bundles;
-using EpicGames.Horde.Storage.Clients;
 
 namespace Horde.Commands
 {
@@ -32,14 +31,14 @@ namespace Horde.Commands
 		/// </summary>
 		public BundleCache BundleCache { get; }
 
-		readonly HttpStorageNamespaceFactory _storageNamespaceFactory;
+		readonly HttpStorageClient _storageClient;
 
 		/// <summary>
 		/// Constructor
 		/// </summary>
-		public StorageCommandBase(HttpStorageNamespaceFactory storageNamespaceFactory, BundleCache bundleCache)
+		public StorageCommandBase(HttpStorageClient storageClient, BundleCache bundleCache)
 		{
-			_storageNamespaceFactory = storageNamespaceFactory;
+			_storageClient = storageClient;
 
 			BundleCache = bundleCache;
 		}
@@ -51,11 +50,11 @@ namespace Horde.Commands
 		{
 			if (String.IsNullOrEmpty(Path))
 			{
-				return _storageNamespaceFactory.CreateClient(Namespace);
+				return _storageClient.GetNamespace(Namespace);
 			}
 			else
 			{
-				return _storageNamespaceFactory.CreateClientWithPath(Path);
+				return _storageClient.GetNamespaceWithPath(Path);
 			}
 		}
 	}
diff --git a/Engine/Source/Programs/Horde/HordeCmd/Commands/Tools/ToolUpload.cs b/Engine/Source/Programs/Horde/HordeCmd/Commands/Tools/ToolUpload.cs
--- a/Engine/Source/Programs/Horde/HordeCmd/Commands/Tools/ToolUpload.cs
+++ b/Engine/Source/Programs/Horde/HordeCmd/Commands/Tools/ToolUpload.cs
@@ -4,7 +4,6 @@ using System.ComponentModel;
 using EpicGames.Core;
 using EpicGames.Horde;
 using EpicGames.Horde.Storage;
-using EpicGames.Horde.Storage.Clients;
 using EpicGames.Horde.Storage.Nodes;
 using EpicGames.Horde.Tools;
 using Microsoft.Extensions.Logging;
@@ -27,18 +26,18 @@ namespace Horde.Commands
 		public DirectoryReference InputDir { get; set; } = null!;
 
 		readonly HordeHttpClient _hordeHttpClient;
-		readonly HttpStorageNamespaceFactory _httpStorageNamespaceFactory;
+		readonly HttpStorageClient _storageClient;
 
-		public ToolUpload(HordeHttpClient httpClient, HttpStorageNamespaceFactory httpStorageNamespaceFactory)
+		public ToolUpload(HordeHttpClient httpClient, HttpStorageClient storageClient)
 		{
 			_hordeHttpClient = httpClient;
-			_httpStorageNamespaceFactory = httpStorageNamespaceFactory;
+			_storageClient = storageClient;
 		}
 
 		/// <inheritdoc/>
 		public override async Task<int> ExecuteAsync(ILogger logger)
 		{
-			IStorageNamespace storageNamespace = _httpStorageNamespaceFactory.CreateClientWithPath($"api/v1/tools/{ToolId}");
+			IStorageNamespace storageNamespace = _storageClient.GetNamespaceWithPath($"api/v1/tools/{ToolId}");
 
 			IHashedBlobRef<DirectoryNode> target;
 			await using (IBlobWriter writer = storageNamespace.CreateBlobWriter())
diff --git a/Engine/Source/Programs/Horde/HordeCmd/Commands/Vcs/VcsBase.cs b/Engine/Source/Programs/Horde/HordeCmd/Commands/Vcs/VcsBase.cs
--- a/Engine/Source/Programs/Horde/HordeCmd/Commands/Vcs/VcsBase.cs
+++ b/Engine/Source/Programs/Horde/HordeCmd/Commands/Vcs/VcsBase.cs
@@ -149,14 +149,14 @@ namespace Horde.Commands.Vcs
 		const string DataDir = ".horde";
 		const string StateFileName = "state.dat";
 
-		readonly IStorageNamespaceFactory _storageNamespaceFactory;
+		readonly IStorageClient _storageClient;
 
-		protected VcsBase(IStorageNamespaceFactory storageNamespaceFactory)
+		protected VcsBase(IStorageClient storageClient)
 		{
-			_storageNamespaceFactory = storageNamespaceFactory;
+			_storageClient = storageClient;
 		}
 
-		protected IStorageNamespace GetStorageNamespace() => _storageNamespaceFactory.GetNamespace(NamespaceId);
+		protected IStorageNamespace GetStorageNamespace() => _storageClient.GetNamespace(NamespaceId);
 
 		protected static async Task<WorkspaceState> ReadStateAsync(DirectoryReference rootDir)
 		{
diff --git a/Engine/Source/Programs/Horde/HordeCmd/Commands/Vcs/VcsBranch.cs b/Engine/Source/Programs/Horde/HordeCmd/Commands/Vcs/VcsBranch.cs
--- a/Engine/Source/Programs/Horde/HordeCmd/Commands/Vcs/VcsBranch.cs
+++ b/Engine/Source/Programs/Horde/HordeCmd/Commands/Vcs/VcsBranch.cs
@@ -12,8 +12,8 @@ namespace Horde.Commands.Vcs
 		[CommandLine(Prefix = "-Name=", Required = true)]
 		public string Name { get; set; } = "";
 
-		public VcsBranch(IStorageNamespaceFactory storageNamespaceFactory)
-			: base(storageNamespaceFactory)
+		public VcsBranch(IStorageClient storageClient)
+			: base(storageClient)
 		{
 		}
 
diff --git a/Engine/Source/Programs/Horde/HordeCmd/Commands/Vcs/VcsCheckout.cs b/Engine/Source/Programs/Horde/HordeCmd/Commands/Vcs/VcsCheckout.cs
--- a/Engine/Source/Programs/Horde/HordeCmd/Commands/Vcs/VcsCheckout.cs
+++ b/Engine/Source/Programs/Horde/HordeCmd/Commands/Vcs/VcsCheckout.cs
@@ -22,8 +22,8 @@ namespace Horde.Commands.Vcs
 		[CommandLine("-Force")]
 		public bool Force { get; set; }
 
-		public VcsCheckout(IStorageNamespaceFactory storageNamespaceFactory)
-			: base(storageNamespaceFactory)
+		public VcsCheckout(IStorageClient storageClient)
+			: base(storageClient)
 		{
 		}
 
diff --git a/Engine/Source/Programs/Horde/HordeCmd/Commands/Vcs/VcsClone.cs b/Engine/Source/Programs/Horde/HordeCmd/Commands/Vcs/VcsClone.cs
--- a/Engine/Source/Programs/Horde/HordeCmd/Commands/Vcs/VcsClone.cs
+++ b/Engine/Source/Programs/Horde/HordeCmd/Commands/Vcs/VcsClone.cs
@@ -10,8 +10,8 @@ namespace Horde.Commands.Vcs
 	[Command("vcs", "clone", "Initialize a directory for VCS-like operations", Advertise = false)]
 	class VcsClone : VcsCheckout
 	{
-		public VcsClone(IStorageNamespaceFactory storageNamespaceFactory)
-			: base(storageNamespaceFactory)
+		public VcsClone(IStorageClient storageClient)
+			: base(storageClient)
 		{
 		}
 
diff --git a/Engine/Source/Programs/Horde/HordeCmd/Commands/Vcs/VcsCommit.cs b/Engine/Source/Programs/Horde/HordeCmd/Commands/Vcs/VcsCommit.cs
--- a/Engine/Source/Programs/Horde/HordeCmd/Commands/Vcs/VcsCommit.cs
+++ b/Engine/Source/Programs/Horde/HordeCmd/Commands/Vcs/VcsCommit.cs
@@ -15,8 +15,8 @@ namespace Horde.Commands.Vcs
 		[CommandLine("-Message=")]
 		public string? Message { get; set; }
 
-		public VcsCommit(IStorageNamespaceFactory storageNamespaceFactory)
-			: base(storageNamespaceFactory)
+		public VcsCommit(IStorageClient storageClient)
+			: base(storageClient)
 		{
 		}
 
diff --git a/Engine/Source/Programs/Horde/HordeCmd/Commands/Vcs/VcsInit.cs b/Engine/Source/Programs/Horde/HordeCmd/Commands/Vcs/VcsInit.cs
--- a/Engine/Source/Programs/Horde/HordeCmd/Commands/Vcs/VcsInit.cs
+++ b/Engine/Source/Programs/Horde/HordeCmd/Commands/Vcs/VcsInit.cs
@@ -9,8 +9,8 @@ namespace Horde.Commands.Vcs
 	[Command("vcs", "init", "Initialize a directory for VCS-like operations", Advertise = false)]
 	class VcsInit : VcsBase
 	{
-		public VcsInit(IStorageNamespaceFactory storageNamespaceFactory)
-			: base(storageNamespaceFactory)
+		public VcsInit(IStorageClient storageClient)
+			: base(storageClient)
 		{
 		}
 
diff --git a/Engine/Source/Programs/Horde/HordeCmd/Commands/Vcs/VcsLog.cs b/Engine/Source/Programs/Horde/HordeCmd/Commands/Vcs/VcsLog.cs
--- a/Engine/Source/Programs/Horde/HordeCmd/Commands/Vcs/VcsLog.cs
+++ b/Engine/Source/Programs/Horde/HordeCmd/Commands/Vcs/VcsLog.cs
@@ -14,8 +14,8 @@ namespace Horde.Commands.Vcs
 		[CommandLine("-Count")]
 		public int Count { get; set; } = 20;
 
-		public VcsLog(IStorageNamespaceFactory storageNamespaceFactory)
-			: base(storageNamespaceFactory)
+		public VcsLog(IStorageClient storageClient)
+			: base(storageClient)
 		{
 		}
 
diff --git a/Engine/Source/Programs/Horde/HordeCmd/Commands/Vcs/VcsStatus.cs b/Engine/Source/Programs/Horde/HordeCmd/Commands/Vcs/VcsStatus.cs
--- a/Engine/Source/Programs/Horde/HordeCmd/Commands/Vcs/VcsStatus.cs
+++ b/Engine/Source/Programs/Horde/HordeCmd/Commands/Vcs/VcsStatus.cs
@@ -9,8 +9,8 @@ namespace Horde.Commands.Vcs
 	[Command("vcs", "status", "Find status of local files", Advertise = false)]
 	class VcsStatus : VcsBase
 	{
-		public VcsStatus(IStorageNamespaceFactory storageNamespaceFactory)
-			: base(storageNamespaceFactory)
+		public VcsStatus(IStorageClient storageClient)
+			: base(storageClient)
 		{
 		}
 
diff --git a/Engine/Source/Programs/Horde/HordeCmd/Commands/Workspace/WorkspaceSync.cs b/Engine/Source/Programs/Horde/HordeCmd/Commands/Workspace/WorkspaceSync.cs
--- a/Engine/Source/Programs/Horde/HordeCmd/Commands/Workspace/WorkspaceSync.cs
+++ b/Engine/Source/Programs/Horde/HordeCmd/Commands/Workspace/WorkspaceSync.cs
@@ -6,7 +6,6 @@ using EpicGames.Core;
 using EpicGames.Horde.Storage;
 using EpicGames.Horde.Storage.Backends;
 using EpicGames.Horde.Storage.Bundles;
-using EpicGames.Horde.Storage.Clients;
 using EpicGames.Horde.Storage.Nodes;
 using Microsoft.Extensions.Logging;
 
@@ -41,8 +40,8 @@ namespace Horde.Commands.Workspace
 		[Description("Outputs stats for the extraction operation.")]
 		public bool Stats { get; set; }
 
-		public WorkspaceSync(HttpStorageNamespaceFactory storageNamespaceFactory, BundleCache bundleCache)
-			: base(storageNamespaceFactory, bundleCache)
+		public WorkspaceSync(HttpStorageClient storageClient, BundleCache bundleCache)
+			: base(storageClient, bundleCache)
 		{
 		}
 
diff --git a/Engine/Source/Programs/Horde/Plugins/Ddc/HordeServer.Ddc/Ddc/BlobService.cs b/Engine/Source/Programs/Horde/Plugins/Ddc/HordeServer.Ddc/Ddc/BlobService.cs
--- a/Engine/Source/Programs/Horde/Plugins/Ddc/HordeServer.Ddc/Ddc/BlobService.cs
+++ b/Engine/Source/Programs/Horde/Plugins/Ddc/HordeServer.Ddc/Ddc/BlobService.cs
@@ -11,12 +11,12 @@ namespace HordeServer.Ddc
 	{
 		static readonly BlobType s_rawBlobType = new BlobType(new Guid("{03E6C37B-33C1-491F-8541-D3C401B8B8EF}"), 1);
 
-		readonly IStorageNamespaceFactory _storageNamespaceFactory;
+		readonly IStorageClient _storageClient;
 		readonly Tracer _tracer;
 
-		public BlobService(IStorageNamespaceFactory storageNamespaceFactory, Tracer tracer)
+		public BlobService(IStorageClient storageClient, Tracer tracer)
 		{
-			_storageNamespaceFactory = storageNamespaceFactory;
+			_storageClient = storageClient;
 			_tracer = tracer;
 		}
 
@@ -60,13 +60,13 @@ namespace HordeServer.Ddc
 
 		public async Task<bool> ExistsAsync(NamespaceId ns, BlobId blob, List<string>? storageLayers = null, CancellationToken cancellationToken = default)
 		{
-			IStorageNamespace storageNamespace = _storageNamespaceFactory.GetNamespace(ns);
+			IStorageNamespace storageNamespace = _storageClient.GetNamespace(ns);
 			return await storageNamespace.FindAliasAsync(GetAlias(blob), cancellationToken) != null;
 		}
 
 		public async Task DeleteObjectAsync(NamespaceId ns, BlobId blob, CancellationToken cancellationToken)
 		{
-			IStorageNamespace storageNamespace = _storageNamespaceFactory.GetNamespace(ns);
+			IStorageNamespace storageNamespace = _storageClient.GetNamespace(ns);
 			string aliasName = GetAlias(blob);
 
 			BlobAlias[] aliases = await storageNamespace.FindAliasesAsync(aliasName, cancellationToken: cancellationToken);
@@ -78,7 +78,7 @@ namespace HordeServer.Ddc
 
 		public async Task<BlobId[]> FilterOutKnownBlobsAsync(NamespaceId ns, IEnumerable<BlobId> blobIds, CancellationToken cancellationToken)
 		{
-			IStorageNamespace storageNamespace = _storageNamespaceFactory.GetNamespace(ns);
+			IStorageNamespace storageNamespace = _storageClient.GetNamespace(ns);
 
 			List<BlobId> unknownBlobIds = new List<BlobId>();
 			foreach (BlobId blobId in blobIds)
@@ -123,7 +123,7 @@ namespace HordeServer.Ddc
 
 		public async Task<BlobContents> GetObjectAsync(NamespaceId ns, BlobId blob, List<string>? storageLayers, bool supportsRedirectUri, bool allowOndemandReplication = true, CancellationToken cancellationToken = default)
 		{
-			IStorageNamespace storageNamespace = _storageNamespaceFactory.GetNamespace(ns);
+			IStorageNamespace storageNamespace = _storageClient.GetNamespace(ns);
 
 			BlobAlias? alias = await storageNamespace.FindAliasAsync(GetAlias(blob), cancellationToken);
 			if (alias == null)
@@ -174,7 +174,7 @@ namespace HordeServer.Ddc
 
 		public async Task<BlobId> PutObjectKnownHashAsync(NamespaceId ns, IBufferedPayload content, BlobId identifier, CancellationToken cancellationToken)
 		{
-			IStorageNamespace storageNamespace = _storageNamespaceFactory.GetNamespace(ns);
+			IStorageNamespace storageNamespace = _storageClient.GetNamespace(ns);
 
 			IHashedBlobRef blobRef;
 			await using (IBlobWriter writer = storageNamespace.CreateBlobWriter())
diff --git a/Engine/Source/Programs/Horde/Plugins/Ddc/HordeServer.Ddc/Ddc/ContentIdStore.cs b/Engine/Source/Programs/Horde/Plugins/Ddc/HordeServer.Ddc/Ddc/ContentIdStore.cs
--- a/Engine/Source/Programs/Horde/Plugins/Ddc/HordeServer.Ddc/Ddc/ContentIdStore.cs
+++ b/Engine/Source/Programs/Horde/Plugins/Ddc/HordeServer.Ddc/Ddc/ContentIdStore.cs
@@ -7,11 +7,11 @@ namespace HordeServer.Ddc
 {
 	class ContentIdStore : IContentIdStore
 	{
-		readonly IStorageNamespaceFactory _storageNamespaceFactory;
+		readonly IStorageClient _storageClient;
 
-		public ContentIdStore(IStorageNamespaceFactory storageService)
+		public ContentIdStore(IStorageClient storageClient)
 		{
-			_storageNamespaceFactory = storageService;
+			_storageClient = storageClient;
 		}
 
 		static string GetAlias(BlobId blobId) => BlobService.GetAlias(blobId);
@@ -19,7 +19,7 @@ namespace HordeServer.Ddc
 
 		public async Task<BlobId[]?> ResolveAsync(NamespaceId ns, ContentId contentId, bool mustBeContentId = false, CancellationToken cancellationToken = default)
 		{
-			IStorageNamespace storageNamespace = _storageNamespaceFactory.GetNamespace(ns);
+			IStorageNamespace storageNamespace = _storageClient.GetNamespace(ns);
 
 			BlobAlias? blobAlias = await storageNamespace.FindAliasAsync(GetAlias(contentId), cancellationToken);
 			if (blobAlias == null && !mustBeContentId)
@@ -37,7 +37,7 @@ namespace HordeServer.Ddc
 
 		public async Task PutAsync(NamespaceId ns, ContentId contentId, BlobId blobId, int contentWeight, CancellationToken cancellationToken)
 		{
-			IStorageNamespace storageNamespace = _storageNamespaceFactory.GetNamespace(ns);
+			IStorageNamespace storageNamespace = _storageClient.GetNamespace(ns);
 
 			BlobAlias? blobAlias = await storageNamespace.FindAliasAsync(GetAlias(blobId), cancellationToken);
 			if (blobAlias == null)
diff --git a/Engine/Source/Programs/Horde/Plugins/Ddc/HordeServer.Ddc/Ddc/RefService.cs b/Engine/Source/Programs/Horde/Plugins/Ddc/HordeServer.Ddc/Ddc/RefService.cs
--- a/Engine/Source/Programs/Horde/Plugins/Ddc/HordeServer.Ddc/Ddc/RefService.cs
+++ b/Engine/Source/Programs/Horde/Plugins/Ddc/HordeServer.Ddc/Ddc/RefService.cs
@@ -8,13 +8,13 @@ namespace HordeServer.Ddc
 {
 	class RefService : IRefService
 	{
-		readonly IStorageNamespaceFactory _storageNamespaceFactory;
+		readonly IStorageClient _storageClient;
 		readonly IReferenceResolver _referenceResolver;
 		readonly IBlobService _blobService;
 
-		public RefService(IStorageNamespaceFactory storageService, IReferenceResolver referenceResolver, IBlobService blobService)
+		public RefService(IStorageClient storageClient, IReferenceResolver referenceResolver, IBlobService blobService)
 		{
-			_storageNamespaceFactory = storageService;
+			_storageClient = storageClient;
 			_referenceResolver = referenceResolver;
 			_blobService = blobService;
 		}
@@ -23,19 +23,19 @@ namespace HordeServer.Ddc
 
 		public async Task<bool> DeleteAsync(NamespaceId ns, BucketId bucket, RefId key, CancellationToken cancellationToken)
 		{
-			IStorageNamespace storageNamespace = _storageNamespaceFactory.GetNamespace(ns);
+			IStorageNamespace storageNamespace = _storageClient.GetNamespace(ns);
 			return await storageNamespace.DeleteRefAsync(GetRefName(bucket, key), cancellationToken);
 		}
 
 		public async Task<bool> ExistsAsync(NamespaceId ns, BucketId bucket, RefId key, CancellationToken cancellationToken)
 		{
-			IStorageNamespace storageNamespace = _storageNamespaceFactory.GetNamespace(ns);
+			IStorageNamespace storageNamespace = _storageClient.GetNamespace(ns);
 			return await storageNamespace.RefExistsAsync(GetRefName(bucket, key), cancellationToken: cancellationToken);
 		}
 
 		public async Task<(ContentId[], BlobId[])> FinalizeAsync(NamespaceId ns, BucketId bucket, RefId key, BlobId blobHash, CancellationToken cancellationToken)
 		{
-			IStorageNamespace storageNamespace = _storageNamespaceFactory.GetNamespace(ns);
+			IStorageNamespace storageNamespace = _storageClient.GetNamespace(ns);
 
 			BlobAlias? blobAlias = await storageNamespace.FindAliasAsync(BlobService.GetAlias(blobHash), cancellationToken);
 			if (blobAlias == null)
@@ -125,7 +125,7 @@ namespace HordeServer.Ddc
 
 		public async Task<(RefRecord, BlobContents?)> GetAsync(NamespaceId ns, BucketId bucket, RefId key, string[] fields, bool doLastAccessTracking, CancellationToken cancellationToken)
 		{
-			IStorageNamespace storageNamespace = _storageNamespaceFactory.GetNamespace(ns);
+			IStorageNamespace storageNamespace = _storageClient.GetNamespace(ns);
 
 			DdcRefNode? node = await storageNamespace.TryReadRefTargetAsync<DdcRefNode>(GetRefName(bucket, key), cancellationToken: cancellationToken);
 			if (node == null)
@@ -142,7 +142,7 @@ namespace HordeServer.Ddc
 
 		public async Task<List<BlobId>> GetReferencedBlobsAsync(NamespaceId ns, BucketId bucket, RefId key, CancellationToken cancellationToken)
 		{
-			IStorageNamespace storageNamespace = _storageNamespaceFactory.GetNamespace(ns);
+			IStorageNamespace storageNamespace = _storageClient.GetNamespace(ns);
 
 			DdcRefNode? node = await storageNamespace.TryReadRefTargetAsync<DdcRefNode>(GetRefName(bucket, key), cancellationToken: cancellationToken);
 			if (node == null)
diff --git a/Engine/Source/Programs/Horde/Plugins/Ddc/HordeServer.Ddc/Ddc/RequestHelper.cs b/Engine/Source/Programs/Horde/Plugins/Ddc/HordeServer.Ddc/Ddc/RequestHelper.cs
--- a/Engine/Source/Programs/Horde/Plugins/Ddc/HordeServer.Ddc/Ddc/RequestHelper.cs
+++ b/Engine/Source/Programs/Horde/Plugins/Ddc/HordeServer.Ddc/Ddc/RequestHelper.cs
@@ -12,12 +12,12 @@ namespace HordeServer.Ddc
 {
 	class RequestHelper : IRequestHelper
 	{
-		readonly IStorageNamespaceFactory _storageNamespaceFactory;
+		readonly IStorageClient _storageClient;
 		readonly StorageConfig _storageConfig;
 
-		public RequestHelper(IStorageNamespaceFactory storageNamespaceFactory, IOptionsSnapshot<StorageConfig> storageConfig)
+		public RequestHelper(IStorageClient storageClient, IOptionsSnapshot<StorageConfig> storageConfig)
 		{
-			_storageNamespaceFactory = storageNamespaceFactory;
+			_storageClient = storageClient;
 			_storageConfig = storageConfig.Value;
 		}
 
@@ -28,7 +28,7 @@ namespace HordeServer.Ddc
 
 		public Task<ActionResult?> HasAccessToNamespaceAsync(ClaimsPrincipal user, HttpRequest request, NamespaceId ns, AclAction[] aclActions)
 		{
-			IStorageNamespace? storageNamespace = _storageNamespaceFactory.TryGetNamespace(ns);
+			IStorageNamespace? storageNamespace = _storageClient.TryGetNamespace(ns);
 			if (storageNamespace == null)
 			{
 				return Task.FromResult<ActionResult?>(new ForbidResult());
diff --git a/Engine/Source/Programs/Horde/Plugins/Storage/HordeServer.Storage/Storage/StorageController.cs b/Engine/Source/Programs/Horde/Plugins/Storage/HordeServer.Storage/Storage/StorageController.cs
--- a/Engine/Source/Programs/Horde/Plugins/Storage/HordeServer.Storage/Storage/StorageController.cs
+++ b/Engine/Source/Programs/Horde/Plugins/Storage/HordeServer.Storage/Storage/StorageController.cs
@@ -332,9 +332,9 @@ namespace HordeServer.Storage
 		/// <summary>
 		/// Reads a ref from storage, without performing namespace access checks.
 		/// </summary>
-		public static async Task<ActionResult<ReadRefResponse>> ReadRefInternalAsync(IStorageNamespaceFactory storageService, string basePath, NamespaceId namespaceId, RefName refName, IHeaderDictionary headers, CancellationToken cancellationToken)
+		public static async Task<ActionResult<ReadRefResponse>> ReadRefInternalAsync(IStorageClient storageClient, string basePath, NamespaceId namespaceId, RefName refName, IHeaderDictionary headers, CancellationToken cancellationToken)
 		{
-			IStorageNamespace client = storageService.GetNamespace(namespaceId);
+			IStorageNamespace storageNamespace = storageClient.GetNamespace(namespaceId);
 
 			RefCacheTime cacheTime = new RefCacheTime();
 			foreach (string? entry in headers.CacheControl)
@@ -345,7 +345,7 @@ namespace HordeServer.Storage
 				}
 			}
 
-			IHashedBlobRef? target = await client.TryReadRefAsync(refName, cacheTime, cancellationToken: cancellationToken);
+			IHashedBlobRef? target = await storageNamespace.TryReadRefAsync(refName, cacheTime, cancellationToken: cancellationToken);
 			if (target == null)
 			{
 				return new NotFoundResult();
diff --git a/Engine/Source/Programs/Horde/Plugins/Storage/HordeServer.Storage/Storage/StorageService.cs b/Engine/Source/Programs/Horde/Plugins/Storage/HordeServer.Storage/Storage/StorageService.cs
--- a/Engine/Source/Programs/Horde/Plugins/Storage/HordeServer.Storage/Storage/StorageService.cs
+++ b/Engine/Source/Programs/Horde/Plugins/Storage/HordeServer.Storage/Storage/StorageService.cs
@@ -32,7 +32,7 @@ namespace HordeServer.Storage
 	/// <summary>
 	/// Interface for the storage service
 	/// </summary>
-	public interface IStorageService : IStorageNamespaceFactory
+	public interface IStorageService : IStorageClient
 	{
 	}
 
@@ -450,12 +450,12 @@ namespace HordeServer.Storage
 
 		internal static ObjectKey GetObjectKey(BlobLocator locator) => new ObjectKey($"{locator.Path}.blob");
 
-		class StorageNamespaceFactory : IStorageNamespaceFactory
+		class StorageClient : IStorageClient
 		{
 			readonly StorageService _storageService;
 			readonly StorageConfig _storageConfig;
 
-			public StorageNamespaceFactory(StorageService storageService, StorageConfig storageConfig)
+			public StorageClient(StorageService storageService, StorageConfig storageConfig)
 			{
 				_storageService = storageService;
 				_storageConfig = storageConfig;
@@ -468,8 +468,8 @@ namespace HordeServer.Storage
 		/// <summary>
 		/// Creates a new storage namespace factory using the current global config value
 		/// </summary>
-		public IStorageNamespaceFactory CreateStorageNamespaceFactory(StorageConfig storageConfig)
-			=> new StorageNamespaceFactory(this, storageConfig);
+		public IStorageClient CreateStorageClient(StorageConfig storageConfig)
+			=> new StorageClient(this, storageConfig);
 
 		/// <inheritdoc/>
 		public async Task StartAsync(CancellationToken cancellationToken)
diff --git a/Engine/Source/Programs/Horde/Plugins/Storage/HordeServer.Storage/StoragePlugin.cs b/Engine/Source/Programs/Horde/Plugins/Storage/HordeServer.Storage/StoragePlugin.cs
--- a/Engine/Source/Programs/Horde/Plugins/Storage/HordeServer.Storage/StoragePlugin.cs
+++ b/Engine/Source/Programs/Horde/Plugins/Storage/HordeServer.Storage/StoragePlugin.cs
@@ -40,7 +40,7 @@ namespace HordeServer.Storage
 		{
 			services.AddSingleton<StorageService>();
 			services.AddSingleton<IStorageService>(sp => sp.GetRequiredService<StorageService>());
-			services.AddScoped(sp => sp.GetRequiredService<StorageService>().CreateStorageNamespaceFactory(sp.GetRequiredService<IOptionsSnapshot<StorageConfig>>().Value));
+			services.AddScoped(sp => sp.GetRequiredService<StorageService>().CreateStorageClient(sp.GetRequiredService<IOptionsSnapshot<StorageConfig>>().Value));
 
 			services.AddSingleton<IObjectStoreFactory, ObjectStoreFactory>();
 			services.AddSingleton<AwsObjectStoreFactory>();
diff --git a/Engine/Source/Programs/Shared/EpicGames.Horde/HordeExtensions.cs b/Engine/Source/Programs/Shared/EpicGames.Horde/HordeExtensions.cs
--- a/Engine/Source/Programs/Shared/EpicGames.Horde/HordeExtensions.cs
+++ b/Engine/Source/Programs/Shared/EpicGames.Horde/HordeExtensions.cs
@@ -5,7 +5,6 @@ using EpicGames.Core;
 using EpicGames.Horde.Storage;
 using EpicGames.Horde.Storage.Backends;
 using EpicGames.Horde.Storage.Bundles;
-using EpicGames.Horde.Storage.Clients;
 using Microsoft.Extensions.DependencyInjection;
 using Microsoft.Extensions.Logging;
 using Microsoft.Extensions.Options;
@@ -35,7 +34,7 @@ namespace EpicGames.Horde
 			serviceCollection.AddSingleton<BundleCache>(sp => new BundleCache(sp.GetRequiredService<IOptions<HordeOptions>>().Value.BundleCache));
 			serviceCollection.AddSingleton<StorageBackendCache>(CreateBackendCache);
 			serviceCollection.AddSingleton<HttpStorageBackendFactory>();
-			serviceCollection.AddSingleton<HttpStorageNamespaceFactory>();
+			serviceCollection.AddSingleton<HttpStorageClient>();
 			serviceCollection.AddSingleton<IHordeClient>(sp => sp.GetRequiredService<HordeClientFactory>().Create());
 			serviceCollection.AddSingleton<HordeClientFactory>();
 		}
diff --git a/Engine/Source/Programs/Shared/EpicGames.Horde/Storage/Bundles/BundleStorageClient.cs b/Engine/Source/Programs/Shared/EpicGames.Horde/Storage/Bundles/BundleStorageClient.cs
deleted file mode 100644
diff --git a/Engine/Source/Programs/Shared/EpicGames.Horde/Storage/Bundles/BundleStorageNamespace.cs b/Engine/Source/Programs/Shared/EpicGames.Horde/Storage/Bundles/BundleStorageNamespace.cs
new file mode 100644
--- /dev/null
+++ b/Engine/Source/Programs/Shared/EpicGames.Horde/Storage/Bundles/BundleStorageNamespace.cs
@@ -0,0 +1,264 @@
+// Copyright Epic Games, Inc. All Rights Reserved.
+
+using System;
+using System.Buffers.Text;
+using System.Collections.Generic;
+using System.Linq;
+using System.Threading;
+using System.Threading.Tasks;
+using EpicGames.Core;
+using EpicGames.Horde.Storage.Backends;
+using EpicGames.Horde.Storage.Bundles.V1;
+using EpicGames.Horde.Storage.Bundles.V2;
+using EpicGames.Horde.Storage.ObjectStores;
+using Microsoft.Extensions.Logging;
+
+namespace EpicGames.Horde.Storage.Bundles
+{
+	/// <summary>
+	/// Base class for an implementation of <see cref="IStorageNamespace"/>, providing implementations for some common functionality using bundles.
+	/// </summary>
+	public sealed class BundleStorageNamespace : IStorageNamespace
+	{
+		readonly IStorageBackend _backend;
+		readonly BundleCache _cache;
+		readonly BundleOptions _options;
+		readonly PacketReaderStats _packetReaderStats = new PacketReaderStats();
+		readonly Bundles.V1.BundleReader _bundleReader;
+
+		internal Bundles.V1.BundleReader BundleReader => _bundleReader;
+
+		/// <summary>
+		/// Allocator which trims the cache to keep below a maximum size
+		/// </summary>
+		public IMemoryAllocator<byte> Allocator => _cache.Allocator;
+
+		/// <summary>
+		/// Accessor for the storage backend
+		/// </summary>
+		public IStorageBackend Backend => _backend;
+
+		/// <summary>
+		/// Cache for bundle data
+		/// </summary>
+		public BundleCache Cache => _cache;
+
+		/// <summary>
+		/// Stats for reading bundles
+		/// </summary>
+		internal PacketReaderStats PacketReaderStats => _packetReaderStats;
+
+		/// <summary>
+		/// Constructor
+		/// </summary>
+		public BundleStorageNamespace(IStorageBackend backend, BundleCache cache, BundleOptions? options, ILogger logger)
+		{
+			_backend = backend;
+			_cache = cache;
+			_options = options ?? BundleOptions.Default;
+			_bundleReader = new Bundles.V1.BundleReader(this, cache, logger);
+		}
+
+		/// <summary>
+		/// Creates a bundle storage namespace around a memory backend
+		/// </summary>
+		public static BundleStorageNamespace CreateInMemory(ILogger logger)
+			=> CreateInMemory(null, logger);
+
+		/// <summary>
+		/// Creates a bundle storage namespace around a memory backend
+		/// </summary>
+		public static BundleStorageNamespace CreateInMemory(BundleOptions? options, ILogger logger)
+		{
+			MemoryStorageBackend backend = new MemoryStorageBackend();
+			return new BundleStorageNamespace(backend, BundleCache.None, options, logger);
+		}
+
+		/// <summary>
+		/// Creates a bundle storage namespace around a directory on the filesystem
+		/// </summary>
+		public static BundleStorageNamespace CreateFromDirectory(DirectoryReference rootDir, BundleCache cache, MemoryMappedFileCache memoryMappedFileCache, ILogger logger)
+			=> CreateFromDirectory(rootDir, cache, memoryMappedFileCache, null, logger);
+
+		/// <summary>
+		/// Creates a bundle storage namespace around a directory on the filesystem
+		/// </summary>
+		public static BundleStorageNamespace CreateFromDirectory(DirectoryReference rootDir, BundleCache cache, MemoryMappedFileCache memoryMappedFileCache, BundleOptions? options, ILogger logger)
+		{
+			FileStorageBackend backend = new FileStorageBackend(new FileObjectStore(rootDir, memoryMappedFileCache), logger);
+			return new BundleStorageNamespace(backend, cache, options, logger);
+		}
+
+		/// <summary>
+		/// Helper method for GC which allows enumerating all references to other bundles
+		/// </summary>
+		public async Task<IEnumerable<BlobLocator>> ReadBundleReferencesAsync(BlobLocator locator, CancellationToken cancellationToken)
+		{
+			using IReadOnlyMemoryOwner<byte> data = await Backend.ReadBlobAsync(locator, 0, 32 * 1024 * 1024, cancellationToken);
+			return ReadBundleReferences(data.Memory);
+		}
+
+		/// <summary>
+		/// Helper method for GC which allows enumerating all references to other bundles
+		/// </summary>
+		IEnumerable<BlobLocator> ReadBundleReferences(ReadOnlyMemory<byte> data)
+		{
+			BundleSignature signature = BundleSignature.Read(data.Span);
+			if (signature.Version <= BundleVersion.LatestV1)
+			{
+				return ReadImportsFromDataV1(data);
+			}
+			else if (signature.Version <= BundleVersion.LatestV2)
+			{
+				return ReadImportsFromDataV2(data);
+			}
+			else
+			{
+				throw new InvalidOperationException($"Unsupported bundle version {(int)signature.Version}");
+			}
+		}
+
+		static IEnumerable<BlobLocator> ReadImportsFromDataV1(ReadOnlyMemory<byte> data)
+		{
+			BundleHeader header = BundleHeader.Read(data);
+			return header.Imports.Select(x => x.BaseLocator);
+		}
+
+		IEnumerable<BlobLocator> ReadImportsFromDataV2(ReadOnlyMemory<byte> data)
+		{
+			HashSet<BlobLocator> locators = new HashSet<BlobLocator>();
+			while (data.Length > 0)
+			{
+				BundleSignature signature = BundleSignature.Read(data.Span);
+
+				using IRefCountedHandle<Bundles.V2.Packet> packet = Bundles.V2.Packet.Decode(data, Cache.Allocator, nameof(ReadImportsFromDataV2));
+				for (int idx = 0; idx < packet.Target.GetImportCount(); idx++)
+				{
+					PacketImport import = packet.Target.GetImport(idx);
+					if (import.BaseIdx == -1)
+					{
+						locators.Add(new BlobLocator(import.Fragment.Clone()));
+					}
+				}
+
+				data = data.Slice(signature.HeaderLength);
+			}
+			return locators;
+		}
+
+		#region Nodes
+
+		/// <inheritdoc/>
+		public IBlobRef CreateBlobRef(BlobLocator locator)
+		{
+			if (!locator.TryUnwrap(out BlobLocator baseLocator, out Utf8String fragment))
+			{
+				throw new Exception($"{locator} is not valid");
+			}
+
+			FlushedBundleHandle bundleHandle = new FlushedBundleHandle(this, baseLocator);
+
+			int exportIdx;
+			if (Utf8Parser.TryParse(fragment.Span, out exportIdx, out int numBytesRead) && numBytesRead == fragment.Length)
+			{
+				return new Bundles.V1.FlushedNodeHandle(_bundleReader, baseLocator, bundleHandle, exportIdx);
+			}
+
+			int ampIdx = fragment.IndexOf('&');
+			if (ampIdx == -1)
+			{
+				throw new Exception($"{locator} is not valid");
+			}
+
+			Bundles.V2.PacketHandle packetHandle = new Bundles.V2.FlushedPacketHandle(this, bundleHandle, fragment.Slice(0, ampIdx).Span, _cache);
+			return new Bundles.V2.ExportHandle(packetHandle, fragment.Slice(ampIdx + 1));
+		}
+
+		/// <inheritdoc/>
+		public IBlobWriter CreateBlobWriter(string? basePath = null, BlobSerializerOptions? serializerOptions = null)
+		{
+			BundleVersion version = _options.MaxVersion;
+			if (version == BundleVersion.LatestV1)
+			{
+				return new Bundles.V1.BundleWriter(this, _bundleReader, basePath, _options);
+			}
+			else if (version == BundleVersion.LatestV2)
+			{
+				return new Bundles.V2.BundleWriter(this, basePath, _cache, _options, serializerOptions);
+			}
+			else
+			{
+				throw new InvalidOperationException($"Unsupported bundle version: {(int)version}");
+			}
+		}
+
+		#endregion
+
+		#region Aliases
+
+		/// <inheritdoc/>
+		public async Task AddAliasAsync(string name, IBlobRef handle, int rank, ReadOnlyMemory<byte> data, CancellationToken cancellationToken)
+		{
+			await handle.FlushAsync(cancellationToken);
+			await _backend.AddAliasAsync(name, handle.GetLocator(), rank, data, cancellationToken);
+		}
+
+		/// <inheritdoc/>
+		public async Task RemoveAliasAsync(string name, IBlobRef handle, CancellationToken cancellationToken)
+		{
+			await handle.FlushAsync(cancellationToken);
+			await _backend.RemoveAliasAsync(name, handle.GetLocator(), cancellationToken);
+		}
+
+		/// <inheritdoc/>
+		public async Task<BlobAlias[]> FindAliasesAsync(string name, int? maxLength = null, CancellationToken cancellationToken = default)
+		{
+			BlobAliasLocator[] aliases = await _backend.FindAliasesAsync(name, maxLength, cancellationToken);
+
+			BlobAlias[] result = new BlobAlias[aliases.Length];
+			for (int idx = 0; idx < aliases.Length; idx++)
+			{
+				BlobAliasLocator alias = aliases[idx];
+				result[idx] = new BlobAlias(CreateBlobRef(alias.Target), alias.Rank, alias.Data);
+			}
+			return result;
+		}
+
+		#endregion
+
+		#region Refs
+
+		/// <inheritdoc/>
+		public Task<bool> DeleteRefAsync(RefName name, CancellationToken cancellationToken = default)
+			=> _backend.DeleteRefAsync(name, cancellationToken);
+
+		/// <inheritdoc/>
+		public async Task<IHashedBlobRef?> TryReadRefAsync(RefName name, RefCacheTime cacheTime = default, CancellationToken cancellationToken = default)
+		{
+			HashedBlobRefValue? target = await _backend.TryReadRefAsync(name, cacheTime, cancellationToken);
+			if (target == null)
+			{
+				return null;
+			}
+			return this.CreateBlobRef(target.Hash, target.Locator);
+		}
+
+		/// <inheritdoc/>
+		public async Task WriteRefAsync(RefName name, IHashedBlobRef target, RefOptions? options = null, CancellationToken cancellationToken = default)
+		{
+			await target.FlushAsync(cancellationToken);
+			await _backend.WriteRefAsync(name, new HashedBlobRefValue(target.Hash, target.GetLocator()), options, cancellationToken);
+		}
+
+		#endregion
+
+		/// <inheritdoc/>
+		public void GetStats(StorageStats stats)
+		{
+			_cache.GetStats(stats);
+			_backend.GetStats(stats);
+			_bundleReader.GetStats(stats);
+			_packetReaderStats.GetStats(stats);
+		}
+	}
+}
diff --git a/Engine/Source/Programs/Shared/EpicGames.Horde/Storage/Clients/HttpStorageClient.cs b/Engine/Source/Programs/Shared/EpicGames.Horde/Storage/Clients/HttpStorageClient.cs
deleted file mode 100644
diff --git a/Engine/Source/Programs/Shared/EpicGames.Horde/Storage/Clients/KeyValueStorageClient.cs b/Engine/Source/Programs/Shared/EpicGames.Horde/Storage/Clients/KeyValueStorageClient.cs
deleted file mode 100644
diff --git a/Engine/Source/Programs/Shared/EpicGames.Horde/Storage/HttpStorageClient.cs b/Engine/Source/Programs/Shared/EpicGames.Horde/Storage/HttpStorageClient.cs
new file mode 100644
--- /dev/null
+++ b/Engine/Source/Programs/Shared/EpicGames.Horde/Storage/HttpStorageClient.cs
@@ -0,0 +1,54 @@
+// Copyright Epic Games, Inc. All Rights Reserved.
+
+using EpicGames.Horde.Storage.Backends;
+using EpicGames.Horde.Storage.Bundles;
+using Microsoft.Extensions.Logging;
+using Microsoft.Extensions.Options;
+
+namespace EpicGames.Horde.Storage
+{
+	/// <summary>
+	/// Factory for constructing storage namespace instances with an HTTP backend
+	/// </summary>
+	public sealed class HttpStorageClient : IStorageClient
+	{
+		readonly HttpStorageBackendFactory _backendFactory;
+		readonly BundleCache _bundleCache;
+		readonly IOptionsSnapshot<HordeOptions> _hordeOptions;
+		readonly ILogger<BundleStorageNamespace> _clientLogger;
+
+		/// <summary>
+		/// Constructor
+		/// </summary>
+		public HttpStorageClient(HttpStorageBackendFactory backendFactory, BundleCache bundleCache, IOptionsSnapshot<HordeOptions> hordeOptions, ILogger<BundleStorageNamespace> clientLogger)
+		{
+			_backendFactory = backendFactory;
+			_bundleCache = bundleCache;
+			_hordeOptions = hordeOptions;
+			_clientLogger = clientLogger;
+		}
+
+		/// <summary>
+		/// Creates a new HTTP storage namespace
+		/// </summary>
+		/// <param name="basePath">Base path for all requests</param>
+		/// <param name="accessToken">Custom access token to use for requests</param>
+		/// <param name="withBackendCache"></param>
+		public IStorageNamespace GetNamespaceWithPath(string basePath, string? accessToken = null, bool withBackendCache = true)
+		{
+			IStorageBackend backend = _backendFactory.CreateBackend(basePath, accessToken, withBackendCache);
+			return new BundleStorageNamespace(backend, _bundleCache, _hordeOptions.Value.Bundle, _clientLogger);
+		}
+
+		/// <summary>
+		/// Creates a new HTTP storage namespace
+		/// </summary>
+		/// <param name="namespaceId">Namespace to create a client for</param>
+		/// <param name="accessToken">Custom access token to use for requests</param>
+		/// <param name="withBackendCache">Whether to enable the backend cache, which caches full bundles to disk</param>
+		public IStorageNamespace GetNamespace(NamespaceId namespaceId, string? accessToken = null, bool withBackendCache = true) => GetNamespaceWithPath($"api/v1/storage/{namespaceId}", accessToken, withBackendCache);
+
+		/// <inheritdoc/>
+		IStorageNamespace? IStorageClient.TryGetNamespace(NamespaceId namespaceId) => GetNamespace(namespaceId);
+	}
+}
diff --git a/Engine/Source/Programs/Shared/EpicGames.Horde/Storage/IStorageClient.cs b/Engine/Source/Programs/Shared/EpicGames.Horde/Storage/IStorageClient.cs
--- a/Engine/Source/Programs/Shared/EpicGames.Horde/Storage/IStorageClient.cs
+++ b/Engine/Source/Programs/Shared/EpicGames.Horde/Storage/IStorageClient.cs
@@ -1,118 +1,13 @@
 // Copyright Epic Games, Inc. All Rights Reserved.
 
 using System;
-using System.Collections.Generic;
-using System.Diagnostics.CodeAnalysis;
-using System.Linq;
-using System.Threading;
-using System.Threading.Tasks;
-using EpicGames.Core;
-using Microsoft.Extensions.Logging;
 
 namespace EpicGames.Horde.Storage
 {
 	/// <summary>
-	/// Options for a new ref
+	/// Client for the storage system
 	/// </summary>
-	/// <param name="Lifetime">Time until a ref is expired</param>
-	/// <param name="Extend">Whether to extend the remaining lifetime of a ref whenever it is fetched. Defaults to true.</param>
-	public record class RefOptions(TimeSpan? Lifetime = null, bool? Extend = null);
-
-	/// <summary>
-	/// Interface for the storage system.
-	/// </summary>
-	public interface IStorageNamespace
-	{
-		#region Blobs
-
-		/// <summary>
-		/// Creates a new blob handle by parsing a locator
-		/// </summary>
-		/// <param name="locator">Path to the blob</param>
-		/// <returns>New handle to the blob</returns>
-		IBlobRef CreateBlobRef(BlobLocator locator);
-
-		/// <summary>
-		/// Creates a new writer for storage blobs
-		/// </summary>
-		/// <param name="basePath">Base path for any nodes written from the writer.</param>
-		/// <param name="serializerOptions">Options for serializing classes</param>
-		/// <returns>New writer instance. Must be disposed after use.</returns>
-		IBlobWriter CreateBlobWriter(string? basePath = null, BlobSerializerOptions? serializerOptions = null);
-
-		#endregion
-
-		#region Aliases
-
-		/// <summary>
-		/// Adds an alias to a given blob
-		/// </summary>
-		/// <param name="name">Alias for the blob</param>
-		/// <param name="handle">Locator for the blob</param>
-		/// <param name="rank">Rank for this alias. In situations where an alias has multiple mappings, the alias with the highest rank will be returned by default.</param>
-		/// <param name="data">Additional data to be stored inline with the alias</param>
-		/// <param name="cancellationToken">Cancellation token for the operation</param>
-		Task AddAliasAsync(string name, IBlobRef handle, int rank = 0, ReadOnlyMemory<byte> data = default, CancellationToken cancellationToken = default);
-
-		/// <summary>
-		/// Removes an alias from a blob
-		/// </summary>
-		/// <param name="name">Name of the alias</param>
-		/// <param name="handle">Locator for the blob</param>
-		/// <param name="cancellationToken">Cancellation token for the operation</param>
-		Task RemoveAliasAsync(string name, IBlobRef handle, CancellationToken cancellationToken = default);
-
-		/// <summary>
-		/// Finds blobs with the given alias. Unlike refs, aliases do not serve as GC roots.
-		/// </summary>
-		/// <param name="name">Alias for the blob</param>
-		/// <param name="maxResults">Maximum number of aliases to return</param>
-		/// <param name="cancellationToken">Cancellation token for the operation</param>
-		/// <returns>Blobs matching the given handle</returns>
-		Task<BlobAlias[]> FindAliasesAsync(string name, int? maxResults = null, CancellationToken cancellationToken = default);
-
-		#endregion
-
-		#region Refs
-
-		/// <summary>
-		/// Reads data for a ref from the store
-		/// </summary>
-		/// <param name="name">The ref name</param>
-		/// <param name="cacheTime">Minimum coherency for any cached value to be returned</param>
-		/// <param name="cancellationToken">Cancellation token for the operation</param>
-		/// <returns>Blob pointed to by the ref</returns>
-		Task<IHashedBlobRef?> TryReadRefAsync(RefName name, RefCacheTime cacheTime = default, CancellationToken cancellationToken = default);
-
-		/// <summary>
-		/// Writes a new ref to the store
-		/// </summary>
-		/// <param name="name">Ref to write</param>
-		/// <param name="target">Handle to the target blob</param>
-		/// <param name="options">Options for the new ref</param>
-		/// <param name="cancellationToken">Cancellation token for the operation</param>
-		/// <returns>Unique identifier for the blob</returns>
-		Task WriteRefAsync(RefName name, IHashedBlobRef target, RefOptions? options = null, CancellationToken cancellationToken = default);
-
-		/// <summary>
-		/// Reads data for a ref from the store
-		/// </summary>
-		/// <param name="name">The ref identifier</param>
-		/// <param name="cancellationToken">Cancellation token for the operation</param>
-		Task<bool> DeleteRefAsync(RefName name, CancellationToken cancellationToken = default);
-
-		#endregion
-
-		/// <summary>
-		/// Gets a snapshot of the stats for the storage namespace.
-		/// </summary>
-		void GetStats(StorageStats stats);
-	}
-
-	/// <summary>
-	/// Allows creating storage namespaces for different namespace ids
-	/// </summary>
-	public interface IStorageNamespaceFactory
+	public interface IStorageClient
 	{
 		/// <summary>
 		/// Creates a storage namespace for the given id
@@ -123,276 +18,21 @@ namespace EpicGames.Horde.Storage
 	}
 
 	/// <summary>
-	/// Extension methods for <see cref="IStorageNamespaceFactory"/>
+	/// Extension methods for <see cref="IStorageClient"/>
 	/// </summary>
-	public static class StorageNamespaceFactoryExtensions
+	public static class StorageClientExtensions
 	{
 		/// <summary>
 		/// Creates a new storage namespace, throwing an exception if it does not exist
 		/// </summary>
-		public static IStorageNamespace GetNamespace(this IStorageNamespaceFactory factory, NamespaceId namespaceId)
+		public static IStorageNamespace GetNamespace(this IStorageClient storageClient, NamespaceId namespaceId)
 		{
-			IStorageNamespace? client = factory.TryGetNamespace(namespaceId);
-			if (client == null)
+			IStorageNamespace? storageNamespace = storageClient.TryGetNamespace(namespaceId);
+			if (storageNamespace == null)
 			{
 				throw new InvalidOperationException($"No namespace '{namespaceId}' is configured");
 			}
-			return client;
-		}
-	}
-
-	/// <summary>
-	/// Indicates the maximum age of a entry returned from a cache in the hierarchy
-	/// </summary>
-	/// <param name="Utc">Oldest allowed timestamp for a returned result</param>
-	public readonly record struct RefCacheTime(DateTime Utc)
-	{
-		/// <summary>
-		/// Maximum age for a cached value to be returned
-		/// </summary>
-		public readonly TimeSpan MaxAge => DateTime.UtcNow - Utc;
-
-		/// <summary>
-		/// Sets the earliest time at which the entry must have been valid
-		/// </summary>
-		/// <param name="age">Maximum age of any returned cache value. Taken from the moment that this object was created.</param>
-		public RefCacheTime(TimeSpan age) : this(DateTime.UtcNow - age) { }
-
-		/// <summary>
-		/// Tests whether this value is set
-		/// </summary>
-		public readonly bool IsSet() => Utc != default;
-
-		/// <summary>
-		/// Determines if this cache time deems a particular cache entry stale
-		/// </summary>
-		/// <param name="entryTime">Time at which the cache entry was valid</param>
-		/// <param name="cacheTime">Maximum cache time to test against</param>
-		public static bool IsStaleCacheEntry(DateTime entryTime, RefCacheTime cacheTime) => cacheTime.IsSet() && cacheTime.Utc < entryTime;
-
-		/// <summary>
-		/// Implicit conversion operator from datetime values.
-		/// </summary>
-		public static implicit operator RefCacheTime(DateTime time) => new RefCacheTime(time);
-
-		/// <summary>
-		/// Implicit conversion operator from timespan values.
-		/// </summary>
-		public static implicit operator RefCacheTime(TimeSpan age) => new RefCacheTime(age);
-	}
-
-	/// <summary>
-	/// Stats for the storage system
-	/// </summary>
-	public class StorageStats
-	{
-		/// <summary>
-		/// Stat name to value
-		/// </summary>
-		public List<(string, long)> Values { get; } = new List<(string, long)>();
-
-		/// <summary>
-		/// Add a new stat to the list
-		/// </summary>
-		public void Add(string name, long value) => Values.Add((name, value));
-
-		/// <summary>
-		/// Prints the table of stats to the logger
-		/// </summary>
-		public void Print(ILogger logger)
-		{
-			foreach ((string key, long value) in Values)
-			{
-				logger.LogInformation("{Key}: {Value:n0}", key, value);
-			}
-		}
-
-		/// <summary>
-		/// Subtract a base set of stats from this one
-		/// </summary>
-		public static StorageStats GetDelta(StorageStats initial, StorageStats finish)
-		{
-			StorageStats result = new StorageStats();
-
-			Dictionary<string, long> initialValues = initial.Values.ToDictionary(x => x.Item1, x => x.Item2, StringComparer.Ordinal);
-			foreach ((string name, long value) in finish.Values.ToArray())
-			{
-				initialValues.TryGetValue(name, out long otherValue);
-				result.Add(name, value - otherValue);
-			}
-
-			return result;
-		}
-	}
-
-	/// <summary>
-	/// Extension methods for <see cref="IStorageNamespace"/>
-	/// </summary>
-	public static class StorageNamespaceExtensions
-	{
-		#region Blobs
-
-		/// <summary>
-		/// Creates a new blob handle by parsing a locator
-		/// </summary>
-		/// <param name="storageNamespace">Storage namespace to operate on</param>
-		/// <param name="locator">Path to the blob</param>
-		/// <param name="serializerOptions">Options for deserializing the blob</param>
-		/// <returns>New handle to the blob</returns>
-		public static IBlobRef<T> CreateBlobRef<T>(this IStorageNamespace storageNamespace, BlobLocator locator, BlobSerializerOptions? serializerOptions = null)
-			=> storageNamespace.CreateBlobRef(locator).ForType<T>(serializerOptions);
-
-		/// <summary>
-		/// Creates a new blob reference from a locator and hash
-		/// </summary>
-		/// <param name="storageNamespace">Storage namespace to operate on</param>
-		/// <param name="hash">Hash of the target blob</param>
-		/// <param name="locator">Path to the blob</param>
-		/// <returns>New handle to the blob</returns>
-		public static IHashedBlobRef CreateBlobRef(this IStorageNamespace storageNamespace, IoHash hash, BlobLocator locator)
-			=> HashedBlobRef.Create(hash, storageNamespace.CreateBlobRef(locator));
-
-		/// <summary>
-		/// Creates a new blob reference from a locator and hash
-		/// </summary>
-		/// <param name="storageNamespace">Storage namespace to operate on</param>
-		/// <param name="hash">Hash of the target blob</param>
-		/// <param name="locator">Path to the blob</param>
-		/// <param name="serializerOptions">Options for deserializing the blob</param>
-		/// <returns>New handle to the blob</returns>
-		public static IHashedBlobRef<T> CreateBlobRef<T>(this IStorageNamespace storageNamespace, IoHash hash, BlobLocator locator, BlobSerializerOptions? serializerOptions = null)
-			=> HashedBlobRef.Create<T>(hash, storageNamespace.CreateBlobRef(locator), serializerOptions);
-
-		/// <summary>
-		/// Create a blob ref from a RefValue
-		/// </summary>
-		public static IHashedBlobRef CreateBlobRef(this IStorageNamespace storageNamespace, HashedBlobRefValue refValue)
-			=> storageNamespace.CreateBlobRef(refValue.Hash, refValue.Locator);
-
-		/// <summary>
-		/// Create a typed blob ref from a RefValue
-		/// </summary>
-		public static IHashedBlobRef<T> CreateBlobRef<T>(this IStorageNamespace storageNamespace, HashedBlobRefValue refValue, BlobSerializerOptions? options)
-			=> storageNamespace.CreateBlobRef<T>(refValue.Hash, refValue.Locator, options);
-
-		class NamedBlobRef : IBlobRef
-		{
-			readonly IStorageNamespace _storageNamespace;
-			readonly RefName _name;
-			readonly RefCacheTime _cacheTime;
-
-			public NamedBlobRef(IStorageNamespace storageNamespace, RefName name, RefCacheTime cacheTime = default)
-			{
-				_storageNamespace = storageNamespace;
-				_name = name;
-				_cacheTime = cacheTime;
-			}
-
-			public IBlobRef Innermost => this;
-
-			public ValueTask FlushAsync(CancellationToken cancellationToken = default)
-				=> default;
-
-			public async ValueTask<BlobData> ReadBlobDataAsync(CancellationToken cancellationToken = default)
-			{
-				IHashedBlobRef handle = await _storageNamespace.ReadRefAsync(_name, _cacheTime, cancellationToken);
-				return await handle.ReadBlobDataAsync(cancellationToken);
-			}
-
-			public bool TryGetLocator([NotNullWhen(true)] out BlobLocator locator)
-			{
-				locator = default;
-				return false;
-			}
-		}
-
-		/// <summary>
-		/// Creates a new blob ref from a ref name
-		/// </summary>
-		/// <param name="storageNamespace">The store instance to read from</param>
-		/// <param name="name">Name of the reference</param>
-		/// <param name="cacheTime">Maximum age for cached responses</param>
-		/// <returns>New handle to the blob</returns>
-		public static IBlobRef CreateBlobRef(this IStorageNamespace storageNamespace, RefName name, RefCacheTime cacheTime = default)
-			=> new NamedBlobRef(storageNamespace, name, cacheTime);
-
-		/// <summary>
-		/// Creates a new blob ref from a ref name
-		/// </summary>
-		/// <param name="storageNamespace">The store instance to read from</param>
-		/// <param name="name">Name of the reference</param>
-		/// <param name="cacheTime">Maximum age for cached responses</param>
-		/// <param name="serializerOptions">Options for deserializing the blob</param>
-		/// <returns>New handle to the blob</returns>
-		public static IBlobRef<T> CreateBlobRef<T>(this IStorageNamespace storageNamespace, RefName name, RefCacheTime cacheTime = default, BlobSerializerOptions? serializerOptions = null)
-			=> CreateBlobRef(storageNamespace, name, cacheTime).ForType<T>(serializerOptions);
-
-		/// <summary>
-		/// Creates a writer using a refname as a base path
-		/// </summary>
-		/// <param name="store">The store instance to read from</param>
-		/// <param name="refName">Ref name to use as a base path</param>
-		public static IBlobWriter CreateBlobWriter(this IStorageNamespace store, RefName refName) => store.CreateBlobWriter(refName.ToString());
-
-		#endregion
-
-		#region Aliases
-
-		/// <summary>
-		/// Finds blobs with the given alias. Unlike refs, aliases do not serve as GC roots.
-		/// </summary>
-		/// <param name="store">The store instance to read from</param>
-		/// <param name="name">Alias for the blob</param>
-		/// <param name="cancellationToken">Cancellation token for the operation</param>
-		/// <returns>Blobs matching the given handle</returns>
-		public static async Task<BlobAlias?> FindAliasAsync(this IStorageNamespace store, string name, CancellationToken cancellationToken = default)
-		{
-			BlobAlias[] aliases = await store.FindAliasesAsync(name, 1, cancellationToken);
-			return aliases.FirstOrDefault();
-		}
-
-		#endregion
-
-		#region Refs
-
-		/// <summary>
-		/// Checks if the given ref exists
-		/// </summary>
-		/// <param name="store">The store instance to read from</param>
-		/// <param name="name">Name of the reference to look for</param>
-		/// <param name="cacheTime">Minimum coherency for any cached value to be returned</param>
-		/// <param name="cancellationToken">Cancellation token for the operation</param>
-		/// <returns>True if the ref exists, false if it did not exist</returns>
-		public static async Task<bool> RefExistsAsync(this IStorageNamespace store, RefName name, RefCacheTime cacheTime = default, CancellationToken cancellationToken = default)
-		{
-			IBlobRef? target = await store.TryReadRefAsync(name, cacheTime, cancellationToken);
-			return target != null;
-		}
-
-		/// <summary>
-		/// Reads a ref from the store, throwing an exception if it does not exist
-		/// </summary>
-		/// <param name="store">The store instance to read from</param>
-		/// <param name="name">Id for the ref</param>
-		/// <param name="cacheTime">Minimum coherency of any cached result</param>
-		/// <param name="cancellationToken">Cancellation token for the operation</param>
-		/// <returns>The ref target</returns>
-		public static async Task<IHashedBlobRef> ReadRefAsync(this IStorageNamespace store, RefName name, RefCacheTime cacheTime = default, CancellationToken cancellationToken = default)
-		{
-			IHashedBlobRef? refTarget = await store.TryReadRefAsync(name, cacheTime, cancellationToken);
-			return refTarget ?? throw new RefNameNotFoundException(name);
-		}
-
-		#endregion
-
-		/// <summary>
-		/// Gets a snapshot of the stats for the storage namespace.
-		/// </summary>
-		public static StorageStats GetStats(this IStorageNamespace store)
-		{
-			StorageStats stats = new StorageStats();
-			store.GetStats(stats);
-			return stats;
+			return storageNamespace;
 		}
 	}
 }
diff --git a/Engine/Source/Programs/Shared/EpicGames.Horde/Storage/IStorageNamespace.cs b/Engine/Source/Programs/Shared/EpicGames.Horde/Storage/IStorageNamespace.cs
new file mode 100644
--- /dev/null
+++ b/Engine/Source/Programs/Shared/EpicGames.Horde/Storage/IStorageNamespace.cs
@@ -0,0 +1,366 @@
+// Copyright Epic Games, Inc. All Rights Reserved.
+
+using System;
+using System.Collections.Generic;
+using System.Diagnostics.CodeAnalysis;
+using System.Linq;
+using System.Threading;
+using System.Threading.Tasks;
+using EpicGames.Core;
+using Microsoft.Extensions.Logging;
+
+namespace EpicGames.Horde.Storage
+{
+	/// <summary>
+	/// Options for a new ref
+	/// </summary>
+	/// <param name="Lifetime">Time until a ref is expired</param>
+	/// <param name="Extend">Whether to extend the remaining lifetime of a ref whenever it is fetched. Defaults to true.</param>
+	public record class RefOptions(TimeSpan? Lifetime = null, bool? Extend = null);
+
+	/// <summary>
+	/// Interface for the storage system.
+	/// </summary>
+	public interface IStorageNamespace
+	{
+		#region Blobs
+
+		/// <summary>
+		/// Creates a new blob handle by parsing a locator
+		/// </summary>
+		/// <param name="locator">Path to the blob</param>
+		/// <returns>New handle to the blob</returns>
+		IBlobRef CreateBlobRef(BlobLocator locator);
+
+		/// <summary>
+		/// Creates a new writer for storage blobs
+		/// </summary>
+		/// <param name="basePath">Base path for any nodes written from the writer.</param>
+		/// <param name="serializerOptions">Options for serializing classes</param>
+		/// <returns>New writer instance. Must be disposed after use.</returns>
+		IBlobWriter CreateBlobWriter(string? basePath = null, BlobSerializerOptions? serializerOptions = null);
+
+		#endregion
+
+		#region Aliases
+
+		/// <summary>
+		/// Adds an alias to a given blob
+		/// </summary>
+		/// <param name="name">Alias for the blob</param>
+		/// <param name="handle">Locator for the blob</param>
+		/// <param name="rank">Rank for this alias. In situations where an alias has multiple mappings, the alias with the highest rank will be returned by default.</param>
+		/// <param name="data">Additional data to be stored inline with the alias</param>
+		/// <param name="cancellationToken">Cancellation token for the operation</param>
+		Task AddAliasAsync(string name, IBlobRef handle, int rank = 0, ReadOnlyMemory<byte> data = default, CancellationToken cancellationToken = default);
+
+		/// <summary>
+		/// Removes an alias from a blob
+		/// </summary>
+		/// <param name="name">Name of the alias</param>
+		/// <param name="handle">Locator for the blob</param>
+		/// <param name="cancellationToken">Cancellation token for the operation</param>
+		Task RemoveAliasAsync(string name, IBlobRef handle, CancellationToken cancellationToken = default);
+
+		/// <summary>
+		/// Finds blobs with the given alias. Unlike refs, aliases do not serve as GC roots.
+		/// </summary>
+		/// <param name="name">Alias for the blob</param>
+		/// <param name="maxResults">Maximum number of aliases to return</param>
+		/// <param name="cancellationToken">Cancellation token for the operation</param>
+		/// <returns>Blobs matching the given handle</returns>
+		Task<BlobAlias[]> FindAliasesAsync(string name, int? maxResults = null, CancellationToken cancellationToken = default);
+
+		#endregion
+
+		#region Refs
+
+		/// <summary>
+		/// Reads data for a ref from the store
+		/// </summary>
+		/// <param name="name">The ref name</param>
+		/// <param name="cacheTime">Minimum coherency for any cached value to be returned</param>
+		/// <param name="cancellationToken">Cancellation token for the operation</param>
+		/// <returns>Blob pointed to by the ref</returns>
+		Task<IHashedBlobRef?> TryReadRefAsync(RefName name, RefCacheTime cacheTime = default, CancellationToken cancellationToken = default);
+
+		/// <summary>
+		/// Writes a new ref to the store
+		/// </summary>
+		/// <param name="name">Ref to write</param>
+		/// <param name="target">Handle to the target blob</param>
+		/// <param name="options">Options for the new ref</param>
+		/// <param name="cancellationToken">Cancellation token for the operation</param>
+		/// <returns>Unique identifier for the blob</returns>
+		Task WriteRefAsync(RefName name, IHashedBlobRef target, RefOptions? options = null, CancellationToken cancellationToken = default);
+
+		/// <summary>
+		/// Reads data for a ref from the store
+		/// </summary>
+		/// <param name="name">The ref identifier</param>
+		/// <param name="cancellationToken">Cancellation token for the operation</param>
+		Task<bool> DeleteRefAsync(RefName name, CancellationToken cancellationToken = default);
+
+		#endregion
+
+		/// <summary>
+		/// Gets a snapshot of the stats for the storage namespace.
+		/// </summary>
+		void GetStats(StorageStats stats);
+	}
+
+	/// <summary>
+	/// Indicates the maximum age of a entry returned from a cache in the hierarchy
+	/// </summary>
+	/// <param name="Utc">Oldest allowed timestamp for a returned result</param>
+	public readonly record struct RefCacheTime(DateTime Utc)
+	{
+		/// <summary>
+		/// Maximum age for a cached value to be returned
+		/// </summary>
+		public readonly TimeSpan MaxAge => DateTime.UtcNow - Utc;
+
+		/// <summary>
+		/// Sets the earliest time at which the entry must have been valid
+		/// </summary>
+		/// <param name="age">Maximum age of any returned cache value. Taken from the moment that this object was created.</param>
+		public RefCacheTime(TimeSpan age) : this(DateTime.UtcNow - age) { }
+
+		/// <summary>
+		/// Tests whether this value is set
+		/// </summary>
+		public readonly bool IsSet() => Utc != default;
+
+		/// <summary>
+		/// Determines if this cache time deems a particular cache entry stale
+		/// </summary>
+		/// <param name="entryTime">Time at which the cache entry was valid</param>
+		/// <param name="cacheTime">Maximum cache time to test against</param>
+		public static bool IsStaleCacheEntry(DateTime entryTime, RefCacheTime cacheTime) => cacheTime.IsSet() && cacheTime.Utc < entryTime;
+
+		/// <summary>
+		/// Implicit conversion operator from datetime values.
+		/// </summary>
+		public static implicit operator RefCacheTime(DateTime time) => new RefCacheTime(time);
+
+		/// <summary>
+		/// Implicit conversion operator from timespan values.
+		/// </summary>
+		public static implicit operator RefCacheTime(TimeSpan age) => new RefCacheTime(age);
+	}
+
+	/// <summary>
+	/// Stats for the storage system
+	/// </summary>
+	public class StorageStats
+	{
+		/// <summary>
+		/// Stat name to value
+		/// </summary>
+		public List<(string, long)> Values { get; } = new List<(string, long)>();
+
+		/// <summary>
+		/// Add a new stat to the list
+		/// </summary>
+		public void Add(string name, long value) => Values.Add((name, value));
+
+		/// <summary>
+		/// Prints the table of stats to the logger
+		/// </summary>
+		public void Print(ILogger logger)
+		{
+			foreach ((string key, long value) in Values)
+			{
+				logger.LogInformation("{Key}: {Value:n0}", key, value);
+			}
+		}
+
+		/// <summary>
+		/// Subtract a base set of stats from this one
+		/// </summary>
+		public static StorageStats GetDelta(StorageStats initial, StorageStats finish)
+		{
+			StorageStats result = new StorageStats();
+
+			Dictionary<string, long> initialValues = initial.Values.ToDictionary(x => x.Item1, x => x.Item2, StringComparer.Ordinal);
+			foreach ((string name, long value) in finish.Values.ToArray())
+			{
+				initialValues.TryGetValue(name, out long otherValue);
+				result.Add(name, value - otherValue);
+			}
+
+			return result;
+		}
+	}
+
+	/// <summary>
+	/// Extension methods for <see cref="IStorageNamespace"/>
+	/// </summary>
+	public static class StorageNamespaceExtensions
+	{
+		#region Blobs
+
+		/// <summary>
+		/// Creates a new blob handle by parsing a locator
+		/// </summary>
+		/// <param name="storageNamespace">Storage namespace to operate on</param>
+		/// <param name="locator">Path to the blob</param>
+		/// <param name="serializerOptions">Options for deserializing the blob</param>
+		/// <returns>New handle to the blob</returns>
+		public static IBlobRef<T> CreateBlobRef<T>(this IStorageNamespace storageNamespace, BlobLocator locator, BlobSerializerOptions? serializerOptions = null)
+			=> storageNamespace.CreateBlobRef(locator).ForType<T>(serializerOptions);
+
+		/// <summary>
+		/// Creates a new blob reference from a locator and hash
+		/// </summary>
+		/// <param name="storageNamespace">Storage namespace to operate on</param>
+		/// <param name="hash">Hash of the target blob</param>
+		/// <param name="locator">Path to the blob</param>
+		/// <returns>New handle to the blob</returns>
+		public static IHashedBlobRef CreateBlobRef(this IStorageNamespace storageNamespace, IoHash hash, BlobLocator locator)
+			=> HashedBlobRef.Create(hash, storageNamespace.CreateBlobRef(locator));
+
+		/// <summary>
+		/// Creates a new blob reference from a locator and hash
+		/// </summary>
+		/// <param name="storageNamespace">Storage namespace to operate on</param>
+		/// <param name="hash">Hash of the target blob</param>
+		/// <param name="locator">Path to the blob</param>
+		/// <param name="serializerOptions">Options for deserializing the blob</param>
+		/// <returns>New handle to the blob</returns>
+		public static IHashedBlobRef<T> CreateBlobRef<T>(this IStorageNamespace storageNamespace, IoHash hash, BlobLocator locator, BlobSerializerOptions? serializerOptions = null)
+			=> HashedBlobRef.Create<T>(hash, storageNamespace.CreateBlobRef(locator), serializerOptions);
+
+		/// <summary>
+		/// Create a blob ref from a RefValue
+		/// </summary>
+		public static IHashedBlobRef CreateBlobRef(this IStorageNamespace storageNamespace, HashedBlobRefValue refValue)
+			=> storageNamespace.CreateBlobRef(refValue.Hash, refValue.Locator);
+
+		/// <summary>
+		/// Create a typed blob ref from a RefValue
+		/// </summary>
+		public static IHashedBlobRef<T> CreateBlobRef<T>(this IStorageNamespace storageNamespace, HashedBlobRefValue refValue, BlobSerializerOptions? options)
+			=> storageNamespace.CreateBlobRef<T>(refValue.Hash, refValue.Locator, options);
+
+		class NamedBlobRef : IBlobRef
+		{
+			readonly IStorageNamespace _storageNamespace;
+			readonly RefName _name;
+			readonly RefCacheTime _cacheTime;
+
+			public NamedBlobRef(IStorageNamespace storageNamespace, RefName name, RefCacheTime cacheTime = default)
+			{
+				_storageNamespace = storageNamespace;
+				_name = name;
+				_cacheTime = cacheTime;
+			}
+
+			public IBlobRef Innermost => this;
+
+			public ValueTask FlushAsync(CancellationToken cancellationToken = default)
+				=> default;
+
+			public async ValueTask<BlobData> ReadBlobDataAsync(CancellationToken cancellationToken = default)
+			{
+				IHashedBlobRef handle = await _storageNamespace.ReadRefAsync(_name, _cacheTime, cancellationToken);
+				return await handle.ReadBlobDataAsync(cancellationToken);
+			}
+
+			public bool TryGetLocator([NotNullWhen(true)] out BlobLocator locator)
+			{
+				locator = default;
+				return false;
+			}
+		}
+
+		/// <summary>
+		/// Creates a new blob ref from a ref name
+		/// </summary>
+		/// <param name="storageNamespace">The store instance to read from</param>
+		/// <param name="name">Name of the reference</param>
+		/// <param name="cacheTime">Maximum age for cached responses</param>
+		/// <returns>New handle to the blob</returns>
+		public static IBlobRef CreateBlobRef(this IStorageNamespace storageNamespace, RefName name, RefCacheTime cacheTime = default)
+			=> new NamedBlobRef(storageNamespace, name, cacheTime);
+
+		/// <summary>
+		/// Creates a new blob ref from a ref name
+		/// </summary>
+		/// <param name="storageNamespace">The store instance to read from</param>
+		/// <param name="name">Name of the reference</param>
+		/// <param name="cacheTime">Maximum age for cached responses</param>
+		/// <param name="serializerOptions">Options for deserializing the blob</param>
+		/// <returns>New handle to the blob</returns>
+		public static IBlobRef<T> CreateBlobRef<T>(this IStorageNamespace storageNamespace, RefName name, RefCacheTime cacheTime = default, BlobSerializerOptions? serializerOptions = null)
+			=> CreateBlobRef(storageNamespace, name, cacheTime).ForType<T>(serializerOptions);
+
+		/// <summary>
+		/// Creates a writer using a refname as a base path
+		/// </summary>
+		/// <param name="store">The store instance to read from</param>
+		/// <param name="refName">Ref name to use as a base path</param>
+		public static IBlobWriter CreateBlobWriter(this IStorageNamespace store, RefName refName) => store.CreateBlobWriter(refName.ToString());
+
+		#endregion
+
+		#region Aliases
+
+		/// <summary>
+		/// Finds blobs with the given alias. Unlike refs, aliases do not serve as GC roots.
+		/// </summary>
+		/// <param name="store">The store instance to read from</param>
+		/// <param name="name">Alias for the blob</param>
+		/// <param name="cancellationToken">Cancellation token for the operation</param>
+		/// <returns>Blobs matching the given handle</returns>
+		public static async Task<BlobAlias?> FindAliasAsync(this IStorageNamespace store, string name, CancellationToken cancellationToken = default)
+		{
+			BlobAlias[] aliases = await store.FindAliasesAsync(name, 1, cancellationToken);
+			return aliases.FirstOrDefault();
+		}
+
+		#endregion
+
+		#region Refs
+
+		/// <summary>
+		/// Checks if the given ref exists
+		/// </summary>
+		/// <param name="store">The store instance to read from</param>
+		/// <param name="name">Name of the reference to look for</param>
+		/// <param name="cacheTime">Minimum coherency for any cached value to be returned</param>
+		/// <param name="cancellationToken">Cancellation token for the operation</param>
+		/// <returns>True if the ref exists, false if it did not exist</returns>
+		public static async Task<bool> RefExistsAsync(this IStorageNamespace store, RefName name, RefCacheTime cacheTime = default, CancellationToken cancellationToken = default)
+		{
+			IBlobRef? target = await store.TryReadRefAsync(name, cacheTime, cancellationToken);
+			return target != null;
+		}
+
+		/// <summary>
+		/// Reads a ref from the store, throwing an exception if it does not exist
+		/// </summary>
+		/// <param name="store">The store instance to read from</param>
+		/// <param name="name">Id for the ref</param>
+		/// <param name="cacheTime">Minimum coherency of any cached result</param>
+		/// <param name="cancellationToken">Cancellation token for the operation</param>
+		/// <returns>The ref target</returns>
+		public static async Task<IHashedBlobRef> ReadRefAsync(this IStorageNamespace store, RefName name, RefCacheTime cacheTime = default, CancellationToken cancellationToken = default)
+		{
+			IHashedBlobRef? refTarget = await store.TryReadRefAsync(name, cacheTime, cancellationToken);
+			return refTarget ?? throw new RefNameNotFoundException(name);
+		}
+
+		#endregion
+
+		/// <summary>
+		/// Gets a snapshot of the stats for the storage namespace.
+		/// </summary>
+		public static StorageStats GetStats(this IStorageNamespace store)
+		{
+			StorageStats stats = new StorageStats();
+			store.GetStats(stats);
+			return stats;
+		}
+	}
+}
diff --git a/Engine/Source/Programs/Shared/EpicGames.Horde/Storage/KeyValueStorageNamespace.cs b/Engine/Source/Programs/Shared/EpicGames.Horde/Storage/KeyValueStorageNamespace.cs
new file mode 100644
--- /dev/null
+++ b/Engine/Source/Programs/Shared/EpicGames.Horde/Storage/KeyValueStorageNamespace.cs
@@ -0,0 +1,226 @@
+// Copyright Epic Games, Inc. All Rights Reserved.
+
+using System;
+using System.Collections.Generic;
+using System.IO;
+using System.Linq;
+using System.Threading;
+using System.Threading.Tasks;
+using EpicGames.Core;
+using EpicGames.Horde.Storage.Backends;
+
+namespace EpicGames.Horde.Storage.Clients
+{
+	/// <summary>
+	/// Base class for storage namespaces that wrap a diirect key/value type store without any merging/splitting.
+	/// </summary>
+	public sealed class KeyValueStorageNamespace : IStorageNamespace
+	{
+		class Handle : IBlobRef
+		{
+			readonly KeyValueStorageNamespace _keyValueStorageNamespace;
+			readonly BlobLocator _locator;
+
+			/// <inheritdoc/>
+			public IBlobRef Innermost => this;
+
+			/// <summary>
+			/// Constructor
+			/// </summary>
+			public Handle(KeyValueStorageNamespace keyValueStorageNamespace, BlobLocator locator)
+			{
+				_keyValueStorageNamespace = keyValueStorageNamespace;
+				_locator = locator;
+			}
+
+			/// <inheritdoc/>
+			public ValueTask FlushAsync(CancellationToken cancellationToken = default) => default;
+
+			/// <inheritdoc/>
+			public ValueTask<BlobData> ReadBlobDataAsync(CancellationToken cancellationToken = default) => _keyValueStorageNamespace.ReadBlobAsync(_locator, cancellationToken);
+
+			/// <inheritdoc/>
+			public bool TryGetLocator(out BlobLocator locator)
+			{
+				locator = _locator;
+				return true;
+			}
+
+			/// <inheritdoc/>
+			public override bool Equals(object? obj) => obj is Handle other && _locator == other._locator;
+
+			/// <inheritdoc/>
+			public override int GetHashCode() => _locator.GetHashCode();
+		}
+
+		class Writer : BlobWriter
+		{
+			readonly KeyValueStorageNamespace _outer;
+			readonly string _basePath;
+			byte[] _data = Array.Empty<byte>();
+			int _offset;
+
+			/// <summary>
+			/// Constructor
+			/// </summary>
+			public Writer(KeyValueStorageNamespace outer, string? basePath, BlobSerializerOptions? options)
+				: base(options)
+			{
+				_outer = outer;
+				_basePath = basePath ?? String.Empty;
+
+				if (!_basePath.EndsWith("/", StringComparison.Ordinal))
+				{
+					_basePath += "/";
+				}
+			}
+
+			/// <inheritdoc/>
+			public override ValueTask DisposeAsync() => new ValueTask();
+
+			/// <inheritdoc/>
+			public override Task FlushAsync(CancellationToken cancellationToken = default) => Task.CompletedTask;
+
+			/// <inheritdoc/>
+			public override IBlobWriter Fork() => new Writer(_outer, _basePath, Options);
+
+			/// <inheritdoc/>
+			public override Memory<byte> GetOutputBuffer(int usedSize, int desiredSize)
+			{
+				if (_offset + desiredSize > _data.Length)
+				{
+					byte[] newData = new byte[(desiredSize + 4095) & ~4095];
+					_data.AsSpan(_offset, usedSize).CopyTo(newData);
+					_data = newData;
+					_offset = 0;
+				}
+				return _data.AsMemory(_offset);
+			}
+
+			/// <inheritdoc/>
+			public override async ValueTask<IHashedBlobRef> WriteBlobAsync(BlobType type, int size, IReadOnlyList<IBlobRef> imports, IReadOnlyList<AliasInfo> aliases, CancellationToken cancellationToken = default)
+			{
+				ReadOnlyMemory<byte> data = _data.AsMemory(_offset, size);
+				_offset += size;
+
+				using ReadOnlyMemoryStream stream = new ReadOnlyMemoryStream(data);
+
+				IHashedBlobRef handle = await _outer.WriteBlobAsync(type, stream, imports, _basePath, cancellationToken);
+				foreach (AliasInfo aliasInfo in aliases)
+				{
+					await _outer.AddAliasAsync(aliasInfo.Name, handle, aliasInfo.Rank, aliasInfo.Data, cancellationToken);
+				}
+				return handle;
+			}
+		}
+
+		/// <inheritdoc/>
+		public bool SupportsRedirects => _backend.SupportsRedirects;
+
+		readonly IStorageBackend _backend;
+
+		/// <summary>
+		/// Constructor
+		/// </summary>
+		public KeyValueStorageNamespace(IStorageBackend backend)
+		{
+			_backend = backend;
+		}
+
+		/// <summary>
+		/// Create an in-memory storage namespace
+		/// </summary>
+		public static KeyValueStorageNamespace CreateInMemory() => new KeyValueStorageNamespace(new MemoryStorageBackend());
+
+		#region Blobs
+
+		/// <inheritdoc/>
+		public IBlobRef CreateBlobRef(BlobLocator locator)
+		{
+			return new Handle(this, locator);
+		}
+
+		/// <inheritdoc/>
+		public IBlobWriter CreateBlobWriter(string? basePath = null, BlobSerializerOptions? options = null) => new Writer(this, basePath, options);
+
+		/// <inheritdoc/>
+		public async ValueTask<BlobData> ReadBlobAsync(BlobLocator locator, CancellationToken cancellationToken = default)
+		{
+			IReadOnlyMemoryOwner<byte> owner = await _backend.ReadBlobAsync(locator, cancellationToken);
+			EncodedBlobData encodedData = new EncodedBlobData(owner.Memory.ToArray());
+
+			return new BlobDataWithOwner(encodedData.Type, encodedData.Payload, encodedData.Imports.Select(x => CreateBlobRef(x)).ToArray(), owner);
+		}
+
+		/// <inheritdoc/>
+		public async ValueTask<IHashedBlobRef> WriteBlobAsync(BlobType type, Stream stream, IReadOnlyList<IBlobRef> imports, string? basePath = null, CancellationToken cancellationToken = default)
+		{
+			BlobLocator[] locators = imports.ConvertAll(x => x.GetLocator()).ToArray();
+			byte[] payload = await stream.ReadAllBytesAsync(cancellationToken);
+			byte[] encodedData = EncodedBlobData.Create(type, locators, payload);
+
+			using ReadOnlyMemoryStream encodedStream = new ReadOnlyMemoryStream(encodedData);
+			BlobLocator locator = await _backend.WriteBlobAsync(encodedStream, locators, basePath, cancellationToken);
+
+			IoHash hash = IoHash.Compute(payload);
+			return HashedBlobRef.Create(hash, CreateBlobRef(locator));
+		}
+
+		#endregion
+
+		#region Aliases
+
+		/// <inheritdoc/>
+		public async Task AddAliasAsync(string name, IBlobRef handle, int rank = 0, ReadOnlyMemory<byte> data = default, CancellationToken cancellationToken = default)
+		{
+			await handle.FlushAsync(cancellationToken);
+			await _backend.AddAliasAsync(name, handle.GetLocator(), rank, data, cancellationToken);
+		}
+
+		/// <inheritdoc/>
+		public async Task RemoveAliasAsync(string name, IBlobRef handle, CancellationToken cancellationToken = default)
+		{
+			await handle.FlushAsync(cancellationToken);
+			await _backend.RemoveAliasAsync(name, handle.GetLocator(), cancellationToken);
+		}
+
+		/// <inheritdoc/>
+		public async Task<BlobAlias[]> FindAliasesAsync(string name, int? maxResults = null, CancellationToken cancellationToken = default)
+		{
+			BlobAliasLocator[] locators = await _backend.FindAliasesAsync(name, maxResults, cancellationToken);
+			return locators.Select(x => new BlobAlias(CreateBlobRef(x.Target), x.Rank, x.Data)).ToArray();
+		}
+
+		#endregion
+
+		#region Refs
+
+		/// <inheritdoc/>
+		public async Task<IHashedBlobRef?> TryReadRefAsync(RefName name, RefCacheTime cacheTime = default, CancellationToken cancellationToken = default)
+		{
+			HashedBlobRefValue? value = await _backend.TryReadRefAsync(name, cacheTime, cancellationToken);
+			if (value == null)
+			{
+				return null;
+			}
+			return this.CreateBlobRef(value.Hash, value.Locator);
+		}
+
+		/// <inheritdoc/>
+		public async Task WriteRefAsync(RefName name, IHashedBlobRef value, RefOptions? options = null, CancellationToken cancellationToken = default)
+		{
+			await value.FlushAsync(cancellationToken);
+			await _backend.WriteRefAsync(name, new HashedBlobRefValue(value.Hash, value.GetLocator()), options, cancellationToken);
+		}
+
+		/// <inheritdoc/>
+		public Task<bool> DeleteRefAsync(RefName name, CancellationToken cancellationToken = default)
+			=> _backend.DeleteRefAsync(name, cancellationToken);
+
+		#endregion
+
+		/// <inheritdoc/>
+		public void GetStats(StorageStats stats)
+			=> _backend.GetStats(stats);
+	}
+}
