Horde: Manage lifetime of the HttpMessageHandler used by Horde clients into a separate class, so consumers don't have to be disposable.

#rnx

diff --git a/Engine/Source/Programs/Shared/EpicGames.Horde/HordeClient.cs b/Engine/Source/Programs/Shared/EpicGames.Horde/HordeClient.cs
--- a/Engine/Source/Programs/Shared/EpicGames.Horde/HordeClient.cs
+++ b/Engine/Source/Programs/Shared/EpicGames.Horde/HordeClient.cs
@@ -261,22 +261,21 @@ namespace EpicGames.Horde
 	class HordeClientWithStaticCredentials : HordeClient
 	{
 		readonly string? _accessToken;
-		readonly HttpMessageHandler _httpMessageHandler;
+		readonly IHordeHttpMessageHandler _httpMessageHandler;
 
 		/// <summary>
 		/// Constructor
 		/// </summary>
-		public HordeClientWithStaticCredentials(Uri serverUrl, string? accessToken, BundleCache bundleCache, IOptionsSnapshot<HordeOptions> hordeOptions, ILoggerFactory loggerFactory)
+		public HordeClientWithStaticCredentials(Uri serverUrl, string? accessToken, IHordeHttpMessageHandler httpMessageHandler, BundleCache bundleCache, IOptionsSnapshot<HordeOptions> hordeOptions, ILoggerFactory loggerFactory)
 			: base(serverUrl, bundleCache, hordeOptions, loggerFactory)
 		{
 			_accessToken = accessToken;
-			_httpMessageHandler = HordeHttpMessageHandler.Create(Logger);
+			_httpMessageHandler = httpMessageHandler;
 		}
 
 		/// <inheritdoc/>
 		public override async ValueTask DisposeAsync()
 		{
-			_httpMessageHandler.Dispose();
 			await base.DisposeAsync();
 		}
 
@@ -295,7 +294,7 @@ namespace EpicGames.Horde
 		/// <inheritdoc/>
 		protected override HttpClient CreateAuthenticatedHttpClient()
 		{
-			HttpClient httpClient = new HttpClient(_httpMessageHandler, false);
+			HttpClient httpClient = new HttpClient(_httpMessageHandler.Instance, false);
 			httpClient.BaseAddress = ServerUrl;
 			if (!String.IsNullOrEmpty(_accessToken))
 			{
@@ -307,7 +306,7 @@ namespace EpicGames.Horde
 		/// <inheritdoc/>
 		protected override HttpClient CreateUnauthenticatedHttpClient()
 		{
-			HttpClient httpClient = new HttpClient(_httpMessageHandler, false);
+			HttpClient httpClient = new HttpClient(_httpMessageHandler.Instance, false);
 			httpClient.BaseAddress = ServerUrl;
 			return httpClient;
 		}
@@ -319,19 +318,19 @@ namespace EpicGames.Horde
 	class HordeClientWithDynamicCredentials : HordeClient
 	{
 		readonly HordeHttpAuthHandlerState _authHandlerState;
-		readonly HttpMessageHandler _baseHttpMessageHandler;
+		readonly IHordeHttpMessageHandler _baseHttpMessageHandler;
 		readonly HttpMessageHandler _authHttpMessageHandler;
 
 		/// <summary>
 		/// Constructor
 		/// </summary>
-		public HordeClientWithDynamicCredentials(Uri serverUrl, BundleCache bundleCache, IOptionsSnapshot<HordeOptions> hordeOptions, ILoggerFactory loggerFactory)
+		public HordeClientWithDynamicCredentials(Uri serverUrl, BundleCache bundleCache, IHordeHttpMessageHandler baseHttpMessageHandler, IOptionsSnapshot<HordeOptions> hordeOptions, ILoggerFactory loggerFactory)
 			: base(serverUrl, bundleCache, hordeOptions, loggerFactory)
 		{
-			_baseHttpMessageHandler = HordeHttpMessageHandler.Create(Logger);
+			_baseHttpMessageHandler = baseHttpMessageHandler;
 
-			_authHandlerState = new HordeHttpAuthHandlerState(_baseHttpMessageHandler, serverUrl, hordeOptions, loggerFactory.CreateLogger<HordeHttpAuthHandlerState>());
-			_authHttpMessageHandler = new HordeHttpAuthHandler(_baseHttpMessageHandler, _authHandlerState, hordeOptions);
+			_authHandlerState = new HordeHttpAuthHandlerState(_baseHttpMessageHandler.Instance, serverUrl, hordeOptions, loggerFactory.CreateLogger<HordeHttpAuthHandlerState>());
+			_authHttpMessageHandler = new HordeHttpAuthHandler(_baseHttpMessageHandler.Instance, _authHandlerState, hordeOptions);
 		}
 
 		/// <inheritdoc/>
@@ -339,7 +338,6 @@ namespace EpicGames.Horde
 		{
 			await _authHandlerState.DisposeAsync();
 			_authHttpMessageHandler.Dispose();
-			_baseHttpMessageHandler.Dispose();
 			await base.DisposeAsync();
 		}
 
@@ -377,7 +375,7 @@ namespace EpicGames.Horde
 		/// <inheritdoc/>
 		protected override HttpClient CreateUnauthenticatedHttpClient()
 		{
-			HttpClient httpClient = new HttpClient(_baseHttpMessageHandler, false);
+			HttpClient httpClient = new HttpClient(_baseHttpMessageHandler.Instance, false);
 			httpClient.BaseAddress = ServerUrl;
 			return httpClient;
 		}
@@ -389,15 +387,17 @@ namespace EpicGames.Horde
 	public class HordeClientFactory
 	{
 		readonly BundleCache _bundleCache;
+		readonly IHordeHttpMessageHandler _httpMessageHandler;
 		readonly IOptionsSnapshot<HordeOptions> _hordeOptions;
 		readonly ILoggerFactory _loggerFactory;
 
 		/// <summary>
 		/// Constructor 
 		/// </summary>
-		public HordeClientFactory(BundleCache bundleCache, IOptionsSnapshot<HordeOptions> hordeOptions, ILoggerFactory loggerFactory)
+		public HordeClientFactory(BundleCache bundleCache, IHordeHttpMessageHandler httpMessageHandler, IOptionsSnapshot<HordeOptions> hordeOptions, ILoggerFactory loggerFactory)
 		{
 			_bundleCache = bundleCache;
+			_httpMessageHandler = httpMessageHandler;
 			_hordeOptions = hordeOptions;
 			_loggerFactory = loggerFactory;
 		}
@@ -420,11 +420,11 @@ namespace EpicGames.Horde
 			serverUrl ??= GetDefaultServerUrl();
 			if (accessToken == null)
 			{
-				return new HordeClientWithDynamicCredentials(serverUrl, _bundleCache, _hordeOptions, _loggerFactory);
+				return new HordeClientWithDynamicCredentials(serverUrl, _bundleCache, _httpMessageHandler, _hordeOptions, _loggerFactory);
 			}
 			else
 			{
-				return new HordeClientWithStaticCredentials(serverUrl, accessToken, _bundleCache, _hordeOptions, _loggerFactory);
+				return new HordeClientWithStaticCredentials(serverUrl, accessToken, _httpMessageHandler, _bundleCache, _hordeOptions, _loggerFactory);
 			}
 		}
 	}
diff --git a/Engine/Source/Programs/Shared/EpicGames.Horde/HordeExtensions.cs b/Engine/Source/Programs/Shared/EpicGames.Horde/HordeExtensions.cs
--- a/Engine/Source/Programs/Shared/EpicGames.Horde/HordeExtensions.cs
+++ b/Engine/Source/Programs/Shared/EpicGames.Horde/HordeExtensions.cs
@@ -35,6 +35,7 @@ namespace EpicGames.Horde
 			serviceCollection.AddSingleton<StorageBackendCache>(CreateBackendCache);
 			serviceCollection.AddSingleton<HttpStorageBackendFactory>();
 			serviceCollection.AddSingleton<HttpStorageClient>();
+			serviceCollection.AddSingleton<IHordeHttpMessageHandler, HordeHttpMessageHandler>();
 			serviceCollection.AddSingleton<IHordeClient>(sp => sp.GetRequiredService<HordeClientFactory>().Create());
 			serviceCollection.AddSingleton<HordeClientFactory>();
 		}
diff --git a/Engine/Source/Programs/Shared/EpicGames.Horde/HordeHttpMessageHandler.cs b/Engine/Source/Programs/Shared/EpicGames.Horde/HordeHttpMessageHandler.cs
--- a/Engine/Source/Programs/Shared/EpicGames.Horde/HordeHttpMessageHandler.cs
+++ b/Engine/Source/Programs/Shared/EpicGames.Horde/HordeHttpMessageHandler.cs
@@ -14,32 +14,30 @@ using Polly.Timeout;
 namespace EpicGames.Horde
 {
 	/// <summary>
-	/// Helper class to create <see cref="HttpMessageHandler"/> instances for Horde with default resiliance pipeline.
+	/// Concrete implementation of <see cref="IHordeHttpMessageHandler"/> which manages the lifetime of the <see cref="HttpMessageHandler"/> instance.
 	/// </summary>
-	static class HordeHttpMessageHandler
+	class HordeHttpMessageHandler : IHordeHttpMessageHandler, IDisposable
 	{
+		readonly HttpMessageHandler _instance;
+
+		/// <inheritdoc/>
+		public HttpMessageHandler Instance => _instance;
+
 		/// <summary>
-		/// Creates a new message handler with default resilience policies
+		/// Constructor
 		/// </summary>
-		public static HttpMessageHandler Create(ILogger logger)
+		/// <param name="logger"></param>
+		public HordeHttpMessageHandler(ILogger<HordeHttpMessageHandler> logger)
 		{
-			HttpMessageHandler? httpMessageHandler = null;
-			try
-			{
-#pragma warning disable CA2000 // Call dispose on httpMessageHandler (disposed by child handlers)
-				httpMessageHandler = new SocketsHttpHandler();
-				httpMessageHandler = new PolicyHttpMessageHandler(request => CreateDefaultTimeoutRetryPolicy(request, logger)) { InnerHandler = httpMessageHandler };
-				httpMessageHandler = new PolicyHttpMessageHandler(request => CreateDefaultTransientErrorPolicy(request, logger)) { InnerHandler = httpMessageHandler };
-#pragma warning restore CA2000
-				return httpMessageHandler;
-			}
-			catch
-			{
-				httpMessageHandler?.Dispose();
-				throw;
-			}
+			_instance = new SocketsHttpHandler();
+			_instance = new PolicyHttpMessageHandler(request => CreateDefaultTimeoutRetryPolicy(request, logger)) { InnerHandler = _instance };
+			_instance = new PolicyHttpMessageHandler(request => CreateDefaultTransientErrorPolicy(request, logger)) { InnerHandler = _instance };
 		}
 
+		/// <inheritdoc/>
+		public void Dispose()
+			=> _instance.Dispose();
+
 		/// <summary>
 		/// Create a default timeout retry policy
 		/// </summary>
diff --git a/Engine/Source/Programs/Shared/EpicGames.Horde/IHordeHttpMessageHandler.cs b/Engine/Source/Programs/Shared/EpicGames.Horde/IHordeHttpMessageHandler.cs
new file mode 100644
--- /dev/null
+++ b/Engine/Source/Programs/Shared/EpicGames.Horde/IHordeHttpMessageHandler.cs
@@ -0,0 +1,17 @@
+// Copyright Epic Games, Inc. All Rights Reserved.
+
+using System.Net.Http;
+
+namespace EpicGames.Horde
+{
+	/// <summary>
+	/// Provides access to a <see cref="HttpMessageHandler"/> instances for Horde with a default resiliance pipeline.
+	/// </summary>
+	public interface IHordeHttpMessageHandler
+	{
+		/// <summary>
+		/// Instance of the http message handler
+		/// </summary>
+		HttpMessageHandler Instance { get; }
+	}
+}
