Horde: Move http utility methods to a separate class.

#rnx

diff --git a/Engine/Source/Programs/Shared/EpicGames.Horde/Compute/Clients/ServerComputeClient.cs b/Engine/Source/Programs/Shared/EpicGames.Horde/Compute/Clients/ServerComputeClient.cs
--- a/Engine/Source/Programs/Shared/EpicGames.Horde/Compute/Clients/ServerComputeClient.cs
+++ b/Engine/Source/Programs/Shared/EpicGames.Horde/Compute/Clients/ServerComputeClient.cs
@@ -223,7 +223,7 @@ namespace EpicGames.Horde.Compute.Clients
 				Protocol = (int)ComputeProtocol.Latest
 			};
 			
-			using HttpResponseMessage httpResponse = await HordeHttpClient.PostAsync(_httpClient, "api/v2/compute/_cluster", request, _cancellationSource.Token);
+			using HttpResponseMessage httpResponse = await HordeHttpRequest.PostAsync(_httpClient, "api/v2/compute/_cluster", request, _cancellationSource.Token);
 			if (!httpResponse.IsSuccessStatusCode)
 			{
 				string body = await httpResponse.Content.ReadAsStringAsync(cancellationToken);
@@ -263,7 +263,7 @@ namespace EpicGames.Horde.Compute.Clients
 		public async Task DeclareResourceNeedsAsync(ClusterId clusterId, string pool, Dictionary<string, int> resourceNeeds, CancellationToken cancellationToken = default)
 		{
 			ResourceNeedsMessage request = new() { SessionId = _sessionId, Pool = pool, ResourceNeeds = resourceNeeds };
-			using HttpResponseMessage response = await HordeHttpClient.PostAsync(_httpClient, $"api/v2/compute/{clusterId}/resource-needs", request, _cancellationSource.Token);
+			using HttpResponseMessage response = await HordeHttpRequest.PostAsync(_httpClient, $"api/v2/compute/{clusterId}/resource-needs", request, _cancellationSource.Token);
 			response.EnsureSuccessStatusCode();
 		}
 
@@ -285,7 +285,7 @@ namespace EpicGames.Horde.Compute.Clients
 
 			AssignComputeResponse? response;
 			string path = clusterId == null ? "api/v2/compute" : $"api/v2/compute/{clusterId}";
-			using (HttpResponseMessage httpResponse = await HordeHttpClient.PostAsync(_httpClient, path, request, _cancellationSource.Token))
+			using (HttpResponseMessage httpResponse = await HordeHttpRequest.PostAsync(_httpClient, path, request, _cancellationSource.Token))
 			{
 				if (httpResponse.StatusCode == HttpStatusCode.NotFound)
 				{
diff --git a/Engine/Source/Programs/Shared/EpicGames.Horde/HordeHttpClient.cs b/Engine/Source/Programs/Shared/EpicGames.Horde/HordeHttpClient.cs
--- a/Engine/Source/Programs/Shared/EpicGames.Horde/HordeHttpClient.cs
+++ b/Engine/Source/Programs/Shared/EpicGames.Horde/HordeHttpClient.cs
@@ -28,6 +28,8 @@ using EpicGames.Horde.Streams;
 using EpicGames.Horde.Tools;
 using EpicGames.Horde.Ugs;
 
+using static EpicGames.Horde.HordeHttpRequest;
+
 #pragma warning disable CA2234
 
 namespace EpicGames.Horde
@@ -66,8 +68,7 @@ namespace EpicGames.Horde
 
 		readonly HttpClient _httpClient;
 
-		static readonly JsonSerializerOptions s_jsonSerializerOptions = CreateJsonSerializerOptions();
-		internal static JsonSerializerOptions JsonSerializerOptions => s_jsonSerializerOptions;
+		internal static JsonSerializerOptions JsonSerializerOptions => HordeHttpRequest.JsonSerializerOptions;
 
 		/// <summary>
 		/// Base address for the Horde server
@@ -89,32 +90,12 @@ namespace EpicGames.Horde
 			_httpClient.Dispose();
 		}
 
-		/// <summary>
-		/// Create the shared instance of JSON options for HordeHttpClient instances
-		/// </summary>
-		static JsonSerializerOptions CreateJsonSerializerOptions()
-		{
-			JsonSerializerOptions options = new JsonSerializerOptions();
-			ConfigureJsonSerializer(options);
-			return options;
-		}
-
 		/// <summary>
 		/// Configures a JSON serializer to read Horde responses
 		/// </summary>
 		/// <param name="options">options for the serializer</param>
 		public static void ConfigureJsonSerializer(JsonSerializerOptions options)
-		{
-			options.AllowTrailingCommas = true;
-			options.ReadCommentHandling = JsonCommentHandling.Skip;
-			options.DefaultIgnoreCondition = JsonIgnoreCondition.WhenWritingNull;
-			options.PropertyNamingPolicy = JsonNamingPolicy.CamelCase;
-			options.PropertyNameCaseInsensitive = true;
-			options.Converters.Add(new JsonStringEnumConverter());
-			options.Converters.Add(new StringIdJsonConverterFactory());
-			options.Converters.Add(new BinaryIdJsonConverterFactory());
-			options.Converters.Add(new SubResourceIdJsonConverterFactory());
-		}
+			=> HordeHttpRequest.ConfigureJsonSerializer(options);
 
 		#region Connection
 		/// <summary>
@@ -128,6 +109,7 @@ namespace EpicGames.Horde
 		}
 
 		#endregion
+
 		#region Artifacts
 
 		/// <summary>
@@ -552,113 +534,5 @@ namespace EpicGames.Horde
 			return GetAsync<GetUgsMetadataResponse>(_httpClient, $"/ugs/api/metadata?stream={perforceStream}&change={commitId.GetPerforceChange()}&project={projectId}", cancellationToken);
 		}
 		#endregion
-
-		#region Utility Methods
-
-		/// <summary>
-		/// Deletes a resource from an HTTP endpoint
-		/// </summary>
-		/// <param name="httpClient">Http client instance</param>
-		/// <param name="relativePath">The url to retrieve</param>
-		/// <param name="cancellationToken">Cancels the request</param>
-		internal static async Task DeleteAsync(HttpClient httpClient, string relativePath, CancellationToken cancellationToken = default)
-		{
-			using HttpResponseMessage response = await httpClient.DeleteAsync(relativePath, cancellationToken);
-			response.EnsureSuccessStatusCode();
-		}
-
-		/// <summary>
-		/// Gets a resource from an HTTP endpoint and parses it as a JSON object
-		/// </summary>
-		/// <typeparam name="TResponse">The object type to return</typeparam>
-		/// <param name="httpClient">Http client instance</param>
-		/// <param name="relativePath">The url to retrieve</param>
-		/// <param name="cancellationToken">Cancels the request</param>
-		/// <returns>New instance of the object</returns>
-		internal static async Task<TResponse> GetAsync<TResponse>(HttpClient httpClient, string relativePath, CancellationToken cancellationToken = default)
-		{
-			TResponse? response = await httpClient.GetFromJsonAsync<TResponse>(relativePath, s_jsonSerializerOptions, cancellationToken);
-			return response ?? throw new InvalidCastException($"Expected non-null response from GET to {relativePath}");
-		}
-
-		/// <summary>
-		/// Posts an object to an HTTP endpoint as a JSON object, and parses the response object
-		/// </summary>
-		/// <typeparam name="TRequest">The object type to post</typeparam>
-		/// <param name="httpClient">Http client instance</param>
-		/// <param name="relativePath">The url to retrieve</param>
-		/// <param name="request">The object to post</param>
-		/// <param name="cancellationToken">Cancels the request</param>
-		/// <returns>The response parsed into the requested type</returns>
-		internal static async Task<HttpResponseMessage> PostAsync<TRequest>(HttpClient httpClient, string relativePath, TRequest request, CancellationToken cancellationToken = default)
-		{
-			return await httpClient.PostAsJsonAsync<TRequest>(relativePath, request, s_jsonSerializerOptions, cancellationToken);
-		}
-
-		/// <summary>
-		/// Posts an object to an HTTP endpoint as a JSON object, and parses the response object
-		/// </summary>
-		/// <typeparam name="TResponse">The object type to return</typeparam>
-		/// <typeparam name="TRequest">The object type to post</typeparam>
-		/// <param name="httpClient">Http client instance</param>
-		/// <param name="relativePath">The url to retrieve</param>
-		/// <param name="request">The object to post</param>
-		/// <param name="cancellationToken">Cancels the request</param>
-		/// <returns>The response parsed into the requested type</returns>
-		internal static async Task<TResponse> PostAsync<TResponse, TRequest>(HttpClient httpClient, string relativePath, TRequest request, CancellationToken cancellationToken = default)
-		{
-			using (HttpResponseMessage response = await PostAsync<TRequest>(httpClient, relativePath, request, cancellationToken))
-			{
-				if (!response.IsSuccessStatusCode)
-				{
-					string body = await response.Content.ReadAsStringAsync(cancellationToken);
-					throw new HttpRequestException($"{(int)response.StatusCode} ({response.StatusCode}) posting to {new Uri(httpClient.BaseAddress!, relativePath)}: {body}", null, response.StatusCode);
-				}
-
-				TResponse? responseValue = await response.Content.ReadFromJsonAsync<TResponse>(s_jsonSerializerOptions, cancellationToken);
-				return responseValue ?? throw new InvalidCastException($"Expected non-null response from POST to {relativePath}");
-			}
-		}
-
-		/// <summary>
-		/// Puts an object to an HTTP endpoint as a JSON object
-		/// </summary>
-		/// <typeparam name="TRequest">The object type to post</typeparam>
-		/// <param name="httpClient">Http client instance</param>
-		/// <param name="relativePath">The url to write to</param>
-		/// <param name="request">The object to post</param>
-		/// <param name="cancellationToken">Cancels the request</param>
-		/// <returns>Response message</returns>
-		internal static async Task<HttpResponseMessage> PutAsync<TRequest>(HttpClient httpClient, string relativePath, TRequest request, CancellationToken cancellationToken)
-		{
-			return await httpClient.PutAsJsonAsync<TRequest>(relativePath, request, s_jsonSerializerOptions, cancellationToken);
-		}
-
-		/// <summary>
-		/// Puts an object to an HTTP endpoint as a JSON object
-		/// </summary>
-		/// <typeparam name="TResponse">The object type to return</typeparam>
-		/// <typeparam name="TRequest">The object type to post</typeparam>
-		/// <param name="httpClient">Http client instance</param>
-		/// <param name="relativePath">The url to write to</param>
-		/// <param name="request">The object to post</param>
-		/// <param name="cancellationToken">Cancels the request</param>
-		/// <returns>Response message</returns>
-		internal static async Task<TResponse> PutAsync<TResponse, TRequest>(HttpClient httpClient, string relativePath, TRequest request, CancellationToken cancellationToken)
-		{
-			using (HttpResponseMessage response = await httpClient.PutAsJsonAsync<TRequest>(relativePath, request, s_jsonSerializerOptions, cancellationToken))
-			{
-				if (!response.IsSuccessStatusCode)
-				{
-					string body = await response.Content.ReadAsStringAsync(cancellationToken);
-					throw new HttpRequestException($"{response.StatusCode} put to {new Uri(httpClient.BaseAddress!, relativePath)}: {body}", null, response.StatusCode);
-				}
-
-				TResponse? responseValue = await response.Content.ReadFromJsonAsync<TResponse>(s_jsonSerializerOptions, cancellationToken);
-				return responseValue ?? throw new InvalidCastException($"Expected non-null response from PUT to {relativePath}");
-			}
-		}
-
-		#endregion
 	}
 }
diff --git a/Engine/Source/Programs/Shared/EpicGames.Horde/HordeHttpRequest.cs b/Engine/Source/Programs/Shared/EpicGames.Horde/HordeHttpRequest.cs
new file mode 100644
--- /dev/null
+++ b/Engine/Source/Programs/Shared/EpicGames.Horde/HordeHttpRequest.cs
@@ -0,0 +1,156 @@
+// Copyright Epic Games, Inc. All Rights Reserved.
+
+using System;
+using System.Collections.Generic;
+using System.Linq;
+using System.Net.Http;
+using System.Net.Http.Json;
+using System.Text;
+using System.Text.Json;
+using System.Text.Json.Serialization;
+using System.Threading;
+using System.Threading.Tasks;
+using EpicGames.Horde.Jobs;
+
+namespace EpicGames.Horde
+{
+	/// <summary>
+	/// Static helper methods for implementing Horde HTTP requests with standard semantics
+	/// </summary>
+	public static class HordeHttpRequest
+	{
+		static readonly JsonSerializerOptions s_jsonSerializerOptions = CreateJsonSerializerOptions();
+		internal static JsonSerializerOptions JsonSerializerOptions => s_jsonSerializerOptions;
+
+		/// <summary>
+		/// Create the shared instance of JSON options for HordeHttpClient instances
+		/// </summary>
+		static JsonSerializerOptions CreateJsonSerializerOptions()
+		{
+			JsonSerializerOptions options = new JsonSerializerOptions();
+			ConfigureJsonSerializer(options);
+			return options;
+		}
+
+		/// <summary>
+		/// Configures a JSON serializer to read Horde responses
+		/// </summary>
+		/// <param name="options">options for the serializer</param>
+		public static void ConfigureJsonSerializer(JsonSerializerOptions options)
+		{
+			options.AllowTrailingCommas = true;
+			options.ReadCommentHandling = JsonCommentHandling.Skip;
+			options.DefaultIgnoreCondition = JsonIgnoreCondition.WhenWritingNull;
+			options.PropertyNamingPolicy = JsonNamingPolicy.CamelCase;
+			options.PropertyNameCaseInsensitive = true;
+			options.Converters.Add(new JsonStringEnumConverter());
+			options.Converters.Add(new StringIdJsonConverterFactory());
+			options.Converters.Add(new BinaryIdJsonConverterFactory());
+			options.Converters.Add(new SubResourceIdJsonConverterFactory());
+		}
+
+		/// <summary>
+		/// Deletes a resource from an HTTP endpoint
+		/// </summary>
+		/// <param name="httpClient">Http client instance</param>
+		/// <param name="relativePath">The url to retrieve</param>
+		/// <param name="cancellationToken">Cancels the request</param>
+		public static async Task DeleteAsync(HttpClient httpClient, string relativePath, CancellationToken cancellationToken = default)
+		{
+			using HttpResponseMessage response = await httpClient.DeleteAsync(relativePath, cancellationToken);
+			response.EnsureSuccessStatusCode();
+		}
+
+		/// <summary>
+		/// Gets a resource from an HTTP endpoint and parses it as a JSON object
+		/// </summary>
+		/// <typeparam name="TResponse">The object type to return</typeparam>
+		/// <param name="httpClient">Http client instance</param>
+		/// <param name="relativePath">The url to retrieve</param>
+		/// <param name="cancellationToken">Cancels the request</param>
+		/// <returns>New instance of the object</returns>
+		public static async Task<TResponse> GetAsync<TResponse>(HttpClient httpClient, string relativePath, CancellationToken cancellationToken = default)
+		{
+			TResponse? response = await httpClient.GetFromJsonAsync<TResponse>(relativePath, s_jsonSerializerOptions, cancellationToken);
+			return response ?? throw new InvalidCastException($"Expected non-null response from GET to {relativePath}");
+		}
+
+		/// <summary>
+		/// Posts an object to an HTTP endpoint as a JSON object, and parses the response object
+		/// </summary>
+		/// <typeparam name="TRequest">The object type to post</typeparam>
+		/// <param name="httpClient">Http client instance</param>
+		/// <param name="relativePath">The url to retrieve</param>
+		/// <param name="request">The object to post</param>
+		/// <param name="cancellationToken">Cancels the request</param>
+		/// <returns>The response parsed into the requested type</returns>
+		internal static async Task<HttpResponseMessage> PostAsync<TRequest>(HttpClient httpClient, string relativePath, TRequest request, CancellationToken cancellationToken = default)
+		{
+			return await httpClient.PostAsJsonAsync<TRequest>(relativePath, request, s_jsonSerializerOptions, cancellationToken);
+		}
+
+		/// <summary>
+		/// Posts an object to an HTTP endpoint as a JSON object, and parses the response object
+		/// </summary>
+		/// <typeparam name="TResponse">The object type to return</typeparam>
+		/// <typeparam name="TRequest">The object type to post</typeparam>
+		/// <param name="httpClient">Http client instance</param>
+		/// <param name="relativePath">The url to retrieve</param>
+		/// <param name="request">The object to post</param>
+		/// <param name="cancellationToken">Cancels the request</param>
+		/// <returns>The response parsed into the requested type</returns>
+		public static async Task<TResponse> PostAsync<TResponse, TRequest>(HttpClient httpClient, string relativePath, TRequest request, CancellationToken cancellationToken = default)
+		{
+			using (HttpResponseMessage response = await PostAsync<TRequest>(httpClient, relativePath, request, cancellationToken))
+			{
+				if (!response.IsSuccessStatusCode)
+				{
+					string body = await response.Content.ReadAsStringAsync(cancellationToken);
+					throw new HttpRequestException($"{(int)response.StatusCode} ({response.StatusCode}) posting to {new Uri(httpClient.BaseAddress!, relativePath)}: {body}", null, response.StatusCode);
+				}
+
+				TResponse? responseValue = await response.Content.ReadFromJsonAsync<TResponse>(s_jsonSerializerOptions, cancellationToken);
+				return responseValue ?? throw new InvalidCastException($"Expected non-null response from POST to {relativePath}");
+			}
+		}
+
+		/// <summary>
+		/// Puts an object to an HTTP endpoint as a JSON object
+		/// </summary>
+		/// <typeparam name="TRequest">The object type to post</typeparam>
+		/// <param name="httpClient">Http client instance</param>
+		/// <param name="relativePath">The url to write to</param>
+		/// <param name="request">The object to post</param>
+		/// <param name="cancellationToken">Cancels the request</param>
+		/// <returns>Response message</returns>
+		public static async Task<HttpResponseMessage> PutAsync<TRequest>(HttpClient httpClient, string relativePath, TRequest request, CancellationToken cancellationToken)
+		{
+			return await httpClient.PutAsJsonAsync<TRequest>(relativePath, request, s_jsonSerializerOptions, cancellationToken);
+		}
+
+		/// <summary>
+		/// Puts an object to an HTTP endpoint as a JSON object
+		/// </summary>
+		/// <typeparam name="TResponse">The object type to return</typeparam>
+		/// <typeparam name="TRequest">The object type to post</typeparam>
+		/// <param name="httpClient">Http client instance</param>
+		/// <param name="relativePath">The url to write to</param>
+		/// <param name="request">The object to post</param>
+		/// <param name="cancellationToken">Cancels the request</param>
+		/// <returns>Response message</returns>
+		public static async Task<TResponse> PutAsync<TResponse, TRequest>(HttpClient httpClient, string relativePath, TRequest request, CancellationToken cancellationToken)
+		{
+			using (HttpResponseMessage response = await httpClient.PutAsJsonAsync<TRequest>(relativePath, request, s_jsonSerializerOptions, cancellationToken))
+			{
+				if (!response.IsSuccessStatusCode)
+				{
+					string body = await response.Content.ReadAsStringAsync(cancellationToken);
+					throw new HttpRequestException($"{response.StatusCode} put to {new Uri(httpClient.BaseAddress!, relativePath)}: {body}", null, response.StatusCode);
+				}
+
+				TResponse? responseValue = await response.Content.ReadFromJsonAsync<TResponse>(s_jsonSerializerOptions, cancellationToken);
+				return responseValue ?? throw new InvalidCastException($"Expected non-null response from PUT to {relativePath}");
+			}
+		}
+	}
+}
