Horde: Fix static analysis warning.

#rnx

diff --git a/Engine/Source/Programs/Shared/EpicGames.Horde/HordeHttpRequest.cs b/Engine/Source/Programs/Shared/EpicGames.Horde/HordeHttpRequest.cs
--- a/Engine/Source/Programs/Shared/EpicGames.Horde/HordeHttpRequest.cs
+++ b/Engine/Source/Programs/Shared/EpicGames.Horde/HordeHttpRequest.cs
@@ -12,6 +12,8 @@ using System.Threading;
 using System.Threading.Tasks;
 using EpicGames.Horde.Jobs;
 
+#pragma warning disable CA2234 // Use Uri instead of string
+
 namespace EpicGames.Horde
 {
 	/// <summary>
