Horde: Fix static analysis warnings.

#rnx

diff --git a/Engine/Source/Programs/Shared/EpicGames.Horde/HordeHttpClient.cs b/Engine/Source/Programs/Shared/EpicGames.Horde/HordeHttpClient.cs
--- a/Engine/Source/Programs/Shared/EpicGames.Horde/HordeHttpClient.cs
+++ b/Engine/Source/Programs/Shared/EpicGames.Horde/HordeHttpClient.cs
@@ -9,7 +9,6 @@ using System.Net.Http;
 using System.Net.Http.Headers;
 using System.Net.Http.Json;
 using System.Text.Json;
-using System.Text.Json.Serialization;
 using System.Threading;
 using System.Threading.Tasks;
 using EpicGames.Core;
diff --git a/Engine/Source/Programs/Shared/EpicGames.Horde/HordeHttpRequest.cs b/Engine/Source/Programs/Shared/EpicGames.Horde/HordeHttpRequest.cs
--- a/Engine/Source/Programs/Shared/EpicGames.Horde/HordeHttpRequest.cs
+++ b/Engine/Source/Programs/Shared/EpicGames.Horde/HordeHttpRequest.cs
@@ -1,11 +1,8 @@
 // Copyright Epic Games, Inc. All Rights Reserved.
 
 using System;
-using System.Collections.Generic;
-using System.Linq;
 using System.Net.Http;
 using System.Net.Http.Json;
-using System.Text;
 using System.Text.Json;
 using System.Text.Json.Serialization;
 using System.Threading;
