[EpicGames.Perforce.Managed]
* Changed code traversing workspace on disk to use EnumerateFileSystemInfos instead of directories and files separately. This reduces number of kernel calls and save 20% time on a machine with attached ssd

diff --git a/Engine/Source/Programs/Shared/EpicGames.Perforce.Managed/WorkspaceDirectoryInfo.cs b/Engine/Source/Programs/Shared/EpicGames.Perforce.Managed/WorkspaceDirectoryInfo.cs
--- a/Engine/Source/Programs/Shared/EpicGames.Perforce.Managed/WorkspaceDirectoryInfo.cs
+++ b/Engine/Source/Programs/Shared/EpicGames.Perforce.Managed/WorkspaceDirectoryInfo.cs
@@ -149,45 +149,50 @@ namespace EpicGames.Perforce.Managed
 		{
 			// Recurse through subdirectories
 			Dictionary<Utf8String, WorkspaceDirectoryInfo> newNameToSubDirectory = new Dictionary<Utf8String, WorkspaceDirectoryInfo>(NameToSubDirectory.Count, NameToSubDirectory.Comparer);
-			foreach (DirectoryInfo subDirectoryInfo in info.EnumerateDirectories())
-			{
-				WorkspaceDirectoryInfo? subDirectory;
-				if (NameToSubDirectory.TryGetValue(new Utf8String(subDirectoryInfo.Name), out subDirectory))
-				{
-					newNameToSubDirectory.Add(subDirectory.Name, subDirectory);
-					await queue.EnqueueAsync(_ => subDirectory.RefreshAsync(subDirectoryInfo, removeUntracked, filesToDelete, directoriesToDelete, queue));
-				}
-				else if (removeUntracked)
-				{
-					directoriesToDelete.Enqueue(subDirectoryInfo);
-				}
-			}
-			NameToSubDirectory = newNameToSubDirectory;
 
 			// Figure out which files have changed.
 			Dictionary<Utf8String, WorkspaceFileInfo> newNameToFile = new Dictionary<Utf8String, WorkspaceFileInfo>(NameToFile.Count, NameToFile.Comparer);
-			foreach (FileInfo file in info.EnumerateFiles())
+
+			foreach (FileSystemInfo fileSystemInfo in info.EnumerateFileSystemInfos())
 			{
-				WorkspaceFileInfo? stagedFile;
-				if (NameToFile.TryGetValue(new Utf8String(file.Name), out stagedFile))
+				if (fileSystemInfo is DirectoryInfo subDirectoryInfo)
 				{
-					if (stagedFile.MatchesAttributes(file))
+					WorkspaceDirectoryInfo? subDirectory;
+					if (NameToSubDirectory.TryGetValue(new Utf8String(subDirectoryInfo.Name), out subDirectory))
 					{
-						newNameToFile.Add(stagedFile.Name, stagedFile);
+						newNameToSubDirectory.Add(subDirectory.Name, subDirectory);
+						await queue.EnqueueAsync(_ => subDirectory.RefreshAsync(subDirectoryInfo, removeUntracked, filesToDelete, directoriesToDelete, queue));
 					}
-					else
+					else if (removeUntracked)
 					{
-						filesToDelete.Enqueue(file);
+						directoriesToDelete.Enqueue(subDirectoryInfo);
 					}
 				}
 				else
 				{
-					if (removeUntracked)
+					FileInfo file = (FileInfo)fileSystemInfo;
+					WorkspaceFileInfo? stagedFile;
+					if (NameToFile.TryGetValue(new Utf8String(file.Name), out stagedFile))
 					{
-						filesToDelete.Enqueue(file);
+						if (stagedFile.MatchesAttributes(file))
+						{
+							newNameToFile.Add(stagedFile.Name, stagedFile);
+						}
+						else
+						{
+							filesToDelete.Enqueue(file);
+						}
+					}
+					else
+					{
+						if (removeUntracked)
+						{
+							filesToDelete.Enqueue(file);
+						}
 					}
 				}
 			}
+			NameToSubDirectory = newNameToSubDirectory;
 
 			// If the file state has changed, clear the directory hashes
 			if (NameToFile.Count != newNameToFile.Count)
