Horde: Handle .horde.json file not existing when configuring server url on Mac/Linux.

diff --git a/Engine/Source/Programs/Shared/EpicGames.Horde/HordeOptions.cs b/Engine/Source/Programs/Shared/EpicGames.Horde/HordeOptions.cs
--- a/Engine/Source/Programs/Shared/EpicGames.Horde/HordeOptions.cs
+++ b/Engine/Source/Programs/Shared/EpicGames.Horde/HordeOptions.cs
@@ -95,7 +95,7 @@ namespace EpicGames.Horde
 			else
 			{
 				FileReference? configFile = GetConfigFile();
-				if (configFile != null)
+				if (configFile != null && FileReference.Exists(configFile))
 				{
 					byte[] data = FileReference.ReadAllBytes(configFile);
 
@@ -153,9 +153,13 @@ namespace EpicGames.Horde
 				FileReference? configFile = GetConfigFile();
 				if (configFile != null)
 				{
-					byte[] data = FileReference.ReadAllBytes(configFile);
+					JsonObject? root = null;
+					if (FileReference.Exists(configFile))
+					{
+						byte[] data = FileReference.ReadAllBytes(configFile);
+						root = JsonNode.Parse(data, new JsonNodeOptions { PropertyNameCaseInsensitive = true }, new JsonDocumentOptions { AllowTrailingCommas = true }) as JsonObject;
+					}
 
-					JsonObject? root = JsonNode.Parse(data, new JsonNodeOptions { PropertyNameCaseInsensitive = true }, new JsonDocumentOptions { AllowTrailingCommas = true }) as JsonObject;
 					root ??= new JsonObject();
 					root["server"] = serverUrl.ToString();
 
