Horde: Warning fix.

#rnx

diff --git a/Engine/Source/Programs/Shared/EpicGames.Horde.Tests/ServerLogPacketBuilderTests.cs b/Engine/Source/Programs/Shared/EpicGames.Horde.Tests/ServerLogPacketBuilderTests.cs
--- a/Engine/Source/Programs/Shared/EpicGames.Horde.Tests/ServerLogPacketBuilderTests.cs
+++ b/Engine/Source/Programs/Shared/EpicGames.Horde.Tests/ServerLogPacketBuilderTests.cs
@@ -224,7 +224,11 @@ namespace EpicGames.Horde.Tests
 			public ISecretCollection Secrets => throw new NotImplementedException();
 			public IToolCollection Tools => throw new NotImplementedException();
 
-			public event Action? OnAccessTokenStateChanged;
+			public event Action? OnAccessTokenStateChanged
+			{
+				add { }
+				remove { }
+			}
 
 			public Task<bool> LoginAsync(bool allowLogin, CancellationToken cancellationToken)
 				=> throw new NotImplementedException();
