Horde: Add a flag for showing a tool in the Horde tray app.

diff --git a/Engine/Source/Programs/Horde/Plugins/Tools/HordeServer.Tools/Commands/Generate/ToolDataCommand.cs b/Engine/Source/Programs/Horde/Plugins/Tools/HordeServer.Tools/Commands/Generate/ToolDataCommand.cs
--- a/Engine/Source/Programs/Horde/Plugins/Tools/HordeServer.Tools/Commands/Generate/ToolDataCommand.cs
+++ b/Engine/Source/Programs/Horde/Plugins/Tools/HordeServer.Tools/Commands/Generate/ToolDataCommand.cs
@@ -46,6 +46,10 @@ namespace HordeServer.Commands.Generate
 		[Description("Shows the tool for download in UGS")]
 		public bool ShowInUgs { get; set; }
 
+		[CommandLine]
+		[Description("Shows the tool for download in the Horde Launcher")]
+		public bool ShowInLauncher { get; set; }
+
 		[CommandLine("-ShowInDashboard=")]
 		[Description("Shows the tool for download on the dashboard")]
 		public bool ShowInDashboard { get; set; } = true;
@@ -126,6 +130,10 @@ namespace HordeServer.Commands.Generate
 				{
 					bundledTool[nameof(BundledToolConfig.Public)] = true;
 				}
+				if (ShowInLauncher)
+				{
+					bundledTool[nameof(BundledToolConfig.ShowInLauncher)] = true;
+				}
 				if (ShowInUgs)
 				{
 					bundledTool[nameof(BundledToolConfig.ShowInUgs)] = true;
diff --git a/Engine/Source/Programs/Horde/Plugins/Tools/HordeServer.Tools/Tools/ToolCollection.cs b/Engine/Source/Programs/Horde/Plugins/Tools/HordeServer.Tools/Tools/ToolCollection.cs
--- a/Engine/Source/Programs/Horde/Plugins/Tools/HordeServer.Tools/Tools/ToolCollection.cs
+++ b/Engine/Source/Programs/Horde/Plugins/Tools/HordeServer.Tools/Tools/ToolCollection.cs
@@ -104,6 +104,7 @@ namespace HordeServer.Tools
 			public bool Bundled => _config is BundledToolConfig;
 			public bool ShowInUgs => _config.ShowInUgs;
 			public bool ShowInDashboard => _config.ShowInDashboard;
+			public bool ShowInLauncher => _config.ShowInLauncher;
 			public IReadOnlyList<IToolDeployment> Deployments => _deployments;
 
 			public Tool(ToolCollection collection, ToolDocument document, ToolConfig config, DateTime utcNow)
diff --git a/Engine/Source/Programs/Horde/Plugins/Tools/HordeServer.Tools/Tools/ToolConfig.cs b/Engine/Source/Programs/Horde/Plugins/Tools/HordeServer.Tools/Tools/ToolConfig.cs
--- a/Engine/Source/Programs/Horde/Plugins/Tools/HordeServer.Tools/Tools/ToolConfig.cs
+++ b/Engine/Source/Programs/Horde/Plugins/Tools/HordeServer.Tools/Tools/ToolConfig.cs
@@ -70,6 +70,11 @@ namespace HordeServer.Tools
 		/// </summary>
 		public bool ShowInDashboard { get; set; } = true;
 
+		/// <summary>
+		/// Whether to show this tool for download in the Horde launcher
+		/// </summary>
+		public bool ShowInLauncher { get; set; }
+
 		/// <summary>
 		/// Default namespace for new deployments of this tool
 		/// </summary>
diff --git a/Engine/Source/Programs/Horde/Plugins/Tools/HordeServer.Tools/Tools/ToolsController.cs b/Engine/Source/Programs/Horde/Plugins/Tools/HordeServer.Tools/Tools/ToolsController.cs
--- a/Engine/Source/Programs/Horde/Plugins/Tools/HordeServer.Tools/Tools/ToolsController.cs
+++ b/Engine/Source/Programs/Horde/Plugins/Tools/HordeServer.Tools/Tools/ToolsController.cs
@@ -204,7 +204,7 @@ namespace HordeServer.Tools
 		static GetToolSummaryResponse CreateGetToolSummaryResponse(ITool tool)
 		{
 			IToolDeployment? deployment = (tool.Deployments.Count == 0) ? null : tool.Deployments[^1];
-			return new GetToolSummaryResponse(tool.Id, tool.Name, tool.Description, tool.Category, tool.Group, tool.Platforms?.ToList(), deployment?.Version, deployment?.Id, deployment?.State, deployment?.Progress, tool.Bundled, tool.ShowInUgs, tool.ShowInDashboard);
+			return new GetToolSummaryResponse(tool.Id, tool.Name, tool.Description, tool.Category, tool.Group, tool.Platforms?.ToList(), deployment?.Version, deployment?.Id, deployment?.State, deployment?.Progress, tool.Bundled, tool.ShowInUgs, tool.ShowInDashboard, tool.ShowInLauncher);
 		}
 
 		/// <summary>
@@ -248,7 +248,7 @@ namespace HordeServer.Tools
 
 		static GetToolResponse CreateGetToolResponse(ITool tool, List<GetToolDeploymentResponse> deployments)
 		{
-			return new GetToolResponse(tool.Id, tool.Name, tool.Description, tool.Category, tool.Group, tool.Platforms?.ToList(), deployments, tool.Public, tool.Bundled, tool.ShowInUgs, tool.ShowInDashboard);
+			return new GetToolResponse(tool.Id, tool.Name, tool.Description, tool.Category, tool.Group, tool.Platforms?.ToList(), deployments, tool.Public, tool.Bundled, tool.ShowInUgs, tool.ShowInDashboard, tool.ShowInLauncher);
 		}
 
 		/// <summary>
diff --git a/Engine/Source/Programs/Shared/EpicGames.Horde/Tools/HttpToolCollection.cs b/Engine/Source/Programs/Shared/EpicGames.Horde/Tools/HttpToolCollection.cs
--- a/Engine/Source/Programs/Shared/EpicGames.Horde/Tools/HttpToolCollection.cs
+++ b/Engine/Source/Programs/Shared/EpicGames.Horde/Tools/HttpToolCollection.cs
@@ -32,6 +32,7 @@ namespace EpicGames.Horde.Tools
 			bool ITool.Bundled => _response.Bundled;
 			bool ITool.ShowInUgs => _response.ShowInUgs;
 			bool ITool.ShowInDashboard => _response.ShowInDashboard;
+			bool ITool.ShowInLauncher => _response.ShowInLauncher;
 			IReadOnlyList<IToolDeployment> ITool.Deployments => _deployments;
 
 			public Tool(HttpToolCollection collection, GetToolResponse response)
diff --git a/Engine/Source/Programs/Shared/EpicGames.Horde/Tools/ITool.cs b/Engine/Source/Programs/Shared/EpicGames.Horde/Tools/ITool.cs
--- a/Engine/Source/Programs/Shared/EpicGames.Horde/Tools/ITool.cs
+++ b/Engine/Source/Programs/Shared/EpicGames.Horde/Tools/ITool.cs
@@ -70,6 +70,11 @@ namespace EpicGames.Horde.Tools
 		/// </summary>
 		bool ShowInDashboard { get; }
 
+		/// <summary>
+		/// Whether to show the tool for download in the launcher
+		/// </summary>
+		bool ShowInLauncher { get; }
+
 		/// <summary>
 		/// Current deployments of this tool, sorted by time.
 		/// </summary>
diff --git a/Engine/Source/Programs/Shared/EpicGames.Horde/Tools/ToolMessages.cs b/Engine/Source/Programs/Shared/EpicGames.Horde/Tools/ToolMessages.cs
--- a/Engine/Source/Programs/Shared/EpicGames.Horde/Tools/ToolMessages.cs
+++ b/Engine/Source/Programs/Shared/EpicGames.Horde/Tools/ToolMessages.cs
@@ -21,7 +21,8 @@ namespace EpicGames.Horde.Tools
 	/// <param name="Bundled">Whether this tool is bundled with the server</param>
 	/// <param name="ShowInUgs">Whether to show this tool for download inside UGS</param>
 	/// <param name="ShowInDashboard">Whether to show this tool for download on the dashboard</param>
-	public record class GetToolResponse(ToolId Id, string Name, string Description, string? Category, string? Group, List<string>? Platforms, List<GetToolDeploymentResponse> Deployments, bool Public, bool Bundled, bool ShowInUgs, bool ShowInDashboard);
+	/// <param name="ShowInLauncher">Whether to show this tool for download in the launcher</param>
+	public record class GetToolResponse(ToolId Id, string Name, string Description, string? Category, string? Group, List<string>? Platforms, List<GetToolDeploymentResponse> Deployments, bool Public, bool Bundled, bool ShowInUgs, bool ShowInDashboard, bool ShowInLauncher);
 
 	/// <summary>
 	/// Summary for a particular tool.
@@ -39,7 +40,8 @@ namespace EpicGames.Horde.Tools
 	/// <param name="Bundled">Whether this tool is bundled with the server</param>
 	/// <param name="ShowInUgs">Whether to show this tool for download inside UGS</param>
 	/// <param name="ShowInDashboard">Whether to show this tool for download on the dashboard</param>
-	public record class GetToolSummaryResponse(ToolId Id, string Name, string Description, string? Category, string? Group, List<string>? Platforms, string? Version, ToolDeploymentId? DeploymentId, ToolDeploymentState? DeploymentState, double? DeploymentProgress, bool Bundled, bool ShowInUgs, bool ShowInDashboard);
+	/// <param name="ShowInLauncher">Whether to show this tool for download in the launcher</param>
+	public record class GetToolSummaryResponse(ToolId Id, string Name, string Description, string? Category, string? Group, List<string>? Platforms, string? Version, ToolDeploymentId? DeploymentId, ToolDeploymentState? DeploymentState, double? DeploymentProgress, bool Bundled, bool ShowInUgs, bool ShowInDashboard, bool ShowInLauncher);
 
 	/// <summary>
 	/// Response when querying all tools
