Added an additional sanitise case for the PerforceMetadataLogger when trying to get file annotations.
We could potentially handle this in LogEventParser AddLine but I wanted to limit the impact of these changes.

diff --git a/Engine/Source/Programs/Shared/EpicGames.Perforce/PerforceMetadataLogger.cs b/Engine/Source/Programs/Shared/EpicGames.Perforce/PerforceMetadataLogger.cs
--- a/Engine/Source/Programs/Shared/EpicGames.Perforce/PerforceMetadataLogger.cs
+++ b/Engine/Source/Programs/Shared/EpicGames.Perforce/PerforceMetadataLogger.cs
@@ -221,9 +221,12 @@ namespace EpicGames.Perforce
 
 				file ??= Encoding.UTF8.GetString(text);
 
+				file = file.Replace(@"\\", Path.DirectorySeparatorChar.ToString(), StringComparison.Ordinal);
+				file = file.Replace('\\', Path.DirectorySeparatorChar);
+
 				foreach (ClientView client in _clients)
 				{
-					FileReference location = FileReference.Combine(client.BaseDir, file.Replace('\\', Path.DirectorySeparatorChar));
+					FileReference location = FileReference.Combine(client.BaseDir, file);
 					if (location.IsUnderDirectory(client.BaseDir))
 					{
 						string relativePath = location.MakeRelativeTo(client.BaseDir).Replace('\\', '/');
