Horde: Add information about processes with a file locked when it cannot be overwritten when extracting data from a bundle.

diff --git a/Engine/Source/Programs/Shared/EpicGames.Core/FileUtils.cs b/Engine/Source/Programs/Shared/EpicGames.Core/FileUtils.cs
--- a/Engine/Source/Programs/Shared/EpicGames.Core/FileUtils.cs
+++ b/Engine/Source/Programs/Shared/EpicGames.Core/FileUtils.cs
@@ -163,6 +163,38 @@ namespace EpicGames.Core
 			}
 		}
 
+		/// <summary>
+		/// Gets information about the processes that have a file locked
+		/// </summary>
+		/// <param name="fileName">Filename to check</param>
+		/// <returns>String containing locking information</returns>
+		public static string? GetFileLockInfo(string fileName)
+		{
+			if (RuntimePlatform.IsWindows)
+			{
+				List<FileLockInfoWin32>? lockInfoList;
+				try
+				{
+					lockInfoList = GetFileLockInfoWin32(fileName);
+				}
+				catch
+				{
+					lockInfoList = null;
+				}
+
+				if (lockInfoList != null && lockInfoList.Count > 0)
+				{
+					StringBuilder message = new StringBuilder("Processes with open handles to file:");
+					foreach (FileLockInfoWin32 lockInfo in lockInfoList)
+					{
+						message.Append($"\n  {lockInfo}");
+					}
+					return message.ToString();
+				}
+			}
+			return null;
+		}
+
 		/// <summary>
 		/// Deletes a file, whether it's read-only or not
 		/// </summary>
@@ -413,7 +445,7 @@ namespace EpicGames.Core
 		{
 			stat64_t stat = new stat64_t();
 			int result = stat64(fileName, stat);
-			return (result >= 0)? stat.st_mode : -1;
+			return (result >= 0) ? stat.st_mode : -1;
 		}
 
 		/// <summary>
@@ -435,7 +467,7 @@ namespace EpicGames.Core
 		{
 			stat64_linux_t stat = new stat64_linux_t();
 			int result = stat64_linux(1, fileName, stat);
-			return (result >= 0)? (int)stat.st_mode : -1;
+			return (result >= 0) ? (int)stat.st_mode : -1;
 		}
 
 		/// <summary>
@@ -509,7 +541,7 @@ namespace EpicGames.Core
 			const string RawPathPrefix = "\\\\?\\";
 			if (!dirName.StartsWith(RawPathPrefix, StringComparison.Ordinal))
 			{
-				dirName = RawPathPrefix + dirName;				
+				dirName = RawPathPrefix + dirName;
 			}
 
 			IntPtr hFind = FindFirstFileW(dirName + "\\*", ref findData);
@@ -709,7 +741,7 @@ namespace EpicGames.Core
 
 		[DllImport("kernel32.dll", SetLastError = true)]
 #pragma warning disable CA1838 // Avoid 'StringBuilder' parameters for P/Invokes
-		static extern bool QueryFullProcessImageName([In]SafeProcessHandle hProcess, [In]int dwFlags, [Out]StringBuilder lpExeName, ref int lpdwSize);
+		static extern bool QueryFullProcessImageName([In] SafeProcessHandle hProcess, [In] int dwFlags, [Out] StringBuilder lpExeName, ref int lpdwSize);
 #pragma warning restore CA1838 // Avoid 'StringBuilder' parameters for P/Invokes
 #pragma warning restore IDE1006
 
@@ -861,10 +893,10 @@ namespace EpicGames.Core
 		};
 
 		/* stat tends to get compiled to another symbol and libc doesnt directly have that entry point */
-		[DllImport("libc", EntryPoint="__xstat64")]
+		[DllImport("libc", EntryPoint = "__xstat64")]
 		static extern int stat64_linux(int ver, string pathname, stat64_linux_t stat);
 
-		[DllImport("libc", EntryPoint="chmod")]
+		[DllImport("libc", EntryPoint = "chmod")]
 		static extern int chmod_linux(string path, ushort mode);
 
 #pragma warning restore IDE1006
diff --git a/Engine/Source/Programs/Shared/EpicGames.Horde/Storage/Nodes/DirectoryNode.Extract.cs b/Engine/Source/Programs/Shared/EpicGames.Horde/Storage/Nodes/DirectoryNode.Extract.cs
--- a/Engine/Source/Programs/Shared/EpicGames.Horde/Storage/Nodes/DirectoryNode.Extract.cs
+++ b/Engine/Source/Programs/Shared/EpicGames.Horde/Storage/Nodes/DirectoryNode.Extract.cs
@@ -212,7 +212,23 @@ namespace EpicGames.Horde.Storage.Nodes
 					FileStream? stream = null;
 					try
 					{
-						stream = FileInfo.Open(FileMode.OpenOrCreate, FileAccess.ReadWrite, FileShare.ReadWrite);
+						try
+						{
+							stream = FileInfo.Open(FileMode.OpenOrCreate, FileAccess.ReadWrite, FileShare.ReadWrite);
+						}
+						catch (IOException ex)
+						{
+							string? lockInfo = FileUtils.GetFileLockInfo(FileInfo.FullName);
+							if (lockInfo == null)
+							{
+								throw;
+							}
+							else
+							{
+								throw new WrappedFileOrDirectoryException(ex, $"{ex.Message}\n{lockInfo}");
+							}
+						}
+
 						if (!_createdFile)
 						{
 							stream.SetLength(FileEntry.Length);
