Fixed a bug where issues would have their FixCommitId set to 0 rather than cleared after re-assignment to a new owner.
Note we set fixChange to 0 in the onAssign function in IssueViewV2.tsx.
I now check whether fixChange is == 0 and if so set FixCommitId to be empty, this causes the FIX CL to be cleared.

diff --git a/Engine/Source/Programs/Shared/EpicGames.Horde/Issues/IssueMessages.cs b/Engine/Source/Programs/Shared/EpicGames.Horde/Issues/IssueMessages.cs
--- a/Engine/Source/Programs/Shared/EpicGames.Horde/Issues/IssueMessages.cs
+++ b/Engine/Source/Programs/Shared/EpicGames.Horde/Issues/IssueMessages.cs
@@ -732,7 +732,7 @@ namespace EpicGames.Horde.Issues
 		/// </summary>
 		public CommitId? FixCommitId
 		{
-			get => _fixCommitId ?? ((_fixChange < 0)? null : CommitId.FromPerforceChange(_fixChange));
+			get => _fixCommitId ?? ((_fixChange < 0)? null : (_fixChange == 0)? CommitId.Empty : CommitId.FromPerforceChange(_fixChange));
 			set => _fixCommitId = value;
 		}
 		CommitId? _fixCommitId;
