Changed "static const FName" back to "static FName".  The change was causing issues with live coding.

#rnx

diff --git a/Engine/Source/Programs/Shared/EpicGames.UHT/Exporters/CodeGen/UhtHeaderCodeGeneratorCppFile.cs b/Engine/Source/Programs/Shared/EpicGames.UHT/Exporters/CodeGen/UhtHeaderCodeGeneratorCppFile.cs
--- a/Engine/Source/Programs/Shared/EpicGames.UHT/Exporters/CodeGen/UhtHeaderCodeGeneratorCppFile.cs
+++ b/Engine/Source/Programs/Shared/EpicGames.UHT/Exporters/CodeGen/UhtHeaderCodeGeneratorCppFile.cs
@@ -1152,7 +1152,8 @@ namespace EpicGames.UHT.Exporters.CodeGen
 
 			if (!isInterfaceClass)
 			{
-				builder.Append("static const FName NAME_").Append(classObj.SourceName).Append('_').Append(function.EngineName).Append(" = FName(TEXT(\"").Append(function.EngineName).Append("\"));\r\n");
+				// Do not make this a static const.  It causes issues with live coding
+				builder.Append("static FName NAME_").Append(classObj.SourceName).Append('_').Append(function.EngineName).Append(" = FName(TEXT(\"").Append(function.EngineName).Append("\"));\r\n");
 			}
 
 			AppendNativeFunctionHeader(builder, function, UhtPropertyTextType.EventFunctionArgOrRetVal, false, null, null, UhtFunctionExportFlags.None, 0, "\r\n");
@@ -1891,7 +1892,8 @@ namespace EpicGames.UHT.Exporters.CodeGen
 				{
 					if (property.PropertyFlags.HasAnyFlags(EPropertyFlags.Net))
 					{
-						builder.Append("\tstatic const FName Name_").Append(property.SourceName).Append("(TEXT(\"").Append(property.SourceName).Append("\"));\r\n");
+						// Do not make this a static const.  It causes issues with live coding
+						builder.Append("\tstatic FName Name_").Append(property.SourceName).Append("(TEXT(\"").Append(property.SourceName).Append("\"));\r\n");
 					}
 				}
 				builder.Append("\tconst bool bIsValid = true");
