Updating UHT to latest code gen to support Verse

#rnx

diff --git a/Engine/Source/Programs/Shared/EpicGames.UHT/Exporters/CodeGen/UhtHeaderCodeGeneratorHFile.cs b/Engine/Source/Programs/Shared/EpicGames.UHT/Exporters/CodeGen/UhtHeaderCodeGeneratorHFile.cs
--- a/Engine/Source/Programs/Shared/EpicGames.UHT/Exporters/CodeGen/UhtHeaderCodeGeneratorHFile.cs
+++ b/Engine/Source/Programs/Shared/EpicGames.UHT/Exporters/CodeGen/UhtHeaderCodeGeneratorHFile.cs
@@ -427,8 +427,9 @@ namespace EpicGames.UHT.Exporters.CodeGen
 				// Verse declarations
 				if (scriptStruct.IsVerseField)
 				{
+					// TODO: When VNI gets deprecated, V_DECLARE_STRUCT_RUNTIME_TRAITS will be needed for both VMs and not just BPVM.
 					builder.Append("#if WITH_VERSE_BPVM\r\n");
-					builder.Append("VUHT_DECLARE_STRUCT_RUNTIME_TRAITS(struct ").Append(scriptStruct.SourceName).Append(")\r\n");
+					builder.Append("V_DECLARE_STRUCT_RUNTIME_TRAITS(, ").Append(scriptStruct.SourceName).Append(", \"").AppendVerseUEPackageName(scriptStruct).Append("\", \"").AppendVerseName(scriptStruct).Append("\")\r\n");
 					builder.Append("#endif // WITH_VERSE_BPVM\r\n");
 					builder.Append("\r\n");
 				}
diff --git a/Engine/Source/Programs/Shared/EpicGames.UHT/Specifiers/UhtFieldSpecifiers.cs b/Engine/Source/Programs/Shared/EpicGames.UHT/Specifiers/UhtFieldSpecifiers.cs
--- a/Engine/Source/Programs/Shared/EpicGames.UHT/Specifiers/UhtFieldSpecifiers.cs
+++ b/Engine/Source/Programs/Shared/EpicGames.UHT/Specifiers/UhtFieldSpecifiers.cs
@@ -18,12 +18,13 @@ namespace EpicGames.UHT.Parsers
 	[SuppressMessage("CodeQuality", "IDE0051:Remove unused private members", Justification = "Attribute accessed method")]
 	public static class UhtFieldSpecifiers
 	{
-		[UhtSpecifier(Extends = UhtTableNames.Struct, ValueType = UhtSpecifierValueType.KeyValuePairList)]
-		[UhtSpecifier(Extends = UhtTableNames.Class, ValueType = UhtSpecifierValueType.KeyValuePairList)]
+		[UhtSpecifier(Extends = UhtTableNames.Struct, ValueType = UhtSpecifierValueType.OptionalEqualsKeyValuePairList)]
+		[UhtSpecifier(Extends = UhtTableNames.Class, ValueType = UhtSpecifierValueType.OptionalEqualsKeyValuePairList)]
 		private static void VerseSpecifier(UhtSpecifierContext specifierContext, List<KeyValuePair<StringView, StringView>> values)
 		{
 			UhtField fieldObj = (UhtField)specifierContext.Type;
 
+			// Extract all the elements of the meta data
 			foreach (KeyValuePair<StringView, StringView> kvp in values)
 			{
 				ReadOnlySpan<char> key = kvp.Key.Span;
@@ -41,6 +42,35 @@ namespace EpicGames.UHT.Parsers
 				}
 			}
 
+			// If no name was specified, generate one using the following algorithm FVerseXxYyZz -> xx_yy_zz
+			if (String.IsNullOrEmpty(fieldObj.VerseName))
+			{
+				if (fieldObj.EngineName.StartsWith("Verse", StringComparison.Ordinal))
+				{
+					using BorrowStringBuilder borrowStringBuilder = new BorrowStringBuilder(StringBuilderCache.Small);
+					ReadOnlySpan<char> span = fieldObj.EngineName.AsSpan(5);
+					bool first = true;
+					foreach (char c in span)
+					{
+						if (c >= 'A' && c <= 'Z')
+						{
+							if (!first)
+							{
+								borrowStringBuilder.StringBuilder.Append('_');
+							}
+							borrowStringBuilder.StringBuilder.Append(Char.ToLower(c));
+						}
+						else
+						{
+							borrowStringBuilder.StringBuilder.Append(c);
+						}
+						first = false;
+					}
+					fieldObj.VerseName = borrowStringBuilder.StringBuilder.ToString();
+				}
+			}
+
+			// Final validation
 			if (!fieldObj.IsVerseField)
 			{
 				fieldObj.LogError($"Verse specifier must include the name");
diff --git a/Engine/Source/Programs/Shared/EpicGames.UHT/Types/UhtField.cs b/Engine/Source/Programs/Shared/EpicGames.UHT/Types/UhtField.cs
--- a/Engine/Source/Programs/Shared/EpicGames.UHT/Types/UhtField.cs
+++ b/Engine/Source/Programs/Shared/EpicGames.UHT/Types/UhtField.cs
@@ -61,5 +61,37 @@ namespace EpicGames.UHT.Types
 			UhtModule module = fieldObj.Module;
 			return builder.Append('/').Append(module.Module.VerseMountPoint).Append("/_Verse/VNI/").Append(module.Module.Name);
 		}
+
+		/// <summary>
+		/// Append the Verse UE package name (without a leading forward slash)
+		/// </summary>
+		/// <param name="builder">Destination builder</param>
+		/// <param name="fieldObj">Field to serialize</param>
+		/// <returns>Builder</returns>
+		public static StringBuilder AppendVerseUEPackageName(this StringBuilder builder, UhtField fieldObj)
+		{
+			if (!fieldObj.IsVerseField)
+			{
+				throw new UhtException(fieldObj, "Attempt to write the Verse package name on a field that isn't part of Verse");
+			}
+			UhtModule module = fieldObj.Module;
+			return builder.Append(module.Module.VerseMountPoint).Append('/').Append(module.Module.Name);
+		}
+
+		/// <summary>
+		/// Append the full verse name of the field
+		/// </summary>
+		/// <param name="builder">Destination builder</param>
+		/// <param name="fieldObj">Field to serialize</param>
+		/// <returns>Builder</returns>
+		public static StringBuilder AppendVerseName(this StringBuilder builder, UhtField fieldObj)
+		{
+			if (!fieldObj.IsVerseField)
+			{
+				throw new UhtException(fieldObj, "Attempt to write the Verse name on a field that isn't part of Verse");
+			}
+			UhtModule module = fieldObj.Module;
+			return builder.Append('(').Append(module.Module.VersePath).Append(":)").Append(fieldObj.VerseName);
+		}
 	}
 }
