Horde: Use a simpler pipeline for managing the set of async tasks that are sweeping blobs for GC.

diff --git a/Engine/Source/Programs/Horde/Plugins/Storage/HordeServer.Storage/Storage/StorageService.cs b/Engine/Source/Programs/Horde/Plugins/Storage/HordeServer.Storage/Storage/StorageService.cs
--- a/Engine/Source/Programs/Horde/Plugins/Storage/HordeServer.Storage/Storage/StorageService.cs
+++ b/Engine/Source/Programs/Horde/Plugins/Storage/HordeServer.Storage/Storage/StorageService.cs
@@ -4,7 +4,7 @@ using System.Buffers;
 using System.ComponentModel;
 using System.Diagnostics;
 using System.Linq.Expressions;
-using System.Threading.Tasks.Dataflow;
+using System.Threading.Channels;
 using EpicGames.Core;
 using EpicGames.Horde.Storage;
 using EpicGames.Horde.Storage.Bundles;
@@ -1123,33 +1123,27 @@ namespace HordeServer.Storage
 						break;
 					}
 
-					// Create a dataflow pipeline for deleting blobs
-					BufferBlock<ObjectId> sourceBlock = new BufferBlock<ObjectId>(
-						new DataflowBlockOptions { CancellationToken = cancellationToken });
+					// Create a pipeline for deleting blobs
+					Channel<ObjectId> channel = Channel.CreateUnbounded<ObjectId>();
 
-					ActionBlock<ObjectId> deleteBlock = new ActionBlock<ObjectId>(blobInfoId => RemoveBlobIfUnreferencedAsync(namespaceInfo, blobInfoId, sweepState, storageConfig, cancellationToken),
-						new ExecutionDataflowBlockOptions { MaxDegreeOfParallelism = 4, CancellationToken = cancellationToken });
+					await using AsyncPipeline pipeline = new AsyncPipeline(cancellationToken);
+					pipeline.AddTasks<ObjectId>(4, channel.Reader, (blobInfoId, ctx) => RemoveBlobIfUnreferencedAsync(namespaceInfo, blobInfoId, sweepState, storageConfig, ctx));
 
-					sourceBlock.LinkTo(deleteBlock,
-						new DataflowLinkOptions { PropagateCompletion = true });
-
-					// Push items into the pipeline
 					RedisValue[] entries = await _redisService.GetDatabase().SortedSetRangeByRankAsync(checkSet, 0, batchSize);
-					foreach (RedisValue entry in entries)
+					foreach(RedisValue entry in entries)
 					{
 						ObjectId blobInfoId = new ObjectId(((byte[]?)entry)!);
 						if (blobInfoId <= lastImportBlobInfoId)
 						{
-							await sourceBlock.SendAsync(blobInfoId, cancellationToken);
+							await channel.Writer.WriteAsync(blobInfoId, cancellationToken);
 						}
 					}
 
-					// Mark the dataflow as complete and wait for it to flush
-					sourceBlock.Complete();
-					await deleteBlock.Completion;
+					channel.Writer.TryComplete();
+					await pipeline.WaitForCompletionAsync();
 				}
 
-				// Update the 
+				// Update the GC timestamp
 				await _gcState.UpdateAsync(state => state.FindOrAddNamespace(namespaceInfo.Id).LastTime = utcNow, cancellationToken);
 				_logger.LogInformation("Finished garbage collection for namespace {NamespaceId} in {TimeSecs}s ({NumItems} removed)", namespaceInfo.Id, timer.Elapsed.TotalSeconds, sweepState.NumRemovedItems);
 			}
diff --git a/Engine/Source/Programs/Shared/EpicGames.Core/AsyncPipeline.cs b/Engine/Source/Programs/Shared/EpicGames.Core/AsyncPipeline.cs
new file mode 100644
--- /dev/null
+++ b/Engine/Source/Programs/Shared/EpicGames.Core/AsyncPipeline.cs
@@ -0,0 +1,151 @@
+// Copyright Epic Games, Inc. All Rights Reserved.
+
+using System;
+using System.Collections.Generic;
+using System.Runtime.ExceptionServices;
+using System.Threading;
+using System.Threading.Channels;
+using System.Threading.Tasks;
+
+namespace EpicGames.Core
+{
+	/// <summary>
+	/// Manages a set of async tasks forming a pipeline, which can be cancelled and awaited as a complete unit. The first exception thrown by an individual task is captured and reported back to the main thread.
+	/// </summary>
+	public class AsyncPipeline : IAsyncDisposable
+	{
+		readonly List<Task> _tasks;
+		readonly CancellationToken _cancellationToken;
+		readonly CancellationTokenSource _cancellationSource;
+
+		ExceptionDispatchInfo? _exceptionDispatchInfo;
+
+		/// <summary>
+		/// Tests whether the pipeline has failed
+		/// </summary>
+		public bool IsFaulted
+			=> _exceptionDispatchInfo != null;
+
+		/// <summary>
+		/// Constructor
+		/// </summary>
+		/// <param name="cancellationToken">Cancellation token for the pipeline</param>
+		public AsyncPipeline(CancellationToken cancellationToken)
+		{
+			_tasks = new List<Task>();
+			_cancellationToken = cancellationToken;
+			_cancellationSource = CancellationTokenSource.CreateLinkedTokenSource(cancellationToken);
+		}
+
+		/// <inheritdoc/>
+		public async ValueTask DisposeAsync()
+		{
+			await _cancellationSource.CancelAsync();
+			await Task.WhenAll(_tasks);
+		}
+
+		/// <summary>
+		/// Adds a new task to the pipeline.
+		/// </summary>
+		/// <param name="taskFunc">Method to execute</param>
+		public void AddTask(Func<CancellationToken, Task> taskFunc)
+		{
+			_tasks.Add(Task.Run(() => RunGuarded(taskFunc), _cancellationSource.Token));
+		}
+
+		/// <summary>
+		/// Adds several tasks to the pipeline.
+		/// </summary>
+		/// <param name="count">Number of tasks to run</param>
+		/// <param name="taskFunc">Method to execute</param>
+		public void AddTasks(int count, Func<CancellationToken, Task> taskFunc)
+		{
+			for (int idx = 0; idx < count; idx++)
+			{
+				AddTask(taskFunc);
+			}
+		}
+
+		async Task RunGuarded(Func<CancellationToken, Task> taskFunc)
+		{
+			try
+			{
+				await taskFunc(_cancellationSource.Token);
+			}
+			catch (OperationCanceledException)
+			{
+				// Ignore
+			}
+			catch (Exception ex)
+			{
+				if (_exceptionDispatchInfo == null)
+				{
+					ExceptionDispatchInfo dispatchInfo = ExceptionDispatchInfo.Capture(ex);
+					Interlocked.CompareExchange(ref _exceptionDispatchInfo, dispatchInfo, null);
+				}
+				await _cancellationSource.CancelAsync();
+			}
+		}
+
+		/// <summary>
+		/// Waits for all tasks to complete, and throws any exceptions 
+		/// </summary>
+		public async Task WaitForCompletionAsync()
+		{
+			await Task.WhenAll(_tasks);
+
+			_cancellationToken.ThrowIfCancellationRequested();
+
+			if (_exceptionDispatchInfo != null)
+			{
+				_exceptionDispatchInfo.Throw();
+			}
+		}
+	}
+
+	/// <summary>
+	/// Extension methods for async pipelines
+	/// </summary>
+	public static class AsyncPipelineExtensions
+	{
+		/// <summary>
+		/// Adds a worker to process items from a channel
+		/// </summary>
+		/// <typeparam name="T">Item type</typeparam>
+		/// <param name="pipeline">Pipeline to add the worker to</param>
+		/// <param name="reader">Source for the items</param>
+		/// <param name="taskFunc">Action to execute for each item</param>
+		public static void AddTask<T>(this AsyncPipeline pipeline, ChannelReader<T> reader, Func<T, CancellationToken, Task> taskFunc)
+		{
+			pipeline.AddTask(ctx => ProcessItemsAsync(reader, taskFunc, ctx));
+		}
+
+		/// <summary>
+		/// Adds a worker to process items from a channel
+		/// </summary>
+		/// <typeparam name="T">Item type</typeparam>
+		/// <param name="pipeline">Pipeline to add the worker to</param>
+		/// <param name="count">Number of workers to add</param>
+		/// <param name="reader">Source for the items</param>
+		/// <param name="taskFunc">Action to execute for each item</param>
+		public static void AddTasks<T>(this AsyncPipeline pipeline, int count, ChannelReader<T> reader, Func<T, CancellationToken, Task> taskFunc)
+		{
+			for (int idx = 0; idx < count; idx++)
+			{
+				AddTask(pipeline, reader, taskFunc);
+			}
+		}
+
+		static async Task ProcessItemsAsync<T>(ChannelReader<T> reader, Func<T, CancellationToken, Task> taskFunc, CancellationToken cancellationToken)
+		{
+			while (await reader.WaitToReadAsync(cancellationToken))
+			{
+				T? item;
+				if (reader.TryRead(out item))
+				{
+					await taskFunc(item, cancellationToken);
+				}
+			}
+		}
+	}
+}
