EpicGames.Core: Fix static analysis warnings.

#rnx

diff --git a/Engine/Source/Programs/Shared/EpicGames.Core/AsyncPipeline.cs b/Engine/Source/Programs/Shared/EpicGames.Core/AsyncPipeline.cs
--- a/Engine/Source/Programs/Shared/EpicGames.Core/AsyncPipeline.cs
+++ b/Engine/Source/Programs/Shared/EpicGames.Core/AsyncPipeline.cs
@@ -12,12 +12,11 @@ namespace EpicGames.Core
 	/// <summary>
 	/// Manages a set of async tasks forming a pipeline, which can be cancelled and awaited as a complete unit. The first exception thrown by an individual task is captured and reported back to the main thread.
 	/// </summary>
-	public class AsyncPipeline : IAsyncDisposable
+	public sealed class AsyncPipeline : IAsyncDisposable
 	{
 		readonly List<Task> _tasks;
 		readonly CancellationToken _cancellationToken;
-		readonly CancellationTokenSource _cancellationSource;
-
+		CancellationTokenSource _cancellationSource;
 		ExceptionDispatchInfo? _exceptionDispatchInfo;
 
 		/// <summary>
@@ -40,8 +39,22 @@ namespace EpicGames.Core
 		/// <inheritdoc/>
 		public async ValueTask DisposeAsync()
 		{
-			await _cancellationSource.CancelAsync();
-			await Task.WhenAll(_tasks);
+			if (_cancellationSource != null)
+			{
+				await _cancellationSource.CancelAsync();
+			}
+
+			if (_tasks.Count > 0)
+			{
+				await Task.WhenAll(_tasks);
+				_tasks.Clear();
+			}
+
+			if(_cancellationSource != null)
+			{
+				_cancellationSource.Dispose();
+				_cancellationSource = null!;
+			}
 		}
 
 		/// <summary>
@@ -50,7 +63,7 @@ namespace EpicGames.Core
 		/// <param name="taskFunc">Method to execute</param>
 		public void AddTask(Func<CancellationToken, Task> taskFunc)
 		{
-			_tasks.Add(Task.Run(() => RunGuarded(taskFunc), _cancellationSource.Token));
+			_tasks.Add(Task.Run(() => RunGuardedAsync(taskFunc), _cancellationSource.Token));
 		}
 
 		/// <summary>
@@ -66,7 +79,7 @@ namespace EpicGames.Core
 			}
 		}
 
-		async Task RunGuarded(Func<CancellationToken, Task> taskFunc)
+		async Task RunGuardedAsync(Func<CancellationToken, Task> taskFunc)
 		{
 			try
 			{
diff --git a/Engine/Source/Programs/Shared/EpicGames.Core/FileUtils.cs b/Engine/Source/Programs/Shared/EpicGames.Core/FileUtils.cs
--- a/Engine/Source/Programs/Shared/EpicGames.Core/FileUtils.cs
+++ b/Engine/Source/Programs/Shared/EpicGames.Core/FileUtils.cs
@@ -543,7 +543,7 @@ namespace EpicGames.Core
 			{
 				dirName = RawPathPrefix + dirName;
 			}
-			if (!dirName.EndsWith("\\"))
+			if (!dirName.EndsWith("\\", StringComparison.Ordinal))
 			{
 				dirName += "\\";
 			}
