Update Blake3 package

#rnx

diff --git a/Engine/Source/Programs/Shared/EpicGames.IoHash/EpicGames.IoHash.csproj b/Engine/Source/Programs/Shared/EpicGames.IoHash/EpicGames.IoHash.csproj
--- a/Engine/Source/Programs/Shared/EpicGames.IoHash/EpicGames.IoHash.csproj
+++ b/Engine/Source/Programs/Shared/EpicGames.IoHash/EpicGames.IoHash.csproj
@@ -37,7 +37,7 @@
     </Compile>
   </ItemGroup>
   <ItemGroup>
-    <PackageReference Include="Blake3" Version="0.6.1" />
+    <PackageReference Include="Blake3" Version="1.1.0" />
     <PackageReference Include="System.Memory" Version="4.5.5" />
   </ItemGroup>
   <ItemGroup>
