Horde: Avoid cleaning a workspace during sync if a previous cleanup was successful *and* same P4 client was used

diff --git a/Engine/Source/Programs/Horde/Drivers/JobDriver/Utility/WorkspaceInfo.cs b/Engine/Source/Programs/Horde/Drivers/JobDriver/Utility/WorkspaceInfo.cs
--- a/Engine/Source/Programs/Horde/Drivers/JobDriver/Utility/WorkspaceInfo.cs
+++ b/Engine/Source/Programs/Horde/Drivers/JobDriver/Utility/WorkspaceInfo.cs
@@ -341,7 +341,7 @@ namespace JobDriver.Utility
 		public async Task CleanAsync(IPerforceConnection perforceClient, CancellationToken cancellationToken)
 		{
 			await Repository.RevertAsync(perforceClient, cancellationToken);
-			await Repository.CleanAsync(RemoveUntrackedFiles, cancellationToken);
+			await Repository.CleanAsync(perforceClient, RemoveUntrackedFiles, cancellationToken);
 		}
 
 		/// <summary>
diff --git a/Engine/Source/Programs/Horde/HordeCmd/Commands/Perforce/PerforceClean.cs b/Engine/Source/Programs/Horde/HordeCmd/Commands/Perforce/PerforceClean.cs
--- a/Engine/Source/Programs/Horde/HordeCmd/Commands/Perforce/PerforceClean.cs
+++ b/Engine/Source/Programs/Horde/HordeCmd/Commands/Perforce/PerforceClean.cs
@@ -17,7 +17,7 @@ namespace Horde.Commands.Perforce
 
 		protected override Task ExecuteAsync(IPerforceConnection perforce, ManagedWorkspace repo, ILogger logger)
 		{
-			return repo.CleanAsync(!Incremental, CancellationToken.None);
+			return repo.CleanAsync(perforce, !Incremental, CancellationToken.None);
 		}
 	}
 }
diff --git a/Engine/Source/Programs/Horde/HordeCmd/Commands/Perforce/PerforceRepair.cs b/Engine/Source/Programs/Horde/HordeCmd/Commands/Perforce/PerforceRepair.cs
--- a/Engine/Source/Programs/Horde/HordeCmd/Commands/Perforce/PerforceRepair.cs
+++ b/Engine/Source/Programs/Horde/HordeCmd/Commands/Perforce/PerforceRepair.cs
@@ -12,7 +12,7 @@ namespace Horde.Commands.Perforce
 	{
 		protected override Task ExecuteAsync(IPerforceConnection perforce, ManagedWorkspace repo, ILogger logger)
 		{
-			return repo.RepairAsync(CancellationToken.None);
+			return repo.RepairCacheAsync(CancellationToken.None);
 		}
 	}
 }
diff --git a/Engine/Source/Programs/Shared/EpicGames.Perforce.Managed.Tests/ManagedWorkspaceTest.cs b/Engine/Source/Programs/Shared/EpicGames.Perforce.Managed.Tests/ManagedWorkspaceTest.cs
--- a/Engine/Source/Programs/Shared/EpicGames.Perforce.Managed.Tests/ManagedWorkspaceTest.cs
+++ b/Engine/Source/Programs/Shared/EpicGames.Perforce.Managed.Tests/ManagedWorkspaceTest.cs
@@ -41,6 +41,63 @@ public class ManagedWorkspaceTest : BasePerforceFixtureTest
 		Stream.GetChangelist(6).AssertDepotFiles(SyncDir);
 		await Stream.GetChangelist(6).AssertHaveTableAsync(PerforceConnection, useHaveTable);
 	}
+	
+	private class CleanCounter
+	{
+		public int NumCleans { get; private set; } = 0;
+		public int NumFilesDeleted { get; private set; } = 0;
+		public int NumDirsDeleted { get; private set; } = 0;
+		
+		public void OnClean(int numFiles, int numDirs)
+		{
+			NumCleans += 1;
+			NumFilesDeleted += numFiles;
+			NumDirsDeleted += numDirs;
+			Console.WriteLine($"Clean performed: numFiles={numFiles} numDirs={numDirs}");
+		}
+		
+		public CleanCounter Reset()
+		{
+			NumCleans = 0;
+			NumFilesDeleted = 0;
+			NumDirsDeleted = 0;
+			return this;
+		}
+		
+		public CleanCounter Assert(int expectedCount, int expectedNumFilesCleaned, int expectedNumDirsCleaned)
+		{
+			Microsoft.VisualStudio.TestTools.UnitTesting.Assert.AreEqual(expectedCount, NumCleans);
+			Microsoft.VisualStudio.TestTools.UnitTesting.Assert.AreEqual(expectedNumFilesCleaned, NumFilesDeleted);
+			Microsoft.VisualStudio.TestTools.UnitTesting.Assert.AreEqual(expectedNumDirsCleaned, NumDirsDeleted);
+			return this;
+		}
+	}
+	
+	[TestMethod]
+	[DataRow(true, DisplayName = "With have-table")]
+	public async Task SyncSkipCleanIfPossibleAsync(bool useHaveTable)
+	{
+		ManagedWorkspace ws = await CreateManagedWorkspaceAsync(useHaveTable);
+		await AssertHaveTableFileCountAsync(0);
+		
+		CleanCounter cc = new();
+		ws.OnClean += cc.OnClean;
+		
+		// Sync first time, clean is expected
+		await SyncAsync(ws, 6, assertFiles: true, useHaveTable: useHaveTable);
+		cc.Assert(1, 0, 0).Reset();
+		
+		// Simulate some build output inside workspace
+		await File.WriteAllTextAsync(Path.Join(SyncDir, "build-output.obj"), "somecontent");
+		
+		// Finalize and clean up is called when job finishes
+		await ws.CleanAsync(PerforceConnection, true, CancellationToken.None);
+		cc.Assert(1, 1, 0).Reset();
+		
+		// Syncing again should not result in a clean up, as it was called during finalize and we're using the same P4 client
+		await SyncAsync(ws, 7, assertFiles: true, useHaveTable: useHaveTable);
+		cc.Assert(0, 0, 0).Reset();
+	}
 
 	[TestMethod]
 	public async Task SyncBackwardsToOlderChangelistRemoveUntrackedAsync()
@@ -300,9 +357,14 @@ public class ManagedWorkspaceTest : BasePerforceFixtureTest
 		return ws;
 	}
 
-	private async Task SyncAsync(ManagedWorkspace managedWorkspace, int changeNumber, FileReference? cacheFile = null, bool removeUntracked = true)
+	private async Task SyncAsync(ManagedWorkspace managedWorkspace, int changeNumber, FileReference? cacheFile = null, bool removeUntracked = true, bool assertFiles = false, bool useHaveTable = true)
 	{
 		await managedWorkspace.SyncAsync(PerforceConnection, StreamName, changeNumber, Array.Empty<string>(), removeUntracked, false, cacheFile, CancellationToken.None);
+		if (assertFiles)
+		{
+			Stream.GetChangelist(changeNumber).AssertDepotFiles(SyncDir);
+			await Stream.GetChangelist(changeNumber).AssertHaveTableAsync(PerforceConnection, useHaveTable);
+		}
 	}
 
 	private async Task AssertHaveTableFileCountAsync(int expected)
diff --git a/Engine/Source/Programs/Shared/EpicGames.Perforce.Managed/ManagedWorkspace.cs b/Engine/Source/Programs/Shared/EpicGames.Perforce.Managed/ManagedWorkspace.cs
--- a/Engine/Source/Programs/Shared/EpicGames.Perforce.Managed/ManagedWorkspace.cs
+++ b/Engine/Source/Programs/Shared/EpicGames.Perforce.Managed/ManagedWorkspace.cs
@@ -162,6 +162,11 @@ namespace EpicGames.Perforce.Managed
 		/// The current revision number for cache archives.
 		/// </summary>
 		static int CurrentVersion { get; } = Enum.GetValues(typeof(ManagedWorkspaceVersion)).Cast<int>().Max();
+		
+		/// <summary>
+		/// File format version for untracked state file
+		/// </summary>
+		private const int UntrackedFileVersion = 1;
 
 		/// <summary>
 		/// Externally configurable options
@@ -182,6 +187,11 @@ namespace EpicGames.Perforce.Managed
 		/// Name of the main data file for a repository
 		/// </summary>
 		const string DataFileName = "Repository.dat";
+		
+		/// <summary>
+		/// Name of the untracked state file for a repository
+		/// </summary>
+		const string UntrackedStateFileName = "Untracked.dat";
 
 		/// <summary>
 		/// Name of the host
@@ -237,6 +247,16 @@ namespace EpicGames.Perforce.Managed
 		/// All the files which are currently being tracked
 		/// </summary>
 		Dictionary<FileContentId, CachedFileInfo> _contentIdToTrackedFile = new Dictionary<FileContentId, CachedFileInfo>();
+		
+		/// <summary>
+		/// Notification that a clean has been performed
+		/// </summary>
+		public delegate void OnCleanDelegate(int numFilesDeleted, int numDirsDeleted);
+		
+		/// <summary>
+		/// Occurs when a clean operation has been performed.
+		/// </summary>
+		public event OnCleanDelegate? OnClean;
 
 		/// <summary>
 		/// Constructor
@@ -422,7 +442,7 @@ namespace EpicGames.Perforce.Managed
 
 			reader.ReadWorkspaceDirectoryInfo(repo._workspace, (ManagedWorkspaceVersion)version);
 
-			await repo.RunOptionalRepairAsync(cancellationToken);
+			await repo.RunOptionalCacheRepairAsync(cancellationToken);
 			return repo;
 		}
 
@@ -454,15 +474,75 @@ namespace EpicGames.Perforce.Managed
 			await FileReference.WriteAllBytesAsync(dataFile, buffer, cancellationToken);
 			CompleteTransaction(dataFile);
 		}
-
+		
+		/// <summary>
+		/// Mark whether current repository in combination with P4 client contains untracked files
+		/// Stores data in a file separate to the repository metadata as that can grow large (hundreds of megabytes)
+		/// </summary>
+		private async Task SaveUntrackedStateAsync(IPerforceConnection connection, bool containsUntrackedFiles, CancellationToken cancellationToken)
+		{
+			if (connection.Settings.ClientName == null)
+			{
+				throw new ArgumentException("Unable to save untracked state as P4 client name is not set");
+			}
+			
+			using MemoryStream ms = new(300);
+			await using BinaryWriter bw = new(ms);
+			bw.Write(UntrackedFileVersion);
+			bw.Write(connection.Settings.ClientName);
+			bw.Write(containsUntrackedFiles);
+			
+			// Write it to disk
+			FileReference dataFile = FileReference.Combine(_baseDir, UntrackedStateFileName);
+			BeginTransaction(dataFile);
+			await FileReference.WriteAllBytesAsync(dataFile, ms.ToArray(), cancellationToken);
+			CompleteTransaction(dataFile);
+		}
+		
+		/// <summary>
+		/// Check if current repository in combination with P4 client contains untracked files
+		/// If deserialization is not possible, assume the repo contains untracked files 
+		/// </summary>
+		/// <returns>True if repository contains untracked files</returns>
+		private async Task<bool> LoadUntrackedStateAsync(IPerforceConnection connection, CancellationToken cancellationToken)
+		{
+			try
+			{
+				FileReference dataFile = FileReference.Combine(_baseDir, UntrackedStateFileName);
+				if (!File.Exists(dataFile.FullName))
+				{
+					return true;
+				}
+				
+				byte[] data = await File.ReadAllBytesAsync(dataFile.FullName, cancellationToken);
+				using MemoryStream ms = new (data);
+				using BinaryReader br = new (ms);
+				
+				int version = br.ReadInt32();
+				if (version != UntrackedFileVersion)
+				{
+					return true;
+				}
+				
+				string clientName = br.ReadString();
+				bool containsUntrackedFiles = br.ReadBoolean();
+				return clientName != connection.Settings.ClientName || containsUntrackedFiles;
+			}
+			catch (Exception)
+			{
+				return true;
+			}
+		}
+		
 		#region Commands
 
 		/// <summary>
 		/// Cleans the current workspace
 		/// </summary>
+		/// <param name="perforce">Optional Perforce client, if not set untracked state file cannot be updated</param>
 		/// <param name="removeUntracked">Whether to remove untracked files</param>
 		/// <param name="cancellationToken">Cancellation token</param>
-		public async Task CleanAsync(bool removeUntracked, CancellationToken cancellationToken)
+		public async Task CleanAsync(IPerforceConnection? perforce, bool removeUntracked, CancellationToken cancellationToken)
 		{
 			Stopwatch timer = Stopwatch.StartNew();
 
@@ -471,6 +551,11 @@ namespace EpicGames.Perforce.Managed
 			{
 				await CleanInternalAsync(removeUntracked, cancellationToken);
 			}
+			
+			if (perforce != null && removeUntracked)
+			{
+				await SaveUntrackedStateAsync(perforce, false, cancellationToken);
+			}
 
 			_logger.LogInformation("Completed in {ElapsedTime}s", $"{timer.Elapsed.TotalSeconds:0.0}");
 		}
@@ -535,6 +620,8 @@ namespace EpicGames.Perforce.Managed
 
 				await SaveAsync(TransactionState.Clean, cancellationToken);
 			}
+			
+			OnClean?.Invoke(filesToDelete.Length, directoriesToDelete.Length);
 		}
 
 		/// <summary>
@@ -590,7 +677,7 @@ namespace EpicGames.Perforce.Managed
 		/// <summary>
 		/// Checks the integrity of the cache
 		/// </summary>
-		public async Task RepairAsync(CancellationToken cancellationToken)
+		public async Task RepairCacheAsync(CancellationToken cancellationToken)
 		{
 			using (Trace("Repair"))
 			using (ILoggerProgress status = _logger.BeginProgressScope("Checking cache..."))
@@ -640,11 +727,11 @@ namespace EpicGames.Perforce.Managed
 		/// <summary>
 		/// Checks the <see cref="_requiresRepair"/> flag, and repairs/resets it if set.
 		/// </summary>
-		private async Task RunOptionalRepairAsync(CancellationToken cancellationToken)
+		private async Task RunOptionalCacheRepairAsync(CancellationToken cancellationToken)
 		{
 			if (_requiresRepair)
 			{
-				await RepairAsync(cancellationToken);
+				await RepairCacheAsync(cancellationToken);
 			}
 		}
 
@@ -878,8 +965,12 @@ namespace EpicGames.Perforce.Managed
 					// Force the P4 metadata to match up
 					Task updateHaveTableTask = Task.Run(() => UpdateClientHaveTableAsync(perforce, changeNumber, view, cancellationToken), cancellationToken);
 
-					// Clean the current workspace
-					await CleanInternalAsync(removeUntracked, cancellationToken);
+					// Check if clean up can be skipped
+					bool containsUntrackedFiles = await LoadUntrackedStateAsync(perforce, cancellationToken);
+					if (containsUntrackedFiles)
+					{
+						await CleanInternalAsync(removeUntracked, cancellationToken);
+					}
 
 					// Wait for the have table update to finish
 					await updateHaveTableTask;
@@ -908,7 +999,10 @@ namespace EpicGames.Perforce.Managed
 				await RemoveFilesFromWorkspaceAsync(contents, cancellationToken);
 				await AddFilesToWorkspaceAsync(perforce, contents, fakeSync, cancellationToken);
 			}
-
+			
+			// Once sync is complete and about to be handed off, mark repository having untracked files
+			await SaveUntrackedStateAsync(perforce, true, cancellationToken);
+			
 			_logger.LogInformation("Completed in {ElapsedTime}s", $"{timer.Elapsed.TotalSeconds:0.0}");
 		}
 
@@ -1040,7 +1134,7 @@ namespace EpicGames.Perforce.Managed
 			}
 
 			// Clean the current workspace
-			await CleanAsync(true, cancellationToken);
+			await CleanAsync(null, true, cancellationToken);
 
 			// Update the list of files in each stream
 			Tuple<int, StreamSnapshot>[] streamState = new Tuple<int, StreamSnapshot>[requests.Count];
@@ -1763,7 +1857,7 @@ namespace EpicGames.Perforce.Managed
 		private async Task RemoveFilesFromWorkspaceAsync(StreamSnapshot contents, CancellationToken cancellationToken)
 		{
 			// Make sure the repair flag is clear before we start
-			await RunOptionalRepairAsync(cancellationToken);
+			await RunOptionalCacheRepairAsync(cancellationToken);
 
 			// Figure out what to remove
 			RemoveTransaction transaction;
@@ -1915,7 +2009,7 @@ namespace EpicGames.Perforce.Managed
 		private async Task AddFilesToWorkspaceAsync(IPerforceConnection client, StreamSnapshot stream, bool fakeSync, CancellationToken cancellationToken)
 		{
 			// Make sure the repair flag is reset
-			await RunOptionalRepairAsync(cancellationToken);
+			await RunOptionalCacheRepairAsync(cancellationToken);
 
 			// Figure out what we need to do
 			AddTransaction transaction;
