EpicGames.Core: Add async versions of functions that write a file if different

#rnx

diff --git a/Engine/Source/Programs/Shared/EpicGames.Core/FileReference.cs b/Engine/Source/Programs/Shared/EpicGames.Core/FileReference.cs
--- a/Engine/Source/Programs/Shared/EpicGames.Core/FileReference.cs
+++ b/Engine/Source/Programs/Shared/EpicGames.Core/FileReference.cs
@@ -510,9 +510,9 @@ namespace EpicGames.Core
 		/// <param name="contents">Contents of the file</param>
 		public static bool WriteAllBytesIfDifferent(FileReference location, byte[] contents)
 		{
-			if (FileReference.Exists(location))
+			if (Exists(location))
 			{
-				byte[] currentContents = FileReference.ReadAllBytes(location);
+				byte[] currentContents = ReadAllBytes(location);
 				if (contents.AsSpan().SequenceEqual(currentContents))
 				{
 					return false;
@@ -522,6 +522,27 @@ namespace EpicGames.Core
 			return true;
 		}
 
+		/// <summary>
+		/// Writes the data to the given file, if it's different from what's there already.
+		/// Returns true if contents were written.
+		/// </summary>
+		/// <param name="location">Location of the file</param>
+		/// <param name="contents">Contents of the file</param>
+		/// <param name="cancellationToken">Cancellation token for the operation</param>
+		public static async Task<bool> WriteAllBytesIfDifferentAsync(FileReference location, byte[] contents, CancellationToken cancellationToken = default)
+		{
+			if (Exists(location))
+			{
+				byte[] currentContents = await ReadAllBytesAsync(location, cancellationToken);
+				if (contents.AsSpan().SequenceEqual(currentContents))
+				{
+					return false;
+				}
+			}
+			await WriteAllBytesAsync(location, contents, cancellationToken);
+			return true;
+		}
+
 		/// <summary>
 		/// Writes the string to the given file, if it's different from what's there already.
 		/// Returns true if contents were written.
@@ -530,9 +551,9 @@ namespace EpicGames.Core
 		/// <param name="contents">Contents of the file</param>
 		public static bool WriteAllTextIfDifferent(FileReference location, string contents)
 		{
-			if (FileReference.Exists(location))
+			if (Exists(location))
 			{
-				string currentContents = FileReference.ReadAllText(location);
+				string currentContents = ReadAllText(location);
 				if (String.Equals(contents, currentContents, StringComparison.Ordinal))
 				{
 					return false;
@@ -542,6 +563,27 @@ namespace EpicGames.Core
 			return true;
 		}
 
+		/// <summary>
+		/// Writes the string to the given file, if it's different from what's there already.
+		/// Returns true if contents were written.
+		/// </summary>
+		/// <param name="location">Location of the file</param>
+		/// <param name="contents">Contents of the file</param>
+		/// <param name="cancellationToken">Cancellation token for the operation</param>
+		public static async Task<bool> WriteAllTextIfDifferentAsync(FileReference location, string contents, CancellationToken cancellationToken = default)
+		{
+			if (Exists(location))
+			{
+				string currentContents = await ReadAllTextAsync(location, cancellationToken);
+				if (String.Equals(contents, currentContents, StringComparison.Ordinal))
+				{
+					return false;
+				}
+			}
+			await WriteAllTextAsync(location, contents, cancellationToken);
+			return true;
+		}
+
 		/// <summary>
 		/// Writes the contents of a file
 		/// </summary>
@@ -626,7 +668,8 @@ namespace EpicGames.Core
 		/// </summary>
 		/// <param name="location">Location of the file</param>
 		/// <param name="contents">Contents of the file</param>
-		public static Task WriteAllTextAsync(FileReference location, string contents) => File.WriteAllTextAsync(location.FullName, contents);
+		/// <param name="cancellationToken">Cancellation token for the operation</param>
+		public static Task WriteAllTextAsync(FileReference location, string contents, CancellationToken cancellationToken = default) => File.WriteAllTextAsync(location.FullName, contents, cancellationToken);
 
 		/// <summary>
 		/// Writes the contents of a file
@@ -634,7 +677,8 @@ namespace EpicGames.Core
 		/// <param name="location">Location of the file</param>
 		/// <param name="contents">Contents of the file</param>
 		/// <param name="encoding">The encoding to use when parsing the file</param>
-		public static Task WriteAllTextAsync(FileReference location, string contents, Encoding encoding) => File.WriteAllTextAsync(location.FullName, contents, encoding);
+		/// <param name="cancellationToken">Cancellation token for the operation</param>
+		public static Task WriteAllTextAsync(FileReference location, string contents, Encoding encoding, CancellationToken cancellationToken = default) => File.WriteAllTextAsync(location.FullName, contents, encoding, cancellationToken);
 
 		/// <summary>
 		/// Appends the contents to a file
@@ -648,7 +692,8 @@ namespace EpicGames.Core
 		/// </summary>
 		/// <param name="location">Location of the file</param>
 		/// <param name="contents">Contents to append to the file</param>
-		public static Task AppendAllLinesAsync(FileReference location, IEnumerable<string> contents) => File.AppendAllLinesAsync(location.FullName, contents);
+		/// <param name="cancellationToken">Cancellation token for the operation</param>
+		public static Task AppendAllLinesAsync(FileReference location, IEnumerable<string> contents, CancellationToken cancellationToken = default) => File.AppendAllLinesAsync(location.FullName, contents, cancellationToken);
 
 		/// <summary>
 		/// Appends the contents to a file
@@ -664,7 +709,8 @@ namespace EpicGames.Core
 		/// <param name="location">Location of the file</param>
 		/// <param name="contents">Contents to append to the file</param>
 		/// <param name="encoding">The encoding to use when parsing the file</param>
-		public static Task AppendAllLinesAsync(FileReference location, IEnumerable<string> contents, Encoding encoding) => File.AppendAllLinesAsync(location.FullName, contents, encoding);
+		/// <param name="cancellationToken">Cancellation token for the operation</param>
+		public static Task AppendAllLinesAsync(FileReference location, IEnumerable<string> contents, Encoding encoding, CancellationToken cancellationToken = default) => File.AppendAllLinesAsync(location.FullName, contents, encoding, cancellationToken);
 
 		/// <summary>
 		/// Appends the contents to a file
@@ -678,7 +724,8 @@ namespace EpicGames.Core
 		/// </summary>
 		/// <param name="location">Location of the file</param>
 		/// <param name="contents">Contents to append to the file</param>
-		public static Task AppendAllLinesAsync(FileReference location, string[] contents) => File.AppendAllLinesAsync(location.FullName, contents);
+		/// <param name="cancellationToken">Cancellation token for the operation</param>
+		public static Task AppendAllLinesAsync(FileReference location, string[] contents, CancellationToken cancellationToken = default) => File.AppendAllLinesAsync(location.FullName, contents, cancellationToken);
 
 		/// <summary>
 		/// Appends the contents to a file
@@ -694,7 +741,8 @@ namespace EpicGames.Core
 		/// <param name="location">Location of the file</param>
 		/// <param name="contents">Contents to append to the file</param>
 		/// <param name="encoding">The encoding to use when parsing the file</param>
-		public static Task AppendAllLinesAsync(FileReference location, string[] contents, Encoding encoding) => File.AppendAllLinesAsync(location.FullName, contents, encoding);
+		/// <param name="cancellationToken">Cancellation token for the operation</param>
+		public static Task AppendAllLinesAsync(FileReference location, string[] contents, Encoding encoding, CancellationToken cancellationToken = default) => File.AppendAllLinesAsync(location.FullName, contents, encoding, cancellationToken);
 
 		/// <summary>
 		/// Appends the contents to a file
@@ -708,7 +756,8 @@ namespace EpicGames.Core
 		/// </summary>
 		/// <param name="location">Location of the file</param>
 		/// <param name="contents">Contents to append to the file</param>
-		public static Task AppendAllTextAsync(FileReference location, string contents) => File.AppendAllTextAsync(location.FullName, contents);
+		/// <param name="cancellationToken">Cancellation token for the operation</param>
+		public static Task AppendAllTextAsync(FileReference location, string contents, CancellationToken cancellationToken = default) => File.AppendAllTextAsync(location.FullName, contents, cancellationToken);
 
 		/// <summary>
 		/// Appends the contents to a file
@@ -724,7 +773,8 @@ namespace EpicGames.Core
 		/// <param name="location">Location of the file</param>
 		/// <param name="contents">Contents to append to the file</param>
 		/// <param name="encoding">The encoding to use when parsing the file</param>
-		public static Task AppendAllTextAsync(FileReference location, string contents, Encoding encoding) => File.AppendAllTextAsync(location.FullName, contents, encoding);
+		/// <param name="cancellationToken">Cancellation token for the operation</param>
+		public static Task AppendAllTextAsync(FileReference location, string contents, Encoding encoding, CancellationToken cancellationToken = default) => File.AppendAllTextAsync(location.FullName, contents, encoding, cancellationToken	);
 
 		#endregion
 	}
@@ -826,7 +876,7 @@ namespace EpicGames.Core
 		/// Deserializes a file reference, using a lookup table to avoid writing the same name more than once.
 		/// </summary>
 		/// <param name="reader">The source to read from</param>
-		/// <param name="uniqueFiles">List of previously read file references. The index into this array is used in place of subsequent ocurrences of the file.</param>
+		/// <param name="uniqueFiles">List of previously read file references. The index into this array is used in place of subsequent occurrences of the file.</param>
 		/// <returns>The file reference that was read</returns>
 		public static FileReference ReadFileReference(this BinaryReader reader, List<FileReference> uniqueFiles) => BinaryArchiveReader.NotNull(ReadFileReferenceOrNull(reader, uniqueFiles));
 
@@ -834,7 +884,7 @@ namespace EpicGames.Core
 		/// Deserializes a file reference, using a lookup table to avoid writing the same name more than once.
 		/// </summary>
 		/// <param name="reader">The source to read from</param>
-		/// <param name="uniqueFiles">List of previously read file references. The index into this array is used in place of subsequent ocurrences of the file.</param>
+		/// <param name="uniqueFiles">List of previously read file references. The index into this array is used in place of subsequent occurrences of the file.</param>
 		/// <returns>The file reference that was read</returns>
 		public static FileReference? ReadFileReferenceOrNull(this BinaryReader reader, List<FileReference> uniqueFiles)
 		{
