Fix analysis warnings in Shared projects

#rnx

diff --git a/Engine/Source/Programs/Horde/Drivers/JobDriver/Execution/PerforceExecutor.cs b/Engine/Source/Programs/Horde/Drivers/JobDriver/Execution/PerforceExecutor.cs
--- a/Engine/Source/Programs/Horde/Drivers/JobDriver/Execution/PerforceExecutor.cs
+++ b/Engine/Source/Programs/Horde/Drivers/JobDriver/Execution/PerforceExecutor.cs
@@ -100,7 +100,7 @@ namespace JobDriver.Execution
 					await _autoSdkWorkspace.UpdateLocalCacheMarkerAsync(autoSdkCacheFile, autoSdkChangeNumber, -1);
 					await _autoSdkWorkspace.SyncAsync(autoSdkPerforce, autoSdkChangeNumber, -1, autoSdkCacheFile, cancellationToken);
 
-					await FileReference.WriteAllTextAsync(syncFile, syncText);
+					await FileReference.WriteAllTextAsync(syncFile, syncText, cancellationToken);
 				}
 			}
 
diff --git a/Engine/Source/Programs/Shared/EpicGames.Horde/Storage/Backends/FileStorageBackend.cs b/Engine/Source/Programs/Shared/EpicGames.Horde/Storage/Backends/FileStorageBackend.cs
--- a/Engine/Source/Programs/Shared/EpicGames.Horde/Storage/Backends/FileStorageBackend.cs
+++ b/Engine/Source/Programs/Shared/EpicGames.Horde/Storage/Backends/FileStorageBackend.cs
@@ -167,7 +167,7 @@ namespace EpicGames.Horde.Storage.Backends
 			{
 				try
 				{
-					await FileReference.WriteAllTextAsync(file, $"{value.Hash}\n{value.Locator}\n");
+					await FileReference.WriteAllTextAsync(file, $"{value.Hash}\n{value.Locator}\n", cancellationToken);
 					break;
 				}
 				catch (IOException ex) when (attempt < 3)
diff --git a/Engine/Source/Programs/Shared/EpicGames.Perforce.Managed/ManagedWorkspace.cs b/Engine/Source/Programs/Shared/EpicGames.Perforce.Managed/ManagedWorkspace.cs
--- a/Engine/Source/Programs/Shared/EpicGames.Perforce.Managed/ManagedWorkspace.cs
+++ b/Engine/Source/Programs/Shared/EpicGames.Perforce.Managed/ManagedWorkspace.cs
@@ -103,7 +103,7 @@ namespace EpicGames.Perforce.Managed
 		bool PreferNativeClient
 	)
 	{
-		/// <inheritdoc cref="ManagedWorkspaceOptions.ManagedWorkspaceOptions(int, int, long, bool, bool, bool)"/>
+		/// <inheritdoc cref="ManagedWorkspaceOptions(Int32, Int32, Int64, Boolean, Boolean, Boolean)"/>
 		public ManagedWorkspaceOptions(int? numParallelSyncThreads = null, int? maxFileConcurrency = null, long minScratchSpace = 50L * 1024, bool useHaveTable = true, bool partitioned = false, bool preferNativeClient = false)
 			: this(
 				  numParallelSyncThreads ?? GetDefaultThreadCount(4),
diff --git a/Engine/Source/Programs/Shared/EpicGames.Redis/RedisHash.cs b/Engine/Source/Programs/Shared/EpicGames.Redis/RedisHash.cs
--- a/Engine/Source/Programs/Shared/EpicGames.Redis/RedisHash.cs
+++ b/Engine/Source/Programs/Shared/EpicGames.Redis/RedisHash.cs
@@ -177,7 +177,7 @@ namespace EpicGames.Redis
 
 		#region HashScanAsync
 
-		/// <inheritdoc cref="IDatabaseAsync.HashScanAsync(RedisKey, RedisValue, int, long, int, CommandFlags)"/>
+		/// <inheritdoc cref="IDatabaseAsync.HashScanAsync(RedisKey, RedisValue, Int32, Int64, Int32, CommandFlags)"/>
 		public static IAsyncEnumerable<HashEntry<TName, TValue>> ScanAsync<TName, TValue>(this RedisHash<TName, TValue> target, RedisValue pattern, int pageSize = 250, long cursor = 0, int pageOffset = 0, CommandFlags flags = CommandFlags.None)
 			=> target.Database.HashScanAsync(target.Key, pattern, pageSize, cursor, pageOffset, flags);
 
diff --git a/Engine/Source/Programs/Shared/EpicGames.Redis/RedisHashKey.cs b/Engine/Source/Programs/Shared/EpicGames.Redis/RedisHashKey.cs
--- a/Engine/Source/Programs/Shared/EpicGames.Redis/RedisHashKey.cs
+++ b/Engine/Source/Programs/Shared/EpicGames.Redis/RedisHashKey.cs
@@ -356,7 +356,7 @@ namespace EpicGames.Redis
 
 		#region HashScanAsync
 
-		/// <inheritdoc cref="IDatabaseAsync.HashScanAsync(RedisKey, RedisValue, int, long, int, CommandFlags)"/>
+		/// <inheritdoc cref="IDatabaseAsync.HashScanAsync(RedisKey, RedisValue, Int32, Int64, Int32, CommandFlags)"/>
 		public static async IAsyncEnumerable<HashEntry<TName, TValue>> HashScanAsync<TName, TValue>(this IDatabaseAsync target, RedisHashKey<TName, TValue> key, RedisValue pattern, int pageSize = 250, long cursor = 0, int pageOffset = 0, CommandFlags flags = CommandFlags.None)
 		{
 			await foreach (HashEntry entry in target.HashScanAsync(key.Inner, pattern, pageSize, cursor, pageOffset, flags))
diff --git a/Engine/Source/Programs/Shared/EpicGames.Redis/RedisSet.cs b/Engine/Source/Programs/Shared/EpicGames.Redis/RedisSet.cs
--- a/Engine/Source/Programs/Shared/EpicGames.Redis/RedisSet.cs
+++ b/Engine/Source/Programs/Shared/EpicGames.Redis/RedisSet.cs
@@ -112,7 +112,7 @@ namespace EpicGames.Redis
 
 		#region SetScanAsync
 
-		/// <inheritdoc cref="IDatabaseAsync.SetScanAsync(RedisKey, RedisValue, int, long, int, CommandFlags)"/>
+		/// <inheritdoc cref="IDatabaseAsync.SetScanAsync(RedisKey, RedisValue, Int32, Int64, Int32, CommandFlags)"/>
 		public static IAsyncEnumerable<TElement> ScanAsync<TElement>(this RedisSet<TElement> target, RedisValue pattern = default, int pageSize = 250, long cursor = 0, int pageOffset = 0, CommandFlags flags = CommandFlags.None)
 			=> target.Database.SetScanAsync(target.Key, pattern, pageSize, cursor, pageOffset, flags);
 
diff --git a/Engine/Source/Programs/Shared/EpicGames.Redis/RedisSetKey.cs b/Engine/Source/Programs/Shared/EpicGames.Redis/RedisSetKey.cs
--- a/Engine/Source/Programs/Shared/EpicGames.Redis/RedisSetKey.cs
+++ b/Engine/Source/Programs/Shared/EpicGames.Redis/RedisSetKey.cs
@@ -121,7 +121,7 @@ namespace EpicGames.Redis
 
 		#region SetRandomMembersAsync
 
-		/// <inheritdoc cref="IDatabaseAsync.SetRandomMembersAsync(RedisKey, long, CommandFlags)"/>
+		/// <inheritdoc cref="IDatabaseAsync.SetRandomMembersAsync(RedisKey, Int64, CommandFlags)"/>
 		public static Task<TElement[]> SetRandomMembersAsync<TElement>(this IDatabaseAsync target, RedisSetKey<TElement> key, long count, CommandFlags flags = CommandFlags.None)
 			=> target.SetRandomMembersAsync(key.Inner, count, flags).DeserializeAsync<TElement>();
 
@@ -145,7 +145,7 @@ namespace EpicGames.Redis
 
 		#region SetScanAsync
 
-		/// <inheritdoc cref="IDatabaseAsync.SetScanAsync(RedisKey, RedisValue, int, long, int, CommandFlags)"/>
+		/// <inheritdoc cref="IDatabaseAsync.SetScanAsync(RedisKey, RedisValue, Int32, Int64, Int32, CommandFlags)"/>
 		public static async IAsyncEnumerable<TElement> SetScanAsync<TElement>(this IDatabaseAsync target, RedisSetKey<TElement> key, RedisValue pattern = default, int pageSize = 250, long cursor = 0, int pageOffset = 0, CommandFlags flags = CommandFlags.None)
 		{
 			await foreach (RedisValue value in target.SetScanAsync(key.Inner, pattern, pageSize, cursor, pageOffset, flags))
diff --git a/Engine/Source/Programs/Shared/EpicGames.Redis/RedisSortedSet.cs b/Engine/Source/Programs/Shared/EpicGames.Redis/RedisSortedSet.cs
--- a/Engine/Source/Programs/Shared/EpicGames.Redis/RedisSortedSet.cs
+++ b/Engine/Source/Programs/Shared/EpicGames.Redis/RedisSortedSet.cs
@@ -187,11 +187,11 @@ namespace EpicGames.Redis
 
 		#region SortedSetUpdateAsync
 
-		/// <inheritdoc cref="IDatabaseAsync.SortedSetUpdateAsync(RedisKey, RedisValue, double, SortedSetWhen, CommandFlags)"/>
+		/// <inheritdoc cref="IDatabaseAsync.SortedSetUpdateAsync(RedisKey, RedisValue, Double, SortedSetWhen, CommandFlags)"/>
 		public static Task<bool> UpdateAsync<TElement>(this RedisSortedSet<TElement> target, TElement member, double score, SortedSetWhen when = SortedSetWhen.Always, CommandFlags flags = CommandFlags.None)
 			=> target.Database.SortedSetUpdateAsync(target.Key, member, score, when, flags);
 
-		/// <inheritdoc cref="IDatabaseAsync.SortedSetUpdateAsync(RedisKey, RedisValue, double, SortedSetWhen, CommandFlags)"/>
+		/// <inheritdoc cref="IDatabaseAsync.SortedSetUpdateAsync(RedisKey, RedisValue, Double, SortedSetWhen, CommandFlags)"/>
 		public static Task<long> UpdateAsync<TElement>(this RedisSortedSet<TElement> target, SortedSetEntry<TElement>[] values, SortedSetWhen when = SortedSetWhen.Always, CommandFlags flags = CommandFlags.None)
 			=> target.Database.SortedSetUpdateAsync(target.Key, values, when, flags);
 
diff --git a/Engine/Source/Programs/Shared/EpicGames.Redis/RedisSortedSetKey.cs b/Engine/Source/Programs/Shared/EpicGames.Redis/RedisSortedSetKey.cs
--- a/Engine/Source/Programs/Shared/EpicGames.Redis/RedisSortedSetKey.cs
+++ b/Engine/Source/Programs/Shared/EpicGames.Redis/RedisSortedSetKey.cs
@@ -355,13 +355,13 @@ namespace EpicGames.Redis
 
 		#region SortedSetUpdateAsync
 
-		/// <inheritdoc cref="IDatabaseAsync.SortedSetUpdateAsync(RedisKey, RedisValue, double, SortedSetWhen, CommandFlags)"/>
+		/// <inheritdoc cref="IDatabaseAsync.SortedSetUpdateAsync(RedisKey, RedisValue, Double, SortedSetWhen, CommandFlags)"/>
 		public static Task<bool> SortedSetUpdateAsync<TElement>(this IDatabaseAsync target, RedisSortedSetKey<TElement> key, TElement member, double score, SortedSetWhen when = SortedSetWhen.Always, CommandFlags flags = CommandFlags.None)
 		{
 			return target.SortedSetUpdateAsync(key.Inner, RedisSerializer.Serialize(member), score, when, flags);
 		}
 
-		/// <inheritdoc cref="IDatabaseAsync.SortedSetUpdateAsync(RedisKey, RedisValue, double, SortedSetWhen, CommandFlags)"/>
+		/// <inheritdoc cref="IDatabaseAsync.SortedSetUpdateAsync(RedisKey, RedisValue, Double, SortedSetWhen, CommandFlags)"/>
 		public static Task<long> SortedSetUpdateAsync<TElement>(this IDatabaseAsync target, RedisSortedSetKey<TElement> key, SortedSetEntry<TElement>[] values, SortedSetWhen when = SortedSetWhen.Always, CommandFlags flags = CommandFlags.None)
 		{
 			return target.SortedSetUpdateAsync(key.Inner, values.ConvertAll(x => new SortedSetEntry(RedisSerializer.Serialize(x.Element), x.Score)), when, flags);
