EpicGames.IoHash: Compute a hash of a string

#rnx

diff --git a/Engine/Source/Programs/Shared/EpicGames.IoHash/IoHash.cs b/Engine/Source/Programs/Shared/EpicGames.IoHash/IoHash.cs
--- a/Engine/Source/Programs/Shared/EpicGames.IoHash/IoHash.cs
+++ b/Engine/Source/Programs/Shared/EpicGames.IoHash/IoHash.cs
@@ -6,6 +6,7 @@ using System.Buffers.Binary;
 using System.ComponentModel;
 using System.Globalization;
 using System.IO;
+using System.Text;
 using System.Text.Json;
 using System.Text.Json.Serialization;
 using System.Threading;
@@ -126,6 +127,22 @@ namespace EpicGames.Core
 			}
 		}
 
+		/// <summary>
+		/// Creates the IoHash for a string with a given encoding
+		/// </summary>
+		/// <param name="data">Data to compute the hash for</param>
+		/// <param name="encoding">The character encoding of the data</param>
+		/// <returns>New hash instance containing the hash of the data</returns>
+		public static IoHash Compute(string data, Encoding encoding) => Compute(encoding.GetBytes(data));
+
+
+		/// <summary>
+		/// Creates the IoHash for a string with the default encoding
+		/// </summary>
+		/// <param name="data">Data to compute the hash for</param>
+		/// <returns>New hash instance containing the hash of the data</returns>
+		public static IoHash Compute(string data) => Compute(data, Encoding.Default);
+
 		/// <summary>
 		/// Creates the IoHash for a stream.
 		/// </summary>
