UnrealBuildTool: Move BuildException to EpicGames.Build

#rnx

diff --git a/Engine/Source/Programs/Shared/EpicGames.Build/Exceptions/BuildException.cs b/Engine/Source/Programs/Shared/EpicGames.Build/Exceptions/BuildException.cs
new file mode 100644
--- /dev/null
+++ b/Engine/Source/Programs/Shared/EpicGames.Build/Exceptions/BuildException.cs
@@ -0,0 +1,77 @@
+// Copyright Epic Games, Inc. All Rights Reserved.
+
+using System;
+using EpicGames.Core;
+using Microsoft.Extensions.Logging;
+using UnrealBuildBase;
+
+// This namespace is intentional for backwards compatibility.
+// BuildException will eventually be marked Obsolete in favor of BuildLogEventException
+namespace UnrealBuildTool
+{
+	/// <summary>
+	/// Base class for exceptions thrown by UnrealBuildTool and AutomationTool
+	/// </summary>
+	public class BuildException : Exception
+	{
+		/// <summary>
+		/// Constructor
+		/// </summary>
+		/// <param name="Message">The error message to display.</param>
+		public BuildException(string Message)
+			: base(Message)
+		{
+		}
+
+		/// <summary>
+		/// Constructor which wraps another exception
+		/// </summary>
+		/// <param name="InnerException">An inner exception to wrap</param>
+		/// <param name="Message">The error message to display.</param>
+		public BuildException(Exception? InnerException, string Message)
+			: base(Message, InnerException)
+		{
+		}
+
+		/// <summary>
+		/// Constructor
+		/// </summary>
+		/// <param name="Format">Formatting string for the error message</param>
+		/// <param name="Arguments">Arguments for the formatting string</param>
+		public BuildException(string Format, params object?[] Arguments)
+			: base(String.Format(Format, Arguments))
+		{
+		}
+
+		/// <summary>
+		/// Constructor which wraps another exception
+		/// </summary>
+		/// <param name="InnerException">The inner exception being wrapped</param>
+		/// <param name="Format">Format for the message string</param>
+		/// <param name="Arguments">Format arguments</param>
+		public BuildException(Exception InnerException, string Format, params object?[] Arguments)
+			: base(String.Format(Format, Arguments), InnerException)
+		{
+		}
+
+		/// <summary>
+		/// Log BuildException with a provided ILogger
+		/// </summary>
+		/// <param name="Logger">The ILogger to use to log this exception</param>
+		public virtual void LogException(ILogger Logger)
+		{
+			Logger.LogError(this, "{Ex}", ExceptionUtils.FormatException(this));
+			Logger.LogDebug(this, "{Ex}", ExceptionUtils.FormatExceptionDetails(this));
+		}
+
+		/// <summary>
+		/// Returns the string representing the exception. Our build exceptions do not show the callstack since they are used to report known error conditions.
+		/// </summary>
+		/// <returns>Message for the exception</returns>
+		public override string ToString()
+		{
+			return Message;
+		}
+	}
+}
+
diff --git a/Engine/Source/Programs/Shared/EpicGames.Build/Exceptions/BuildExceptionExtensions.cs b/Engine/Source/Programs/Shared/EpicGames.Build/Exceptions/BuildExceptionExtensions.cs
new file mode 100644
--- /dev/null
+++ b/Engine/Source/Programs/Shared/EpicGames.Build/Exceptions/BuildExceptionExtensions.cs
@@ -0,0 +1,62 @@
+﻿// Copyright Epic Games, Inc. All Rights Reserved.
+
+using System;
+using System.Linq;
+using System.Text.Json;
+using EpicGames.Core;
+using Microsoft.Extensions.Logging;
+
+namespace UnrealBuildBase
+{
+	/// <summary>
+	/// Extension methods for build exceptions
+	/// </summary>
+	public static class BuildExceptionExtensions
+	{
+		/// <summary>
+		/// Log Exception with a provided ILogger
+		/// </summary>
+		/// <param name="ex">The exception to log</param>
+		/// <param name="logger">The ILogger to use to log this exception</param>
+		public static void LogException(this Exception ex, ILogger logger)
+		{
+			if (ex is UnrealBuildTool.BuildException buildException)
+			{
+				buildException.LogException(logger);
+			}
+			else if (ex is JsonException jsonException)
+			{
+				FileReference source = new FileReference(jsonException.Path ?? jsonException.Source ?? "unknown");
+				LogValue fileValue = LogValue.SourceFile(source, source.GetFileName());
+				long line = jsonException.LineNumber ?? 0;
+				logger.LogError(KnownLogEvents.Compiler, "{File}({Line}): error: {Message}", fileValue, line, ExceptionUtils.FormatExceptionDetails(ex));
+			}
+			else if (ex is AggregateException aggregateException)
+			{
+				logger.LogError(ex, "Unhandled {Count} aggregate exceptions", aggregateException.InnerExceptions.Count);
+
+				foreach (Exception innerException in aggregateException.InnerExceptions)
+				{
+					innerException.LogException(logger);
+				}
+			}
+			else
+			{
+				logger.LogError(ex, "Unhandled exception: {Ex}", ExceptionUtils.FormatExceptionDetails(ex));
+			}
+		}
+
+		/// <summary>
+		/// Get the CompilationResult for a provided Exception
+		/// </summary>
+		/// <param name="ex">The exception to get the result for</param>
+		/// <returns>CompilationResult</returns>
+		public static CompilationResult GetCompilationResult(this Exception ex)
+		{
+			return (ex as CompilationResultException)?.Result
+				?? (ex.InnerException as CompilationResultException)?.Result
+				?? (ex as AggregateException)?.InnerExceptions.OfType<CompilationResultException>().FirstOrDefault()?.Result
+				?? CompilationResult.OtherCompilationError;
+		}
+	}
+}
diff --git a/Engine/Source/Programs/Shared/EpicGames.Build/Exceptions/BuildLogEventException.cs b/Engine/Source/Programs/Shared/EpicGames.Build/Exceptions/BuildLogEventException.cs
new file mode 100644
--- /dev/null
+++ b/Engine/Source/Programs/Shared/EpicGames.Build/Exceptions/BuildLogEventException.cs
@@ -0,0 +1,70 @@
+﻿// Copyright Epic Games, Inc. All Rights Reserved.
+
+using System;
+using EpicGames.Core;
+using Microsoft.Extensions.Logging;
+
+namespace UnrealBuildBase
+{
+	/// <summary>
+	/// Implementation of <see cref="UnrealBuildTool.BuildException"/> that captures a full structured logging event.
+	/// </summary>
+	public class BuildLogEventException : UnrealBuildTool.BuildException
+	{
+		/// <summary>
+		/// The event object
+		/// </summary>
+		public LogEvent Event { get; }
+
+		/// <summary>
+		/// Constructor
+		/// </summary>
+		/// <param name="Event">Event to construct from</param>
+		public BuildLogEventException(LogEvent Event)
+			: this(null, Event)
+		{
+		}
+
+		/// <summary>
+		/// Constructor which wraps another exception
+		/// </summary>
+		/// <param name="InnerException">The inner exception</param>
+		/// <param name="Event">Event to construct from</param>
+		public BuildLogEventException(Exception? InnerException, LogEvent Event)
+			: base(InnerException, Event.ToString())
+		{
+			this.Event = Event;
+		}
+
+		/// <inheritdoc/>
+		public BuildLogEventException(string Format, params object[] Arguments)
+			: this(LogEvent.Create(LogLevel.Error, Format, Arguments))
+		{
+		}
+
+		/// <inheritdoc/>
+		public BuildLogEventException(Exception? InnerException, string Format, params object[] Arguments)
+			: this(InnerException, LogEvent.Create(LogLevel.Error, default, InnerException, Format, Arguments))
+		{
+		}
+
+		/// <summary>
+		/// Constructor which wraps another exception
+		/// </summary>
+		/// <param name="EventId">Event id for the error</param>
+		/// <param name="InnerException">Inner exception to wrap</param>
+		/// <param name="Format">Structured logging format string</param>
+		/// <param name="Arguments">Argument objects</param>
+		public BuildLogEventException(Exception? InnerException, EventId EventId, string Format, params object[] Arguments)
+			: this(InnerException, LogEvent.Create(LogLevel.Error, EventId, InnerException, Format, Arguments))
+		{
+		}
+
+		/// <inheritdoc/>
+		public override void LogException(ILogger Logger)
+		{
+			Logger.Log(Event.Level, Event.Id, Event, this, (s, e) => s.ToString());
+			Logger.LogDebug(Event.Id, this, "{Ex}", ExceptionUtils.FormatExceptionDetails(this));
+		}
+	}
+}
diff --git a/Engine/Source/Programs/Shared/EpicGames.Build/Exceptions/CompilationResultException.cs b/Engine/Source/Programs/Shared/EpicGames.Build/Exceptions/CompilationResultException.cs
new file mode 100644
--- /dev/null
+++ b/Engine/Source/Programs/Shared/EpicGames.Build/Exceptions/CompilationResultException.cs
@@ -0,0 +1,117 @@
+﻿// Copyright Epic Games, Inc. All Rights Reserved.
+
+using System;
+using EpicGames.Core;
+using Microsoft.Extensions.Logging;
+
+namespace UnrealBuildBase
+{
+	/// <summary>
+	/// Implementation of <see cref="BuildLogEventException"/> that will return a unique exit code.
+	/// </summary>
+	public class CompilationResultException : BuildLogEventException
+	{
+		/// <summary>
+		/// The exit code associated with this exception
+		/// </summary>
+		public CompilationResult Result { get; }
+
+		readonly bool HasMessage = true;
+
+		/// <summary>
+		/// Constructor
+		/// </summary>
+		/// <param name="Result">The resulting exit code</param>
+		public CompilationResultException(CompilationResult Result)
+			: base(LogEvent.Create(LogLevel.Error, "{CompilationResult}", Result))
+		{
+			HasMessage = false;
+			this.Result = Result;
+		}
+
+		/// <summary>
+		/// Constructor
+		/// </summary>
+		/// <param name="Result">The resulting exit code</param>
+		/// <param name="Event">Event to construct from</param>
+		public CompilationResultException(CompilationResult Result, LogEvent Event)
+			: base(Event)
+		{
+			this.Result = Result;
+		}
+
+		/// <summary>
+		/// Constructor which wraps another exception
+		/// </summary>
+		/// <param name="Result">The resulting exit code</param>
+		/// <param name="InnerException">The inner exception</param>
+		/// <param name="Event">Event to construct from</param>
+		public CompilationResultException(CompilationResult Result, Exception? InnerException, LogEvent Event)
+			: base(InnerException, Event)
+		{
+			this.Result = Result;
+		}
+
+		/// <summary>
+		/// Constructor
+		/// </summary>
+		/// <param name="Result">The resulting exit code</param>
+		/// <param name="EventId">Event id for the error</param>
+		/// <param name="Format">Formatting string for the error message</param>
+		/// <param name="Arguments">Arguments for the formatting string</param>
+		public CompilationResultException(CompilationResult Result, EventId EventId, string Format, params object[] Arguments)
+			: base(LogEvent.Create(LogLevel.Error, EventId, Format, Arguments))
+		{
+			this.Result = Result;
+		}
+
+		/// <summary>
+		/// Constructor
+		/// </summary>
+		/// <param name="Result">The resulting exit code</param>
+		/// <param name="Format">Formatting string for the error message</param>
+		/// <param name="Arguments">Arguments for the formatting string</param>
+		public CompilationResultException(CompilationResult Result, string Format, params object[] Arguments)
+			: base(LogEvent.Create(LogLevel.Error, Format, Arguments))
+		{
+			this.Result = Result;
+		}
+
+		/// <summary>
+		/// Constructor which wraps another exception
+		/// </summary>
+		/// <param name="Result">The resulting exit code</param>
+		/// <param name="InnerException">The inner exception being wrapped</param>
+		/// <param name="Format">Format for the message string</param>
+		/// <param name="Arguments">Format arguments</param>
+		public CompilationResultException(CompilationResult Result, Exception? InnerException, string Format, params object[] Arguments)
+			: base(InnerException, LogEvent.Create(LogLevel.Error, default, InnerException, Format, Arguments))
+		{
+			this.Result = Result;
+		}
+
+		/// <summary>
+		/// Constructor which wraps another exception
+		/// </summary>
+		/// <param name="Result">The resulting exit code</param>
+		/// <param name="EventId">Event id for the error</param>
+		/// <param name="InnerException">Inner exception to wrap</param>
+		/// <param name="Format">Structured logging format string</param>
+		/// <param name="Arguments">Argument objects</param>
+		public CompilationResultException(CompilationResult Result, Exception? InnerException, EventId EventId, string Format, params object[] Arguments)
+			: base(InnerException, LogEvent.Create(LogLevel.Error, EventId, InnerException, Format, Arguments))
+		{
+			this.Result = Result;
+		}
+
+		/// <inheritdoc/>
+		public override void LogException(ILogger Logger)
+		{
+			if (HasMessage)
+			{
+				Logger.Log(Event.Level, Event.Id, Event, this, (s, e) => s.ToString());
+			}
+			Logger.LogDebug(Event.Id, this, "{Ex}", ExceptionUtils.FormatExceptionDetails(this));
+		}
+	}
+}
diff --git a/Engine/Source/Programs/Shared/EpicGames.Build/System/GlobalSingleInstanceMutex.cs b/Engine/Source/Programs/Shared/EpicGames.Build/System/GlobalSingleInstanceMutex.cs
new file mode 100644
--- /dev/null
+++ b/Engine/Source/Programs/Shared/EpicGames.Build/System/GlobalSingleInstanceMutex.cs
@@ -0,0 +1,91 @@
+// Copyright Epic Games, Inc. All Rights Reserved.
+
+using System;
+using System.Threading;
+using EpicGames.Core;
+
+namespace UnrealBuildBase
+{
+	/// <summary>
+	/// System-wide mutex allowing only one instance of the program to run at a time
+	/// </summary>
+	public class GlobalSingleInstanceMutex : IDisposable
+	{
+		/// <summary>
+		/// The global mutex instance
+		/// </summary>
+		Mutex? GlobalMutex;
+
+		/// <summary>
+		/// Constructor. Attempts to acquire the global mutex
+		/// </summary>
+		/// <param name="mutexName">Name of the mutex to acquire</param>
+		/// <param name="bWaitMutex">Allow waiting for the mutex to be acquired</param>
+		public GlobalSingleInstanceMutex(string mutexName, bool bWaitMutex)
+		{
+			// Try to create the mutex, with it initially locked
+			bool bCreatedMutex;
+			GlobalMutex = new Mutex(true, mutexName, out bCreatedMutex);
+
+			// If we didn't create the mutex, we can wait for it or fail immediately
+			if (!bCreatedMutex)
+			{
+				if (bWaitMutex)
+				{
+					try
+					{
+						GlobalMutex.WaitOne();
+					}
+					catch (AbandonedMutexException)
+					{
+					}
+				}
+				else
+				{
+					throw new BuildLogEventException("A conflicting instance of {Mutex} is already running.", mutexName);
+				}
+			}
+		}
+
+		/// <summary>
+		/// Gets the name of a mutex unique for the given path
+		/// </summary>
+		/// <param name="name">Base name of the mutex</param>
+		/// <param name="uniquePath">Path to identify a unique mutex</param>
+		public static string GetUniqueMutexForPath(string name, string uniquePath)
+		{
+			// generate an IoHash of the path, as GetHashCode is not guaranteed to generate a stable hash
+			return $"Global\\{name}_{IoHash.Compute(uniquePath.ToUpperInvariant())}";
+		}
+
+		/// <summary>
+		/// Gets the name of a mutex unique for the given path
+		/// </summary>
+		/// <param name="name">Base name of the mutex</param>
+		/// <param name="uniquePath">Path to identify a unique mutex</param>
+		public static string GetUniqueMutexForPath(string name, FileSystemReference? uniquePath) => GetUniqueMutexForPath(name, uniquePath?.FullName ?? String.Empty);
+
+		/// <summary>
+		/// Release the mutex and dispose of the object
+		/// </summary>
+		public void Dispose()
+		{
+			Dispose(true);
+			GC.SuppressFinalize(this);
+		}
+
+		/// <summary>
+		/// Overridable dispose method
+		/// </summary>
+		/// <param name="disposing">Whether the object should be disposed</param>
+		protected virtual void Dispose(bool disposing)
+		{
+			if (disposing)
+			{
+				GlobalMutex?.ReleaseMutex();
+				GlobalMutex?.Dispose();
+				GlobalMutex = null;
+			}
+		}
+	}
+}
diff --git a/Engine/Source/Programs/UnrealBuildTool/System/BuildException.cs b/Engine/Source/Programs/UnrealBuildTool/System/BuildException.cs
deleted file mode 100644
diff --git a/Engine/Source/Programs/UnrealBuildTool/System/GlobalSingleInstanceMutex.cs b/Engine/Source/Programs/UnrealBuildTool/System/GlobalSingleInstanceMutex.cs
deleted file mode 100644
