EpicGames.Build: Allow resolving root path from Engine/Source/Programs

#rnx

diff --git a/Engine/Source/Programs/Shared/EpicGames.Build/Unreal.cs b/Engine/Source/Programs/Shared/EpicGames.Build/Unreal.cs
--- a/Engine/Source/Programs/Shared/EpicGames.Build/Unreal.cs
+++ b/Engine/Source/Programs/Shared/EpicGames.Build/Unreal.cs
@@ -57,9 +57,34 @@ namespace UnrealBuildBase
 				FoundRootDirectory = FoundRootDirectory.ParentDirectory;
 			}
 
+			// Search up through the directory tree for the deepest instance of the sub-path "Engine/Source/Programs"
 			if (FoundRootDirectory == null)
 			{
-				throw new Exception($"This code requires that applications using it are launched from a path containing \"Engine/Binaries/DotNET\". This application was launched from {Path.GetDirectoryName(AssemblyLocation)}");
+				FoundRootDirectory = DirectoryReference.FindCorrectCase(DirectoryReference.FromString(AssemblyLocation)!);
+				while (FoundRootDirectory != null)
+				{
+					if (String.Equals("Programs", FoundRootDirectory.GetDirectoryName()))
+					{
+						FoundRootDirectory = FoundRootDirectory.ParentDirectory;
+						if (FoundRootDirectory != null && String.Equals("Source", FoundRootDirectory.GetDirectoryName()))
+						{
+							FoundRootDirectory = FoundRootDirectory.ParentDirectory;
+							if (FoundRootDirectory != null && String.Equals("Engine", FoundRootDirectory.GetDirectoryName()))
+							{
+								FoundRootDirectory = FoundRootDirectory.ParentDirectory;
+								break;
+							}
+							continue;
+						}
+						continue;
+					}
+					FoundRootDirectory = FoundRootDirectory.ParentDirectory;
+				}
+			}
+
+			if (FoundRootDirectory == null)
+			{
+				throw new Exception($"This code requires that applications using it are launched from a path containing \"Engine/Binaries/DotNET\" or \"Engine/Source/Programs\". This application was launched from {Path.GetDirectoryName(AssemblyLocation)}");
 			}
 
 			// Confirm that we've found a valid root directory, by testing for the existence of a well-known file
diff --git a/Engine/Source/Programs/UnrealBuildTool.Tests/ConfigFileTests.cs b/Engine/Source/Programs/UnrealBuildTool.Tests/ConfigFileTests.cs
--- a/Engine/Source/Programs/UnrealBuildTool.Tests/ConfigFileTests.cs
+++ b/Engine/Source/Programs/UnrealBuildTool.Tests/ConfigFileTests.cs
@@ -52,8 +52,6 @@ namespace UnrealBuildToolTests
 		[TestMethod]
 		public void ConfigFileTestsQuotations()
 		{
-			UnrealHelper.InitializePath();
-
 			const string TestSection = "MySection";
 			const string TestValueKey = "Name";
 			const string TestValue = "\"https://www.unrealengine.com/\"";
diff --git a/Engine/Source/Programs/UnrealBuildTool.Tests/SourceFileTests.cs b/Engine/Source/Programs/UnrealBuildTool.Tests/SourceFileTests.cs
--- a/Engine/Source/Programs/UnrealBuildTool.Tests/SourceFileTests.cs
+++ b/Engine/Source/Programs/UnrealBuildTool.Tests/SourceFileTests.cs
@@ -20,8 +20,6 @@ namespace UnrealBuildToolTests
 		[TestMethod]
 		public void Run()
 		{
-			UnrealHelper.InitializePath();
-
 			List<DirectoryReference> baseDirectories = new List<DirectoryReference>();
 			baseDirectories.Add(DirectoryReference.Combine(Unreal.EngineSourceDirectory, "Runtime"));
 			baseDirectories.Add(DirectoryReference.Combine(Unreal.EngineSourceDirectory, "Developer"));
diff --git a/Engine/Source/Programs/UnrealBuildTool.Tests/UnrealHeaderToolTests.cs b/Engine/Source/Programs/UnrealBuildTool.Tests/UnrealHeaderToolTests.cs
--- a/Engine/Source/Programs/UnrealBuildTool.Tests/UnrealHeaderToolTests.cs
+++ b/Engine/Source/Programs/UnrealBuildTool.Tests/UnrealHeaderToolTests.cs
@@ -15,8 +15,6 @@ namespace UnrealBuildToolTests
 		[Ignore]
 		public void Run()
 		{
-			UnrealHelper.InitializePath();
-
 			string[] Arguments = System.Array.Empty<string>();
 			CommandLineArguments CommandLineArguments = new CommandLineArguments(Arguments);
 			UhtGlobalOptions Options = new UhtGlobalOptions(CommandLineArguments);
diff --git a/Engine/Source/Programs/UnrealBuildTool.Tests/UnrealHelper.cs b/Engine/Source/Programs/UnrealBuildTool.Tests/UnrealHelper.cs
deleted file mode 100644
