Sdk json can specify a list of banned sdk versions that will be excluded by the auto-selection of manual SDKs

diff --git a/Engine/Source/Programs/Shared/EpicGames.Build/UEBuildPlatformSDK.cs b/Engine/Source/Programs/Shared/EpicGames.Build/UEBuildPlatformSDK.cs
--- a/Engine/Source/Programs/Shared/EpicGames.Build/UEBuildPlatformSDK.cs
+++ b/Engine/Source/Programs/Shared/EpicGames.Build/UEBuildPlatformSDK.cs
@@ -341,7 +341,7 @@ namespace EpicGames.Core
 				}
 				else
 				{
-					Console.WriteLine("Processing parent manual sdk: {0}", ParentManualSDKVersions);
+					Logger.LogInformation("Processing parent manual sdk: {SDKVersion}", ParentManualSDKVersions);
 					CachedManualSDKVersions = JsonSerializer.Deserialize<Dictionary<string, string>>(ParentManualSDKVersions) ?? new Dictionary<string, string>();
 				}
 			}
@@ -353,7 +353,7 @@ namespace EpicGames.Core
 				{
 					if (TrySelectBestManualSDK( out string? SelectedSDKVersion ))
 					{
-						Console.WriteLine("Auto-selected best installed sdk for {0} - {1}", GetAutoSDKPlatformName(), SelectedSDKVersion);
+						Logger.LogInformation("Auto-selected best installed sdk for {PlatformName} - {SDKVersion}", GetAutoSDKPlatformName(), SelectedSDKVersion);
 					}
 				}
 
@@ -685,11 +685,31 @@ namespace EpicGames.Core
 			GetValidVersionRange(out string? MinVersion, out string? MaxVersion);
 			if (TryConvertVersionToInt(MinVersion, out ulong MinVersionInt, null) && TryConvertVersionToInt(MaxVersion, out ulong MaxVersionInt, null))
 			{
+				List<Tuple<ulong,ulong>> BannedSdkVersionRanges = new();
+				foreach (string SdkVersionPair in GetStringArrayFromConfig("BannedSdkVersions"))
+				{
+					string[] Versions = SdkVersionPair.Split("-");
+					if (Versions.Length > 0 && Versions.Length <= 2)
+					{
+						if (TryConvertVersionToInt(Versions.First(), out ulong MinInvalidVersion, null) && TryConvertVersionToInt(Versions.Last(), out ulong MaxInvalidVersion, null))
+						{
+							BannedSdkVersionRanges.Add( new( MinInvalidVersion, MaxInvalidVersion));
+						}
+					}
+				}
+
 				IGrouping<ulong,string>? ResultGroup = InstalledSDKVersions
-					.GroupBy( X => TryConvertVersionToInt(X, out ulong Value, null) ? Value : 0 )   // group by integer version
-					.Where( X => X.Key >= MinVersionInt && X.Key <= MaxVersionInt)                  // select only valid versions
-					.OrderByDescending( X => X.Key )                                                // sort highest version first
+					.GroupBy( X => TryConvertVersionToInt(X, out ulong Value, null) ? Value : 0 )            // group by integer version
+					.Where( X => X.Key >= MinVersionInt && X.Key <= MaxVersionInt)                           // select only valid versions
+					.OrderByDescending( X => X.Key )                                                         // sort highest version first
+					.OrderBy(X => BannedSdkVersionRanges.Any(R => X.Key >= R.Item1 && X.Key <= R.Item2))     // put banned versions at the end of the list - they'll be excluded later
 					.FirstOrDefault();
+				if (ResultGroup?.FirstOrDefault() != null && BannedSdkVersionRanges.Any(R => ResultGroup.Key >= R.Item1 && ResultGroup.Key <= R.Item2))
+				{
+					Logger.LogInformation("The best installed SDK for {Platform} is {Version} but this version is forbidden", PlatformName, ResultGroup.FirstOrDefault());
+					return null;
+				}
+
 				return ResultGroup?.FirstOrDefault();
 			}
 
