Fixes for deadlocks regarding async CancellationTokens and IDisposable in BackgroundTask and HangMonitor.

diff --git a/Engine/Source/Programs/Horde/Plugins/Build/HordeServer.Build/Perforce/PerforceService.cs b/Engine/Source/Programs/Horde/Plugins/Build/HordeServer.Build/Perforce/PerforceService.cs
--- a/Engine/Source/Programs/Horde/Plugins/Build/HordeServer.Build/Perforce/PerforceService.cs
+++ b/Engine/Source/Programs/Horde/Plugins/Build/HordeServer.Build/Perforce/PerforceService.cs
@@ -72,6 +72,15 @@ namespace HordeServer.Perforce
 				_inner = inner;
 			}
 
+			/// <inheritdoc/>
+			public async ValueTask DisposeAsync()
+			{
+				await Task.Run(() =>
+				{
+					Dispose();
+				} );
+			}
+
 			/// <inheritdoc/>
 			public void Dispose()
 			{
diff --git a/Engine/Source/Programs/Shared/EpicGames.Core/BackgroundTask.cs b/Engine/Source/Programs/Shared/EpicGames.Core/BackgroundTask.cs
--- a/Engine/Source/Programs/Shared/EpicGames.Core/BackgroundTask.cs
+++ b/Engine/Source/Programs/Shared/EpicGames.Core/BackgroundTask.cs
@@ -9,7 +9,7 @@ namespace EpicGames.Core
 	/// <summary>
 	/// Runs a task in the background and allows stopping it on demand
 	/// </summary>
-	public sealed class BackgroundTask : IAsyncDisposable
+	public sealed class BackgroundTask : IAsyncDisposable, IDisposable
 	{
 		readonly Func<CancellationToken, Task> _runTask;
 		Task? _task;
@@ -29,6 +29,13 @@ namespace EpicGames.Core
 			_runTask = runTask;
 		}
 
+		/// <inheritdoc/>
+		public void Dispose()
+		{
+			Stop();
+			_cancellationTokenSource.Dispose();
+		}
+
 		/// <inheritdoc/>
 		public async ValueTask DisposeAsync()
 		{
@@ -73,28 +80,41 @@ namespace EpicGames.Core
 		}
 
 		/// <summary>
-		/// Signals the cancellation token and waits for the task to finish
+		/// Signals the cancellation token and waits synchronously for the task to finish
 		/// </summary>
 		/// <param name="cancellationToken">Cancellation token for the operation</param>
-		public async Task StopAsync(CancellationToken cancellationToken = default)
+		public void Stop(CancellationToken cancellationToken = default)
 		{
+
 			if (_task != null && !_task.IsCompleted)
 			{
 				try
 				{
-#if NET8_0_OR_GREATER
-					await _cancellationTokenSource.CancelAsync();
-#else
 					_cancellationTokenSource.Cancel();
-#endif
-					await _task.WaitAsync(cancellationToken).ConfigureAwait(false);
+					_task.Wait(cancellationToken);
 				}
-				catch (OperationCanceledException)
+				catch (AggregateException e)
 				{
+					if (e.InnerException is not TaskCanceledException)
+					{
+						throw;
+					}
 				}
 			}
 		}
 
+		/// <summary>
+		/// Signals the cancellation token and waits for the task to finish
+		/// </summary>
+		/// <param name="cancellationToken">Cancellation token for the operation</param>
+		public async Task StopAsync(CancellationToken cancellationToken = default)
+		{
+			await Task.Run( () =>
+			{
+				Stop( cancellationToken );
+			}, cancellationToken );
+		}
+
 		/// <summary>
 		/// Waits for the task to complete
 		/// </summary>
diff --git a/Engine/Source/Programs/Shared/EpicGames.Perforce/HangMonitor.cs b/Engine/Source/Programs/Shared/EpicGames.Perforce/HangMonitor.cs
--- a/Engine/Source/Programs/Shared/EpicGames.Perforce/HangMonitor.cs
+++ b/Engine/Source/Programs/Shared/EpicGames.Perforce/HangMonitor.cs
@@ -12,7 +12,7 @@ namespace EpicGames.Perforce
 	/// <summary>
 	/// Runs a background task that logs warnings if the Tick method isn't called within a certain time period
 	/// </summary>
-	public sealed class HangMonitor : IDisposable
+	public sealed class HangMonitor : IDisposable, IAsyncDisposable
 	{
 		class Scope : IDisposable
 		{
@@ -54,12 +54,21 @@ namespace EpicGames.Perforce
 			_logger = logger;
 		}
 
+		/// <inheritdoc/>
+		public async ValueTask DisposeAsync()
+		{
+			await Task.Run(() =>
+			{
+				Dispose();
+			} );
+		}
+
 		/// <inheritdoc/>
 		public void Dispose()
 		{
-			if (_hangTask != null)
+			if (_hangTask != null && _hangTask is IDisposable)
 			{
-				_hangTask.DisposeAsync().AsTask().Wait();
+				_hangTask.Dispose();
 				_hangTask = null;
 			}
 		}
diff --git a/Engine/Source/Programs/Shared/EpicGames.Perforce/IPerforceConnection.cs b/Engine/Source/Programs/Shared/EpicGames.Perforce/IPerforceConnection.cs
--- a/Engine/Source/Programs/Shared/EpicGames.Perforce/IPerforceConnection.cs
+++ b/Engine/Source/Programs/Shared/EpicGames.Perforce/IPerforceConnection.cs
@@ -9,7 +9,7 @@ namespace EpicGames.Perforce
 	/// <summary>
 	/// Base interface for Perforce clients
 	/// </summary>
-	public interface IPerforceConnection : IDisposable
+	public interface IPerforceConnection : IDisposable, IAsyncDisposable
 	{
 		/// <summary>
 		/// Connection settings
diff --git a/Engine/Source/Programs/Shared/EpicGames.Perforce/NativePerforceConnection.cs b/Engine/Source/Programs/Shared/EpicGames.Perforce/NativePerforceConnection.cs
--- a/Engine/Source/Programs/Shared/EpicGames.Perforce/NativePerforceConnection.cs
+++ b/Engine/Source/Programs/Shared/EpicGames.Perforce/NativePerforceConnection.cs
@@ -20,7 +20,7 @@ namespace EpicGames.Perforce
 	/// <summary>
 	/// Experimental implementation of <see cref="IPerforceConnection"/> which wraps the native C++ API.
 	/// </summary>
-	public sealed class NativePerforceConnection : IPerforceConnection, IDisposable
+	public sealed class NativePerforceConnection : IPerforceConnection, IDisposable, IAsyncDisposable
 	{
 		const string NativeDll = "EpicGames.Perforce.Native";
 
@@ -368,7 +368,11 @@ namespace EpicGames.Perforce
 			}
 			catch
 			{
-				connection?.Dispose();
+				if ( connection != null )
+				{
+					await connection.DisposeAsync();
+				}
+				
 				throw;
 			}
 		}
@@ -421,6 +425,16 @@ namespace EpicGames.Perforce
 			Dispose(false);
 		}
 
+		/// <inheritdoc/>
+		public async ValueTask DisposeAsync()
+		{
+			await Task.Run(() =>
+			{
+				Dispose(true);
+				GC.SuppressFinalize(this);
+			});
+		}
+
 		/// <inheritdoc/>
 		public void Dispose()
 		{
diff --git a/Engine/Source/Programs/Shared/EpicGames.Perforce/PerforceConnection.cs b/Engine/Source/Programs/Shared/EpicGames.Perforce/PerforceConnection.cs
--- a/Engine/Source/Programs/Shared/EpicGames.Perforce/PerforceConnection.cs
+++ b/Engine/Source/Programs/Shared/EpicGames.Perforce/PerforceConnection.cs
@@ -282,6 +282,12 @@ namespace EpicGames.Perforce
 			GlobalOptions.AddRange(other.GlobalOptions);
 		}
 
+		/// <inheritdoc/>
+		public ValueTask DisposeAsync()
+		{
+			return ValueTask.CompletedTask;
+		}
+
 		/// <inheritdoc/>
 		public void Dispose()
 		{
diff --git a/Engine/Source/Programs/UnrealGameSync/UnrealGameSync/Automation/P4Automation.cs b/Engine/Source/Programs/UnrealGameSync/UnrealGameSync/Automation/P4Automation.cs
--- a/Engine/Source/Programs/UnrealGameSync/UnrealGameSync/Automation/P4Automation.cs
+++ b/Engine/Source/Programs/UnrealGameSync/UnrealGameSync/Automation/P4Automation.cs
@@ -58,9 +58,9 @@ namespace UnrealGameSync
 			finally
 			{
 				// If we created a new connection, tear it down now.
-				if (createNewConnection)
+				if (createNewConnection && connection != null)
 				{
-					connection?.Dispose();
+					await connection.DisposeAsync().AsTask();
 				}
 			}
 		}
diff --git a/Engine/Source/Programs/UnrealGameSync/UnrealGameSyncShared/OpenProjectInfo.cs b/Engine/Source/Programs/UnrealGameSync/UnrealGameSyncShared/OpenProjectInfo.cs
--- a/Engine/Source/Programs/UnrealGameSync/UnrealGameSyncShared/OpenProjectInfo.cs
+++ b/Engine/Source/Programs/UnrealGameSync/UnrealGameSyncShared/OpenProjectInfo.cs
@@ -252,7 +252,10 @@ namespace UnrealGameSync
 			}
 			finally
 			{
-				perforceClient?.Dispose();
+				if(perforceClient != null)
+				{
+					await perforceClient.DisposeAsync().AsTask().ConfigureAwait(false);
+				}
 			}
 		}
 
