Undo of UGS deadlock, a simpler fix will be following.

diff --git a/Engine/Source/Programs/Horde/Plugins/Build/HordeServer.Build/Perforce/PerforceService.cs b/Engine/Source/Programs/Horde/Plugins/Build/HordeServer.Build/Perforce/PerforceService.cs
--- a/Engine/Source/Programs/Horde/Plugins/Build/HordeServer.Build/Perforce/PerforceService.cs
+++ b/Engine/Source/Programs/Horde/Plugins/Build/HordeServer.Build/Perforce/PerforceService.cs
@@ -72,15 +72,6 @@ namespace HordeServer.Perforce
 				_inner = inner;
 			}
 
-			/// <inheritdoc/>
-			public async ValueTask DisposeAsync()
-			{
-				await Task.Run(() =>
-				{
-					Dispose();
-				} );
-			}
-
 			/// <inheritdoc/>
 			public void Dispose()
 			{
diff --git a/Engine/Source/Programs/Shared/EpicGames.Core/BackgroundTask.cs b/Engine/Source/Programs/Shared/EpicGames.Core/BackgroundTask.cs
--- a/Engine/Source/Programs/Shared/EpicGames.Core/BackgroundTask.cs
+++ b/Engine/Source/Programs/Shared/EpicGames.Core/BackgroundTask.cs
@@ -9,7 +9,7 @@ namespace EpicGames.Core
 	/// <summary>
 	/// Runs a task in the background and allows stopping it on demand
 	/// </summary>
-	public sealed class BackgroundTask : IAsyncDisposable, IDisposable
+	public sealed class BackgroundTask : IAsyncDisposable
 	{
 		readonly Func<CancellationToken, Task> _runTask;
 		Task? _task;
@@ -29,13 +29,6 @@ namespace EpicGames.Core
 			_runTask = runTask;
 		}
 
-		/// <inheritdoc/>
-		public void Dispose()
-		{
-			Stop();
-			_cancellationTokenSource.Dispose();
-		}
-
 		/// <inheritdoc/>
 		public async ValueTask DisposeAsync()
 		{
@@ -80,41 +73,28 @@ namespace EpicGames.Core
 		}
 
 		/// <summary>
-		/// Signals the cancellation token and waits synchronously for the task to finish
+		/// Signals the cancellation token and waits for the task to finish
 		/// </summary>
 		/// <param name="cancellationToken">Cancellation token for the operation</param>
-		public void Stop(CancellationToken cancellationToken = default)
+		public async Task StopAsync(CancellationToken cancellationToken = default)
 		{
-
 			if (_task != null && !_task.IsCompleted)
 			{
 				try
 				{
+#if NET8_0_OR_GREATER
+					await _cancellationTokenSource.CancelAsync();
+#else
 					_cancellationTokenSource.Cancel();
-					_task.Wait(cancellationToken);
+#endif
+					await _task.WaitAsync(cancellationToken).ConfigureAwait(false);
 				}
-				catch (AggregateException e)
+				catch (OperationCanceledException)
 				{
-					if (e.InnerException is not TaskCanceledException)
-					{
-						throw;
-					}
 				}
 			}
 		}
 
-		/// <summary>
-		/// Signals the cancellation token and waits for the task to finish
-		/// </summary>
-		/// <param name="cancellationToken">Cancellation token for the operation</param>
-		public async Task StopAsync(CancellationToken cancellationToken = default)
-		{
-			await Task.Run( () =>
-			{
-				Stop( cancellationToken );
-			}, cancellationToken );
-		}
-
 		/// <summary>
 		/// Waits for the task to complete
 		/// </summary>
diff --git a/Engine/Source/Programs/Shared/EpicGames.Perforce/HangMonitor.cs b/Engine/Source/Programs/Shared/EpicGames.Perforce/HangMonitor.cs
--- a/Engine/Source/Programs/Shared/EpicGames.Perforce/HangMonitor.cs
+++ b/Engine/Source/Programs/Shared/EpicGames.Perforce/HangMonitor.cs
@@ -12,7 +12,7 @@ namespace EpicGames.Perforce
 	/// <summary>
 	/// Runs a background task that logs warnings if the Tick method isn't called within a certain time period
 	/// </summary>
-	public sealed class HangMonitor : IDisposable, IAsyncDisposable
+	public sealed class HangMonitor : IDisposable
 	{
 		class Scope : IDisposable
 		{
@@ -54,21 +54,12 @@ namespace EpicGames.Perforce
 			_logger = logger;
 		}
 
-		/// <inheritdoc/>
-		public async ValueTask DisposeAsync()
-		{
-			await Task.Run(() =>
-			{
-				Dispose();
-			} );
-		}
-
 		/// <inheritdoc/>
 		public void Dispose()
 		{
-			if (_hangTask != null && _hangTask is IDisposable)
+			if (_hangTask != null)
 			{
-				_hangTask.Dispose();
+				_hangTask.DisposeAsync().AsTask().Wait();
 				_hangTask = null;
 			}
 		}
diff --git a/Engine/Source/Programs/Shared/EpicGames.Perforce/IPerforceConnection.cs b/Engine/Source/Programs/Shared/EpicGames.Perforce/IPerforceConnection.cs
--- a/Engine/Source/Programs/Shared/EpicGames.Perforce/IPerforceConnection.cs
+++ b/Engine/Source/Programs/Shared/EpicGames.Perforce/IPerforceConnection.cs
@@ -9,7 +9,7 @@ namespace EpicGames.Perforce
 	/// <summary>
 	/// Base interface for Perforce clients
 	/// </summary>
-	public interface IPerforceConnection : IDisposable, IAsyncDisposable
+	public interface IPerforceConnection : IDisposable
 	{
 		/// <summary>
 		/// Connection settings
diff --git a/Engine/Source/Programs/Shared/EpicGames.Perforce/NativePerforceConnection.cs b/Engine/Source/Programs/Shared/EpicGames.Perforce/NativePerforceConnection.cs
--- a/Engine/Source/Programs/Shared/EpicGames.Perforce/NativePerforceConnection.cs
+++ b/Engine/Source/Programs/Shared/EpicGames.Perforce/NativePerforceConnection.cs
@@ -20,7 +20,7 @@ namespace EpicGames.Perforce
 	/// <summary>
 	/// Experimental implementation of <see cref="IPerforceConnection"/> which wraps the native C++ API.
 	/// </summary>
-	public sealed class NativePerforceConnection : IPerforceConnection, IDisposable, IAsyncDisposable
+	public sealed class NativePerforceConnection : IPerforceConnection, IDisposable
 	{
 		const string NativeDll = "EpicGames.Perforce.Native";
 
@@ -368,11 +368,7 @@ namespace EpicGames.Perforce
 			}
 			catch
 			{
-				if ( connection != null )
-				{
-					await connection.DisposeAsync();
-				}
-				
+				connection?.Dispose();
 				throw;
 			}
 		}
@@ -425,16 +421,6 @@ namespace EpicGames.Perforce
 			Dispose(false);
 		}
 
-		/// <inheritdoc/>
-		public async ValueTask DisposeAsync()
-		{
-			await Task.Run(() =>
-			{
-				Dispose(true);
-				GC.SuppressFinalize(this);
-			});
-		}
-
 		/// <inheritdoc/>
 		public void Dispose()
 		{
diff --git a/Engine/Source/Programs/Shared/EpicGames.Perforce/PerforceConnection.cs b/Engine/Source/Programs/Shared/EpicGames.Perforce/PerforceConnection.cs
--- a/Engine/Source/Programs/Shared/EpicGames.Perforce/PerforceConnection.cs
+++ b/Engine/Source/Programs/Shared/EpicGames.Perforce/PerforceConnection.cs
@@ -282,12 +282,6 @@ namespace EpicGames.Perforce
 			GlobalOptions.AddRange(other.GlobalOptions);
 		}
 
-		/// <inheritdoc/>
-		public ValueTask DisposeAsync()
-		{
-			return ValueTask.CompletedTask;
-		}
-
 		/// <inheritdoc/>
 		public void Dispose()
 		{
diff --git a/Engine/Source/Programs/UnrealGameSync/UnrealGameSync/Automation/P4Automation.cs b/Engine/Source/Programs/UnrealGameSync/UnrealGameSync/Automation/P4Automation.cs
--- a/Engine/Source/Programs/UnrealGameSync/UnrealGameSync/Automation/P4Automation.cs
+++ b/Engine/Source/Programs/UnrealGameSync/UnrealGameSync/Automation/P4Automation.cs
@@ -58,9 +58,9 @@ namespace UnrealGameSync
 			finally
 			{
 				// If we created a new connection, tear it down now.
-				if (createNewConnection && connection != null)
+				if (createNewConnection)
 				{
-					await connection.DisposeAsync().AsTask();
+					connection?.Dispose();
 				}
 			}
 		}
diff --git a/Engine/Source/Programs/UnrealGameSync/UnrealGameSyncShared/OpenProjectInfo.cs b/Engine/Source/Programs/UnrealGameSync/UnrealGameSyncShared/OpenProjectInfo.cs
--- a/Engine/Source/Programs/UnrealGameSync/UnrealGameSyncShared/OpenProjectInfo.cs
+++ b/Engine/Source/Programs/UnrealGameSync/UnrealGameSyncShared/OpenProjectInfo.cs
@@ -252,10 +252,7 @@ namespace UnrealGameSync
 			}
 			finally
 			{
-				if(perforceClient != null)
-				{
-					await perforceClient.DisposeAsync().AsTask().ConfigureAwait(false);
-				}
+				perforceClient?.Dispose();
 			}
 		}
 
