Fixed a deadlock in UGS when querying streams.

diff --git a/Engine/Source/Programs/Shared/EpicGames.Perforce/HangMonitor.cs b/Engine/Source/Programs/Shared/EpicGames.Perforce/HangMonitor.cs
--- a/Engine/Source/Programs/Shared/EpicGames.Perforce/HangMonitor.cs
+++ b/Engine/Source/Programs/Shared/EpicGames.Perforce/HangMonitor.cs
@@ -59,7 +59,7 @@ namespace EpicGames.Perforce
 		{
 			if (_hangTask != null)
 			{
-				_hangTask.DisposeAsync().AsTask().Wait();
+				Task.Run(async () => await _hangTask.DisposeAsync()).Wait();
 				_hangTask = null;
 			}
 		}
