Horde: Ignore the "Installed" registry key when parsing data into the agent configuration. Fixes exception on startup from installer.

#rnx

diff --git a/Engine/Source/Programs/Horde/HordeAgent/AgentApp.cs b/Engine/Source/Programs/Horde/HordeAgent/AgentApp.cs
--- a/Engine/Source/Programs/Horde/HordeAgent/AgentApp.cs
+++ b/Engine/Source/Programs/Horde/HordeAgent/AgentApp.cs
@@ -261,7 +261,10 @@ namespace HordeAgent
 			IConfigurationBuilder builder = new ConfigurationBuilder();
 			if (readInstalledConfig && OperatingSystem.IsWindows())
 			{
-				builder = builder.Add(new RegistryConfigurationSource(Registry.LocalMachine, "SOFTWARE\\Epic Games\\Horde\\Agent", AgentSettings.SectionName));
+				static bool IncludeRegistrySetting(string name)
+					=> !String.Equals(name, $"{AgentSettings.SectionName}:Installed", StringComparison.OrdinalIgnoreCase);
+
+				builder = builder.Add(new RegistryConfigurationSource(Registry.LocalMachine, "SOFTWARE\\Epic Games\\Horde\\Agent", AgentSettings.SectionName, IncludeRegistrySetting));
 			}
 			
 			string basePath = AppDir.FullName;
diff --git a/Engine/Source/Programs/Shared/EpicGames.Core/RegistryConfigurationSource.cs b/Engine/Source/Programs/Shared/EpicGames.Core/RegistryConfigurationSource.cs
--- a/Engine/Source/Programs/Shared/EpicGames.Core/RegistryConfigurationSource.cs
+++ b/Engine/Source/Programs/Shared/EpicGames.Core/RegistryConfigurationSource.cs
@@ -1,5 +1,6 @@
 ﻿// Copyright Epic Games, Inc. All Rights Reserved.
 
+using System;
 using System.Collections.Generic;
 using System.Runtime.InteropServices;
 using System.Runtime.Versioning;
@@ -17,6 +18,7 @@ namespace EpicGames.Core
 		readonly RegistryKey _baseKey;
 		readonly string _keyPath;
 		readonly string _baseConfigName;
+		readonly Func<string, bool>? _filter;
 
 		/// <summary>
 		/// Constructor
@@ -24,11 +26,13 @@ namespace EpicGames.Core
 		/// <param name="baseKey">Hive to resolve keyPath relative to</param>
 		/// <param name="keyPath">Path within the registry to enumerate</param>
 		/// <param name="baseConfigName">Prefix for returned configuration values</param>
-		public RegistryConfigurationSource(RegistryKey baseKey, string keyPath, string baseConfigName)
+		/// <param name="filter">Optional filter for keys to include</param>
+		public RegistryConfigurationSource(RegistryKey baseKey, string keyPath, string baseConfigName, Func<string, bool>? filter = null)
 		{
 			_baseKey = baseKey;
 			_keyPath = keyPath;
 			_baseConfigName = baseConfigName;
+			_filter = filter;
 		}
 		
 		/// <summary>
@@ -39,7 +43,7 @@ namespace EpicGames.Core
 		
 		/// <inheritdoc/>
 		public IConfigurationProvider Build(IConfigurationBuilder builder)
-			=> new RegistryConfigProvider(_baseKey, _keyPath, _baseConfigName);
+			=> new RegistryConfigProvider(_baseKey, _keyPath, _baseConfigName, _filter);
 	}
 
 	class RegistryConfigProvider : ConfigurationProvider
@@ -47,12 +51,14 @@ namespace EpicGames.Core
 		readonly RegistryKey _baseKey;
 		readonly string _keyPath;
 		readonly string _baseConfigName;
+		readonly Func<string, bool>? _filter;
 
-		public RegistryConfigProvider(RegistryKey baseKey, string keyPath, string baseConfigName)
+		public RegistryConfigProvider(RegistryKey baseKey, string keyPath, string baseConfigName, Func<string, bool>? filter)
 		{
 			_baseKey = baseKey;
 			_keyPath = keyPath;
 			_baseConfigName = baseConfigName;
+			_filter = filter;
 		}
 
 		public override void Load()
@@ -60,13 +66,13 @@ namespace EpicGames.Core
 			if (RuntimeInformation.IsOSPlatform(OSPlatform.Windows))
 			{
 				Dictionary<string, string?> data = new Dictionary<string, string?>();
-				GetValues(_baseKey, _keyPath, _baseConfigName, data);
+				GetValues(_baseKey, _keyPath, _baseConfigName, _filter, data);
 				Data = data;
 			}
 		}
 
 		[SupportedOSPlatform("windows")]
-		static void GetValues(RegistryKey baseKey, string keyPath, string baseConfigName, Dictionary<string, string?> data)
+		static void GetValues(RegistryKey baseKey, string keyPath, string baseConfigName, Func<string, bool>? filter, Dictionary<string, string?> data)
 		{
 			using RegistryKey? registryKey = baseKey.OpenSubKey(keyPath);
 			if (registryKey != null)
@@ -74,7 +80,7 @@ namespace EpicGames.Core
 				string[] subKeyNames = registryKey.GetSubKeyNames();
 				foreach (string subKeyName in subKeyNames)
 				{
-					GetValues(registryKey, subKeyName, $"{baseConfigName}:{subKeyName}", data);
+					GetValues(registryKey, subKeyName, $"{baseConfigName}:{subKeyName}", filter, data);
 				}
 
 				string[] valueNames = registryKey.GetValueNames();
@@ -83,7 +89,11 @@ namespace EpicGames.Core
 					object? value = registryKey.GetValue(valueName);
 					if (value != null)
 					{
-						data[$"{baseConfigName}:{valueName}"] = value.ToString();
+						string name = $"{baseConfigName}:{valueName}";
+						if (filter == null || filter(name))
+						{
+							data[name] = value.ToString();
+						}
 					}
 				}
 			}
