EpicGames.Perforce: Fix image build error for fixture server

#rnx

diff --git a/Engine/Source/Programs/Shared/EpicGames.Perforce.Fixture/setup.sh b/Engine/Source/Programs/Shared/EpicGames.Perforce.Fixture/setup.sh
--- a/Engine/Source/Programs/Shared/EpicGames.Perforce.Fixture/setup.sh
+++ b/Engine/Source/Programs/Shared/EpicGames.Perforce.Fixture/setup.sh
@@ -128,4 +128,4 @@ p4 delete -c 8 main.cpp
 p4 shelve -c 8
 
 # Stop server running in background
-kill "$(pidof p4d)"
\ No newline at end of file
+pkill p4d
\ No newline at end of file
