Horde: Add scratch and min conform space to REST API agent response

diff --git a/Engine/Source/Programs/Horde/Plugins/Compute/HordeServer.Compute/Agents/AgentsController.cs b/Engine/Source/Programs/Horde/Plugins/Compute/HordeServer.Compute/Agents/AgentsController.cs
--- a/Engine/Source/Programs/Horde/Plugins/Compute/HordeServer.Compute/Agents/AgentsController.cs
+++ b/Engine/Source/Programs/Horde/Plugins/Compute/HordeServer.Compute/Agents/AgentsController.cs
@@ -241,7 +241,7 @@ namespace HordeServer.Agents
 
 		internal static GetAgentWorkspaceResponse CreateGetAgentWorkspaceResponse(AgentWorkspaceInfo workspace)
 		{
-			return new GetAgentWorkspaceResponse(workspace.Cluster, workspace.UserName, workspace.Identifier, workspace.Stream, workspace.View, workspace.Incremental, workspace.Method);
+			return new GetAgentWorkspaceResponse(workspace.Cluster, workspace.UserName, workspace.Identifier, workspace.Stream, workspace.View, workspace.Incremental, workspace.Method, workspace.MinScratchSpace, workspace.ConformDiskFreeSpace);
 		}
 
 		/// <summary>
diff --git a/Engine/Source/Programs/Shared/EpicGames.Horde/Agents/AgentMessages.cs b/Engine/Source/Programs/Shared/EpicGames.Horde/Agents/AgentMessages.cs
--- a/Engine/Source/Programs/Shared/EpicGames.Horde/Agents/AgentMessages.cs
+++ b/Engine/Source/Programs/Shared/EpicGames.Horde/Agents/AgentMessages.cs
@@ -85,8 +85,10 @@ namespace EpicGames.Horde.Agents
 	/// <param name="View"> Custom view for the workspace </param>
 	/// <param name="BIncremental"> Whether to use an incremental workspace </param>
 	/// <param name="Method"> Method to use when syncing/materializing data from Perforce </param>
-	public record GetAgentWorkspaceResponse(string? Cluster, string? UserName, string Identifier, string Stream, List<string>? View, bool BIncremental, string? Method);
-
+	/// <param name="MinScratchSpace">Minimum disk space that must be available *after* syncing this workspace (in megabytes)</param>
+	/// <param name="ConformDiskFreeSpace">Threshold for when to trigger an automatic conform of agent. Measured in megabytes free on disk</param>
+	public record GetAgentWorkspaceResponse(string? Cluster, string? UserName, string Identifier, string Stream, List<string>? View, bool BIncremental, string? Method, long? MinScratchSpace, long? ConformDiskFreeSpace);
+	
 	/// <summary>
 	/// Information about an agent
 	/// </summary>
