Fixed a bug where the same error/warning message using different slash directions would not match with an existing issue when handled by the hashed issue handler.

We now standardise the slash direction before hashing to give a consistent result.

An example of this issue can be found here:
- https://horde.devtools.epicgames.com/job/66e1d5bf3451afc9d41b0d69?step=a3cb&issue=580067
- https://horde.devtools.epicgames.com/log/66e1d95a667fefbd2f083bad?issue=580098

diff --git a/Engine/Source/Programs/Shared/EpicGames.Horde/Issues/Handlers/HashedIssueHandler.cs b/Engine/Source/Programs/Shared/EpicGames.Horde/Issues/Handlers/HashedIssueHandler.cs
--- a/Engine/Source/Programs/Shared/EpicGames.Horde/Issues/Handlers/HashedIssueHandler.cs
+++ b/Engine/Source/Programs/Shared/EpicGames.Horde/Issues/Handlers/HashedIssueHandler.cs
@@ -120,6 +120,7 @@ namespace EpicGames.Horde.Issues.Handlers
 			sanitized = Regex.Replace(sanitized, @"(?<![a-zA-Z])(?:[A-Z]:|/)[^ :]+[/\\]SYNC[/\\]", "{root}/"); // Redact things that look like workspace roots; may be different between agents
 			sanitized = Regex.Replace(sanitized, @"0[xX][0-9a-fA-F]+", "H"); // Redact hex strings
 			sanitized = Regex.Replace(sanitized, @"\d[\d.,:]*", "n"); // Redact numbers and timestamp like things
+			sanitized = Regex.Replace(sanitized, @"\\", "/");
 
 			if (sanitized.Length > 30)
 			{
