Resolve warnings

#rnx
[RN] none

diff --git a/Engine/Source/Programs/Shared/EpicGames.UHT/Exporters/CodeGen/UhtHeaderCodeGenerator.cs b/Engine/Source/Programs/Shared/EpicGames.UHT/Exporters/CodeGen/UhtHeaderCodeGenerator.cs
--- a/Engine/Source/Programs/Shared/EpicGames.UHT/Exporters/CodeGen/UhtHeaderCodeGenerator.cs
+++ b/Engine/Source/Programs/Shared/EpicGames.UHT/Exporters/CodeGen/UhtHeaderCodeGenerator.cs
@@ -279,7 +279,7 @@ namespace EpicGames.UHT.Exporters.CodeGen
 
 		// Find the UFUNCTION associated with this event so that we can check
 		// if there is a BP implementation of it
-		protected static StringBuilder AppendFindUFunction(StringBuilder builder, UhtClass classObj, UhtFunction function, int tabs, string endl)
+		protected static StringBuilder AppendFindUFunction(StringBuilder builder, UhtClass classObj, UhtFunction function, int tabs, string endl = "\r\n")
 		{
 			builder
 				.AppendTabs(tabs)
@@ -288,7 +288,8 @@ namespace EpicGames.UHT.Exporters.CodeGen
 				.Append(classObj.SourceName)
 				.Append('_')
 				.Append(function.EngineName)
-				.Append(");\r\n");
+				.Append(");")
+				.Append(endl);
 
 			return builder;
 		}
diff --git a/Engine/Source/Programs/Shared/EpicGames.UHT/Exporters/CodeGen/UhtPackageCodeGeneratorCppFile.cs b/Engine/Source/Programs/Shared/EpicGames.UHT/Exporters/CodeGen/UhtPackageCodeGeneratorCppFile.cs
--- a/Engine/Source/Programs/Shared/EpicGames.UHT/Exporters/CodeGen/UhtPackageCodeGeneratorCppFile.cs
+++ b/Engine/Source/Programs/Shared/EpicGames.UHT/Exporters/CodeGen/UhtPackageCodeGeneratorCppFile.cs
@@ -5,7 +5,6 @@ using System.Collections.Generic;
 using System.Linq;
 using System.Text;
 using EpicGames.Core;
-using EpicGames.UHT.Tables;
 using EpicGames.UHT.Types;
 using EpicGames.UHT.Utils;
 
diff --git a/Engine/Source/Programs/Shared/EpicGames.UHT/Exporters/Json/UhtJsonExporter.cs b/Engine/Source/Programs/Shared/EpicGames.UHT/Exporters/Json/UhtJsonExporter.cs
--- a/Engine/Source/Programs/Shared/EpicGames.UHT/Exporters/Json/UhtJsonExporter.cs
+++ b/Engine/Source/Programs/Shared/EpicGames.UHT/Exporters/Json/UhtJsonExporter.cs
@@ -4,7 +4,6 @@ using System.Collections.Generic;
 using System.Text.Json;
 using System.Text.Json.Serialization;
 using System.Threading.Tasks;
-using EpicGames.Core;
 using EpicGames.UHT.Tables;
 using EpicGames.UHT.Types;
 using EpicGames.UHT.Utils;
diff --git a/Engine/Source/Programs/Shared/EpicGames.UHT/Types/UhtModule.cs b/Engine/Source/Programs/Shared/EpicGames.UHT/Types/UhtModule.cs
--- a/Engine/Source/Programs/Shared/EpicGames.UHT/Types/UhtModule.cs
+++ b/Engine/Source/Programs/Shared/EpicGames.UHT/Types/UhtModule.cs
@@ -3,7 +3,6 @@
 using System;
 using System.Collections.Generic;
 using System.IO;
-using System.Security.Cryptography.X509Certificates;
 using System.Text.Json.Serialization;
 using EpicGames.Core;
 using EpicGames.UHT.Utils;
diff --git a/Engine/Source/Programs/Shared/EpicGames.UHT/Types/UhtObject.cs b/Engine/Source/Programs/Shared/EpicGames.UHT/Types/UhtObject.cs
--- a/Engine/Source/Programs/Shared/EpicGames.UHT/Types/UhtObject.cs
+++ b/Engine/Source/Programs/Shared/EpicGames.UHT/Types/UhtObject.cs
@@ -3,7 +3,6 @@
 using System.Text.Json.Serialization;
 using EpicGames.Core;
 using EpicGames.UHT.Tables;
-using EpicGames.UHT.Utils;
 
 namespace EpicGames.UHT.Types
 {
diff --git a/Engine/Source/Programs/Shared/EpicGames.UHT/Types/UhtPackage.cs b/Engine/Source/Programs/Shared/EpicGames.UHT/Types/UhtPackage.cs
--- a/Engine/Source/Programs/Shared/EpicGames.UHT/Types/UhtPackage.cs
+++ b/Engine/Source/Programs/Shared/EpicGames.UHT/Types/UhtPackage.cs
@@ -1,6 +1,5 @@
 // Copyright Epic Games, Inc. All Rights Reserved.
 
-using System;
 using System.Text.Json.Serialization;
 using EpicGames.Core;
 using EpicGames.UHT.Tables;
diff --git a/Engine/Source/Programs/Shared/EpicGames.UHT/Utils/UhtSession.cs b/Engine/Source/Programs/Shared/EpicGames.UHT/Utils/UhtSession.cs
--- a/Engine/Source/Programs/Shared/EpicGames.UHT/Utils/UhtSession.cs
+++ b/Engine/Source/Programs/Shared/EpicGames.UHT/Utils/UhtSession.cs
@@ -1848,21 +1848,6 @@ namespace EpicGames.UHT.Utils
 			}
 		}
 
-		private UhtHeaderFileParser? ParseHeaderFile(UhtHeaderFile headerFile)
-		{
-			UhtHeaderFileParser? parser = null;
-			try
-			{
-				headerFile.Read();
-				parser = UhtHeaderFileParser.Parse(headerFile);
-			}
-			catch (Exception e)
-			{
-				HandleException(headerFile.MessageSource, e);
-			}
-			return parser;
-		}
-
 		#region Run steps
 		private void StepReadManifestFile(string manifestFilePath)
 		{
diff --git a/Engine/Source/Programs/Shared/UnrealBuildTool.Tests/ConfigFileTests.cs b/Engine/Source/Programs/Shared/UnrealBuildTool.Tests/ConfigFileTests.cs
--- a/Engine/Source/Programs/Shared/UnrealBuildTool.Tests/ConfigFileTests.cs
+++ b/Engine/Source/Programs/Shared/UnrealBuildTool.Tests/ConfigFileTests.cs
@@ -22,14 +22,14 @@ namespace UnrealBuildToolTests
 		{
 			internal string? TemporaryFile { get; private set; }
 
-			internal TestConfigFile(string FileName, string Contents)
+			internal TestConfigFile(string fileName, string contents)
 			{
 				try
 				{
-					TemporaryFile = Path.Combine(System.IO.Path.GetTempPath(), Guid.NewGuid().ToString() + "_" + FileName);
+					TemporaryFile = Path.Combine(System.IO.Path.GetTempPath(), Guid.NewGuid().ToString() + "_" + fileName);
 					using (StreamWriter writer = new StreamWriter(TemporaryFile))
 					{
-						writer.Write(Contents);
+						writer.Write(contents);
 					}
 				}
 				catch (Exception)
@@ -56,27 +56,27 @@ namespace UnrealBuildToolTests
 			const string TestValueKey = "Name";
 			const string TestValue = "\"https://www.unrealengine.com/\"";
 
-			string TestContents = String.Format("[{0}]\r\n{1}={2}", TestSection, TestValueKey, TestValue);
+			string testContents = String.Format("[{0}]\r\n{1}={2}", TestSection, TestValueKey, TestValue);
 
-			using (TestConfigFile TemporaryConfigFile = new TestConfigFile(TEST_FILE_NAME, TestContents))
+			using (TestConfigFile temporaryConfigFile = new TestConfigFile(TEST_FILE_NAME, testContents))
 			{
-				string? TestFilePath = TemporaryConfigFile.TemporaryFile;
+				string? testFilePath = temporaryConfigFile.TemporaryFile;
 
-				Assert.IsFalse(String.IsNullOrEmpty(TestFilePath), "Unable to generate a test file.");
+				Assert.IsFalse(String.IsNullOrEmpty(testFilePath), "Unable to generate a test file.");
 
-				ConfigFile ConfigFile;
-				FileReference ConfigFileLocation = new FileReference(TestFilePath);
+				ConfigFile configFile;
+				FileReference configFileLocation = new FileReference(testFilePath);
 
-				Assert.IsTrue(FileReference.Exists(ConfigFileLocation), "Unable to acquire a test file.");
+				Assert.IsTrue(FileReference.Exists(configFileLocation), "Unable to acquire a test file.");
 				
-				ConfigFile = new ConfigFile(ConfigFileLocation);
+				configFile = new ConfigFile(configFileLocation);
 
-				FileReference.MakeWriteable(ConfigFileLocation);
-				ConfigFile.Write(ConfigFileLocation);
+				FileReference.MakeWriteable(configFileLocation);
+				configFile.Write(configFileLocation);
 
-				Assert.IsNotNull(TemporaryConfigFile.TemporaryFile);
-				string FileContents = File.ReadAllText(TemporaryConfigFile.TemporaryFile);
-				Assert.AreEqual(TestContents.Trim(), FileContents.Trim(), "The file contents after writing do not match the expected contents.");
+				Assert.IsNotNull(temporaryConfigFile.TemporaryFile);
+				string fileContents = File.ReadAllText(temporaryConfigFile.TemporaryFile);
+				Assert.AreEqual(testContents.Trim(), fileContents.Trim(), "The file contents after writing do not match the expected contents.");
 			}
 		}
 	}
diff --git a/Engine/Source/Programs/Shared/UnrealBuildTool.Tests/UnrealHeaderToolTests.cs b/Engine/Source/Programs/Shared/UnrealBuildTool.Tests/UnrealHeaderToolTests.cs
--- a/Engine/Source/Programs/Shared/UnrealBuildTool.Tests/UnrealHeaderToolTests.cs
+++ b/Engine/Source/Programs/Shared/UnrealBuildTool.Tests/UnrealHeaderToolTests.cs
@@ -15,19 +15,19 @@ namespace UnrealBuildToolTests
 		[Ignore]
 		public void Run()
 		{
-			string[] Arguments = System.Array.Empty<string>();
-			CommandLineArguments CommandLineArguments = new CommandLineArguments(Arguments);
-			UhtGlobalOptions Options = new UhtGlobalOptions(CommandLineArguments);
+			string[] arguments = System.Array.Empty<string>();
+			CommandLineArguments commandLineArguments = new CommandLineArguments(arguments);
+			UhtGlobalOptions options = new UhtGlobalOptions(commandLineArguments);
 
 			// Initialize the attributes
-			UhtTables Tables = new UhtTables();
+			UhtTables tables = new UhtTables();
 
 			// Initialize the configuration
-			IUhtConfig Config = new UhtConfigImpl(CommandLineArguments);
+			IUhtConfig config = new UhtConfigImpl(commandLineArguments);
 
 			// Run the tests
 			using ILoggerFactory factory = LoggerFactory.Create(x => x.AddEpicDefault());
-			Assert.IsTrue(UhtTestHarness.RunTests(Tables, Config, Options, factory.CreateLogger<UhtTestHarness>()));
+			Assert.IsTrue(UhtTestHarness.RunTests(tables, config, options, factory.CreateLogger<UhtTestHarness>()));
 		}
 	}
 }
