EpicGames.Build: .net8 suggestions

#rnx

diff --git a/Engine/Source/Programs/Shared/EpicGames.Build/Automation/CommandLine.cs b/Engine/Source/Programs/Shared/EpicGames.Build/Automation/CommandLine.cs
--- a/Engine/Source/Programs/Shared/EpicGames.Build/Automation/CommandLine.cs
+++ b/Engine/Source/Programs/Shared/EpicGames.Build/Automation/CommandLine.cs
@@ -13,7 +13,7 @@ namespace UnrealBuildBase
 	{
 		public string CommandName;
 
-		public List<string> Arguments = new List<string>();
+		public List<string> Arguments = [];
 
 		public CommandInfo(string CommandNameIn)
 		{
@@ -129,6 +129,6 @@ namespace UnrealBuildBase
 			return IgnoredGlobalParameters.Contains(Parameter);
 		}
 
-		public List<CommandInfo> CommandsToExecute = new List<CommandInfo>();
+		public List<CommandInfo> CommandsToExecute = [];
 	}
 }
diff --git a/Engine/Source/Programs/Shared/EpicGames.Build/Automation/FormatCommandLine.cs b/Engine/Source/Programs/Shared/EpicGames.Build/Automation/FormatCommandLine.cs
--- a/Engine/Source/Programs/Shared/EpicGames.Build/Automation/FormatCommandLine.cs
+++ b/Engine/Source/Programs/Shared/EpicGames.Build/Automation/FormatCommandLine.cs
@@ -49,7 +49,7 @@ namespace UnrealBuildBase
 
 			// If it does have a space, and it's formatted as an option (ie. -Something=), try to insert quotes after the equals character
 			int EqualsIdx = Argument.IndexOf('=');
-			if(Argument.StartsWith("-") && EqualsIdx != -1 && EqualsIdx < SpaceIdx)
+			if(Argument.StartsWith('-') && EqualsIdx != -1 && EqualsIdx < SpaceIdx)
 			{
 				return String.Format("{0}=\"{1}\"", Argument.Substring(0, EqualsIdx), Argument.Substring(EqualsIdx + 1));
 			}
diff --git a/Engine/Source/Programs/Shared/EpicGames.Build/Platform/Mac/ApplePlatformSDK.cs b/Engine/Source/Programs/Shared/EpicGames.Build/Platform/Mac/ApplePlatformSDK.cs
--- a/Engine/Source/Programs/Shared/EpicGames.Build/Platform/Mac/ApplePlatformSDK.cs
+++ b/Engine/Source/Programs/Shared/EpicGames.Build/Platform/Mac/ApplePlatformSDK.cs
@@ -12,10 +12,14 @@ using Microsoft.Win32;
 
 namespace UnrealBuildBase
 {
-	static public class ApplePlatformSDK
+	static public partial class ApplePlatformSDK
 	{
 		public static readonly string? InstalledSDKVersion = GetInstalledSDKVersion();
 
+		// 8 bits per component, with high getting extra from high 32
+		[GeneratedRegex(@"^(\d+).(\d+)(.(\d+))?(.(\d+))?(.(\d+))?$")]
+		private static partial Regex VersionToIntRegex();
+
 		public static bool TryConvertVersionToInt(string? StringValue, out ulong OutValue)
 		{
 			OutValue = 0;
@@ -25,8 +29,7 @@ namespace UnrealBuildBase
 				return false;
 			}
 
-			// 8 bits per component, with high getting extra from high 32
-			Match Result = Regex.Match(StringValue, @"^(\d+).(\d+)(.(\d+))?(.(\d+))?(.(\d+))?$");
+			Match Result = VersionToIntRegex().Match(StringValue);
 			if (Result.Success)
 			{
 				OutValue = UInt64.Parse(Result.Groups[1].Value) << 24 | UInt64.Parse(Result.Groups[2].Value) << 16;
diff --git a/Engine/Source/Programs/Shared/EpicGames.Build/System/CompileScriptModules.cs b/Engine/Source/Programs/Shared/EpicGames.Build/System/CompileScriptModules.cs
--- a/Engine/Source/Programs/Shared/EpicGames.Build/System/CompileScriptModules.cs
+++ b/Engine/Source/Programs/Shared/EpicGames.Build/System/CompileScriptModules.cs
@@ -55,7 +55,7 @@ namespace UnrealBuildBase
 		class Hook : CsProjBuildHook
 		{
 			private ILogger Logger;
-			private Dictionary<string, DateTime> WriteTimes = new Dictionary<string, DateTime>();
+			private Dictionary<string, DateTime> WriteTimes = [];
 			private string BuildRecordDirectory;
 
 			public Hook(ILogger InLogger, Rules.RulesFileType InRulesFileType)
@@ -99,7 +99,7 @@ namespace UnrealBuildBase
 				CompileScriptModule.ValidateBuildRecordRecursively(BuildRecords, ProjectPath, this, Logger);
 			}
 
-			private static readonly char[] s_wildcardCharacters = { '*', '?' };
+			private static readonly char[] s_wildcardCharacters = ['*', '?'];
 			public bool HasWildcards(string FileSpec)
 			{
 				// Perf Note: Doing a [Last]IndexOfAny(...) is much faster than compiling a
@@ -159,17 +159,19 @@ namespace UnrealBuildBase
 			List<DirectoryReference> GameBuildDirectories = GetAdditionalBuildDirectories(GameDirectories);
 
 			// List of directories used to locate Intermediate/ScriptModules dirs for writing build records
-			List<DirectoryReference> BaseDirectories = new List<DirectoryReference>(1 + GameDirectories.Count + AdditionalDirectories.Count);
-			BaseDirectories.Add(Unreal.EngineDirectory);
-			BaseDirectories.AddRange(GameDirectories);
-			BaseDirectories.AddRange(AdditionalDirectories);
+			List<DirectoryReference> BaseDirectories =
+			[
+				Unreal.EngineDirectory,
+				.. GameDirectories,
+				.. AdditionalDirectories,
+			];
 			BaseDirectories = BaseDirectories.Distinct().ToList();
 
 			HashSet<FileReference> FoundProjects = new HashSet<FileReference>(
 				Rules.FindAllRulesSourceFiles(RulesFileType,
 				// Project scripts require source engine builds
-				GameFolders: Unreal.IsEngineInstalled() ? GameDirectories : new List<DirectoryReference>(), 
-				ForeignPlugins: null, AdditionalSearchPaths: AdditionalDirectories.Concat(GameBuildDirectories).ToList()));
+				GameFolders: Unreal.IsEngineInstalled() ? GameDirectories : [], 
+				ForeignPlugins: null, AdditionalSearchPaths: [.. AdditionalDirectories, .. GameBuildDirectories]));
 
 			BuildFlags BuildFlags = BuildFlags.None;
 			if (bForceCompile)
@@ -399,7 +401,7 @@ namespace UnrealBuildBase
 			bool bForceCompile, out bool bBuildSuccess, CsProjBuildHook Hook, IEnumerable<DirectoryReference> BaseDirectories,
 			IEnumerable<string>? DefineConstants, Action<int> OnBuildingProjects, ILogger Logger)
 		{
-			return CsProjBuilder.Build(FoundProjects, bForceCompile, out bBuildSuccess, Hook, BaseDirectories, DefineConstants ?? (Enumerable.Empty<string>()), OnBuildingProjects, Logger);
+			return CsProjBuilder.Build(FoundProjects, bForceCompile, out bBuildSuccess, Hook, BaseDirectories, DefineConstants ?? ([]), OnBuildingProjects, Logger);
 		}
 
 		/// <summary>
@@ -411,7 +413,7 @@ namespace UnrealBuildBase
 		/// <returns></returns>
 		static Dictionary<FileReference, CsProjBuildRecordEntry> LoadExistingBuildRecords(CsProjBuildHook Hook, IEnumerable<DirectoryReference> BaseDirectories, ILogger Logger)
         {
-			Dictionary<FileReference, CsProjBuildRecordEntry> LoadedBuildRecords = new Dictionary<FileReference, CsProjBuildRecordEntry>();
+			Dictionary<FileReference, CsProjBuildRecordEntry> LoadedBuildRecords = [];
 
 			foreach (DirectoryReference Directory in BaseDirectories)
 			{
@@ -463,28 +465,28 @@ namespace UnrealBuildBase
 			
 			// Files are grouped by ItemType (e.g. Compile, EmbeddedResource) to ensure that Exclude and
 			// Remove act as expected.
-			Dictionary<string, HashSet<string>> Files = new Dictionary<string, HashSet<string>>();
+			Dictionary<string, HashSet<string>> Files = [];
 			foreach (CsProjBuildRecord.Glob Glob in Globs)
 			{
 				HashSet<string>? TypedFiles;
 				if (!Files.TryGetValue(Glob.ItemType!, out TypedFiles))
 				{
-					TypedFiles = new HashSet<string>();
+					TypedFiles = [];
 					Files.Add(Glob.ItemType!, TypedFiles);
 				}
 
 				Matcher matcher = new Matcher();
-				if (Glob.Include?.Any() == true)
+				if (Glob.Include?.Count > 0)
 				{
 					matcher.AddIncludePatterns(Glob.Include);
 				}
-				if (Glob.Exclude?.Any() == true)
+				if (Glob.Exclude?.Count > 0)
 				{
 					matcher.AddExcludePatterns(Glob.Exclude);
 				}
 				TypedFiles.UnionWith(matcher.GetResultsInFullPath(ProjectDirectory.FullName).Select(x => Path.GetRelativePath(ProjectDirectory.FullName, x)));
 
-				if (Glob.Remove?.Any() == true)
+				if (Glob.Remove?.Count > 0)
 				{
 					// Matcher.Match() doesn't handle inconsistent path separators correctly - which is why globs
 					// are normalized when they are added to CsProjBuildRecord
@@ -675,7 +677,7 @@ namespace UnrealBuildBase
 
 		static List<DirectoryReference> GetGameDirectories(string? ScriptsForProjectFileName, ILogger Logger)
         {
-			List<DirectoryReference> GameDirectories = new List<DirectoryReference>();
+			List<DirectoryReference> GameDirectories = [];
 
 			if (String.IsNullOrEmpty(ScriptsForProjectFileName))
 			{
@@ -691,7 +693,7 @@ namespace UnrealBuildBase
         }
 
 		static List<DirectoryReference> GetAdditionalDirectories(List<string>? AdditionalScriptsFolders) =>
-			AdditionalScriptsFolders == null ? new List<DirectoryReference>() :
+			AdditionalScriptsFolders == null ? [] :
 				AdditionalScriptsFolders.Select(x => DirectoryReference.FindCorrectCase(new DirectoryReference(x))).ToList();
 
 		static List<DirectoryReference> GetAdditionalBuildDirectories(List<DirectoryReference> GameDirectories) =>
diff --git a/Engine/Source/Programs/Shared/EpicGames.Build/System/DirectoryItem.cs b/Engine/Source/Programs/Shared/EpicGames.Build/System/DirectoryItem.cs
--- a/Engine/Source/Programs/Shared/EpicGames.Build/System/DirectoryItem.cs
+++ b/Engine/Source/Programs/Shared/EpicGames.Build/System/DirectoryItem.cs
@@ -371,7 +371,7 @@ namespace UnrealBuildBase
 		/// </summary>
 		public static void WriteDebugFileWithAllEnumeratedFiles(string OutFile)
 		{
-			SortedSet<string> AllFiles = new SortedSet<string>();
+			SortedSet<string> AllFiles = [];
 			foreach (DirectoryItem Item in DirectoryItem.LocationToItem.Values)
 			{
 				if (Item.Files != null)
diff --git a/Engine/Source/Programs/Shared/EpicGames.Build/System/EnumeratePlugins.cs b/Engine/Source/Programs/Shared/EpicGames.Build/System/EnumeratePlugins.cs
--- a/Engine/Source/Programs/Shared/EpicGames.Build/System/EnumeratePlugins.cs
+++ b/Engine/Source/Programs/Shared/EpicGames.Build/System/EnumeratePlugins.cs
@@ -21,7 +21,7 @@ namespace UnrealBuildBase
 		/// Cache of all UBT plugins that have been built.  When generating projects, there is a good chance that a UBT project
 		/// might need to be rebuilt.
 		/// </summary>
-		static Dictionary<FileReference, FileReference> UbtPluginFileCache = new Dictionary<FileReference, FileReference>();
+		static Dictionary<FileReference, FileReference> UbtPluginFileCache = [];
 
 		/// <summary>
 		/// Invalidate cached plugin data so that we can pickup new things
@@ -40,8 +40,7 @@ namespace UnrealBuildBase
 		/// <returns>List of project files</returns>
 		public static IEnumerable<FileReference> EnumeratePlugins(FileReference? ProjectFile)
 		{
-			List<DirectoryReference> BaseDirs = new List<DirectoryReference>();
-			BaseDirs.AddRange(Unreal.GetExtensionDirs(Unreal.EngineDirectory, "Plugins"));
+			List<DirectoryReference> BaseDirs = [.. Unreal.GetExtensionDirs(Unreal.EngineDirectory, "Plugins")];
 			if(ProjectFile != null)
 			{
 				BaseDirs.AddRange(Unreal.GetExtensionDirs(ProjectFile.Directory, "Plugins"));
@@ -58,7 +57,7 @@ namespace UnrealBuildBase
 		{
 			return PluginFileCache.GetOrAdd(ParentDirectory, _ =>
 			{
-				List<FileReference> FileNames = new List<FileReference>();
+				List<FileReference> FileNames = [];
 
 				DirectoryItem ParentDirectoryItem = DirectoryItem.GetItemByDirectoryReference(ParentDirectory);
 				if (ParentDirectoryItem.Exists)
@@ -70,7 +69,7 @@ namespace UnrealBuildBase
 				}
 
 				// Sort the filenames to ensure that the plugin order is deterministic; otherwise response files will change with each build.
-				FileNames = FileNames.OrderBy(x => x.FullName, StringComparer.OrdinalIgnoreCase).ToList();
+				FileNames = [.. FileNames.OrderBy(x => x.FullName, StringComparer.OrdinalIgnoreCase)];
 				return FileNames;
 			});
 		}
@@ -116,10 +115,12 @@ namespace UnrealBuildBase
 		/// <returns>The collection of UBT project files</returns>
 		public static List<FileReference> EnumerateUbtPlugins(FileReference? ProjectFile)
 		{
-			List<DirectoryReference> PluginDirectories = new List<DirectoryReference>();
-			PluginDirectories.AddRange(Unreal.GetExtensionDirs(Unreal.EngineDirectory, "Build"));
-			PluginDirectories.AddRange(Unreal.GetExtensionDirs(Unreal.EngineDirectory, "Source"));
-			PluginDirectories.AddRange(Unreal.GetExtensionDirs(Unreal.EngineDirectory, "Plugins"));
+			List<DirectoryReference> PluginDirectories =
+			[
+				.. Unreal.GetExtensionDirs(Unreal.EngineDirectory, "Build"),
+				.. Unreal.GetExtensionDirs(Unreal.EngineDirectory, "Source"),
+				.. Unreal.GetExtensionDirs(Unreal.EngineDirectory, "Plugins"),
+			];
 			if (ProjectFile != null)
 			{
 				PluginDirectories.AddRange(Unreal.GetExtensionDirs(ProjectFile.Directory, "Build"));
@@ -168,7 +169,7 @@ namespace UnrealBuildBase
 				}
 			}
 
-			Plugins = Built.Where(P => UbtPluginFileCache.ContainsKey(P)).Select(P => (P, UbtPluginFileCache[P])).OrderBy(X => X.P.FullName).ToArray();
+			Plugins = [.. Built.Where(P => UbtPluginFileCache.ContainsKey(P)).Select(P => (P, UbtPluginFileCache[P])).OrderBy(X => X.P.FullName)];
 			return bBuildSuccess;
 		}
 
@@ -184,7 +185,7 @@ namespace UnrealBuildBase
 		{
 			// Collect the work to be done
 			IEnumerable<FileReference>? ToBuild = null;
-			List<DirectoryReference> BaseDirectories = new();
+			List<DirectoryReference> BaseDirectories = [];
 			CompileScriptModule.BuildFlags BuildFlags = CompileScriptModule.BuildFlags.UseBuildRecords;
 			switch (PluginSet)
 			{
diff --git a/Engine/Source/Programs/Shared/EpicGames.Build/System/NativeProjectsBase.cs b/Engine/Source/Programs/Shared/EpicGames.Build/System/NativeProjectsBase.cs
--- a/Engine/Source/Programs/Shared/EpicGames.Build/System/NativeProjectsBase.cs
+++ b/Engine/Source/Programs/Shared/EpicGames.Build/System/NativeProjectsBase.cs
@@ -52,7 +52,7 @@ namespace UnrealBuildBase
 				{
 					if(CachedBaseDirectories == null)
 					{
-						HashSet<DirectoryReference> BaseDirs = new HashSet<DirectoryReference>();
+						HashSet<DirectoryReference> BaseDirs = [];
 						foreach (FileReference RootFile in DirectoryLookupCache.EnumerateFiles(Unreal.RootDirectory))
 						{
 							if(RootFile.HasExtension(".uprojectdirs"))
@@ -60,7 +60,7 @@ namespace UnrealBuildBase
 								foreach(string Line in File.ReadAllLines(RootFile.FullName))
 								{
 									string TrimLine = Line.Trim();
-									if(!TrimLine.StartsWith(";"))
+									if(!TrimLine.StartsWith(';'))
 									{
 										DirectoryReference BaseProjectDir = DirectoryReference.Combine(Unreal.RootDirectory, TrimLine);
 										if(BaseProjectDir.IsUnderDirectory(Unreal.RootDirectory))
@@ -94,7 +94,7 @@ namespace UnrealBuildBase
 				{
 					if(CachedProjectFiles == null)
 					{
-						HashSet<FileReference> ProjectFiles = new HashSet<FileReference>();
+						HashSet<FileReference> ProjectFiles = [];
 						foreach(DirectoryReference BaseDirectory in EnumerateBaseDirectories(Logger))
 						{
 							if(DirectoryLookupCache.DirectoryExists(BaseDirectory))
@@ -103,7 +103,7 @@ namespace UnrealBuildBase
 								{
 									// Ignore system directories (specifically for Mac, since temporary folders are created here)
 									string DirectoryName = SubDirectory.GetDirectoryName();
-									if (DirectoryName.StartsWith(".", StringComparison.Ordinal))
+									if (DirectoryName.StartsWith('.'))
 									{
 										continue;
 									}
diff --git a/Engine/Source/Programs/Shared/EpicGames.Build/System/Rules.cs b/Engine/Source/Programs/Shared/EpicGames.Build/System/Rules.cs
--- a/Engine/Source/Programs/Shared/EpicGames.Build/System/Rules.cs
+++ b/Engine/Source/Programs/Shared/EpicGames.Build/System/Rules.cs
@@ -42,15 +42,15 @@ namespace UnrealBuildBase
 		/// </summary>
 		class RulesFileCache
 		{
-			public List<FileReference> ModuleRules = new List<FileReference>();
-			public List<FileReference> TargetRules = new List<FileReference>();
-			public List<FileReference> AutomationModules = new List<FileReference>();
-			public List<FileReference> UbtPlugins = new List<FileReference>();
+			public List<FileReference> ModuleRules = [];
+			public List<FileReference> TargetRules = [];
+			public List<FileReference> AutomationModules = [];
+			public List<FileReference> UbtPlugins = [];
 		}
 
 		/// Map of root folders to a cached list of all UBT-related source files in that folder or any of its sub-folders.
 		/// We cache these file names so we can avoid searching for them later on.
-		static Dictionary<DirectoryReference, RulesFileCache> RootFolderToRulesFileCache = new Dictionary<DirectoryReference, RulesFileCache>();
+		static Dictionary<DirectoryReference, RulesFileCache> RootFolderToRulesFileCache = [];
 
 		/// <summary>
 		/// 
@@ -64,7 +64,7 @@ namespace UnrealBuildBase
 		/// <returns></returns>
 		public static List<FileReference> FindAllRulesSourceFiles(RulesFileType RulesFileType, List<DirectoryReference>? GameFolders, List<FileReference>? ForeignPlugins, List<DirectoryReference>? AdditionalSearchPaths, bool bIncludeEngine = true, bool bIncludeTempTargets = true)
 		{
-			List<DirectoryReference> Folders = new List<DirectoryReference>();
+			List<DirectoryReference> Folders = [];
 
 			// Add all engine source (including third party source)
 			if (bIncludeEngine)
@@ -75,7 +75,7 @@ namespace UnrealBuildBase
 			// @todo plugin: Disallow modules from including plugin modules as dependency modules? (except when the module is part of that plugin)
 
 			// Get all the root folders for plugins
-			List<DirectoryReference> RootFolders = new List<DirectoryReference>();
+			List<DirectoryReference> RootFolders = [];
 			if (bIncludeEngine)
 			{
 				RootFolders.AddRange(Unreal.GetExtensionDirs(Unreal.EngineDirectory));
@@ -205,7 +205,7 @@ namespace UnrealBuildBase
 		/// <returns>List of rules files of the given type</returns>
 		public static IReadOnlyList<FileReference> FindAllRulesFiles(DirectoryReference Directory, RulesFileType Type)
 		{
-			return FindAllRulesFiles(new List<DirectoryReference> { Directory }, Type);
+			return FindAllRulesFiles([Directory], Type);
 		}
 
 		/// <summary>
@@ -217,7 +217,7 @@ namespace UnrealBuildBase
 		public static List<FileReference> FindAllRulesFiles(IEnumerable<DirectoryReference> Directories, RulesFileType Type)
 		{
 			List<(DirectoryReference, RulesFileCache)> Caches = PrefetchRulesFilesInternal(Directories);
-			List<FileReference> Files = new List<FileReference>();			
+			List<FileReference> Files = [];			
 			foreach ((DirectoryReference Directory, RulesFileCache Cache) in Caches)
 			{
 				// Get the list of files of the type we're looking for
diff --git a/Engine/Source/Programs/Shared/EpicGames.Build/UEBuildPlatformSDK.cs b/Engine/Source/Programs/Shared/EpicGames.Build/UEBuildPlatformSDK.cs
--- a/Engine/Source/Programs/Shared/EpicGames.Build/UEBuildPlatformSDK.cs
+++ b/Engine/Source/Programs/Shared/EpicGames.Build/UEBuildPlatformSDK.cs
@@ -11,7 +11,6 @@ using System.Text.Json;
 using Microsoft.Extensions.Logging;
 using UnrealBuildBase;
 using EpicGames.Core;
-using System.Runtime.InteropServices;
 
 namespace EpicGames.Core
 {
@@ -132,7 +131,7 @@ namespace EpicGames.Core
 
 
 		public UEBuildPlatformSDK PlatformSDK;
-		public List<SDKDescriptor> Sdks = new List<SDKDescriptor>();
+		public List<SDKDescriptor> Sdks = [];
 		public List<SDKDescriptor> FullSdks => Sdks.Where(x => String.Compare(x.Name, "AutoSdk", true) != 0).ToList();
 		public SDKDescriptor? AutoSdk => Sdks.FirstOrDefault(x => String.Compare(x.Name, "AutoSdk", true) == 0);
 
@@ -337,12 +336,12 @@ namespace EpicGames.Core
 				if (ParentManualSDKVersions == "__None")
 				{
 					// empty it out
-					CachedManualSDKVersions = new Dictionary<string, string>();
+					CachedManualSDKVersions = [];
 				}
 				else
 				{
 					Logger.LogInformation("Processing parent manual sdk: {SDKVersion}", ParentManualSDKVersions);
-					CachedManualSDKVersions = JsonSerializer.Deserialize<Dictionary<string, string>>(ParentManualSDKVersions) ?? new Dictionary<string, string>();
+					CachedManualSDKVersions = JsonSerializer.Deserialize<Dictionary<string, string>>(ParentManualSDKVersions) ?? [];
 				}
 			}
 			else
@@ -357,7 +356,7 @@ namespace EpicGames.Core
 					}
 				}
 
-				CachedManualSDKVersions = new Dictionary<string, string>();
+				CachedManualSDKVersions = [];
 
 				// if there was no parent, get the SDK version before we run AutoSDK to get the manual version
 				SDKCollection InstalledVersions = GetInstalledSDKVersions();
@@ -395,8 +394,7 @@ namespace EpicGames.Core
 			{
 				if (!TempSDKRegistry.TryGetValue(PlatformName, out SDK))
 				{
-					object[] parameter = new object[1];
-					parameter[0] = Log.Logger;
+					object[] parameter = [Log.Logger];
 					SDK = (UEBuildPlatformSDK)Activator.CreateInstance(typeof(T), parameter)!;
 					// by setting this to false, we don't require any of the RequiredVersions to exist, but if they do, they will be read
 					// this is useful on other platforms that 
@@ -412,7 +410,7 @@ namespace EpicGames.Core
 		/// <summary>
 		/// Gets the set of all known SDKs
 		/// </summary>
-		public static UEBuildPlatformSDK[] AllPlatformSDKObjects => SDKRegistry.Values.ToArray();
+		public static UEBuildPlatformSDK[] AllPlatformSDKObjects => [.. SDKRegistry.Values];
 
 		// True if at least one platform has had its version changed from the standard SDK version, which means it may not be compatible with other projects
 		public static bool bHasAnySDKOverride = false;
@@ -421,7 +419,7 @@ namespace EpicGames.Core
 		public bool bHasSDKOverride = false;
 
 		// Contains a list of projects that had per-project SDK overrides, used for validating conflicting SDKs
-		public List<FileReference> ProjectsThatOverrodeSDK = new();
+		public List<FileReference> ProjectsThatOverrodeSDK = [];
 
 		// String name of the platform (will match an UnrealTargetPlatform)
 		public string? PlatformName;
@@ -659,7 +657,7 @@ namespace EpicGames.Core
 		/// <returns></returns>
 		public virtual string[] GetAllInstalledSDKVersions()
 		{
-			return new string[] { };
+			return [];
 		}
 		
 		/// <summary>
@@ -685,7 +683,7 @@ namespace EpicGames.Core
 			GetValidVersionRange(out string? MinVersion, out string? MaxVersion);
 			if (TryConvertVersionToInt(MinVersion, out ulong MinVersionInt, null) && TryConvertVersionToInt(MaxVersion, out ulong MaxVersionInt, null))
 			{
-				List<Tuple<ulong,ulong>> BannedSdkVersionRanges = new();
+				List<Tuple<ulong,ulong>> BannedSdkVersionRanges = [];
 				foreach (string SdkVersionPair in GetStringArrayFromConfig("BannedSdkVersions"))
 				{
 					string[] Versions = SdkVersionPair.Split("-");
@@ -937,7 +935,7 @@ namespace EpicGames.Core
 
 		#region Per-project SDK support
 
-		private static List<FileReference> ProjectsToCheckForSDKOverrides = new();
+		private static List<FileReference> ProjectsToCheckForSDKOverrides = [];
 		public static void InitializePerProjectSDKVersions(IEnumerable<FileReference> ProjectsToCheckForOverrides)
 		{
 			ProjectsToCheckForSDKOverrides = ProjectsToCheckForOverrides.ToList();
@@ -1035,12 +1033,12 @@ namespace EpicGames.Core
 		#region Private/Protected general functionality
 
 		// this is the SDK version that was set before activating AutoSDK, since AutoSDK may remove ability to retrieve the Manual SDK version
-		protected Dictionary<string, string> CachedManualSDKVersions = new Dictionary<string, string>();
-		private static Dictionary<string, UEBuildPlatformSDK> SDKRegistry = new Dictionary<string, UEBuildPlatformSDK>();
+		protected Dictionary<string, string> CachedManualSDKVersions = [];
+		private static Dictionary<string, UEBuildPlatformSDK> SDKRegistry = [];
 
 		// this map holds on to some temporary SDK objects that are generally used once and don't want to stick around, but they could be used multiple times, 
 		// like in MicrosoftPlatofrmSDK.Version.cs, if one of the versions is needed, C# will go construct every version, needing the SDK multiple times
-		private static Dictionary<string, UEBuildPlatformSDK> TempSDKRegistry = new Dictionary<string, UEBuildPlatformSDK>();
+		private static Dictionary<string, UEBuildPlatformSDK> TempSDKRegistry = [];
 
 		private SDKDescriptor? GetSDKVersionForHint(SDKCollection Collection, string? Hint)
 		{
@@ -1149,24 +1147,23 @@ namespace EpicGames.Core
 			}
 		}
 
+		static Lazy<JsonSerializerOptions>ProcessJsonFileOptions = new(() => new() {
+			AllowTrailingCommas = true,
+			ReadCommentHandling = JsonCommentHandling.Skip,
+		});
 		private void ProcessJsonFile(FileReference SDKConfigFile, Dictionary<string, string> VersionMap, Dictionary<string, string[]> VersionArrayMap)
 		{
 			string Contents = FileReference.ReadAllText(SDKConfigFile);
 
 			// load the json into a dictionary
-			JsonSerializerOptions Options = new()
-			{
-				AllowTrailingCommas = true,
-				ReadCommentHandling = JsonCommentHandling.Skip,
-			};
-			Dictionary<string, JsonElement>? LoadedDictionary = JsonSerializer.Deserialize<Dictionary<string, JsonElement>>(Contents, Options);
+			Dictionary<string, JsonElement>? LoadedDictionary = JsonSerializer.Deserialize<Dictionary<string, JsonElement>>(Contents, ProcessJsonFileOptions.Value);
 			if (LoadedDictionary == null)
 			{
 				throw new Exception($"Failed to parse SDK version file '{SDKConfigFile}'");
 			}
 
 			// look for special key to chain to other - after processing everything else (we only add from the "parent" settings not already in the map)
-			List<string> Parents = new();
+			List<string> Parents = [];
 			foreach (string Key in LoadedDictionary.Keys)
 			{
 				JsonElement Obj = LoadedDictionary[Key];
@@ -1180,10 +1177,7 @@ namespace EpicGames.Core
 						Parents.Add(StringValue);
 					}
 					// add this key if it's not already in the (case-insensitive) dictionary 
-					else if (!VersionMap.ContainsKey(Key))
-					{
-						VersionMap.Add(Key, StringValue);
-					}
+					VersionMap.TryAdd(Key, StringValue);
 				}
 				else if (ArrayValue != null)
 				{
@@ -1293,7 +1287,7 @@ namespace EpicGames.Core
 		public VersionNumberRange[] GetVersionNumberRangeArrayFromConfig(string VersionName)
 		{
 			string[]? VersionRanges;
-			List<VersionNumberRange> Ranges = new();
+			List<VersionNumberRange> Ranges = [];
 			if (ConfigSDKVersionArrays.TryGetValue(GetHostSpecificVersionName(VersionName), out VersionRanges) || ConfigSDKVersionArrays.TryGetValue(VersionName, out VersionRanges))
 			{
 				foreach (string VersionRange in VersionRanges)
@@ -1306,7 +1300,7 @@ namespace EpicGames.Core
 				}
 			}
 
-			return Ranges.ToArray();
+			return [.. Ranges];
 		}
 
 		public string[] GetStringArrayFromConfig(string Name)
@@ -1317,7 +1311,7 @@ namespace EpicGames.Core
 				return Results;
 			}
 
-			return Array.Empty<string>();
+			return [];
 		}
 
 		#endregion
@@ -1662,7 +1656,7 @@ namespace EpicGames.Core
 			}
 
 			// write all environment modifications
-			if ( (EnvVarValues != null && EnvVarValues.Any()) || (PathRemoves != null && PathRemoves.Any()) || (PathAdds != null && PathAdds.Any()) )
+			if ( (EnvVarValues != null && EnvVarValues.Count > 0) || (PathRemoves != null && PathRemoves.Count > 0) || (PathAdds != null && PathAdds.Count > 0) )
 			{
 				string ManualSDKEnvironmentVarsPath = Path.Combine(Unreal.EngineDirectory.ToString(), "Intermediate", ManualSDKEnvironmentVarsFile);
 
@@ -1773,19 +1767,19 @@ namespace EpicGames.Core
 			}
 		}
 
-		static private string[] AutoSDKLevels = new string[]
-		{
+		static private string[] AutoSDKLevels =
+		[
 			"NONE",
 			"BUILD",
 			"PACKAGE",
 			"RUN",
 			"FULL",
-		};
+		];
 
 		private bool IsSDKLevelAtLeast(string Value, string ComparedTo)
 		{
-			int ValueIndex = AutoSDKLevels.FindIndex(x => x == Value.ToUpper());
-			int ComparedIndex = AutoSDKLevels.FindIndex(x => x == ComparedTo.ToUpper());
+			int ValueIndex = AutoSDKLevels.FindIndex(x => x.Equals(Value, StringComparison.OrdinalIgnoreCase));
+			int ComparedIndex = AutoSDKLevels.FindIndex(x => x.Equals(ComparedTo, StringComparison.OrdinalIgnoreCase));
 			if (ValueIndex < 0 || ComparedIndex < 0)
 			{
 				throw new Exception($"Passed in a bad value to IsSDKLevelAtLeast: {Value}, {ComparedTo}");
@@ -1799,7 +1793,7 @@ namespace EpicGames.Core
 			string PlatformName = Path.GetFileName(PlatformSDKRoot).ToUpper();
 
 			// parse the envvar
-			Dictionary<string, string> PlatformSpecificLevels = new Dictionary<string, string>();
+			Dictionary<string, string> PlatformSpecificLevels = [];
 			string? DetailedSettings = Environment.GetEnvironmentVariable("UE_AUTOSDK_SPECIFIC_LEVELS"); // "Android=PACKAGE;GDk=RuN;PS5=BUILD";
 			if (!String.IsNullOrEmpty(DetailedSettings) && DetailedSettings.Contains(PlatformName, StringComparison.InvariantCultureIgnoreCase))
 			{
@@ -1815,9 +1809,9 @@ namespace EpicGames.Core
 			}
 
 			string FinalAutoSDKLevel = "FULL";
-			if (PlatformSpecificLevels.ContainsKey(PlatformName))
+			if (PlatformSpecificLevels.TryGetValue(PlatformName, out string? value))
 			{
-				FinalAutoSDKLevel = PlatformSpecificLevels[PlatformName];
+				FinalAutoSDKLevel = value;
 			}
 			else
 			{
@@ -1927,11 +1921,11 @@ namespace EpicGames.Core
 			{
 				using (StreamReader Reader = new StreamReader(EnvVarFile))
 				{
-					List<string> PathAdds = new List<string>();
-					List<string> PathRemoves = new List<string>();
+					List<string> PathAdds = [];
+					List<string> PathRemoves = [];
 
-					List<string> EnvVarNames = new List<string>();
-					List<string> EnvVarValues = new List<string>();
+					List<string> EnvVarNames = [];
+					List<string> EnvVarValues = [];
 
 					bool bNeedsToWriteAutoSetupEnvVar = true;
 					String PlatformSetupEnvVar = GetPlatformAutoSDKSetupEnvVar();
@@ -1986,7 +1980,7 @@ namespace EpicGames.Core
 					// actually perform the PATH stripping / adding.
 					String? OrigPathVar = Environment.GetEnvironmentVariable("PATH");
 					String PathDelimiter = RuntimePlatform.IsWindows ? ";" : ":";
-					String[] PathVars = { };
+					String[] PathVars = [];
 					if (!String.IsNullOrEmpty(OrigPathVar))
 					{
 						PathVars = OrigPathVar.Split(PathDelimiter.ToCharArray());
@@ -1996,8 +1990,7 @@ namespace EpicGames.Core
 						Logger.LogDebug("Path environment variable is null during AutoSDK");
 					}
 
-					List<String> ModifiedPathVars = new List<string>();
-					ModifiedPathVars.AddRange(PathVars);
+					List<String> ModifiedPathVars = [.. PathVars];
 
 					// perform removes first, in case they overlap with any adds.
 					foreach (String PathRemove in PathRemoves)
