EpicGames.Oodle: .net8 suggestions

#rnx

diff --git a/Engine/Source/Programs/Shared/EpicGames.Oodle/EpicGames.Oodle.csproj b/Engine/Source/Programs/Shared/EpicGames.Oodle/EpicGames.Oodle.csproj
--- a/Engine/Source/Programs/Shared/EpicGames.Oodle/EpicGames.Oodle.csproj
+++ b/Engine/Source/Programs/Shared/EpicGames.Oodle/EpicGames.Oodle.csproj
@@ -6,7 +6,7 @@
     <Configuration Condition=" '$(Configuration)' == '' ">Development</Configuration>
     <OutputType>Library</OutputType>
     <Deterministic>true</Deterministic>
-    <RootNamespace>EpicGames.Oodle</RootNamespace>
+    <RootNamespace>EpicGames.Compression</RootNamespace>
     <GenerateAssemblyInfo>false</GenerateAssemblyInfo>
     <GenerateTargetFrameworkAttribute>false</GenerateTargetFrameworkAttribute>
     <Configurations>Debug;Release;Development;Analyze</Configurations>
diff --git a/Engine/Source/Programs/Shared/EpicGames.Oodle/Oodle.cs b/Engine/Source/Programs/Shared/EpicGames.Oodle/Oodle.cs
--- a/Engine/Source/Programs/Shared/EpicGames.Oodle/Oodle.cs
+++ b/Engine/Source/Programs/Shared/EpicGames.Oodle/Oodle.cs
@@ -147,15 +147,8 @@ namespace EpicGames.Compression
 	/// <summary>
 	/// Base class for oodle exceptions
 	/// </summary>
-	public sealed class OodleException : Exception
+	public sealed class OodleException(string message) : Exception(message)
 	{
-		/// <summary>
-		/// Constructor
-		/// </summary>
-		public OodleException(string message)
-			: base(message)
-		{
-		}
 	}
 
 	/// <summary>
