EpicGames.MsBuild: .net8 suggestions

#rnx

diff --git a/Engine/Source/Programs/Shared/EpicGames.MsBuild/CsProjBuilder.cs b/Engine/Source/Programs/Shared/EpicGames.MsBuild/CsProjBuilder.cs
--- a/Engine/Source/Programs/Shared/EpicGames.MsBuild/CsProjBuilder.cs
+++ b/Engine/Source/Programs/Shared/EpicGames.MsBuild/CsProjBuilder.cs
@@ -24,9 +24,9 @@ namespace EpicGames.MsBuild
 	/// </summary>
 	public static class CsProjBuilder
 	{
-		class MLogger : IBuildLogger
+		class MLogger(ILogger inInner) : IBuildLogger
 		{
-			readonly ILogger _inner;
+			readonly ILogger _inner = inInner;
 
 			LoggerVerbosity IBuildLogger.Verbosity { get => LoggerVerbosity.Normal; set => throw new NotImplementedException(); }
 			string IBuildLogger.Parameters { get => throw new NotImplementedException(); set { } }
@@ -35,11 +35,6 @@ namespace EpicGames.MsBuild
 
 			bool _bFirstError = true;
 
-			public MLogger(ILogger inInner)
-			{
-				_inner = inInner;
-			}
-
 			void IBuildLogger.Initialize(IEventSource eventSource)
 			{
 				eventSource.ProjectStarted += new ProjectStartedEventHandler(EventSource_ProjectStarted);
@@ -201,11 +196,11 @@ namespace EpicGames.MsBuild
 				globalProperties.Add("DefineConstants", String.Join(';', defineConstants));
 			}
 
-			Dictionary<FileReference, CsProjBuildRecordEntry> buildRecords = new();
+			Dictionary<FileReference, CsProjBuildRecordEntry> buildRecords = [];
 
 			using ProjectCollection projectCollection = new ProjectCollection(globalProperties);
-			Dictionary<string, Project> projects = new Dictionary<string, Project>();
-			HashSet<string> skippedProjects = new HashSet<string>();
+			Dictionary<string, Project> projects = [];
+			HashSet<string> skippedProjects = [];
 
 			// Microsoft.Build.Evaluation.Project provides access to information stored in the .csproj xml that is 
 			// not available when using Microsoft.Build.Execution.ProjectInstance (used later in this function and
@@ -404,9 +399,9 @@ namespace EpicGames.MsBuild
 						continue;
 					}
 
-					List<string> include = new List<string>(glob.IncludeGlobs.Select(f => CleanGlobString(f))).OrderBy(x => x).ToList();
-					List<string> exclude = new List<string>(glob.Excludes.Select(f => CleanGlobString(f))).OrderBy(x => x).ToList();
-					List<string> remove = new List<string>(glob.Removes.Select(f => CleanGlobString(f))).OrderBy(x => x).ToList();
+					List<string> include = new List<string>(glob.IncludeGlobs.Select(f => CleanGlobString(f)).OrderBy(x => x));
+					List<string> exclude = new List<string>(glob.Excludes.Select(f => CleanGlobString(f)).OrderBy(x => x));
+					List<string> remove = new List<string>(glob.Removes.Select(f => CleanGlobString(f)).OrderBy(x => x));
 
 					buildRecord.Globs.Add(new CsProjBuildRecord.Glob()
 					{
@@ -489,12 +484,12 @@ namespace EpicGames.MsBuild
 			}
 
 			// write all build records
+			JsonSerializerOptions jsonOptions = new JsonSerializerOptions { WriteIndented = true };
 			foreach (ProjectGraphNode projectNode in inputProjectGraph.ProjectNodes)
 			{
 				FileReference projectPath = FileReference.FromString(projectNode.ProjectInstance.FullPath);
 				CsProjBuildRecordEntry entry = buildRecords[projectPath];
-				if (FileReference.WriteAllTextIfDifferent(entry.BuildRecordFile,
-					JsonSerializer.Serialize(entry.BuildRecord, new JsonSerializerOptions { WriteIndented = true })))
+				if (FileReference.WriteAllTextIfDifferent(entry.BuildRecordFile, JsonSerializer.Serialize(entry.BuildRecord, jsonOptions)))
 				{
 					logger.LogDebug("Wrote script module build record to {BuildRecordPath}", entry.BuildRecordFile);
 				}
@@ -503,7 +498,7 @@ namespace EpicGames.MsBuild
 			// todo: re-verify build records after a build to verify that everything is actually up to date
 
 			// even if only a subset was built, this function returns the full list of target assembly paths
-			Dictionary<FileReference, CsProjBuildRecordEntry> outDict = new();
+			Dictionary<FileReference, CsProjBuildRecordEntry> outDict = [];
 			foreach (ProjectGraphNode entryPointNode in inputProjectGraph.EntryPointNodes)
 			{
 				FileReference projectPath = FileReference.FromString(entryPointNode.ProjectInstance.FullPath);
@@ -517,13 +512,13 @@ namespace EpicGames.MsBuild
 			DateTime startTime = DateTime.UtcNow;
 			MLogger buildLogger = new MLogger(logger);
 
-			string[] targetsToBuild = { "Restore", "Build" };
+			string[] targetsToBuild = ["Restore", "Build"];
 
 			bool result = true;
 
 			foreach (string targetToBuild in targetsToBuild)
 			{
-				GraphBuildRequestData graphRequest = new GraphBuildRequestData(projectGraph, new string[] { targetToBuild });
+				GraphBuildRequestData graphRequest = new GraphBuildRequestData(projectGraph, [targetToBuild]);
 
 				BuildManager buildMan = BuildManager.DefaultBuildManager;
 
@@ -531,7 +526,7 @@ namespace EpicGames.MsBuild
 				buildParameters.AllowFailureWithoutError = false;
 				buildParameters.DetailedSummary = true;
 
-				buildParameters.Loggers = new List<IBuildLogger> { buildLogger };
+				buildParameters.Loggers = [buildLogger];
 				buildParameters.MaxNodeCount = 1; // msbuild bug - more than 1 here and the build stalls. Likely related to [REDACTED_PROFILE_URL]/msbuild/issues/1941
 
 				buildParameters.OnlyLogCriticalEvents = false;
@@ -576,7 +571,7 @@ namespace EpicGames.MsBuild
 			s_hasRegiteredMsBuildPath = true;
 
 			// Find our bundled dotnet SDK
-			List<string> listOfSdks = new List<string>();
+			List<string> listOfSdks = [];
 			ProcessStartInfo startInfo = new ProcessStartInfo
 			{
 				FileName = hook.DotnetPath.FullName,
