EpicGames.Serialization: .net8 suggestions

#rnx

diff --git a/Engine/Source/Programs/Shared/EpicGames.Serialization/CbConverter.cs b/Engine/Source/Programs/Shared/EpicGames.Serialization/CbConverter.cs
--- a/Engine/Source/Programs/Shared/EpicGames.Serialization/CbConverter.cs
+++ b/Engine/Source/Programs/Shared/EpicGames.Serialization/CbConverter.cs
@@ -12,22 +12,14 @@ namespace EpicGames.Serialization
 	/// <summary>
 	/// Attribute declaring the converter to use for a type
 	/// </summary>
+	/// <param name="converterType"></param>
 	[AttributeUsage(AttributeTargets.Class | AttributeTargets.Struct | AttributeTargets.Property)]
-	public sealed class CbConverterAttribute : Attribute
+	public sealed class CbConverterAttribute(Type converterType) : Attribute
 	{
 		/// <summary>
 		/// The converter type
 		/// </summary>
-		public Type ConverterType { get; }
-
-		/// <summary>
-		/// Constructor
-		/// </summary>
-		/// <param name="converterType"></param>
-		public CbConverterAttribute(Type converterType)
-		{
-			ConverterType = converterType;
-		}
+		public Type ConverterType { get; } = converterType;
 	}
 
 	/// <summary>
@@ -144,13 +136,13 @@ namespace EpicGames.Serialization
 			}
 		}
 
-		static readonly Dictionary<PropertyInfo, ICbConverterMethods> s_propertyToMethods = new Dictionary<PropertyInfo, ICbConverterMethods>();
-		static readonly Dictionary<Type, ICbConverterMethods> s_typeToMethods = new Dictionary<Type, ICbConverterMethods>();
+		static readonly Dictionary<PropertyInfo, ICbConverterMethods> s_propertyToMethods = [];
+		static readonly Dictionary<Type, ICbConverterMethods> s_typeToMethods = [];
 
 		static ICbConverterMethods CreateWrapper(Type type, CbConverter converter)
 		{
 			Type converterMethodsType = typeof(CbConverterMethodsWrapper<,>).MakeGenericType(converter.GetType(), type);
-			return (ICbConverterMethods)Activator.CreateInstance(converterMethodsType, new object[] { converter })!;
+			return (ICbConverterMethods)Activator.CreateInstance(converterMethodsType, [converter])!;
 		}
 
 		/// <summary>
@@ -214,7 +206,7 @@ namespace EpicGames.Serialization
 		/// <summary>
 		/// Cache of property to converter
 		/// </summary>
-		static readonly Dictionary<PropertyInfo, CbConverter> s_propertyToConverter = new Dictionary<PropertyInfo, CbConverter>();
+		static readonly Dictionary<PropertyInfo, CbConverter> s_propertyToConverter = [];
 
 		/// <summary>
 		/// Cache of type to converter
@@ -242,8 +234,8 @@ namespace EpicGames.Serialization
 		/// <summary>
 		/// List of converter factories. Must be 
 		/// </summary>
-		public static List<CbConverterFactory> ConverterFactories { get; } = new List<CbConverterFactory>
-		{
+		public static List<CbConverterFactory> ConverterFactories { get; } =
+		[
 			new CbClassConverterFactory(),
 			new CbStringConverterFactory(),
 			new CbEnumConverterFactory(),
@@ -251,7 +243,7 @@ namespace EpicGames.Serialization
 			new CbArrayConverterFactory(),
 			new CbDictionaryConverterFactory(),
 			new CbNullableConverterFactory(),
-		};
+		];
 
 		/// <summary>
 		/// Class used to cache values returned by CreateConverterInfo without having to do a dictionary lookup.
@@ -277,20 +269,16 @@ namespace EpicGames.Serialization
 			/// <summary>
 			/// Wrapper class to convert an untyped converter into a typed one
 			/// </summary>
-			class CbConverterWrapper : CbConverter<T>
+			class CbConverterWrapper(CbConverter inner) : CbConverter<T>
 			{
-				readonly CbConverter _inner;
-
-				public CbConverterWrapper(CbConverter inner) => _inner = inner;
-
 				/// <inheritdoc/>
-				public override T Read(CbField field) => (T)_inner.ReadObject(field)!;
+				public override T Read(CbField field) => (T)inner.ReadObject(field)!;
 
 				/// <inheritdoc/>
-				public override void Write(CbWriter writer, T value) => _inner.WriteObject(writer, value);
+				public override void Write(CbWriter writer, T value) => inner.WriteObject(writer, value);
 
 				/// <inheritdoc/>
-				public override void WriteNamed(CbWriter writer, CbFieldName name, T value) => _inner.WriteNamedObject(writer, name, value);
+				public override void WriteNamed(CbWriter writer, CbFieldName name, T value) => inner.WriteNamedObject(writer, name, value);
 			}
 		}
 
diff --git a/Engine/Source/Programs/Shared/EpicGames.Serialization/CbField.cs b/Engine/Source/Programs/Shared/EpicGames.Serialization/CbField.cs
--- a/Engine/Source/Programs/Shared/EpicGames.Serialization/CbField.cs
+++ b/Engine/Source/Programs/Shared/EpicGames.Serialization/CbField.cs
@@ -191,10 +191,11 @@ namespace EpicGames.Serialization
 	/// <summary>
 	/// A binary attachment, referenced by <see cref="IoHash"/>
 	/// </summary>
+	/// <param name="hash">Hash of the referenced object</param>
 	[DebuggerDisplay("{Hash}")]
 	[JsonConverter(typeof(CbBinaryAttachmentJsonConverter))]
 	[TypeConverter(typeof(CbBinaryAttachmentTypeConverter))]
-	public readonly struct CbBinaryAttachment : IEquatable<CbBinaryAttachment>
+	public readonly struct CbBinaryAttachment(IoHash hash) : IEquatable<CbBinaryAttachment>
 	{
 		/// <summary>
 		/// Attachment with a hash of zero
@@ -204,16 +205,7 @@ namespace EpicGames.Serialization
 		/// <summary>
 		/// Hash of the referenced object
 		/// </summary>
-		public IoHash Hash { get; }
-
-		/// <summary>
-		/// Constructor
-		/// </summary>
-		/// <param name="hash">Hash of the referenced object</param>
-		public CbBinaryAttachment(IoHash hash)
-		{
-			Hash = hash;
-		}
+		public IoHash Hash { get; } = hash;
 
 		/// <inheritdoc/>
 		public override string ToString() => Hash.ToString();
@@ -273,10 +265,11 @@ namespace EpicGames.Serialization
 	/// <summary>
 	/// An object attachment, referenced by <see cref="IoHash"/>
 	/// </summary>
+	/// <param name="hash">Hash of the referenced object</param>
 	[DebuggerDisplay("{Hash}")]
 	[JsonConverter(typeof(CbObjectAttachmentJsonConverter))]
 	[TypeConverter(typeof(CbObjectAttachmentTypeConverter))]
-	public readonly struct CbObjectAttachment : IEquatable<CbObjectAttachment>
+	public readonly struct CbObjectAttachment(IoHash hash) : IEquatable<CbObjectAttachment>
 	{
 		/// <summary>
 		/// Attachment with a hash of zero
@@ -286,16 +279,7 @@ namespace EpicGames.Serialization
 		/// <summary>
 		/// Hash of the referenced object
 		/// </summary>
-		public IoHash Hash { get; }
-
-		/// <summary>
-		/// Constructor
-		/// </summary>
-		/// <param name="hash">Hash of the referenced object</param>
-		public CbObjectAttachment(IoHash hash)
-		{
-			Hash = hash;
-		}
+		public IoHash Hash { get; } = hash;
 
 		/// <inheritdoc/>
 		public override string ToString() => Hash.ToString();
@@ -915,7 +899,7 @@ namespace EpicGames.Serialization
 		/// <returns></returns>
 		public byte[] AsBinaryArray()
 		{
-			return AsBinaryArray(Array.Empty<byte>());
+			return AsBinaryArray([]);
 		}
 
 		/// <summary>
@@ -1520,8 +1504,7 @@ namespace EpicGames.Serialization
 		/// <param name="hasher"></param>
 		void AppendHash(Blake3.Hasher hasher)
 		{
-			Span<byte> data = stackalloc byte[1];
-			data[0] = (byte)CbFieldUtils.GetSerializedType(TypeWithFlags);
+			Span<byte> data = [(byte)CbFieldUtils.GetSerializedType(TypeWithFlags)];
 			hasher.Update(data);
 
 			hasher.Update(GetViewNoType().Span);
@@ -1910,17 +1893,10 @@ namespace EpicGames.Serialization
 	/// <summary>
 	/// Enumerator for contents of a field
 	/// </summary>
-	public sealed class CbFieldEnumerator : IEnumerator<CbField>
+	/// <param name="data"></param>
+	/// <param name="uniformType"></param>
+	public sealed class CbFieldEnumerator(ReadOnlyMemory<byte> data, CbFieldType uniformType) : IEnumerator<CbField>
 	{
-		/// <summary>
-		/// The underlying buffer
-		/// </summary>
-		readonly ReadOnlyMemory<byte> _data;
-
-		/// <summary>
-		/// Type for all fields
-		/// </summary>
-		CbFieldType UniformType { get; }
 
 		/// <inheritdoc/>
 		public CbField Current { get; private set; } = null!;
@@ -1928,17 +1904,6 @@ namespace EpicGames.Serialization
 		/// <inheritdoc/>
 		object? IEnumerator.Current => Current;
 
-		/// <summary>
-		/// Constructor
-		/// </summary>
-		/// <param name="data"></param>
-		/// <param name="uniformType"></param>
-		public CbFieldEnumerator(ReadOnlyMemory<byte> data, CbFieldType uniformType)
-		{
-			_data = data;
-			UniformType = uniformType;
-		}
-
 		/// <inheritdoc/>
 		public void Dispose()
 		{
@@ -1953,9 +1918,9 @@ namespace EpicGames.Serialization
 		/// <inheritdoc/>
 		public bool MoveNext()
 		{
-			if (_data.Length > 0)
+			if (data.Length > 0)
 			{
-				Current = new CbField(_data, UniformType);
+				Current = new CbField(data, uniformType);
 				return true;
 			}
 			else
@@ -1971,7 +1936,7 @@ namespace EpicGames.Serialization
 		/// <returns></returns>
 		public CbFieldEnumerator Clone()
 		{
-			return new CbFieldEnumerator(_data, UniformType);
+			return new CbFieldEnumerator(data, uniformType);
 		}
 	}
 
@@ -2172,14 +2137,10 @@ namespace EpicGames.Serialization
 	/// <summary>
 	/// Simplified view of <see cref="CbArray"/> for display in the debugger
 	/// </summary>
-	class CbArrayDebugView
+	class CbArrayDebugView(CbArray array)
 	{
-		readonly CbArray _array;
-
-		public CbArrayDebugView(CbArray array) => _array = array;
-
 		[DebuggerBrowsable(DebuggerBrowsableState.RootHidden)]
-		public object?[] Value => _array.Select(x => x.Value).ToArray();
+		public object?[] Value => array.Select(x => x.Value).ToArray();
 	}
 
 	/// <summary>
@@ -2286,7 +2247,7 @@ namespace EpicGames.Serialization
 		/// </summary>
 		public void AppendHash(Blake3.Hasher hasher)
 		{
-			byte[] serializedType = new byte[] { (byte)_innerField.GetType() };
+			byte[] serializedType = [(byte)_innerField.GetType()];
 			hasher.Update(serializedType);
 			hasher.Update(_innerField.Payload.Span);
 		}
@@ -2365,7 +2326,7 @@ namespace EpicGames.Serialization
 	/// <summary>
 	/// Simplified view of <see cref="CbObject"/> for display in the debugger
 	/// </summary>
-	class CbObjectDebugView
+	class CbObjectDebugView(CbObject obj)
 	{
 		[DebuggerDisplay("{Name}: {Value}")]
 		public class Property
@@ -2374,12 +2335,8 @@ namespace EpicGames.Serialization
 			public object? Value { get; set; }
 		}
 
-		readonly CbObject _object;
-
-		public CbObjectDebugView(CbObject obj) => _object = obj;
-
 		[DebuggerBrowsable(DebuggerBrowsableState.RootHidden)]
-		public Property[] Properties => _object.Select(x => new Property { Name = x.Name.ToString(), Value = x.Value }).ToArray();
+		public Property[] Properties => obj.Select(x => new Property { Name = x.Name.ToString(), Value = x.Value }).ToArray();
 	}
 
 	/// <summary>
@@ -2507,7 +2464,7 @@ namespace EpicGames.Serialization
 		/// <param name="hasher"></param>
 		public void AppendHash(Blake3.Hasher hasher)
 		{
-			byte[] temp = new byte[] { (byte)_innerField.GetType() };
+			byte[] temp = [(byte)_innerField.GetType()];
 			hasher.Update(temp);
 			hasher.Update(_innerField.Payload.Span);
 		}
diff --git a/Engine/Source/Programs/Shared/EpicGames.Serialization/CbPackage.cs b/Engine/Source/Programs/Shared/EpicGames.Serialization/CbPackage.cs
--- a/Engine/Source/Programs/Shared/EpicGames.Serialization/CbPackage.cs
+++ b/Engine/Source/Programs/Shared/EpicGames.Serialization/CbPackage.cs
@@ -171,7 +171,7 @@ namespace EpicGames.Serialization
 		private readonly CbPackageHeader _header;
 #pragma warning restore IDE0052 // Remove unread private members
 		private readonly Stream _attachmentsStream;
-		private readonly List<CbPackageAttachmentEntry> _attachmentEntries = new List<CbPackageAttachmentEntry>();
+		private readonly List<CbPackageAttachmentEntry> _attachmentEntries = [];
 
 		private CbPackageReader(CbPackageHeader header, CbObject rootObject, IoHash rootHash, List<CbPackageAttachmentEntry> attachmentEntries, Stream attachmentsStream)
 		{
@@ -210,7 +210,7 @@ namespace EpicGames.Serialization
 		{
 			CbPackageHeader header = await CbPackageHeader.ReadAsync(s);
 
-			List<CbPackageAttachmentEntry> entries = new List<CbPackageAttachmentEntry>();
+			List<CbPackageAttachmentEntry> entries = [];
 			// we expect a extra entry for the root object
 			uint attachmentsCount = header.AttachmentCount + 1;
 
@@ -287,7 +287,7 @@ namespace EpicGames.Serialization
 	/// </summary>
 	public sealed class CbPackageBuilder : IDisposable
 	{
-		private readonly List<(CbPackageAttachmentEntry, Stream)> _streamAttachments = new List<(CbPackageAttachmentEntry, Stream)>();
+		private readonly List<(CbPackageAttachmentEntry, Stream)> _streamAttachments = [];
 
 		/// <summary>
 		/// Constructor for the package builder
diff --git a/Engine/Source/Programs/Shared/EpicGames.Serialization/CbSerializer.cs b/Engine/Source/Programs/Shared/EpicGames.Serialization/CbSerializer.cs
--- a/Engine/Source/Programs/Shared/EpicGames.Serialization/CbSerializer.cs
+++ b/Engine/Source/Programs/Shared/EpicGames.Serialization/CbSerializer.cs
@@ -60,19 +60,14 @@ namespace EpicGames.Serialization
 	/// <summary>
 	/// Sets the name used for discriminating between derived classes during serialization
 	/// </summary>
+	/// <param name="name">Name used to identify this class</param>
 	[AttributeUsage(AttributeTargets.Class)]
-	public sealed class CbDiscriminatorAttribute : Attribute
+	public sealed class CbDiscriminatorAttribute(string name) : Attribute
 	{
 		/// <summary>
 		/// Name used to identify this class
 		/// </summary>
-		public string Name { get; }
-
-		/// <summary>
-		/// Constructor
-		/// </summary>
-		/// <param name="name">Name used to identify this class</param>
-		public CbDiscriminatorAttribute(string name) => Name = name;
+		public string Name { get; } = name;
 	}
 
 	/// <summary>
@@ -94,21 +89,15 @@ namespace EpicGames.Serialization
 	/// <summary>
 	/// Exception indicating that a class does not have any fields to serialize
 	/// </summary>
-	public sealed class CbEmptyClassException : CbException
+	/// <remarks>
+	/// Constructor
+	/// </remarks>
+	public sealed class CbEmptyClassException(Type classType) : CbException($"{classType.Name} does not have any fields marked with a [CbField] attribute. If this is intended, explicitly mark the class with a [CbObject] attribute.")
 	{
 		/// <summary>
 		/// Type with missing field annotations
 		/// </summary>
-		public Type ClassType { get; }
-
-		/// <summary>
-		/// Constructor
-		/// </summary>
-		public CbEmptyClassException(Type classType)
-			: base($"{classType.Name} does not have any fields marked with a [CbField] attribute. If this is intended, explicitly mark the class with a [CbObject] attribute.")
-		{
-			ClassType = classType;
-		}
+		public Type ClassType { get; } = classType;
 	}
 
 	/// <summary>
diff --git a/Engine/Source/Programs/Shared/EpicGames.Serialization/CbWriter.cs b/Engine/Source/Programs/Shared/EpicGames.Serialization/CbWriter.cs
--- a/Engine/Source/Programs/Shared/EpicGames.Serialization/CbWriter.cs
+++ b/Engine/Source/Programs/Shared/EpicGames.Serialization/CbWriter.cs
@@ -569,7 +569,7 @@ namespace EpicGames.Serialization
 		/// <returns>Sequence of segments</returns>
 		public List<ReadOnlyMemory<byte>> GetSegments()
 		{
-			List<ReadOnlyMemory<byte>> segments = new List<ReadOnlyMemory<byte>>();
+			List<ReadOnlyMemory<byte>> segments = [];
 			GetSegments(_rootScope, segments);
 			return segments;
 		}
@@ -592,18 +592,11 @@ namespace EpicGames.Serialization
 		/// <returns>New stream for the contents of this object</returns>
 		public Stream AsStream() => new ReadStream(GetSegments().GetEnumerator(), GetSize());
 
-		class ReadStream : Stream
+		class ReadStream(IEnumerator<ReadOnlyMemory<byte>> segments, long length) : Stream
 		{
-			readonly IEnumerator<ReadOnlyMemory<byte>> _segments;
 			ReadOnlyMemory<byte> _segment;
 			long _positionInternal;
 
-			public ReadStream(IEnumerator<ReadOnlyMemory<byte>> segments, long length)
-			{
-				_segments = segments;
-				Length = length;
-			}
-
 			/// <inheritdoc/>
 			public override bool CanRead => true;
 
@@ -614,7 +607,7 @@ namespace EpicGames.Serialization
 			public override bool CanWrite => false;
 
 			/// <inheritdoc/>
-			public override long Length { get; }
+			public override long Length { get; } = length;
 
 			/// <inheritdoc/>
 			public override long Position
@@ -634,11 +627,11 @@ namespace EpicGames.Serialization
 				{
 					while (_segment.Length == 0)
 					{
-						if (!_segments.MoveNext())
+						if (!segments.MoveNext())
 						{
 							return readLength;
 						}
-						_segment = _segments.Current;
+						_segment = segments.Current;
 					}
 
 					int copyLength = Math.Min(_segment.Length, buffer.Length);
@@ -676,8 +669,8 @@ namespace EpicGames.Serialization
 		/// </summary>
 		public const int DefaultChunkSize = 1024;
 
-		readonly List<byte[]> _chunks = new List<byte[]>();
-		readonly List<byte[]> _freeChunks = new List<byte[]>();
+		readonly List<byte[]> _chunks = [];
+		readonly List<byte[]> _freeChunks = [];
 
 		/// <summary>
 		/// Constructor
diff --git a/Engine/Source/Programs/Shared/EpicGames.Serialization/Converters/CbArrayConverter.cs b/Engine/Source/Programs/Shared/EpicGames.Serialization/Converters/CbArrayConverter.cs
--- a/Engine/Source/Programs/Shared/EpicGames.Serialization/Converters/CbArrayConverter.cs
+++ b/Engine/Source/Programs/Shared/EpicGames.Serialization/Converters/CbArrayConverter.cs
@@ -19,12 +19,12 @@ namespace EpicGames.Serialization.Converters
 				return null!;
 			}
 
-			List<T> list = new List<T>();
+			List<T> list = [];
 			foreach (CbField elementField in field)
 			{
 				list.Add(CbSerializer.Deserialize<T>(elementField));
 			}
-			return list.ToArray();
+			return [.. list];
 		}
 
 		/// <inheritdoc/>
diff --git a/Engine/Source/Programs/Shared/EpicGames.Serialization/Converters/CbClassConverter.cs b/Engine/Source/Programs/Shared/EpicGames.Serialization/Converters/CbClassConverter.cs
--- a/Engine/Source/Programs/Shared/EpicGames.Serialization/Converters/CbClassConverter.cs
+++ b/Engine/Source/Programs/Shared/EpicGames.Serialization/Converters/CbClassConverter.cs
@@ -38,7 +38,7 @@ namespace EpicGames.Serialization.Converters
 		public DynamicMethod ReadConcreteMethod { get; }
 		public DynamicMethod WriteConcreteContentsMethod { get; }
 
-		static readonly Dictionary<Type, CbClassConverterMethods> s_typeToMethods = new Dictionary<Type, CbClassConverterMethods>();
+		static readonly Dictionary<Type, CbClassConverterMethods> s_typeToMethods = [];
 
 		public static CbClassConverterMethods Create(Type classType)
 		{
@@ -57,16 +57,16 @@ namespace EpicGames.Serialization.Converters
 			ClassType = classType;
 			IsPolymorphic = (classType.GetCustomAttribute<CbPolymorphicAttribute>(true) != null);
 
-			ReadConcreteMethod = new DynamicMethod($"ReadConcrete_{classType.Name}", classType, new Type[] { typeof(CbField) });
-			WriteConcreteContentsMethod = new DynamicMethod($"WriteConcreteContents_{classType.Name}", null, new Type[] { typeof(CbWriter), classType });
+			ReadConcreteMethod = new DynamicMethod($"ReadConcrete_{classType.Name}", classType, [typeof(CbField)]);
+			WriteConcreteContentsMethod = new DynamicMethod($"WriteConcreteContents_{classType.Name}", null, [typeof(CbWriter), classType]);
 
-			WriteMethod = new DynamicMethod($"Write_{classType.Name}", null, new Type[] { typeof(CbWriter), classType });
-			WriteNamedMethod = new DynamicMethod($"WriteNamed_{classType.Name}", null, new Type[] { typeof(CbWriter), typeof(CbFieldName), classType });
+			WriteMethod = new DynamicMethod($"Write_{classType.Name}", null, [typeof(CbWriter), classType]);
+			WriteNamedMethod = new DynamicMethod($"WriteNamed_{classType.Name}", null, [typeof(CbWriter), typeof(CbFieldName), classType]);
 
 			if (IsPolymorphic)
 			{
-				ReadMethod = new DynamicMethod($"Read_{classType.Name}", classType, new Type[] { typeof(CbField) });
-				WriteContentsMethod = new DynamicMethod($"WriteContents_{classType.Name}", null, new Type[] { typeof(CbWriter), classType });
+				ReadMethod = new DynamicMethod($"Read_{classType.Name}", classType, [typeof(CbField)]);
+				WriteContentsMethod = new DynamicMethod($"WriteContents_{classType.Name}", null, [typeof(CbWriter), classType]);
 			}
 			else
 			{
@@ -114,7 +114,7 @@ namespace EpicGames.Serialization.Converters
 
 		static void PopulateDispatchType(Type classType, Type dispatchType)
 		{
-			Dictionary<Utf8String, Type> discriminatorToKnownType = new Dictionary<Utf8String, Type>();
+			Dictionary<Utf8String, Type> discriminatorToKnownType = [];
 
 			Type[] knownTypes = classType.Assembly.GetTypes();
 			foreach (Type knownType in knownTypes)
@@ -141,7 +141,7 @@ namespace EpicGames.Serialization.Converters
 				CbClassConverterMethods methods = Create(knownType);
 
 				{
-					DynamicMethod dynamicMethod = new DynamicMethod("_", typeof(object), new Type[] { typeof(CbField) });
+					DynamicMethod dynamicMethod = new DynamicMethod("_", typeof(object), [typeof(CbField)]);
 					ILGenerator generator = dynamicMethod.GetILGenerator();
 					generator.Emit(OpCodes.Ldarg_0);
 					generator.Emit(OpCodes.Call, methods.ReadConcreteMethod);
@@ -151,7 +151,7 @@ namespace EpicGames.Serialization.Converters
 				}
 
 				{
-					DynamicMethod dynamicMethod = new DynamicMethod("_", null, new Type[] { typeof(CbWriter), typeof(object) });
+					DynamicMethod dynamicMethod = new DynamicMethod("_", null, [typeof(CbWriter), typeof(object)]);
 					ILGenerator generator = dynamicMethod.GetILGenerator();
 					generator.Emit(OpCodes.Ldarg_0);
 					generator.Emit(OpCodes.Ldarg_1);
@@ -276,8 +276,8 @@ namespace EpicGames.Serialization.Converters
 
 		class CbPolymorphicDispatch<T>
 		{
-			public static Dictionary<Utf8String, Func<CbField, object>> NameToReadFunc = new Dictionary<Utf8String, Func<CbField, object>>();
-			public static Dictionary<Type, Action<CbWriter, object>> TypeToWriteContentsFunc = new Dictionary<Type, Action<CbWriter, object>>();
+			public static Dictionary<Utf8String, Func<CbField, object>> NameToReadFunc = [];
+			public static Dictionary<Type, Action<CbWriter, object>> TypeToWriteContentsFunc = [];
 
 			public static object Read(CbField field)
 			{
@@ -460,7 +460,7 @@ namespace EpicGames.Serialization.Converters
 
 		static (Utf8String, PropertyInfo)[] GetProperties(Type type)
 		{
-			List<(Utf8String, PropertyInfo)> propertyList = new List<(Utf8String, PropertyInfo)>();
+			List<(Utf8String, PropertyInfo)> propertyList = [];
 			foreach (PropertyInfo property in type.GetProperties(BindingFlags.Public | BindingFlags.NonPublic | BindingFlags.Instance))
 			{
 				CbIgnoreAttribute? ignoreAttribute = property.GetCustomAttribute<CbIgnoreAttribute>();
@@ -478,7 +478,7 @@ namespace EpicGames.Serialization.Converters
 			{
 				throw new CbEmptyClassException(type);
 			}
-			return propertyList.ToArray();
+			return [.. propertyList];
 		}
 
 		static FieldInfo GetFieldInfo<T>(Expression<Func<T>> expr)
diff --git a/Engine/Source/Programs/Shared/EpicGames.Serialization/Converters/CbDictionaryConverter.cs b/Engine/Source/Programs/Shared/EpicGames.Serialization/Converters/CbDictionaryConverter.cs
--- a/Engine/Source/Programs/Shared/EpicGames.Serialization/Converters/CbDictionaryConverter.cs
+++ b/Engine/Source/Programs/Shared/EpicGames.Serialization/Converters/CbDictionaryConverter.cs
@@ -18,7 +18,7 @@ namespace EpicGames.Serialization.Converters
 				return null!;
 			}
 
-			Dictionary<TKey, TValue> dictionary = new Dictionary<TKey, TValue>();
+			Dictionary<TKey, TValue> dictionary = [];
 			foreach (CbField element in field)
 			{
 				IEnumerator<CbField> enumerator = element.AsArray().GetEnumerator();
diff --git a/Engine/Source/Programs/Shared/EpicGames.Serialization/Converters/CbEnumConverter.cs b/Engine/Source/Programs/Shared/EpicGames.Serialization/Converters/CbEnumConverter.cs
--- a/Engine/Source/Programs/Shared/EpicGames.Serialization/Converters/CbEnumConverter.cs
+++ b/Engine/Source/Programs/Shared/EpicGames.Serialization/Converters/CbEnumConverter.cs
@@ -17,15 +17,15 @@ namespace EpicGames.Serialization.Converters
 		{
 			Type type = typeof(T);
 
-			ReadMethod = new DynamicMethod($"Read_{type.Name}", type, new Type[] { typeof(CbField) });
+			ReadMethod = new DynamicMethod($"Read_{type.Name}", type, [typeof(CbField)]);
 			CreateEnumReader(ReadMethod.GetILGenerator());
 			_readFunc = (Func<CbField, T>)ReadMethod.CreateDelegate(typeof(Func<CbField, T>));
 
-			WriteMethod = new DynamicMethod($"Write_{type.Name}", null, new Type[] { typeof(CbWriter), type });
+			WriteMethod = new DynamicMethod($"Write_{type.Name}", null, [typeof(CbWriter), type]);
 			CreateEnumWriter(WriteMethod.GetILGenerator());
 			_writeFunc = (Action<CbWriter, T>)WriteMethod.CreateDelegate(typeof(Action<CbWriter, T>));
 
-			WriteNamedMethod = new DynamicMethod($"WriteNamed_{type.Name}", null, new Type[] { typeof(CbWriter), typeof(CbFieldName), type });
+			WriteNamedMethod = new DynamicMethod($"WriteNamed_{type.Name}", null, [typeof(CbWriter), typeof(CbFieldName), type]);
 			CreateNamedEnumWriter(WriteNamedMethod.GetILGenerator());
 			_writeNamedFunc = (Action<CbWriter, CbFieldName, T>)WriteNamedMethod.CreateDelegate(typeof(Action<CbWriter, CbFieldName, T>));
 		}
diff --git a/Engine/Source/Programs/Shared/EpicGames.Serialization/Converters/CbListConverter.cs b/Engine/Source/Programs/Shared/EpicGames.Serialization/Converters/CbListConverter.cs
--- a/Engine/Source/Programs/Shared/EpicGames.Serialization/Converters/CbListConverter.cs
+++ b/Engine/Source/Programs/Shared/EpicGames.Serialization/Converters/CbListConverter.cs
@@ -21,7 +21,7 @@ namespace EpicGames.Serialization.Converters
 			}
 			else
 			{
-				List<T> list = new List<T>();
+				List<T> list = [];
 				foreach (CbField elementField in field)
 				{
 					list.Add(CbSerializer.Deserialize<T>(elementField));
@@ -75,7 +75,7 @@ namespace EpicGames.Serialization.Converters
 		/// <inheritdoc/>
 		public override List<Utf8String> Read(CbField field)
 		{
-			List<Utf8String> list = new List<Utf8String>();
+			List<Utf8String> list = [];
 			foreach (CbField elementField in field)
 			{
 				list.Add(elementField.AsUtf8String());
diff --git a/Engine/Source/Programs/Shared/EpicGames.Serialization/Converters/CbPrimitiveConverter.cs b/Engine/Source/Programs/Shared/EpicGames.Serialization/Converters/CbPrimitiveConverter.cs
--- a/Engine/Source/Programs/Shared/EpicGames.Serialization/Converters/CbPrimitiveConverter.cs
+++ b/Engine/Source/Programs/Shared/EpicGames.Serialization/Converters/CbPrimitiveConverter.cs
@@ -6,28 +6,16 @@ using System.Reflection;
 
 namespace EpicGames.Serialization.Converters
 {
-	class CbPrimitiveConverter<T> : CbConverter<T>, ICbConverterMethods
+	class CbPrimitiveConverter<T>(Expression<Func<CbField, T>> read, Expression<Action<CbWriter, T>> write, Expression<Action<CbWriter, CbFieldName, T>> writeNamed) : CbConverter<T>, ICbConverterMethods
 	{
-		public MethodInfo ReadMethod { get; }
-		public Func<CbField, T> ReadFunc { get; }
+		public MethodInfo ReadMethod { get; } = ((MethodCallExpression)read.Body).Method;
+		public Func<CbField, T> ReadFunc { get; } = read.Compile();
 
-		public MethodInfo WriteMethod { get; }
-		public Action<CbWriter, T> WriteFunc { get; }
+		public MethodInfo WriteMethod { get; } = ((MethodCallExpression)write.Body).Method;
+		public Action<CbWriter, T> WriteFunc { get; } = write.Compile();
 
-		public MethodInfo WriteNamedMethod { get; }
-		public Action<CbWriter, CbFieldName, T> WriteNamedFunc { get; }
-
-		public CbPrimitiveConverter(Expression<Func<CbField, T>> read, Expression<Action<CbWriter, T>> write, Expression<Action<CbWriter, CbFieldName, T>> writeNamed)
-		{
-			ReadMethod = ((MethodCallExpression)read.Body).Method;
-			ReadFunc = read.Compile();
-
-			WriteMethod = ((MethodCallExpression)write.Body).Method;
-			WriteFunc = write.Compile();
-
-			WriteNamedMethod = ((MethodCallExpression)writeNamed.Body).Method;
-			WriteNamedFunc = writeNamed.Compile();
-		}
+		public MethodInfo WriteNamedMethod { get; } = ((MethodCallExpression)writeNamed.Body).Method;
+		public Action<CbWriter, CbFieldName, T> WriteNamedFunc { get; } = writeNamed.Compile();
 
 		public override T Read(CbField field) => ReadFunc(field);
 
diff --git a/Engine/Source/Programs/Shared/EpicGames.Serialization/Converters/CbStringConverter.cs b/Engine/Source/Programs/Shared/EpicGames.Serialization/Converters/CbStringConverter.cs
--- a/Engine/Source/Programs/Shared/EpicGames.Serialization/Converters/CbStringConverter.cs
+++ b/Engine/Source/Programs/Shared/EpicGames.Serialization/Converters/CbStringConverter.cs
@@ -5,28 +5,21 @@ using System.ComponentModel;
 
 namespace EpicGames.Serialization.Converters
 {
-	class CbStringConverter<T> : CbConverter<T>
+	class CbStringConverter<T>(TypeConverter typeConverter) : CbConverter<T>
 	{
-		readonly TypeConverter _typeConverter;
-
-		public CbStringConverter(TypeConverter typeConverter)
-		{
-			_typeConverter = typeConverter;
-		}
-
 		public override T Read(CbField field)
 		{
-			return (T)_typeConverter.ConvertFromInvariantString(field.AsString())!;
+			return (T)typeConverter.ConvertFromInvariantString(field.AsString())!;
 		}
 
 		public override void WriteNamed(CbWriter writer, CbFieldName name, T value)
 		{
-			writer.WriteString(name, _typeConverter.ConvertToInvariantString(value));
+			writer.WriteString(name, typeConverter.ConvertToInvariantString(value));
 		}
 
 		public override void Write(CbWriter writer, T value)
 		{
-			writer.WriteStringValue(_typeConverter.ConvertToInvariantString(value)!);
+			writer.WriteStringValue(typeConverter.ConvertToInvariantString(value)!);
 		}
 	}
 
