EpicGames.Build: .net8 suggestions

#rnx

diff --git a/Engine/Source/Programs/Shared/EpicGames.Build/Automation/CommandLine.cs b/Engine/Source/Programs/Shared/EpicGames.Build/Automation/CommandLine.cs
--- a/Engine/Source/Programs/Shared/EpicGames.Build/Automation/CommandLine.cs
+++ b/Engine/Source/Programs/Shared/EpicGames.Build/Automation/CommandLine.cs
@@ -9,17 +9,12 @@ namespace UnrealBuildBase
 	/// <summary>
 	/// Command to execute info.
 	/// </summary>
-	public class CommandInfo
+	public class CommandInfo(string CommandNameIn)
 	{
-		public string CommandName;
+		public string CommandName = CommandNameIn;
 
 		public List<string> Arguments = [];
 
-		public CommandInfo(string CommandNameIn)
-		{
-			CommandName = CommandNameIn;
-		}
-
 		public override string ToString()
 		{
 			string Result = CommandName;
@@ -37,7 +32,7 @@ namespace UnrealBuildBase
 		}
 	}
 
-	public partial class ParsedCommandLine
+	public partial class ParsedCommandLine(Dictionary<string, string> GlobalParameters, HashSet<string> IgnoredGlobalParameters)
 	{
 		/// <summary>
 		/// Argument name and description for known arguments
@@ -46,19 +41,13 @@ namespace UnrealBuildBase
 		{
 			get;
 			private set;
-		}
+		} = new Dictionary<string, string>(GlobalParameters, StringComparer.InvariantCultureIgnoreCase);
 
 		public HashSet<string> IgnoredGlobalParameters
 		{
 			get;
 			private set;
-		}
-
-		public ParsedCommandLine(Dictionary<string, string> GlobalParameters, HashSet<string> IgnoredGlobalParameters)
-		{
-			this.GlobalParameters = new Dictionary<string, string>(GlobalParameters, StringComparer.InvariantCultureIgnoreCase);
-			this.IgnoredGlobalParameters = new HashSet<string>(IgnoredGlobalParameters, StringComparer.InvariantCultureIgnoreCase);
-		}
+		} = new HashSet<string>(IgnoredGlobalParameters, StringComparer.InvariantCultureIgnoreCase);
 
 		HashSet<string> GlobalParameterValues = new HashSet<string>(StringComparer.InvariantCultureIgnoreCase);
 		
diff --git a/Engine/Source/Programs/Shared/EpicGames.Build/Exceptions/BuildLogEventException.cs b/Engine/Source/Programs/Shared/EpicGames.Build/Exceptions/BuildLogEventException.cs
--- a/Engine/Source/Programs/Shared/EpicGames.Build/Exceptions/BuildLogEventException.cs
+++ b/Engine/Source/Programs/Shared/EpicGames.Build/Exceptions/BuildLogEventException.cs
@@ -9,12 +9,14 @@ namespace UnrealBuildBase
 	/// <summary>
 	/// Implementation of <see cref="UnrealBuildTool.BuildException"/> that captures a full structured logging event.
 	/// </summary>
-	public class BuildLogEventException : UnrealBuildTool.BuildException
+	/// <param name="InnerException">The inner exception</param>
+	/// <param name="Event">Event to construct from</param>
+	public class BuildLogEventException(Exception? InnerException, LogEvent Event) : UnrealBuildTool.BuildException(InnerException, Event.ToString())
 	{
 		/// <summary>
 		/// The event object
 		/// </summary>
-		public LogEvent Event { get; }
+		public LogEvent Event { get; } = Event;
 
 		/// <summary>
 		/// Constructor
@@ -25,17 +27,6 @@ namespace UnrealBuildBase
 		{
 		}
 
-		/// <summary>
-		/// Constructor which wraps another exception
-		/// </summary>
-		/// <param name="InnerException">The inner exception</param>
-		/// <param name="Event">Event to construct from</param>
-		public BuildLogEventException(Exception? InnerException, LogEvent Event)
-			: base(InnerException, Event.ToString())
-		{
-			this.Event = Event;
-		}
-
 		/// <inheritdoc/>
 		public BuildLogEventException(string Format, params object[] Arguments)
 			: this(LogEvent.Create(LogLevel.Error, Format, Arguments))
diff --git a/Engine/Source/Programs/Shared/EpicGames.Build/System/DirectoryItem.cs b/Engine/Source/Programs/Shared/EpicGames.Build/System/DirectoryItem.cs
--- a/Engine/Source/Programs/Shared/EpicGames.Build/System/DirectoryItem.cs
+++ b/Engine/Source/Programs/Shared/EpicGames.Build/System/DirectoryItem.cs
@@ -233,9 +233,9 @@ namespace UnrealBuildBase
 						NewDirectories = new Dictionary<string, DirectoryItem>(Directories.Length, DirectoryReference.Comparer);
 						foreach (DirectoryInfo SubDirectoryInfo in Directories)
 						{
-							if (NewDirectories.ContainsKey(SubDirectoryInfo.Name))
+							if (NewDirectories.TryGetValue(SubDirectoryInfo.Name, out DirectoryItem? value))
 							{
-								throw new Exception($"Trying to add {SubDirectoryInfo.FullName} as '{SubDirectoryInfo.Name}' yet exists as {NewDirectories[SubDirectoryInfo.Name].FullName}");
+								throw new Exception($"Trying to add {SubDirectoryInfo.FullName} as '{SubDirectoryInfo.Name}' yet exists as {value.FullName}");
 							}
 
 							NewDirectories.Add(SubDirectoryInfo.Name, DirectoryItem.GetItemByDirectoryInfo(SubDirectoryInfo));
diff --git a/Engine/Source/Programs/Shared/EpicGames.Build/System/FileHasher.cs b/Engine/Source/Programs/Shared/EpicGames.Build/System/FileHasher.cs
--- a/Engine/Source/Programs/Shared/EpicGames.Build/System/FileHasher.cs
+++ b/Engine/Source/Programs/Shared/EpicGames.Build/System/FileHasher.cs
@@ -14,23 +14,20 @@ using Microsoft.Extensions.Logging;
 
 namespace UnrealBuildBase
 {
-	public class FileHasher
+	/// <summary>
+	/// FileHasher constructor
+	/// </summary>
+	/// <param name="logger">Optional logger</param>
+	public class FileHasher(ILogger? logger = null)
 	{
-		private readonly struct CachedDigest
+		private readonly struct CachedDigest(FileItem item, IoHash digest)
 		{
 			[JsonPropertyName("l")]
-			public long Length { get; init; }
+			public long Length { get; init; } = item.Length;
 			[JsonPropertyName("w")]
-			public long LastWriteTimeUtc { get; init; }
+			public long LastWriteTimeUtc { get; init; } = item.LastWriteTimeUtc.Ticks;
 			[JsonPropertyName("d")]
-			public IoHash Digest { get; init; }
-
-			public CachedDigest(FileItem item, IoHash digest)
-			{
-				Length = item.Length;
-				LastWriteTimeUtc = item.LastWriteTimeUtc.Ticks;
-				Digest = digest;
-			}
+			public IoHash Digest { get; init; } = digest;
 
 			/// <summary>
 			/// Check if a digest is up to date for a FileItem
@@ -40,20 +37,9 @@ namespace UnrealBuildBase
 			public bool UpToDate(FileItem item) => item.Exists && Length == item.Length && LastWriteTimeUtc == item.LastWriteTimeUtc.Ticks;
 		}
 
-		// Dictionary containg cached digests, to prevent rehashing data if the file is unchanged
+		// Dictionary containing cached digests, to prevent rehashing data if the file is unchanged
 		ConcurrentDictionary<string, CachedDigest> CachedDigests = new();
 
-		ILogger? Logger;
-
-		/// <summary>
-		/// FileHasher constructor
-		/// </summary>
-		/// <param name="logger">Optional logger</param>
-		public FileHasher(ILogger? logger = null)
-		{
-			Logger = logger;
-		}
-
 		/// <summary>
 		/// Saves cached digests to disk.
 		/// </summary>
@@ -62,7 +48,7 @@ namespace UnrealBuildBase
 		{
 			using FileStream stream = FileReference.Open(location, FileMode.Create);
 			await JsonSerializer.SerializeAsync(stream, new SortedDictionary<string, CachedDigest>(CachedDigests));
-			Logger?.LogDebug("Saved {Entries} cached digest entries", CachedDigests.Count);
+			logger?.LogDebug("Saved {Entries} cached digest entries", CachedDigests.Count);
 		}
 
 		/// <summary>
@@ -81,7 +67,7 @@ namespace UnrealBuildBase
 			if (loadedDigests != null)
 			{
 				CachedDigests = loadedDigests;
-				Logger?.LogDebug("Loaded {Entries} cached digest entries", CachedDigests.Count);
+				logger?.LogDebug("Loaded {Entries} cached digest entries", CachedDigests.Count);
 			}
 		}
 
@@ -93,7 +79,7 @@ namespace UnrealBuildBase
 			int count = CachedDigests.Count;
 			IEnumerable<KeyValuePair<string, CachedDigest>> valid = CachedDigests.Where((x) => x.Value.UpToDate(FileItem.GetItemByFileReference(FileReference.FromString(x.Key))));
 			CachedDigests = new(valid);
-			Logger?.LogDebug("Purged {Entries} stale cached digest entries", count - CachedDigests.Count);
+			logger?.LogDebug("Purged {Entries} stale cached digest entries", count - CachedDigests.Count);
 		}
 
 		/// <summary>
@@ -119,7 +105,7 @@ namespace UnrealBuildBase
 				return IoHash.Zero;
 			}
 
-			CachedDigest digest = await ComputeDigest(item, Logger, cancellationToken);
+			CachedDigest digest = await ComputeDigest(item, logger, cancellationToken);
 			return CachedDigests.AddOrUpdate(item.FullName, digest, (key, oldValue) => digest).Digest;
 		}
 
diff --git a/Engine/Source/Programs/Shared/EpicGames.Build/UEBuildPlatformSDK.cs b/Engine/Source/Programs/Shared/EpicGames.Build/UEBuildPlatformSDK.cs
--- a/Engine/Source/Programs/Shared/EpicGames.Build/UEBuildPlatformSDK.cs
+++ b/Engine/Source/Programs/Shared/EpicGames.Build/UEBuildPlatformSDK.cs
@@ -120,7 +120,7 @@ namespace EpicGames.Core
 			return $"{Descriptor}{DisplayName}={Min}" + (CurrentDescriptor != null ? $", Current{CurrentDescriptor}{DisplayName}={Current ?? ""}" : "");
 		}
 	}
-	public class SDKCollection
+	public class SDKCollection(UEBuildPlatformSDK InPlatformSDK)
 	{
 		public virtual string DefaultName => "Sdk";
 		public virtual bool IsValid(string Version, SDKDescriptor Info)
@@ -128,19 +128,11 @@ namespace EpicGames.Core
 			return PlatformSDK.IsVersionValid(Version, Info);
 		}
 
-
-
-		public UEBuildPlatformSDK PlatformSDK;
+		public UEBuildPlatformSDK PlatformSDK = InPlatformSDK;
 		public List<SDKDescriptor> Sdks = [];
 		public List<SDKDescriptor> FullSdks => Sdks.Where(x => String.Compare(x.Name, "AutoSdk", true) != 0).ToList();
 		public SDKDescriptor? AutoSdk => Sdks.FirstOrDefault(x => String.Compare(x.Name, "AutoSdk", true) == 0);
 
-
-		public SDKCollection(UEBuildPlatformSDK PlatformSDK)
-		{
-			this.PlatformSDK = PlatformSDK;
-		}
-
 		public SDKCollection(SDKCollection Other, UEBuildPlatformSDK PlatformSDK)
 			: this(PlatformSDK)
 		{
@@ -284,16 +276,11 @@ namespace EpicGames.Core
 	/// <summary>
 	/// SDK for a platform
 	/// </summary>
-	abstract public class UEBuildPlatformSDK
+	abstract public class UEBuildPlatformSDK(ILogger InLogger)
 	{
 		// Public SDK handling, not specific to AutoSDK
 
-		protected readonly ILogger Logger;
-
-		public UEBuildPlatformSDK(ILogger InLogger)
-		{
-			Logger = InLogger;
-		}
+		protected readonly ILogger Logger = InLogger;
 
 		#region Global SDK Registration
 
