Update MetaData.cs for 5.6

diff --git a/Engine/Source/Programs/Shared/MetaData.cs b/Engine/Source/Programs/Shared/MetaData.cs
--- a/Engine/Source/Programs/Shared/MetaData.cs
+++ b/Engine/Source/Programs/Shared/MetaData.cs
@@ -10,7 +10,7 @@ using System.Runtime.InteropServices;
 [assembly: AssemblyCopyright("Copyright Epic Games, Inc. All Rights Reserved.")]
 [assembly: AssemblyTrademark("")]
 
-// Use a neutral culture to avoid some localisation issues
+// Use a neutral culture to avoid some localization issues
 [assembly: AssemblyCulture("")]
 
 [assembly: ComVisible(false)]
@@ -18,8 +18,8 @@ using System.Runtime.InteropServices;
 
 #if !SPECIFIC_VERSION
 // Automatically generate a version number based on the time of compilation
-[assembly: AssemblyVersion("5.5.0.0")]
-[assembly: AssemblyFileVersion("5.5.0.0")]
-[assembly: AssemblyInformationalVersion("5.5.0.0")]
+[assembly: AssemblyVersion("5.6.0.0")]
+[assembly: AssemblyFileVersion("5.6.0.0")]
+[assembly: AssemblyInformationalVersion("5.6.0.0")]
 #endif
 
