Horde: Allow clients to specify a locator when uploading blobs. The server will not allow overwriting any existing blob, but allowing the client to determine the locator opens the door to it being able to write multiple blobs in parallel.

diff --git a/Engine/Source/Programs/Horde/Plugins/Compute/HordeServer.Compute/Logs/LogsController.cs b/Engine/Source/Programs/Horde/Plugins/Compute/HordeServer.Compute/Logs/LogsController.cs
--- a/Engine/Source/Programs/Horde/Plugins/Compute/HordeServer.Compute/Logs/LogsController.cs
+++ b/Engine/Source/Programs/Horde/Plugins/Compute/HordeServer.Compute/Logs/LogsController.cs
@@ -127,7 +127,7 @@ namespace HordeServer.Logs
 			request.Prefix = $"{log.RefName}";
 
 			IStorageBackend storageBackend = _storageService.CreateBackend(log.NamespaceId);
-			return await StorageController.WriteBlobAsync(storageBackend, request, cancellationToken);
+			return await StorageController.WriteBlobAsync(storageBackend, null, request, cancellationToken);
 		}
 
 		/// <summary>
diff --git a/Engine/Source/Programs/Horde/Plugins/Storage/HordeServer.Storage/Storage/StorageController.cs b/Engine/Source/Programs/Horde/Plugins/Storage/HordeServer.Storage/Storage/StorageController.cs
--- a/Engine/Source/Programs/Horde/Plugins/Storage/HordeServer.Storage/Storage/StorageController.cs
+++ b/Engine/Source/Programs/Horde/Plugins/Storage/HordeServer.Storage/Storage/StorageController.cs
@@ -105,17 +105,42 @@ namespace HordeServer.Storage
 				return Forbid(StorageAclAction.WriteBlobs, namespaceId);
 			}
 
-			return await WriteBlobAsync(storageBackend, request, cancellationToken);
+			return await WriteBlobAsync(storageBackend, null, request, cancellationToken);
+		}
+
+		/// <summary>
+		/// Uploads data to the storage service using a client-determined path.
+		/// </summary>
+		/// <param name="namespaceId">Namespace to fetch from</param>
+		/// <param name="request">Information about the blob to write</param>
+		/// <param name="locator">Location for the blob</param>
+		/// <param name="cancellationToken">Cancellation token for the operation</param>
+		[HttpPut]
+		[Route("/api/v1/storage/{namespaceId}/blobs/{*locator}")]
+		public async Task<ActionResult<WriteBlobResponse>> WriteBlobAsync(NamespaceId namespaceId, BlobLocator locator, WriteBlobRequest request, CancellationToken cancellationToken = default)
+		{
+			IStorageBackend? storageBackend = _storageService.TryCreateBackend(namespaceId);
+			if (storageBackend == null)
+			{
+				return NotFound(namespaceId);
+			}
+			if (!Authorize(namespaceId, StorageAclAction.WriteBlobs) && !HasPathClaim(User, HordeClaimTypes.WriteNamespace, namespaceId, request.Prefix ?? String.Empty))
+			{
+				return Forbid(StorageAclAction.WriteBlobs, namespaceId);
+			}
+
+			return await WriteBlobAsync(storageBackend, locator, request, cancellationToken);
 		}
 
 		/// <summary>
 		/// Writes a blob to storage. Exposed as a public utility method to allow other routes with their own authentication methods to wrap their own authentication/redirection.
 		/// </summary>
 		/// <param name="storageBackend">The backend to write to</param>
+		/// <param name="locator">Locator for the blob to write</param>
 		/// <param name="request">Information about the blob to write</param>
 		/// <param name="cancellationToken">Cancellation token</param>
 		/// <returns>Information about the written blob, or redirect information</returns>
-		public static async Task<ActionResult<WriteBlobResponse>> WriteBlobAsync(IStorageBackend storageBackend, WriteBlobRequest request, CancellationToken cancellationToken = default)
+		public static async Task<ActionResult<WriteBlobResponse>> WriteBlobAsync(IStorageBackend storageBackend, BlobLocator? locator, WriteBlobRequest request, CancellationToken cancellationToken = default)
 		{
 			IReadOnlyCollection<BlobLocator> imports = request.Imports ?? (IReadOnlyCollection<BlobLocator>)Array.Empty<BlobLocator>();
 			if (request.File == null)
@@ -131,8 +156,15 @@ namespace HordeServer.Storage
 			else
 			{
 				using Stream stream = request.File.OpenReadStream();
-				BlobLocator locator = await storageBackend.WriteBlobAsync(stream, imports, request.Prefix, cancellationToken);
-				return new WriteBlobResponse { Blob = locator.ToString(), SupportsRedirects = storageBackend.SupportsRedirects };
+				if (locator == null)
+				{
+					locator = await storageBackend.WriteBlobAsync(stream, imports, request.Prefix, cancellationToken);
+				}
+				else
+				{
+					await storageBackend.WriteBlobAsync(locator.Value, stream, imports, cancellationToken);
+				}
+				return new WriteBlobResponse { Blob = locator.ToString()!, SupportsRedirects = storageBackend.SupportsRedirects };
 			}
 		}
 
diff --git a/Engine/Source/Programs/Horde/Plugins/Storage/HordeServer.Storage/Storage/StorageService.cs b/Engine/Source/Programs/Horde/Plugins/Storage/HordeServer.Storage/Storage/StorageService.cs
--- a/Engine/Source/Programs/Horde/Plugins/Storage/HordeServer.Storage/Storage/StorageService.cs
+++ b/Engine/Source/Programs/Horde/Plugins/Storage/HordeServer.Storage/Storage/StorageService.cs
@@ -91,11 +91,15 @@ namespace HordeServer.Storage
 			public async Task<BlobLocator> WriteBlobAsync(Stream stream, IReadOnlyCollection<BlobLocator> imports, string? basePath = null, CancellationToken cancellationToken = default)
 			{
 				BlobLocator locator = StorageHelpers.CreateUniqueLocator(basePath);
+				await WriteBlobAsync(locator, stream, imports, cancellationToken);
+				return locator;
+			}
 
+			/// <inheritdoc/>
+			public async Task WriteBlobAsync(BlobLocator locator, Stream stream, IReadOnlyCollection<BlobLocator> imports, CancellationToken cancellationToken = default)
+			{
 				await _store.WriteAsync(GetObjectKey(locator), stream, cancellationToken);
 				await _outer.AddBlobAsync(NamespaceId, locator, imports, null, cancellationToken);
-
-				return locator;
 			}
 
 			/// <inheritdoc/>
@@ -106,15 +110,13 @@ namespace HordeServer.Storage
 			}
 
 			/// <inheritdoc/>
-			public async ValueTask<(BlobLocator, Uri)?> TryGetBlobWriteRedirectAsync(IReadOnlyCollection<BlobLocator> imports, string? prefix = null, CancellationToken cancellationToken = default)
+			public async ValueTask<Uri?> TryGetBlobWriteRedirectAsync(BlobLocator locator, IReadOnlyCollection<BlobLocator> imports, CancellationToken cancellationToken = default)
 			{
 				if (!_store.SupportsRedirects)
 				{
 					return null;
 				}
 
-				BlobLocator locator = StorageHelpers.CreateUniqueLocator(prefix);
-
 				Uri? url = await _store.TryGetWriteRedirectAsync(GetObjectKey(locator), cancellationToken);
 				if (url == null)
 				{
@@ -122,6 +124,24 @@ namespace HordeServer.Storage
 				}
 
 				await _outer.AddBlobAsync(NamespaceId, locator, imports, null, cancellationToken);
+				return url;
+			}
+
+			/// <inheritdoc/>
+			public async ValueTask<(BlobLocator, Uri)?> TryGetBlobWriteRedirectAsync(IReadOnlyCollection<BlobLocator> imports, string? prefix = null, CancellationToken cancellationToken = default)
+			{
+				if (!_store.SupportsRedirects)
+				{
+					return null;
+				}
+
+				BlobLocator locator = StorageHelpers.CreateUniqueLocator(prefix);
+
+				Uri? url = await TryGetBlobWriteRedirectAsync(locator, imports, cancellationToken);
+				if (url == null)
+				{
+					return null;
+				}
 
 				return (locator, url);
 			}
diff --git a/Engine/Source/Programs/Horde/Plugins/Tools/HordeServer.Tools/Tools/ToolsController.cs b/Engine/Source/Programs/Horde/Plugins/Tools/HordeServer.Tools/Tools/ToolsController.cs
--- a/Engine/Source/Programs/Horde/Plugins/Tools/HordeServer.Tools/Tools/ToolsController.cs
+++ b/Engine/Source/Programs/Horde/Plugins/Tools/HordeServer.Tools/Tools/ToolsController.cs
@@ -62,7 +62,7 @@ namespace HordeServer.Tools
 			}
 
 			IStorageBackend storageBackend = tool.GetStorageBackend();
-			return await StorageController.WriteBlobAsync(storageBackend, request, cancellationToken);
+			return await StorageController.WriteBlobAsync(storageBackend, null, request, cancellationToken);
 		}
 
 		/// <summary>
diff --git a/Engine/Source/Programs/Shared/EpicGames.Horde/Compute/AgentStorageClient.cs b/Engine/Source/Programs/Shared/EpicGames.Horde/Compute/AgentStorageClient.cs
--- a/Engine/Source/Programs/Shared/EpicGames.Horde/Compute/AgentStorageClient.cs
+++ b/Engine/Source/Programs/Shared/EpicGames.Horde/Compute/AgentStorageClient.cs
@@ -67,6 +67,10 @@ namespace EpicGames.Horde.Compute
 			}
 		}
 
+		/// <inheritdoc/>
+		public Task WriteBlobAsync(BlobLocator locator, Stream stream, IReadOnlyCollection<BlobLocator>? imports, CancellationToken cancellationToken = default)
+			=> throw new NotSupportedException();
+
 		/// <inheritdoc/>
 		public Task<BlobLocator> WriteBlobAsync(Stream stream, IReadOnlyCollection<BlobLocator>? imports, string? basePath = null, CancellationToken cancellationToken = default) 
 			=> throw new NotSupportedException();
@@ -75,6 +79,10 @@ namespace EpicGames.Horde.Compute
 		public ValueTask<Uri?> TryGetBlobReadRedirectAsync(BlobLocator locator, CancellationToken cancellationToken = default) 
 			=> default;
 
+		/// <inheritdoc/>
+		public ValueTask<Uri?> TryGetBlobWriteRedirectAsync(BlobLocator locator, IReadOnlyCollection<BlobLocator>? imports = null, CancellationToken cancellationToken = default)
+			=> default;
+
 		/// <inheritdoc/>
 		public ValueTask<(BlobLocator, Uri)?> TryGetBlobWriteRedirectAsync(IReadOnlyCollection<BlobLocator>? imports = null, string? prefix = null, CancellationToken cancellationToken = default) 
 			=> default;
diff --git a/Engine/Source/Programs/Shared/EpicGames.Horde/Storage/Backends/FileStorageBackend.cs b/Engine/Source/Programs/Shared/EpicGames.Horde/Storage/Backends/FileStorageBackend.cs
--- a/Engine/Source/Programs/Shared/EpicGames.Horde/Storage/Backends/FileStorageBackend.cs
+++ b/Engine/Source/Programs/Shared/EpicGames.Horde/Storage/Backends/FileStorageBackend.cs
@@ -68,19 +68,29 @@ namespace EpicGames.Horde.Storage.Backends
 		public Task<IReadOnlyMemoryOwner<byte>> ReadBlobAsync(BlobLocator locator, int offset, int? length, CancellationToken cancellationToken = default)
 			=> _objectStore.ReadAsync(GetBlobFile(locator), offset, length, cancellationToken);
 
+		/// <inheritdoc/>
+		public Task WriteBlobAsync(BlobLocator locator, Stream stream, IReadOnlyCollection<BlobLocator>? imports, CancellationToken cancellationToken = default)
+			=> _objectStore.WriteAsync(GetBlobFile(locator), stream, cancellationToken);
+
 		/// <inheritdoc/>
 		public async Task<BlobLocator> WriteBlobAsync(Stream stream, IReadOnlyCollection<BlobLocator>? imports, string? prefix = null, CancellationToken cancellationToken = default)
 		{
 			BlobLocator locator = StorageHelpers.CreateUniqueLocator(prefix);
-			await _objectStore.WriteAsync(GetBlobFile(locator), stream, cancellationToken);
+			await WriteBlobAsync(locator, stream, imports, cancellationToken);
 			return locator;
 		}
 
 		/// <inheritdoc/>
-		public ValueTask<Uri?> TryGetBlobReadRedirectAsync(BlobLocator path, CancellationToken cancellationToken = default) => default;
+		public ValueTask<Uri?> TryGetBlobReadRedirectAsync(BlobLocator path, CancellationToken cancellationToken = default)
+			=> default;
+
+		/// <inheritdoc/>
+		public ValueTask<Uri?> TryGetBlobWriteRedirectAsync(BlobLocator path, IReadOnlyCollection<BlobLocator>? imports = null, CancellationToken cancellationToken = default)
+			=> default;
 
 		/// <inheritdoc/>
-		public ValueTask<(BlobLocator, Uri)?> TryGetBlobWriteRedirectAsync(IReadOnlyCollection<BlobLocator>? imports = null, string? prefix = null, CancellationToken cancellationToken = default) => default;
+		public ValueTask<(BlobLocator, Uri)?> TryGetBlobWriteRedirectAsync(IReadOnlyCollection<BlobLocator>? imports = null, string? prefix = null, CancellationToken cancellationToken = default)
+			=> default;
 
 		#endregion
 
diff --git a/Engine/Source/Programs/Shared/EpicGames.Horde/Storage/Backends/HttpStorageBackend.cs b/Engine/Source/Programs/Shared/EpicGames.Horde/Storage/Backends/HttpStorageBackend.cs
--- a/Engine/Source/Programs/Shared/EpicGames.Horde/Storage/Backends/HttpStorageBackend.cs
+++ b/Engine/Source/Programs/Shared/EpicGames.Horde/Storage/Backends/HttpStorageBackend.cs
@@ -137,13 +137,21 @@ namespace EpicGames.Horde.Storage.Backends
 		}
 
 		/// <inheritdoc/>
-		public async Task<BlobLocator> WriteBlobAsync(Stream stream, IReadOnlyCollection<BlobLocator>? imports, string? prefix = null, CancellationToken cancellationToken = default)
+		public Task WriteBlobAsync(BlobLocator locator, Stream stream, IReadOnlyCollection<BlobLocator>? imports, CancellationToken cancellationToken = default)
+			=> WriteBlobAsync(locator, stream, imports, null, cancellationToken);
+
+		/// <inheritdoc/>
+		public Task<BlobLocator> WriteBlobAsync(Stream stream, IReadOnlyCollection<BlobLocator>? imports, string? prefix = null, CancellationToken cancellationToken = default)
+			=> WriteBlobAsync(null, stream, imports, prefix, cancellationToken);
+
+		/// <inheritdoc/>
+		public async Task<BlobLocator> WriteBlobAsync(BlobLocator? locator, Stream stream, IReadOnlyCollection<BlobLocator>? imports, string? prefix = null, CancellationToken cancellationToken = default)
 		{
 			using StreamContent streamContent = new StreamContent(stream);
 
 			if (_supportsUploadRedirects)
 			{
-				WriteBlobResponse redirectResponse = await SendWriteRequestAsync(null, imports, prefix, cancellationToken);
+				WriteBlobResponse redirectResponse = await SendWriteRequestAsync(locator, null, imports, prefix, cancellationToken);
 				if (redirectResponse.UploadUrl != null)
 				{
 					using HttpClient uploadRedirectClient = _createUploadRedirectClient();
@@ -159,17 +167,19 @@ namespace EpicGames.Horde.Storage.Backends
 				}
 			}
 
-			WriteBlobResponse response = await SendWriteRequestAsync(streamContent, imports, prefix, cancellationToken);
+			WriteBlobResponse response = await SendWriteRequestAsync(null, streamContent, imports, prefix, cancellationToken);
 			_supportsUploadRedirects = response.SupportsRedirects ?? false;
 			_logger.LogDebug("Written {Locator} (direct)", response.Blob);
 			return new BlobLocator(response.Blob);
 		}
 
-		async Task<WriteBlobResponse> SendWriteRequestAsync(StreamContent? streamContent, IReadOnlyCollection<BlobLocator>? imports = null, string? prefix = null, CancellationToken cancellationToken = default)
+		async Task<WriteBlobResponse> SendWriteRequestAsync(BlobLocator? locator, StreamContent? streamContent, IReadOnlyCollection<BlobLocator>? imports = null, string? prefix = null, CancellationToken cancellationToken = default)
 		{
 			using (HttpClient httpClient = _createClient())
 			{
-				using (HttpRequestMessage request = new HttpRequestMessage(HttpMethod.Post, $"{_basePath}/blobs"))
+				using (HttpRequestMessage request = (locator != null)
+					? new HttpRequestMessage(HttpMethod.Put, $"{_basePath}/blobs/{locator.Value}")
+					: new HttpRequestMessage(HttpMethod.Post, $"{_basePath}/blobs"))
 				{
 					using MultipartFormDataContent form = new MultipartFormDataContent();
 					if (streamContent != null)
@@ -214,12 +224,26 @@ namespace EpicGames.Horde.Storage.Backends
 			return default;
 		}
 
+		/// <inheritdoc/>
+		public async ValueTask<Uri?> TryGetBlobWriteRedirectAsync(BlobLocator locator, IReadOnlyCollection<BlobLocator> imports, CancellationToken cancellationToken = default)
+		{
+			if (_supportsUploadRedirects)
+			{
+				WriteBlobResponse redirectResponse = await SendWriteRequestAsync(locator, null, imports, null, cancellationToken);
+				if (redirectResponse.UploadUrl != null)
+				{
+					return redirectResponse.UploadUrl;
+				}
+			}
+			return null;
+		}
+
 		/// <inheritdoc/>
 		public async ValueTask<(BlobLocator, Uri)?> TryGetBlobWriteRedirectAsync(IReadOnlyCollection<BlobLocator>? imports = null, string? prefix = null, CancellationToken cancellationToken = default)
 		{
 			if (_supportsUploadRedirects)
 			{
-				WriteBlobResponse redirectResponse = await SendWriteRequestAsync(null, imports, prefix, cancellationToken);
+				WriteBlobResponse redirectResponse = await SendWriteRequestAsync(null, null, imports, prefix, cancellationToken);
 				if (redirectResponse.UploadUrl != null)
 				{
 					return (new BlobLocator(redirectResponse.Blob), redirectResponse.UploadUrl);
diff --git a/Engine/Source/Programs/Shared/EpicGames.Horde/Storage/Backends/MemoryStorageBackend.cs b/Engine/Source/Programs/Shared/EpicGames.Horde/Storage/Backends/MemoryStorageBackend.cs
--- a/Engine/Source/Programs/Shared/EpicGames.Horde/Storage/Backends/MemoryStorageBackend.cs
+++ b/Engine/Source/Programs/Shared/EpicGames.Horde/Storage/Backends/MemoryStorageBackend.cs
@@ -64,11 +64,21 @@ namespace EpicGames.Horde.Storage.Backends
 			return memory;
 		}
 
+		/// <inheritdoc/>
+		public async Task WriteBlobAsync(BlobLocator locator, Stream stream, IReadOnlyCollection<BlobLocator>? imports, CancellationToken cancellationToken = default)
+		{
+			byte[] data = await stream.ReadAllBytesAsync(cancellationToken);
+			if (!_blobs.TryAdd(locator, data))
+			{
+				throw new InvalidOperationException($"Locator {locator} has already been written");
+			}
+		}
+
 		/// <inheritdoc/>
 		public async Task<BlobLocator> WriteBlobAsync(Stream stream, IReadOnlyCollection<BlobLocator>? imports, string? prefix = null, CancellationToken cancellationToken = default)
 		{
 			BlobLocator locator = StorageHelpers.CreateUniqueLocator(prefix);
-			_blobs[locator] = await stream.ReadAllBytesAsync(cancellationToken);
+			await WriteBlobAsync(locator, stream, imports, cancellationToken);
 			return locator;
 		}
 
@@ -77,7 +87,11 @@ namespace EpicGames.Horde.Storage.Backends
 			=> default;
 
 		/// <inheritdoc/>
-		public ValueTask<(BlobLocator, Uri)?> TryGetBlobWriteRedirectAsync(IReadOnlyCollection<BlobLocator>? imports = null, string? prefix = null, CancellationToken cancellationToken = default)
+		public ValueTask<Uri?> TryGetBlobWriteRedirectAsync(BlobLocator locator, IReadOnlyCollection<BlobLocator> imports, CancellationToken cancellationToken = default)
+			=> default;
+
+		/// <inheritdoc/>
+		public ValueTask<(BlobLocator, Uri)?> TryGetBlobWriteRedirectAsync(IReadOnlyCollection<BlobLocator> imports, string? prefix = null, CancellationToken cancellationToken = default)
 			=> default;
 
 		#endregion
diff --git a/Engine/Source/Programs/Shared/EpicGames.Horde/Storage/IStorageBackend.cs b/Engine/Source/Programs/Shared/EpicGames.Horde/Storage/IStorageBackend.cs
--- a/Engine/Source/Programs/Shared/EpicGames.Horde/Storage/IStorageBackend.cs
+++ b/Engine/Source/Programs/Shared/EpicGames.Horde/Storage/IStorageBackend.cs
@@ -54,6 +54,16 @@ namespace EpicGames.Horde.Storage
 		/// <returns>Path to the uploaded object</returns>
 		Task<BlobLocator> WriteBlobAsync(Stream stream, IReadOnlyCollection<BlobLocator> imports, string? prefix = null, CancellationToken cancellationToken = default);
 
+		/// <summary>
+		/// Writes a stream to the storage backend. If the stream throws an exception during read, the write will be aborted.
+		/// </summary>s
+		/// <param name="locator">Locator for the new blob</param>
+		/// <param name="stream">Data stream</param>
+		/// <param name="imports">Imported blobs. If omitted, the backend will parse them from the stream data.</param>
+		/// <param name="cancellationToken">Cancellation token for the operation</param>
+		/// <returns>Path to the uploaded object</returns>
+		Task WriteBlobAsync(BlobLocator locator, Stream stream, IReadOnlyCollection<BlobLocator> imports, CancellationToken cancellationToken = default);
+
 		/// <summary>
 		/// Gets a HTTP redirect for a read request
 		/// </summary>
@@ -62,6 +72,15 @@ namespace EpicGames.Horde.Storage
 		/// <returns>Path to upload the data to</returns>
 		ValueTask<Uri?> TryGetBlobReadRedirectAsync(BlobLocator locator, CancellationToken cancellationToken = default);
 
+		/// <summary>
+		/// Gets a HTTP redirect for a write request
+		/// </summary>
+		/// <param name="locator">Path to write to</param>
+		/// <param name="imports">Imports for this blob</param>
+		/// <param name="cancellationToken">Cancellation token for the operation</param>
+		/// <returns>Path for retrieval, and URI to upload the data to</returns>
+		ValueTask<Uri?> TryGetBlobWriteRedirectAsync(BlobLocator locator, IReadOnlyCollection<BlobLocator> imports, CancellationToken cancellationToken = default);
+
 		/// <summary>
 		/// Gets a HTTP redirect for a write request
 		/// </summary>
diff --git a/Engine/Source/Programs/Shared/EpicGames.Horde/Storage/StorageBackendCache.cs b/Engine/Source/Programs/Shared/EpicGames.Horde/Storage/StorageBackendCache.cs
--- a/Engine/Source/Programs/Shared/EpicGames.Horde/Storage/StorageBackendCache.cs
+++ b/Engine/Source/Programs/Shared/EpicGames.Horde/Storage/StorageBackendCache.cs
@@ -133,9 +133,15 @@ namespace EpicGames.Horde.Storage
 			public Task<BlobLocator> WriteBlobAsync(Stream stream, IReadOnlyCollection<BlobLocator> imports, string? prefix = null, CancellationToken cancellationToken = default) 
 				=> _inner.WriteBlobAsync(stream, imports, prefix, cancellationToken);
 
+			public Task WriteBlobAsync(BlobLocator locator, Stream stream, IReadOnlyCollection<BlobLocator> imports, CancellationToken cancellationToken = default)
+				=> _inner.WriteBlobAsync(locator, stream, imports, cancellationToken);
+
 			public ValueTask<Uri?> TryGetBlobReadRedirectAsync(BlobLocator locator, CancellationToken cancellationToken = default) 
 				=> _inner.TryGetBlobReadRedirectAsync(locator, cancellationToken);
 
+			public ValueTask<Uri?> TryGetBlobWriteRedirectAsync(BlobLocator locator, IReadOnlyCollection<BlobLocator> imports, CancellationToken cancellationToken = default)
+				=> _inner.TryGetBlobWriteRedirectAsync(locator, imports, cancellationToken);
+
 			public ValueTask<(BlobLocator, Uri)?> TryGetBlobWriteRedirectAsync(IReadOnlyCollection<BlobLocator> imports, string? prefix = null, CancellationToken cancellationToken = default) 
 				=> _inner.TryGetBlobWriteRedirectAsync(imports, prefix, cancellationToken);
 
