EpicGames.Core.Tests: .net8 suggestions

#rnx

diff --git a/Engine/Source/Programs/Shared/EpicGames.Core.Tests/AppContainerTests.cs b/Engine/Source/Programs/Shared/EpicGames.Core.Tests/AppContainerTests.cs
--- a/Engine/Source/Programs/Shared/EpicGames.Core.Tests/AppContainerTests.cs
+++ b/Engine/Source/Programs/Shared/EpicGames.Core.Tests/AppContainerTests.cs
@@ -26,7 +26,7 @@ public class AppContainerTests
 	private const string Description = "Created by Unreal Engine tests";
 
 	private static int s_idCounter = 1000;
-	private static readonly HashSet<string> s_allocatedNames = new();
+	private static readonly HashSet<string> s_allocatedNames = [];
 	private readonly DirectoryInfo _tempDir;
 	private readonly FileInfo _tempFile;
 
diff --git a/Engine/Source/Programs/Shared/EpicGames.Core.Tests/ILoggerPropertyTests.cs b/Engine/Source/Programs/Shared/EpicGames.Core.Tests/ILoggerPropertyTests.cs
--- a/Engine/Source/Programs/Shared/EpicGames.Core.Tests/ILoggerPropertyTests.cs
+++ b/Engine/Source/Programs/Shared/EpicGames.Core.Tests/ILoggerPropertyTests.cs
@@ -7,16 +7,11 @@ using Microsoft.VisualStudio.TestTools.UnitTesting;
 
 namespace EpicGames.Core.Tests;
 
-public sealed class StubLogScope : IDisposable
+public sealed class StubLogScope(object? state) : IDisposable
 {
-	public object? State { get;}
+	public object? State { get; } = state;
 	public bool IsDisposed { get; private set; }
 
-	public StubLogScope(object? state)
-	{
-		State = state;
-	}
-
 	public void Dispose()
 	{
 		IsDisposed = true;
@@ -25,8 +20,8 @@ public sealed class StubLogScope : IDisposable
 
 public class SpyLogger : ILogger
 {
-	public List<object?> ReceivedLogStates { get; } = new();
-	public List<StubLogScope> ScopesStarted { get; } = new();
+	public List<object?> ReceivedLogStates { get; } = [];
+	public List<StubLogScope> ScopesStarted { get; } = [];
 	
 	public void Log<TState>(LogLevel logLevel, EventId eventId, TState? state, Exception? exception, Func<TState, Exception?, string> formatter)
 	{
@@ -46,18 +41,11 @@ public class SpyLogger : ILogger
 	}
 }
 
-public sealed class FakeLoggerProvider : ILoggerProvider
+public sealed class FakeLoggerProvider(ILogger logger) : ILoggerProvider
 {
-	private readonly ILogger _logger;
-
-	public FakeLoggerProvider(ILogger logger)
-	{
-		_logger = logger;
-	}
-
 	public ILogger CreateLogger(string categoryName)
 	{
-		return _logger;
+		return logger;
 	}
 	
 	public void Dispose()
diff --git a/Engine/Source/Programs/Shared/EpicGames.Core.Tests/JsonObjectTests.cs b/Engine/Source/Programs/Shared/EpicGames.Core.Tests/JsonObjectTests.cs
--- a/Engine/Source/Programs/Shared/EpicGames.Core.Tests/JsonObjectTests.cs
+++ b/Engine/Source/Programs/Shared/EpicGames.Core.Tests/JsonObjectTests.cs
@@ -18,10 +18,15 @@ namespace EpicGames.Core.Tests
 		private JsonObject _validTestObject = null!;
 		private Enemy _validJsonDataSource = null!;
 		private static DirectoryReference s_tempDirectory = null!;
+		private static JsonSerializerOptions s_jsonOptions = null!;
 
 		public JsonObjectTests()
 		{
 			s_tempDirectory = CreateTempDir();
+			// This is to make sure that enums are serialized as strings in the case that they were written as.
+			s_jsonOptions = new();
+			s_jsonOptions.Converters.Add(new JsonStringEnumConverter(null));
+			s_jsonOptions.WriteIndented = true;
 		}
 
 		[TestCleanup]
@@ -37,11 +42,7 @@ namespace EpicGames.Core.Tests
 		public void InitializeTest()
 		{
 			Enemy enemy1 = CreateValidJsonDataSource();
-			// This is to make sure that enums are serialized as strings in the case that they were written as.
-			JsonSerializerOptions options = new JsonSerializerOptions();
-			options.Converters.Add(new JsonStringEnumConverter(null));
-			options.WriteIndented = true;
-			_testJsonText = JsonSerializer.Serialize(enemy1, options);
+			_testJsonText = JsonSerializer.Serialize(enemy1, s_jsonOptions);
 			_validTestObject = CreateValidJsonObject();
 			_validJsonDataSource = enemy1;
 		}
@@ -237,7 +238,7 @@ namespace EpicGames.Core.Tests
 		public void GetStringArrayValidField()
 		{
 			string fieldName = "Skills";
-			string[] correctValue = _validJsonDataSource.Skills.ToArray();
+			string[] correctValue = [.. _validJsonDataSource.Skills];
 			string[] result = _validTestObject.GetStringArrayField(fieldName);
 			CollectionAssert.AreEqual(correctValue, result);
 			string[]? outValue;
@@ -249,7 +250,7 @@ namespace EpicGames.Core.Tests
 		public void GetStringArrayInvalidField()
 		{
 			string fieldName = "Invalid";
-			string[] correctValue = _validJsonDataSource.Skills.ToArray();
+			string[] correctValue = [.. _validJsonDataSource.Skills];
 			Assert.ThrowsException<JsonException>(() => _validTestObject.GetStringArrayField(fieldName));
 			string[]? outValue;
 			Assert.IsFalse(_validTestObject.TryGetStringArrayField(fieldName, out outValue));
@@ -282,7 +283,7 @@ namespace EpicGames.Core.Tests
 		public void GetEnumArrayValidField()
 		{
 			string fieldName = "ImmuneStatusEffects";
-			StatusEffect[] correctValue = _validJsonDataSource.ImmuneStatusEffects.ToArray();
+			StatusEffect[] correctValue = [.. _validJsonDataSource.ImmuneStatusEffects];
 			StatusEffect[]? outValue;
 			Assert.IsTrue(_validTestObject.TryGetEnumArrayField<StatusEffect>(fieldName, out outValue));
 			CollectionAssert.AreEqual(correctValue, outValue);
@@ -292,7 +293,7 @@ namespace EpicGames.Core.Tests
 		public void GetEnumArrayInvalidField()
 		{
 			string fieldName = "Invalid";
-			StatusEffect[] correctValu = _validJsonDataSource.ImmuneStatusEffects.ToArray();
+			StatusEffect[] correctValu = [.. _validJsonDataSource.ImmuneStatusEffects];
 			StatusEffect[]? outValue;
 			Assert.IsFalse(_validTestObject.TryGetEnumArrayField(fieldName, out outValue));
 			Assert.IsNull(outValue);
@@ -387,7 +388,7 @@ namespace EpicGames.Core.Tests
 		public void SetStringArrayField()
 		{
 			string fieldName = "Skills";
-			string[] setValue = { "SetValue1", "SetValue2", "SetValue3", "SetValue4"};
+			string[] setValue = ["SetValue1", "SetValue2", "SetValue3", "SetValue4"];
 			Assert.IsTrue(_validTestObject.ContainsField(fieldName));
 			_validTestObject.AddOrSetFieldValue(fieldName, setValue);
 			CollectionAssert.AreEqual(setValue, _validTestObject.GetStringArrayField(fieldName));
@@ -400,8 +401,8 @@ namespace EpicGames.Core.Tests
 
 			// test array of strings that contains null. Null should be set as ""
 			string nullFieldName = "Null";
-			string?[] nullStringArray = { "String1", null, "String2" };
-			string[] correctNullStringArray = { "String1", "", "String2" };
+			string?[] nullStringArray = ["String1", null, "String2"];
+			string[] correctNullStringArray = ["String1", "", "String2"];
 			Assert.IsFalse(_validTestObject.ContainsField(nullFieldName));
 			_validTestObject.AddOrSetFieldValue(nullFieldName, nullStringArray);
 			CollectionAssert.AreEqual(correctNullStringArray, _validTestObject.GetStringArrayField(nullFieldName));
@@ -410,7 +411,7 @@ namespace EpicGames.Core.Tests
 			// This test is very important for ModuleDescriptor 
 			string emptyFieldName = "Empty";
 			Assert.IsFalse(_validTestObject.ContainsField(emptyFieldName));
-			string[] emptyStringArray = Array.Empty<string>();
+			string[] emptyStringArray = [];
 			_validTestObject.AddOrSetFieldValue(emptyFieldName, emptyStringArray);
 			Assert.IsTrue(_validTestObject.ContainsField(emptyFieldName));
 			string[] returnArray = _validTestObject.GetStringArrayField(emptyFieldName);
@@ -438,10 +439,10 @@ namespace EpicGames.Core.Tests
 		{
 			string fieldName = "Attacks";
 			EnemyAttack[] attacks = 
-			{ 
+			[ 
 				new EnemyAttack("Attack1", 40, StatusEffect.Sleep), 
 				new EnemyAttack("Attack2", 100, StatusEffect.Burn)
-			};
+			];
 			Assert.IsTrue(_validTestObject.ContainsField(fieldName));
 			JsonObject[] setValue = attacks.Select(x => x.ToJsonObject()).ToArray();
 			_validTestObject.AddOrSetFieldValue(fieldName, setValue);
@@ -474,7 +475,7 @@ namespace EpicGames.Core.Tests
 		public void SetEnumArrayField()
 		{
 			string fieldName = "ImmuneStatusEffects";
-			StatusEffect[] setValue = { StatusEffect.Sleep, StatusEffect.Poison, StatusEffect.Frozen, StatusEffect.Burn };
+			StatusEffect[] setValue = [StatusEffect.Sleep, StatusEffect.Poison, StatusEffect.Frozen, StatusEffect.Burn];
 			Assert.IsTrue(_validTestObject.ContainsField(fieldName));
 			_validTestObject.AddOrSetFieldValue(fieldName, setValue);
 			StatusEffect[]? outValue;
@@ -490,7 +491,7 @@ namespace EpicGames.Core.Tests
 
 			// Test passing in an empty array
 			string emptyFieldName = "Empty";
-			StatusEffect[] setValueEmpty = Array.Empty<StatusEffect>();
+			StatusEffect[] setValueEmpty = [];
 			Assert.IsFalse(_validTestObject.ContainsField(emptyFieldName));
 			_validTestObject.AddOrSetFieldValue(emptyFieldName, setValueEmpty);
 			Assert.IsTrue(_validTestObject.ContainsField(emptyFieldName));
@@ -504,7 +505,7 @@ namespace EpicGames.Core.Tests
 		{
 			string emptyFieldName = "Empty";
 			Assert.IsFalse(_validTestObject.ContainsField(emptyFieldName));
-			string[] setValue = Array.Empty<string>();
+			string[] setValue = [];
 			_validTestObject.AddOrSetFieldValue(emptyFieldName, setValue);
 			Assert.IsTrue(_validTestObject.ContainsField(emptyFieldName));
 			string[] returnValue = _validTestObject.GetStringArrayField(emptyFieldName);
@@ -601,9 +602,7 @@ namespace EpicGames.Core.Tests
 
 			enemy1.StrongestAttack = new EnemyAttack("Ultimate", 100, StatusEffect.Poison);
 			// Add skills 
-			enemy1.Skills = new List<string>();
-			enemy1.Skills.Add("Skill 1");
-			enemy1.Skills.Add("Skill 2");
+			enemy1.Skills = ["Skill 1", "Skill 2"];
 
 			// Add status effects
 			enemy1.ImmuneStatusEffects.Add(StatusEffect.Frozen);
@@ -656,23 +655,16 @@ namespace EpicGames.Core.Tests
 		public string Name { get; set; } = "";
 		public int Hp { get; set; } = 0;
 		public double CriticalHitRate { get; set; } = 0.0;
-		public List<EnemyAttack> Attacks { get; set; } = new List<EnemyAttack>();
+		public List<EnemyAttack> Attacks { get; set; } = [];
 		public EnemyAttack StrongestAttack { get; set; } = null!;
-		public List<StatusEffect> ImmuneStatusEffects { get; set; } = new List<StatusEffect>();
+		public List<StatusEffect> ImmuneStatusEffects { get; set; } = [];
 		public StatusEffect CurrentStatusEffect { get; set; } = StatusEffect.None;
-		public List<string> Skills { get; set; } = new List<string>();
+		public List<string> Skills { get; set; } = [];
 		public bool IsAlive { get; set; } = false;
 	}
 
-	class EnemyAttack
+	class EnemyAttack(string name, int damage, StatusEffect effect = StatusEffect.None)
 	{
-		public EnemyAttack(string name, int damage, StatusEffect effect = StatusEffect.None)
-		{
-			Name = name;
-			Damage = damage;
-			Effect = effect;
-		}
-
 		public JsonObject ToJsonObject()
 		{
 			JsonObject obj = new JsonObject();
@@ -681,8 +673,8 @@ namespace EpicGames.Core.Tests
 			obj.AddOrSetFieldValue<StatusEffect>("Effect", Effect);
 			return obj;
 		}
-		public string Name { get; set; } = "";
-		public int Damage { get; set; } = 0;
-		public StatusEffect Effect { get; set; } = StatusEffect.None;
+		public string Name { get; set; } = name;
+		public int Damage { get; set; } = damage;
+		public StatusEffect Effect { get; set; } = effect;
 	}
 }
diff --git a/Engine/Source/Programs/Shared/EpicGames.Core.Tests/JsonTracerTests.cs b/Engine/Source/Programs/Shared/EpicGames.Core.Tests/JsonTracerTests.cs
--- a/Engine/Source/Programs/Shared/EpicGames.Core.Tests/JsonTracerTests.cs
+++ b/Engine/Source/Programs/Shared/EpicGames.Core.Tests/JsonTracerTests.cs
@@ -66,7 +66,7 @@ namespace EpicGames.Core.Tests
 			{
 				JsonElement root = document.RootElement;
 				JsonElement spansElement = root.GetProperty("Spans");
-				List<JsonElement> spansList = spansElement.EnumerateArray().ToList();
+				List<JsonElement> spansList = [.. spansElement.EnumerateArray()];
 				Console.WriteLine(spansList[0]);
 				Assert.AreEqual("bogusOp1", spansList[0].GetProperty("Name").GetString());
 				Assert.IsNotNull(spansList[0].GetProperty("StartTime").GetDateTime());
diff --git a/Engine/Source/Programs/Shared/EpicGames.Core.Tests/LogEventTests.cs b/Engine/Source/Programs/Shared/EpicGames.Core.Tests/LogEventTests.cs
--- a/Engine/Source/Programs/Shared/EpicGames.Core.Tests/LogEventTests.cs
+++ b/Engine/Source/Programs/Shared/EpicGames.Core.Tests/LogEventTests.cs
@@ -35,7 +35,7 @@ namespace EpicGames.Core.Tests
 		public void RoundTripTests()
 		{
 			DateTime time = new DateTime(2024, 1, 1);
-			LogEvent ev = new LogEvent(time, LogLevel.Information, default, "hello", null, new KeyValuePair<string, object?>[] { KeyValuePair.Create<string, object?>("foo", new LogValue(new Utf8String("type"), "text")) }, null);
+			LogEvent ev = new LogEvent(time, LogLevel.Information, default, "hello", null, [KeyValuePair.Create<string, object?>("foo", new LogValue(new Utf8String("type"), "text"))], null);
 
 			JsonLogEvent jev = JsonLogEvent.Parse(ev.ToJsonBytes());
 			string jevs = jev.ToString();
diff --git a/Engine/Source/Programs/Shared/EpicGames.Core.Tests/MemorySerializerTests.cs b/Engine/Source/Programs/Shared/EpicGames.Core.Tests/MemorySerializerTests.cs
--- a/Engine/Source/Programs/Shared/EpicGames.Core.Tests/MemorySerializerTests.cs
+++ b/Engine/Source/Programs/Shared/EpicGames.Core.Tests/MemorySerializerTests.cs
@@ -73,8 +73,8 @@ namespace EpicGames.Core.Tests
 			byte[] data;
 			{
 				ByteArrayBuilder writer = new ByteArrayBuilder();
-				writer.WriteVariableLengthBytes(new byte[] { 1, 2, 3 });
-				writer.WriteVariableLengthBytes(new byte[] { 4, 5, 6 });
+				writer.WriteVariableLengthBytes([1, 2, 3]);
+				writer.WriteVariableLengthBytes([4, 5, 6]);
 				data = writer.ToByteArray();
 			}
 			{
diff --git a/Engine/Source/Programs/Shared/EpicGames.Core.Tests/MemoryWriterTests.cs b/Engine/Source/Programs/Shared/EpicGames.Core.Tests/MemoryWriterTests.cs
--- a/Engine/Source/Programs/Shared/EpicGames.Core.Tests/MemoryWriterTests.cs
+++ b/Engine/Source/Programs/Shared/EpicGames.Core.Tests/MemoryWriterTests.cs
@@ -13,25 +13,18 @@ namespace EpicGames.Core.Tests
 	{
 		class TestAllocator : IMemoryAllocator<byte>
 		{
-			class Allocation : IMemoryOwner<byte>
+			class Allocation(MemoryWriterTests.TestAllocator allocator, Memory<byte> memory) : IMemoryOwner<byte>
 			{
-				TestAllocator? _allocator;
-				readonly Memory<byte> _memory;
+				TestAllocator? _allocator = allocator;
 
-				public Allocation(TestAllocator allocator, Memory<byte> memory)
-				{
-					_allocator = allocator;
-					_memory = memory;
-				}
-
-				public Memory<byte> Memory => _memory;
+				public Memory<byte> Memory => memory;
 
 				public void Dispose()
 				{
 					if (_allocator != null)
 					{
 						Interlocked.Decrement(ref _allocator._allocatedCount);
-						Interlocked.Add(ref _allocator._allocatedSize, -_memory.Length);
+						Interlocked.Add(ref _allocator._allocatedSize, -memory.Length);
 						_allocator = null;
 					}
 				}
