EpicGames.UBA: .net8 suggestions

#rnx

diff --git a/Engine/Source/Programs/Shared/EpicGames.UBA/ICacheClient.cs b/Engine/Source/Programs/Shared/EpicGames.UBA/ICacheClient.cs
--- a/Engine/Source/Programs/Shared/EpicGames.UBA/ICacheClient.cs
+++ b/Engine/Source/Programs/Shared/EpicGames.UBA/ICacheClient.cs
@@ -1,6 +1,7 @@
 // Copyright Epic Games, Inc. All Rights Reserved.
 
 using System.Collections.Generic;
+using EpicGames.UBA.Impl;
 
 namespace EpicGames.UBA
 {
diff --git a/Engine/Source/Programs/Shared/EpicGames.UBA/IConfig.cs b/Engine/Source/Programs/Shared/EpicGames.UBA/IConfig.cs
--- a/Engine/Source/Programs/Shared/EpicGames.UBA/IConfig.cs
+++ b/Engine/Source/Programs/Shared/EpicGames.UBA/IConfig.cs
@@ -1,5 +1,7 @@
 // Copyright Epic Games, Inc. All Rights Reserved.
 
+using EpicGames.UBA.Impl;
+
 namespace EpicGames.UBA
 {
 	/// <summary>
diff --git a/Engine/Source/Programs/Shared/EpicGames.UBA/ILogger.cs b/Engine/Source/Programs/Shared/EpicGames.UBA/ILogger.cs
--- a/Engine/Source/Programs/Shared/EpicGames.UBA/ILogger.cs
+++ b/Engine/Source/Programs/Shared/EpicGames.UBA/ILogger.cs
@@ -1,5 +1,7 @@
 // Copyright Epic Games, Inc. All Rights Reserved.
 
+using EpicGames.UBA.Impl;
+
 namespace EpicGames.UBA
 {
 	/// <summary>
diff --git a/Engine/Source/Programs/Shared/EpicGames.UBA/IProcess.cs b/Engine/Source/Programs/Shared/EpicGames.UBA/IProcess.cs
--- a/Engine/Source/Programs/Shared/EpicGames.UBA/IProcess.cs
+++ b/Engine/Source/Programs/Shared/EpicGames.UBA/IProcess.cs
@@ -3,6 +3,7 @@
 using System;
 using System.Collections.Generic;
 using System.Diagnostics;
+using EpicGames.UBA.Impl;
 
 namespace EpicGames.UBA
 {
@@ -98,51 +99,38 @@ namespace EpicGames.UBA
 	/// <summary>
 	/// Event args for exited event
 	/// </summary>
-	public class ExitedEventArgs : EventArgs
+	/// <param name="process">The process to pull data from</param>
+	public class ExitedEventArgs(IProcess process) : EventArgs
 	{
 		/// <summary>
 		/// Process exit code
 		/// </summary>
-		public int ExitCode { get; init; }
+		public int ExitCode { get; } = process.ExitCode;
 
 		/// <summary>
 		/// The remote host that ran the process, if run remotely
 		/// </summary>
-		public string? ExecutingHost { get; init; }
+		public string? ExecutingHost { get; } = process.ExecutingHost;
 
 		/// <summary>
 		/// Captured output lines
 		/// </summary>
-		public List<string> LogLines { get; init; }
+		public List<string> LogLines { get; } = process.LogLines;
 
 		/// <summary>
 		/// Total time spent for the processor
 		/// </summary>
-		public TimeSpan TotalProcessorTime { get; init; }
+		public TimeSpan TotalProcessorTime { get; } = process.TotalProcessorTime;
 
 		/// <summary>
 		/// Total wall time spent
 		/// </summary>
-		public TimeSpan TotalWallTime { get; init; }
+		public TimeSpan TotalWallTime { get; } = process.TotalWallTime;
 
 		/// <summary>
 		/// Total wall time spent
 		/// </summary>
-		public object? UserData { get; init; }
-
-		/// <summary>
-		/// Constructor
-		/// </summary>
-		/// <param name="process">The process to pull data from</param>
-		public ExitedEventArgs(IProcess process)
-		{
-			ExitCode = process.ExitCode;
-			ExecutingHost = process.ExecutingHost;
-			LogLines = process.LogLines;
-			TotalProcessorTime = process.TotalProcessorTime;
-			TotalWallTime = process.TotalWallTime;
-			UserData = process.UserData;
-		}
+		public object? UserData { get; } = process.UserData;
 	}
 
 	/// <summary>
diff --git a/Engine/Source/Programs/Shared/EpicGames.UBA/IServer.cs b/Engine/Source/Programs/Shared/EpicGames.UBA/IServer.cs
--- a/Engine/Source/Programs/Shared/EpicGames.UBA/IServer.cs
+++ b/Engine/Source/Programs/Shared/EpicGames.UBA/IServer.cs
@@ -1,5 +1,7 @@
 // Copyright Epic Games, Inc. All Rights Reserved.
 
+using EpicGames.UBA.Impl;
+
 namespace EpicGames.UBA
 {
 	/// <summary>
diff --git a/Engine/Source/Programs/Shared/EpicGames.UBA/ISessionServer.cs b/Engine/Source/Programs/Shared/EpicGames.UBA/ISessionServer.cs
--- a/Engine/Source/Programs/Shared/EpicGames.UBA/ISessionServer.cs
+++ b/Engine/Source/Programs/Shared/EpicGames.UBA/ISessionServer.cs
@@ -1,90 +1,74 @@
 // Copyright Epic Games, Inc. All Rights Reserved.
 
 using System;
+using EpicGames.UBA.Impl;
 
 namespace EpicGames.UBA
 {
 	/// <summary>
 	/// Information needed to create a session server
 	/// </summary>
-	public readonly struct SessionServerCreateInfo
+	/// <param name="rootDirectory">Root directory to store content addressable data</param>
+	/// <param name="traceOutputFile">Path to a trace file that records the build</param>
+	/// <param name="disableCustomAllocator">If the custom allocator should be disabled</param>
+	/// <param name="launchVisualizer">If the visualizer should be launched</param>
+	/// <param name="resetCas">If the content addressable storage should be reset</param>
+	/// <param name="writeToDisk">If intermediate/output files should be written to disk</param>
+	/// <param name="detailedTrace">More detailed trace information</param>
+	/// <param name="allowWaitOnMem">Wait for memory before starting new processes</param>
+	/// <param name="allowKillOnMem">Kill processes when close to run out of memory</param>
+	/// <param name="storeObjFilesCompressed">Store .obj files compressed on disk</param>
+	public readonly struct SessionServerCreateInfo(string rootDirectory, string traceOutputFile, bool disableCustomAllocator, bool launchVisualizer, bool resetCas, bool writeToDisk, bool detailedTrace, bool allowWaitOnMem, bool allowKillOnMem, bool storeObjFilesCompressed)
 	{
 		/// <summary>
 		/// Root directory to store content addressable data
 		/// </summary>
-		public string RootDirectory { get; init; }
+		public string RootDirectory { get; } = rootDirectory;
 
 		/// <summary>
 		/// Path to a trace file that records the build
 		/// </summary>
-		public string TraceOutputFile { get; init; }
+		public string TraceOutputFile { get; } = traceOutputFile;
 
 		/// <summary>
 		/// If the custom allocator should be disabled
 		/// </summary>
-		public bool DisableCustomAllocator { get; init; }
+		public bool DisableCustomAllocator { get; } = disableCustomAllocator;
 
 		/// <summary>
 		/// If the visualizer should be launched
 		/// </summary>
-		public bool LaunchVisualizer { get; init; }
+		public bool LaunchVisualizer { get; } = launchVisualizer;
 
 		/// <summary>
 		/// If the content addressable storage should be reset
 		/// </summary>
-		public bool ResetCas { get; init; }
+		public bool ResetCas { get; } = resetCas;
 
 		/// <summary>
 		/// If intermediate/output files should be written to disk
 		/// </summary>
-		public bool WriteToDisk { get; init; }
+		public bool WriteToDisk { get; } = writeToDisk;
 
 		/// <summary>
 		/// More detailed trace information
 		/// </summary>
-		public bool DetailedTrace { get; init; }
+		public bool DetailedTrace { get; } = detailedTrace;
 
 		/// <summary>
 		/// Wait for memory before starting new processes
 		/// </summary>
-		public bool AllowWaitOnMem { get; init; }
+		public bool AllowWaitOnMem { get; } = allowWaitOnMem;
 
 		/// <summary>
 		/// Kill processes when close to run out of memory
 		/// </summary>
-		public bool AllowKillOnMem { get; init; }
+		public bool AllowKillOnMem { get; } = allowKillOnMem;
 
 		/// <summary>
 		/// Store .obj files compressed on disk
 		/// </summary>
-		public bool StoreObjFilesCompressed { get; init; }
-
-		/// <summary>
-		/// Constructor
-		/// </summary>
-		/// <param name="rootDirectory">Root directory to store content addressable data</param>
-		/// <param name="traceOutputFile">Path to a trace file that records the build</param>
-		/// <param name="disableCustomAllocator">If the custom allocator should be disabled</param>
-		/// <param name="launchVisualizer">If the visualizer should be launched</param>
-		/// <param name="resetCas">If the content addressable storage should be reset</param>
-		/// <param name="writeToDisk">If intermediate/output files should be written to disk</param>
-		/// <param name="detailedTrace">More detailed trace information</param>
-		/// <param name="allowWaitOnMem">Wait for memory before starting new processes</param>
-		/// <param name="allowKillOnMem">Kill processes when close to run out of memory</param>
-		/// <param name="storeObjFilesCompressed">Store .obj files compressed on disk</param>
-		public SessionServerCreateInfo(string rootDirectory, string traceOutputFile, bool disableCustomAllocator, bool launchVisualizer, bool resetCas, bool writeToDisk, bool detailedTrace, bool allowWaitOnMem, bool allowKillOnMem, bool storeObjFilesCompressed)
-		{
-			RootDirectory = rootDirectory;
-			TraceOutputFile = traceOutputFile;
-			DisableCustomAllocator = disableCustomAllocator;
-			LaunchVisualizer = launchVisualizer;
-			ResetCas = resetCas;
-			WriteToDisk = writeToDisk;
-			DetailedTrace = detailedTrace;
-			AllowWaitOnMem = allowWaitOnMem;
-			AllowKillOnMem = allowKillOnMem;
-			StoreObjFilesCompressed = storeObjFilesCompressed;
-		}
+		public bool StoreObjFilesCompressed { get; } = storeObjFilesCompressed;
 	}
 
 	/// <summary>
@@ -116,21 +100,14 @@ namespace EpicGames.UBA
 	/// <summary>
 	/// Event args for remote process returned event
 	/// </summary>
-	public class RemoteProcessReturnedEventArgs : EventArgs
+	/// <param name="process">The process being returned</param>
+	public class RemoteProcessReturnedEventArgs(IProcess process) : EventArgs
 	{
-		/// <summary>
-		/// Constructor
-		/// </summary>
-		/// <param name="process">The process being returned</param>
-		public RemoteProcessReturnedEventArgs(IProcess process)
-		{
-			Process = process;
-		}
 
 		/// <summary>
 		/// The remote process that was returned
 		/// </summary>
-		public IProcess Process { get; init; }
+		public IProcess Process { get; } = process;
 	}
 
 	/// <summary>
diff --git a/Engine/Source/Programs/Shared/EpicGames.UBA/IStorageServer.cs b/Engine/Source/Programs/Shared/EpicGames.UBA/IStorageServer.cs
--- a/Engine/Source/Programs/Shared/EpicGames.UBA/IStorageServer.cs
+++ b/Engine/Source/Programs/Shared/EpicGames.UBA/IStorageServer.cs
@@ -1,46 +1,37 @@
 // Copyright Epic Games, Inc. All Rights Reserved.
 
+using EpicGames.UBA.Impl;
+
 namespace EpicGames.UBA
 {
 	/// <summary>
 	/// Information needed to create a storage server
 	/// </summary>
-	public readonly struct StorageServerCreateInfo
+	/// <param name="rootDirectory">The root directory for the storage</param>
+	/// <param name="capacityBytes">The capacity of the storage in bytes</param>
+	/// <param name="storeCompressed">If the storage should be stored as compressed</param>
+	/// <param name="zone">The geographical zone this machine belongs to. Can be empty</param>
+	public readonly struct StorageServerCreateInfo(string rootDirectory, ulong capacityBytes, bool storeCompressed, string zone)
 	{
 		/// <summary>
 		/// The root directory for the storage
 		/// </summary>
-		public string RootDirectory { get; init; }
+		public string RootDirectory { get; } = rootDirectory;
 
 		/// <summary>
 		/// The capacity of the storage in bytes
 		/// </summary>
-		public ulong CapacityBytes { get; init; }
+		public ulong CapacityBytes { get; } = capacityBytes;
 
 		/// <summary>
 		/// If the storage should be stored as compressed
 		/// </summary>
-		public bool StoreCompressed { get; init; }
+		public bool StoreCompressed { get; } = storeCompressed;
 
 		/// <summary>
 		/// The geographical zone this machine belongs to. Can be empty
 		/// </summary>
-		public string Zone { get; init; }
-
-		/// <summary>
-		/// Constructor
-		/// </summary>
-		/// <param name="rootDirectory">The root directory for the storage</param>
-		/// <param name="capacityBytes">The capacity of the storage in bytes</param>
-		/// <param name="storeCompressed">If the storage should be stored as compressed</param>
-		/// <param name="zone">The geographical zone this machine belongs to. Can be empty</param>
-		public StorageServerCreateInfo(string rootDirectory, ulong capacityBytes, bool storeCompressed, string zone)
-		{
-			RootDirectory = rootDirectory;
-			CapacityBytes = capacityBytes;
-			StoreCompressed = storeCompressed;
-			Zone = zone;
-		}
+		public string Zone { get; } = zone;
 	}
 
 	/// <summary>
@@ -49,7 +40,7 @@ namespace EpicGames.UBA
 	public interface IStorageServer : IBaseInterface
 	{
 		/// <summary>
-		/// Save tge content addressabale storage table
+		/// Save tge content addressable storage table
 		/// </summary>
 		public abstract void SaveCasTable();
 
diff --git a/Engine/Source/Programs/Shared/EpicGames.UBA/Impl/CacheClientImpl.cs b/Engine/Source/Programs/Shared/EpicGames.UBA/Impl/CacheClientImpl.cs
--- a/Engine/Source/Programs/Shared/EpicGames.UBA/Impl/CacheClientImpl.cs
+++ b/Engine/Source/Programs/Shared/EpicGames.UBA/Impl/CacheClientImpl.cs
@@ -4,25 +4,25 @@ using System;
 using System.Collections.Generic;
 using System.Runtime.InteropServices;
 
-namespace EpicGames.UBA
+namespace EpicGames.UBA.Impl
 {
 	internal class RootPathsImpl : IRootPaths
 	{
-		IntPtr _handle = IntPtr.Zero;
+		nint _handle = IntPtr.Zero;
 		readonly ILogger _logger;
 
 		#region DllImport
 		[DllImport("UbaHost", CharSet = CharSet.Auto)]
-		static extern IntPtr RootPaths_Create(IntPtr logWriter);
+		static extern nint RootPaths_Create(nint logWriter);
 
 		[DllImport("UbaHost", CharSet = CharSet.Auto)]
-		static extern void RootPaths_Destroy(IntPtr rootPaths);
+		static extern void RootPaths_Destroy(nint rootPaths);
 
 		[DllImport("UbaHost", CharSet = CharSet.Auto)]
-		static extern bool RootPaths_RegisterRoot(IntPtr rootPaths, string root, bool includeInKey, byte id);
+		static extern bool RootPaths_RegisterRoot(nint rootPaths, string root, bool includeInKey, byte id);
 
 		[DllImport("UbaHost", CharSet = CharSet.Auto)]
-		static extern bool RootPaths_RegisterSystemRoots(IntPtr rootPaths, byte id);
+		static extern bool RootPaths_RegisterSystemRoots(nint rootPaths, byte id);
 		#endregion
 
 		#region IDisposable
@@ -64,49 +64,49 @@ namespace EpicGames.UBA
 			return RootPaths_RegisterSystemRoots(_handle, startId);
 		}
 
-		public IntPtr GetHandle() => _handle;
+		public nint GetHandle() => _handle;
 	}
 
 	internal class CacheClientImpl : ICacheClient
 	{
-		IntPtr _handle = IntPtr.Zero;
+		nint _handle = IntPtr.Zero;
 		readonly ISessionServer _sessionServer;
 
-		public delegate void ExitCallback(IntPtr userData, IntPtr handle);
+		public delegate void ExitCallback(nint userData, nint handle);
 
 		#region DllImport
 		[DllImport("UbaHost", CharSet = CharSet.Auto)]
-		static extern IntPtr ProcessStartInfo_Create(string application, string arguments, string workingDir, string description, uint priorityClass, ulong outputStatsThresholdMs, bool trackInputs, string logFile, ExitCallback? exit);
+		static extern nint ProcessStartInfo_Create(string application, string arguments, string workingDir, string description, uint priorityClass, ulong outputStatsThresholdMs, bool trackInputs, string logFile, ExitCallback? exit);
 
 		[DllImport("UbaHost", CharSet = CharSet.Auto)]
-		static extern void ProcessStartInfo_Destroy(IntPtr server);
+		static extern void ProcessStartInfo_Destroy(nint server);
 
 		[DllImport("UbaHost", CharSet = CharSet.Auto)]
-		static extern IntPtr CacheClient_Create(IntPtr session, bool reportMissReason, string crypto);
+		static extern nint CacheClient_Create(nint session, bool reportMissReason, string crypto);
 
 		[DllImport("UbaHost", CharSet = CharSet.Auto)]
-		static extern bool CacheClient_Connect(IntPtr cacheClient, string host, int port);
+		static extern bool CacheClient_Connect(nint cacheClient, string host, int port);
 
 		[DllImport("UbaHost", CharSet = CharSet.Auto)]
-		static extern bool CacheClient_WriteToCache(IntPtr cacheClient, IntPtr rootPaths, uint bucket, IntPtr info, byte[] inputs, uint inputsSize, byte[] outputs, uint outputsSize);
+		static extern bool CacheClient_WriteToCache(nint cacheClient, nint rootPaths, uint bucket, nint info, byte[] inputs, uint inputsSize, byte[] outputs, uint outputsSize);
 
 		[DllImport("UbaHost", CharSet = CharSet.Auto)]
-		static extern IntPtr CacheClient_FetchFromCache2(IntPtr cacheClient, IntPtr rootPaths, uint bucket, IntPtr info);
+		static extern nint CacheClient_FetchFromCache2(nint cacheClient, nint rootPaths, uint bucket, nint info);
 
 		[DllImport("UbaHost", CharSet = CharSet.Auto)]
-		static extern IntPtr CacheResult_GetLogLine(IntPtr result, uint index);
+		static extern nint CacheResult_GetLogLine(nint result, uint index);
 
 		[DllImport("UbaHost", CharSet = CharSet.Auto)]
-		static extern uint CacheResult_GetLogLineType(IntPtr result, uint index);
+		static extern uint CacheResult_GetLogLineType(nint result, uint index);
 
 		[DllImport("UbaHost", CharSet = CharSet.Auto)]
-		static extern void CacheResult_Delete(IntPtr result);
+		static extern void CacheResult_Delete(nint result);
 
 		[DllImport("UbaHost", CharSet = CharSet.Auto)]
-		static extern void CacheClient_RequestServerShutdown(IntPtr cacheClient, string reason);
+		static extern void CacheClient_RequestServerShutdown(nint cacheClient, string reason);
 
 		[DllImport("UbaHost", CharSet = CharSet.Auto)]
-		static extern void CacheClient_Destroy(IntPtr cacheClient);
+		static extern void CacheClient_Destroy(nint cacheClient);
 		#endregion
 
 		#region IDisposable
@@ -150,22 +150,22 @@ namespace EpicGames.UBA
 
 		public FetchFromCacheResult FetchFromCache(IRootPaths rootPaths, uint bucket, ProcessStartInfo info)
 		{
-			IntPtr si = ProcessStartInfo_Create(info.Application, info.Arguments, info.WorkingDirectory, info.Description, (uint)info.Priority, info.OutputStatsThresholdMs, info.TrackInputs, info.LogFile ?? String.Empty, null);
-			IntPtr result = CacheClient_FetchFromCache2(_handle, rootPaths.GetHandle(), bucket, si);
+			nint si = ProcessStartInfo_Create(info.Application, info.Arguments, info.WorkingDirectory, info.Description, (uint)info.Priority, info.OutputStatsThresholdMs, info.TrackInputs, info.LogFile ?? String.Empty, null);
+			nint result = CacheClient_FetchFromCache2(_handle, rootPaths.GetHandle(), bucket, si);
 			ProcessStartInfo_Destroy(si);
 
 			if (result == IntPtr.Zero)
 			{
-				return new FetchFromCacheResult(false, new List<string>());
+				return new FetchFromCacheResult(false, []);
 			}
 
 			string? line = Marshal.PtrToStringAuto(CacheResult_GetLogLine(result, 0));
 			if (line == null)
 			{
-				return new FetchFromCacheResult(true, new List<string>());
+				return new FetchFromCacheResult(true, []);
 			}
 
-			List<string> logLines = new();
+			List<string> logLines = [];
 			while (line != null)
 			{
 				logLines.Add(line);
@@ -180,6 +180,6 @@ namespace EpicGames.UBA
 			CacheClient_RequestServerShutdown(_handle, reason);
 		}
 
-		public IntPtr GetHandle() => _handle;
+		public nint GetHandle() => _handle;
 	}
 }
diff --git a/Engine/Source/Programs/Shared/EpicGames.UBA/Impl/ConfigImpl.cs b/Engine/Source/Programs/Shared/EpicGames.UBA/Impl/ConfigImpl.cs
--- a/Engine/Source/Programs/Shared/EpicGames.UBA/Impl/ConfigImpl.cs
+++ b/Engine/Source/Programs/Shared/EpicGames.UBA/Impl/ConfigImpl.cs
@@ -2,7 +2,7 @@
 
 using System.Runtime.InteropServices;
 
-namespace EpicGames.UBA
+namespace EpicGames.UBA.Impl
 {
 	internal class ConfigImpl : IConfig
 	{
diff --git a/Engine/Source/Programs/Shared/EpicGames.UBA/Impl/LoggerImpl.cs b/Engine/Source/Programs/Shared/EpicGames.UBA/Impl/LoggerImpl.cs
--- a/Engine/Source/Programs/Shared/EpicGames.UBA/Impl/LoggerImpl.cs
+++ b/Engine/Source/Programs/Shared/EpicGames.UBA/Impl/LoggerImpl.cs
@@ -4,15 +4,15 @@ using Microsoft.Extensions.Logging;
 using System;
 using System.Runtime.InteropServices;
 
-namespace EpicGames.UBA
+namespace EpicGames.UBA.Impl
 {
 	internal class LoggerImpl : ILogger
 	{
 		delegate void BeginScopeCallback();
 		delegate void EndScopeCallback();
-		delegate void LogCallback(LogEntryType type, IntPtr str, uint len);
+		delegate void LogCallback(LogEntryType type, nint str, uint len);
 
-		IntPtr _handle = IntPtr.Zero;
+		nint _handle = IntPtr.Zero;
 		readonly Microsoft.Extensions.Logging.ILogger _logger;
 		readonly BeginScopeCallback _beginScopeCallbackDelegate;
 		readonly EndScopeCallback _endScopeCallbackDelegate;
@@ -20,13 +20,13 @@ namespace EpicGames.UBA
 
 		#region DllImport
 		[DllImport("UbaHost", CharSet = CharSet.Auto)]
-		static extern IntPtr GetDefaultLogWriter();
+		static extern nint GetDefaultLogWriter();
 
 		[DllImport("UbaHost", CharSet = CharSet.Auto)]
-		static extern IntPtr CreateCallbackLogWriter(BeginScopeCallback begin, EndScopeCallback end, LogCallback log);
+		static extern nint CreateCallbackLogWriter(BeginScopeCallback begin, EndScopeCallback end, LogCallback log);
 
 		[DllImport("UbaHost", CharSet = CharSet.Auto)]
-		static extern void DestroyCallbackLogWriter(IntPtr logger);
+		static extern void DestroyCallbackLogWriter(nint logger);
 		#endregion
 
 		public LoggerImpl(Microsoft.Extensions.Logging.ILogger logger)
@@ -63,7 +63,7 @@ namespace EpicGames.UBA
 
 		#region ILogger
 
-		public IntPtr GetHandle() => _handle;
+		public nint GetHandle() => _handle;
 
 		public void BeginScope()
 		{
@@ -86,6 +86,6 @@ namespace EpicGames.UBA
 		}
 		#endregion
 
-		void Log(LogEntryType type, IntPtr ptr, uint len) => Log(type, Marshal.PtrToStringAuto(ptr, (int)len) ?? String.Empty);
+		void Log(LogEntryType type, nint ptr, uint len) => Log(type, Marshal.PtrToStringAuto(ptr, (int)len) ?? String.Empty);
 	}
 }
\ No newline at end of file
diff --git a/Engine/Source/Programs/Shared/EpicGames.UBA/Impl/ProcessImpl.cs b/Engine/Source/Programs/Shared/EpicGames.UBA/Impl/ProcessImpl.cs
--- a/Engine/Source/Programs/Shared/EpicGames.UBA/Impl/ProcessImpl.cs
+++ b/Engine/Source/Programs/Shared/EpicGames.UBA/Impl/ProcessImpl.cs
@@ -5,21 +5,21 @@ using System.Collections.Generic;
 using System.Runtime.InteropServices;
 using System.Threading;
 
-namespace EpicGames.UBA
+namespace EpicGames.UBA.Impl
 {
 	internal class ProcessStartInfoImpl : IProcessStartInfo
 	{
-		IntPtr _handle = IntPtr.Zero;
-		public delegate void ExitCallback(IntPtr userData, IntPtr handle);
+		nint _handle = IntPtr.Zero;
+		public delegate void ExitCallback(nint userData, nint handle);
 		public event ExitCallback? ExitCallbackDelegate;
 		public ProcessImpl? Process { get; set; }
 
 		#region DllImport
 		[DllImport("UbaHost", CharSet = CharSet.Auto)]
-		static extern IntPtr ProcessStartInfo_Create(string application, string arguments, string workingDir, string description, uint priorityClass, ulong outputStatsThresholdMs, bool trackInputs, string logFile, ExitCallback? exit);
+		static extern nint ProcessStartInfo_Create(string application, string arguments, string workingDir, string description, uint priorityClass, ulong outputStatsThresholdMs, bool trackInputs, string logFile, ExitCallback? exit);
 
 		[DllImport("UbaHost", CharSet = CharSet.Auto)]
-		static extern void ProcessStartInfo_Destroy(IntPtr server);
+		static extern void ProcessStartInfo_Destroy(nint server);
 		#endregion
 
 		public ProcessStartInfoImpl(ProcessStartInfo info, bool needCallback)
@@ -29,10 +29,10 @@ namespace EpicGames.UBA
 				ExitCallbackDelegate = RaiseExited;
 			}
 
-			_handle = ProcessStartInfo_Create(info.Application, info.Arguments, info.WorkingDirectory, info.Description, (uint)info.Priority, info.OutputStatsThresholdMs, info.TrackInputs, info.LogFile?? String.Empty, ExitCallbackDelegate);
+			_handle = ProcessStartInfo_Create(info.Application, info.Arguments, info.WorkingDirectory, info.Description, (uint)info.Priority, info.OutputStatsThresholdMs, info.TrackInputs, info.LogFile ?? String.Empty, ExitCallbackDelegate);
 		}
 
-		void RaiseExited(IntPtr userData, IntPtr handle)
+		void RaiseExited(nint userData, nint handle)
 		{
 			while (Process == null)
 			{
@@ -66,43 +66,43 @@ namespace EpicGames.UBA
 		#endregion
 
 		#region IProcessStartInfo
-		public IntPtr GetHandle() => _handle;
+		public nint GetHandle() => _handle;
 		#endregion
 	}
 
 	internal class ProcessImpl : IProcess
 	{
-		IntPtr _handle = IntPtr.Zero;
+		nint _handle = IntPtr.Zero;
 		[System.Diagnostics.CodeAnalysis.SuppressMessage("CodeQuality", "IDE0052:Remove unread private members", Justification = "maintain unmanaged reference")]
 		readonly IProcessStartInfo _processStartInfo;
 
 		#region DllImport
 		[DllImport("UbaHost", CharSet = CharSet.Auto)]
-		static extern uint ProcessHandle_GetExitCode(IntPtr process);
+		static extern uint ProcessHandle_GetExitCode(nint process);
 
 		[DllImport("UbaHost", CharSet = CharSet.Auto)]
-		static extern IntPtr ProcessHandle_GetExecutingHost(IntPtr process);
+		static extern nint ProcessHandle_GetExecutingHost(nint process);
 
 		[DllImport("UbaHost", CharSet = CharSet.Auto)]
-		static extern IntPtr ProcessHandle_GetLogLine(IntPtr process, uint index);
+		static extern nint ProcessHandle_GetLogLine(nint process, uint index);
 
 		[DllImport("UbaHost", CharSet = CharSet.Auto)]
-		static extern ulong ProcessHandle_GetTotalProcessorTime(IntPtr process);
+		static extern ulong ProcessHandle_GetTotalProcessorTime(nint process);
 
 		[DllImport("UbaHost", CharSet = CharSet.Auto)]
-		static extern ulong ProcessHandle_GetTotalWallTime(IntPtr process);
+		static extern ulong ProcessHandle_GetTotalWallTime(nint process);
 
 		[DllImport("UbaHost", CharSet = CharSet.Auto)]
-		static extern ulong ProcessHandle_GetHash(IntPtr process);
+		static extern ulong ProcessHandle_GetHash(nint process);
 
 		[DllImport("UbaHost", CharSet = CharSet.Auto)]
-		static extern void ProcessHandle_Cancel(IntPtr process, bool terminate);
+		static extern void ProcessHandle_Cancel(nint process, bool terminate);
 
 		[DllImport("UbaHost", CharSet = CharSet.Auto)]
-		static extern void DestroyProcessHandle(IntPtr process);
+		static extern void DestroyProcessHandle(nint process);
 		#endregion
 
-		public ProcessImpl(IntPtr handle, IProcessStartInfo info, IProcess.ExitedEventHandler? exitedEventHandler, object? userData)
+		public ProcessImpl(nint handle, IProcessStartInfo info, IProcess.ExitedEventHandler? exitedEventHandler, object? userData)
 		{
 			_handle = handle;
 			_processStartInfo = info;
@@ -138,7 +138,7 @@ namespace EpicGames.UBA
 		#endregion
 
 		#region IServer
-		public IntPtr GetHandle() => _handle;
+		public nint GetHandle() => _handle;
 		public event IProcess.ExitedEventHandler? Exited;
 
 		public int ExitCode => (int)ProcessHandle_GetExitCode(_handle);
@@ -149,7 +149,7 @@ namespace EpicGames.UBA
 		{
 			get
 			{
-				List<string> logLines = new();
+				List<string> logLines = [];
 				string? line = Marshal.PtrToStringAuto(ProcessHandle_GetLogLine(_handle, 0));
 				while (line != null)
 				{
@@ -166,7 +166,7 @@ namespace EpicGames.UBA
 
 		public ulong Hash => ProcessHandle_GetHash(_handle);
 
-		public object? UserData { get; init; }
+		public object? UserData { get; }
 
 		public void Cancel(bool terminate) => ProcessHandle_Cancel(_handle, terminate);
 		#endregion
diff --git a/Engine/Source/Programs/Shared/EpicGames.UBA/Impl/ServerImpl.cs b/Engine/Source/Programs/Shared/EpicGames.UBA/Impl/ServerImpl.cs
--- a/Engine/Source/Programs/Shared/EpicGames.UBA/Impl/ServerImpl.cs
+++ b/Engine/Source/Programs/Shared/EpicGames.UBA/Impl/ServerImpl.cs
@@ -3,32 +3,32 @@
 using System;
 using System.Runtime.InteropServices;
 
-namespace EpicGames.UBA
+namespace EpicGames.UBA.Impl
 {
 	internal class ServerImpl : IServer
 	{
-		IntPtr _handle = IntPtr.Zero;
+		nint _handle = IntPtr.Zero;
 		readonly ILogger _logger;
 
 		#region DllImport
 		[DllImport("UbaHost", CharSet = CharSet.Auto)]
-		static extern IntPtr NetworkServer_Create(IntPtr logger, int workerCount, int sendSize, int receiveTimeoutSeconds, bool useQuic);
+		static extern nint NetworkServer_Create(nint logger, int workerCount, int sendSize, int receiveTimeoutSeconds, bool useQuic);
 
 		[DllImport("UbaHost", CharSet = CharSet.Auto)]
-		static extern void NetworkServer_Destroy(IntPtr server);
+		static extern void NetworkServer_Destroy(nint server);
 
 		// Server Imports
 		[DllImport("UbaHost", CharSet = CharSet.Auto)]
-		static extern bool NetworkServer_StartListen(IntPtr server, int port, string ip, string crypto);
+		static extern bool NetworkServer_StartListen(nint server, int port, string ip, string crypto);
 
 		[DllImport("UbaHost", CharSet = CharSet.Auto)]
-		static extern bool NetworkServer_AddClient(IntPtr server, string ip, int port, string crypto);
+		static extern bool NetworkServer_AddClient(nint server, string ip, int port, string crypto);
 
 		[DllImport("UbaHost", CharSet = CharSet.Auto)]
-		static extern bool NetworkServer_AddNamedConnection(IntPtr server, string name);
+		static extern bool NetworkServer_AddNamedConnection(nint server, string name);
 
 		[DllImport("UbaHost", CharSet = CharSet.Auto)]
-		static extern bool NetworkServer_Stop(IntPtr server);
+		static extern bool NetworkServer_Stop(nint server);
 
 		const int DEFAULT_PORT = 1345;
 		#endregion
@@ -65,7 +65,7 @@ namespace EpicGames.UBA
 
 		#region IServer
 		/// <inheritdoc/>
-		public IntPtr GetHandle() => _handle;
+		public nint GetHandle() => _handle;
 
 		/// <inheritdoc/>
 		public bool StartServer(string ip = "", int port = -1, string crypto = "") => NetworkServer_StartListen(_handle, port > 0 ? port : DEFAULT_PORT, ip, crypto);
diff --git a/Engine/Source/Programs/Shared/EpicGames.UBA/Impl/SessionServerImpl.cs b/Engine/Source/Programs/Shared/EpicGames.UBA/Impl/SessionServerImpl.cs
--- a/Engine/Source/Programs/Shared/EpicGames.UBA/Impl/SessionServerImpl.cs
+++ b/Engine/Source/Programs/Shared/EpicGames.UBA/Impl/SessionServerImpl.cs
@@ -7,22 +7,22 @@ using System.Collections.Generic;
 using System.IO;
 using System.Runtime.InteropServices;
 
-namespace EpicGames.UBA
+namespace EpicGames.UBA.Impl
 {
 	internal class SessionServerCreateInfoImpl : ISessionServerCreateInfo
 	{
-		IntPtr _handle = IntPtr.Zero;
+		nint _handle = IntPtr.Zero;
 		readonly IStorageServer _storage;
 		readonly IServer _client;
 		readonly ILogger _logger;
 
 		#region DllImport
 		[DllImport("UbaHost", CharSet = CharSet.Auto)]
-		static extern IntPtr SessionServerCreateInfo_Create(IntPtr storage, IntPtr client, IntPtr logger, string rootDir, string traceOutputFile,
+		static extern nint SessionServerCreateInfo_Create(nint storage, nint client, nint logger, string rootDir, string traceOutputFile,
 			bool disableCustomAllocator, bool launchVisualizer, bool resetCas, bool writeToDisk, bool detailedTrace, bool allowWaitOnMem, bool allowKillOnMem, bool storeObjFilesCompressed);
 
 		[DllImport("UbaHost", CharSet = CharSet.Auto)]
-		static extern void SessionServerCreateInfo_Destroy(IntPtr server);
+		static extern void SessionServerCreateInfo_Destroy(nint server);
 		#endregion
 
 		public SessionServerCreateInfoImpl(IStorageServer storage, IServer client, ILogger logger, SessionServerCreateInfo info)
@@ -57,16 +57,16 @@ namespace EpicGames.UBA
 		#endregion
 
 		#region ISessionServerCreateInfo
-		public IntPtr GetHandle() => _handle;
+		public nint GetHandle() => _handle;
 		#endregion
 	}
 
 	internal class SessionServerImpl : ISessionServer
 	{
 		delegate void RemoteProcessSlotAvailableCallback();
-		delegate void RemoteProcessReturnedCallback(IntPtr handle);
+		delegate void RemoteProcessReturnedCallback(nint handle);
 
-		IntPtr _handle = IntPtr.Zero;
+		nint _handle = IntPtr.Zero;
 		readonly ISessionServerCreateInfo _info;
 		readonly RemoteProcessSlotAvailableCallback _remoteProcessSlotAvailableCallbackDelegate;
 		readonly RemoteProcessReturnedCallback _remoteProcessReturnedCallbackDelegate;
@@ -74,55 +74,55 @@ namespace EpicGames.UBA
 
 		#region DllImport
 		[DllImport("UbaHost", CharSet = CharSet.Auto)]
-		static extern IntPtr SessionServer_Create(IntPtr info, byte[] environment, uint environmentSize);
+		static extern nint SessionServer_Create(nint info, byte[] environment, uint environmentSize);
 
 		[DllImport("UbaHost", CharSet = CharSet.Auto)]
-		static extern void SessionServer_SetRemoteProcessAvailable(IntPtr server, RemoteProcessSlotAvailableCallback func);
+		static extern void SessionServer_SetRemoteProcessAvailable(nint server, RemoteProcessSlotAvailableCallback func);
 
 		[DllImport("UbaHost", CharSet = CharSet.Auto)]
-		static extern void SessionServer_SetRemoteProcessReturned(IntPtr server, RemoteProcessReturnedCallback func);
+		static extern void SessionServer_SetRemoteProcessReturned(nint server, RemoteProcessReturnedCallback func);
 
 		[DllImport("UbaHost", CharSet = CharSet.Auto)]
-		static extern void SessionServer_RefreshDirectory(IntPtr server, string directory);
+		static extern void SessionServer_RefreshDirectory(nint server, string directory);
 
 		[DllImport("UbaHost", CharSet = CharSet.Auto)]
-		static extern void SessionServer_RegisterNewFile(IntPtr server, string filename);
+		static extern void SessionServer_RegisterNewFile(nint server, string filename);
 
 		[DllImport("UbaHost", CharSet = CharSet.Auto)]
-		static extern uint SessionServer_BeginExternalProcess(IntPtr server, string description);
+		static extern uint SessionServer_BeginExternalProcess(nint server, string description);
 
 		[DllImport("UbaHost", CharSet = CharSet.Auto)]
-		static extern void SessionServer_EndExternalProcess(IntPtr server, uint id, uint exitCode);
+		static extern void SessionServer_EndExternalProcess(nint server, uint id, uint exitCode);
 
 		[DllImport("UbaHost", CharSet = CharSet.Auto)]
-		static extern void SessionServer_UpdateProgress(IntPtr server, uint processesTotal, uint processesDone, uint errorCount);
+		static extern void SessionServer_UpdateProgress(nint server, uint processesTotal, uint processesDone, uint errorCount);
 
 		[DllImport("UbaHost", CharSet = CharSet.Auto)]
-		static extern void SessionServer_UpdateStatus(IntPtr server, uint statusRow, uint statusColumn, string statusText, byte statusType, string? statusLink);
+		static extern void SessionServer_UpdateStatus(nint server, uint statusRow, uint statusColumn, string statusText, byte statusType, string? statusLink);
 
 		[DllImport("UbaHost", CharSet = CharSet.Auto)]
-		static extern IntPtr SessionServer_RunProcess(IntPtr server, IntPtr info, bool async, bool enableDetour);
+		static extern nint SessionServer_RunProcess(nint server, nint info, bool async, bool enableDetour);
 
 		[DllImport("UbaHost", CharSet = CharSet.Auto)]
-		static extern IntPtr SessionServer_RunProcessRemote(IntPtr server, IntPtr info, float weight, byte[]? knownInputs, uint knownInputsCount);
+		static extern nint SessionServer_RunProcessRemote(nint server, nint info, float weight, byte[]? knownInputs, uint knownInputsCount);
 
 		[DllImport("UbaHost", CharSet = CharSet.Auto)]
-		static extern void SessionServer_SetMaxRemoteProcessCount(IntPtr server, uint count);
+		static extern void SessionServer_SetMaxRemoteProcessCount(nint server, uint count);
 
 		[DllImport("UbaHost", CharSet = CharSet.Auto)]
-		static extern void SessionServer_DisableRemoteExecution(IntPtr server);
+		static extern void SessionServer_DisableRemoteExecution(nint server);
 
 		[DllImport("UbaHost", CharSet = CharSet.Auto)]
-		static extern void SessionServer_SetCustomCasKeyFromTrackedInputs(IntPtr server, IntPtr process, string filename, string workingdir);
+		static extern void SessionServer_SetCustomCasKeyFromTrackedInputs(nint server, nint process, string filename, string workingdir);
 
 		[DllImport("UbaHost", CharSet = CharSet.Auto)]
-		static extern void SessionServer_PrintSummary(IntPtr server);
+		static extern void SessionServer_PrintSummary(nint server);
 
 		[DllImport("UbaHost", CharSet = CharSet.Auto)]
-		static extern void SessionServer_CancelAll(IntPtr server);
+		static extern void SessionServer_CancelAll(nint server);
 
 		[DllImport("UbaHost", CharSet = CharSet.Auto)]
-		static extern void SessionServer_Destroy(IntPtr server);
+		static extern void SessionServer_Destroy(nint server);
 		#endregion
 
 		public SessionServerImpl(ISessionServerCreateInfo info)
@@ -175,7 +175,7 @@ namespace EpicGames.UBA
 		#endregion
 
 		#region ISessionServer
-		public IntPtr GetHandle() => _handle;
+		public nint GetHandle() => _handle;
 
 		public event ISessionServer.RemoteProcessSlotAvailableEventHandler? RemoteProcessSlotAvailable;
 
@@ -184,7 +184,7 @@ namespace EpicGames.UBA
 		public IProcess RunProcess(ProcessStartInfo info, bool async, IProcess.ExitedEventHandler? exitedEventHandler, bool enableDetour)
 		{
 			IProcessStartInfo startInfo = IProcessStartInfo.CreateProcessStartInfo(info, exitedEventHandler != null);
-			IntPtr processPtr = SessionServer_RunProcess(_handle, startInfo.GetHandle(), async, enableDetour);
+			nint processPtr = SessionServer_RunProcess(_handle, startInfo.GetHandle(), async, enableDetour);
 			IProcess process = IProcess.CreateProcess(processPtr, startInfo, exitedEventHandler, info.UserData);
 			return process;
 		}
@@ -192,7 +192,7 @@ namespace EpicGames.UBA
 		public IProcess RunProcessRemote(ProcessStartInfo info, IProcess.ExitedEventHandler? exitedEventHandler, double weight, byte[]? knownInputs, uint knownInputsCount)
 		{
 			IProcessStartInfo startInfo = IProcessStartInfo.CreateProcessStartInfo(info, exitedEventHandler != null);
-			IntPtr processPtr = SessionServer_RunProcessRemote(_handle, startInfo.GetHandle(), (float)weight, knownInputs, knownInputsCount);
+			nint processPtr = SessionServer_RunProcessRemote(_handle, startInfo.GetHandle(), (float)weight, knownInputs, knownInputsCount);
 			IProcess process = IProcess.CreateProcess(processPtr, startInfo, exitedEventHandler, info.UserData);
 			_remoteProcesses.AddOrUpdate(process.Hash, process, (k, v) => process);
 			return process;
@@ -217,7 +217,7 @@ namespace EpicGames.UBA
 		public void SetCustomCasKeyFromTrackedInputs(string file, string workingDirectory, IProcess process) => SessionServer_SetCustomCasKeyFromTrackedInputs(_handle, process.GetHandle(), file, workingDirectory);
 
 		public void PrintSummary() => SessionServer_PrintSummary(_handle);
-		
+
 		public void CancelAll() => SessionServer_CancelAll(_handle);
 		#endregion
 
@@ -226,7 +226,7 @@ namespace EpicGames.UBA
 			RemoteProcessSlotAvailable?.Invoke(this, new RemoteProcessSlotAvailableEventArgs());
 		}
 
-		void RaiseRemoteProcessReturned(IntPtr handle)
+		void RaiseRemoteProcessReturned(nint handle)
 		{
 			IProcess? process = null;
 			try
diff --git a/Engine/Source/Programs/Shared/EpicGames.UBA/Impl/StorageServerImpl.cs b/Engine/Source/Programs/Shared/EpicGames.UBA/Impl/StorageServerImpl.cs
--- a/Engine/Source/Programs/Shared/EpicGames.UBA/Impl/StorageServerImpl.cs
+++ b/Engine/Source/Programs/Shared/EpicGames.UBA/Impl/StorageServerImpl.cs
@@ -4,26 +4,26 @@ using System;
 using System.Linq;
 using System.Runtime.InteropServices;
 
-namespace EpicGames.UBA
+namespace EpicGames.UBA.Impl
 {
 	internal class StorageServerImpl : IStorageServer
 	{
-		IntPtr _handle = IntPtr.Zero;
+		nint _handle = IntPtr.Zero;
 		readonly IServer _server;
 		readonly ILogger _logger;
 
 		#region DllImport
 		[DllImport("UbaHost", CharSet = CharSet.Auto)]
-		static extern IntPtr StorageServer_Create(IntPtr server, string rootDir, ulong casCapacityBytes, bool storeCompressed, IntPtr logger, string zone);
+		static extern nint StorageServer_Create(nint server, string rootDir, ulong casCapacityBytes, bool storeCompressed, nint logger, string zone);
 
 		[DllImport("UbaHost", CharSet = CharSet.Auto)]
-		static extern void StorageServer_Destroy(IntPtr server);
+		static extern void StorageServer_Destroy(nint server);
 
 		[DllImport("UbaHost", CharSet = CharSet.Auto)]
-		static extern void StorageServer_SaveCasTable(IntPtr server);
+		static extern void StorageServer_SaveCasTable(nint server);
 
 		[DllImport("UbaHost", CharSet = CharSet.Auto)]
-		static extern void StorageServer_RegisterDisallowedPath(IntPtr server, string path);
+		static extern void StorageServer_RegisterDisallowedPath(nint server, string path);
 		#endregion
 
 		public StorageServerImpl(IServer server, ILogger logger, StorageServerCreateInfo info)
@@ -58,7 +58,7 @@ namespace EpicGames.UBA
 		#endregion
 
 		#region IStorageServer
-		public IntPtr GetHandle() => _handle;
+		public nint GetHandle() => _handle;
 
 		public void SaveCasTable() => StorageServer_SaveCasTable(_handle);
 
diff --git a/Engine/Source/Programs/Shared/EpicGames.UBA/ThreadedLogger.cs b/Engine/Source/Programs/Shared/EpicGames.UBA/ThreadedLogger.cs
--- a/Engine/Source/Programs/Shared/EpicGames.UBA/ThreadedLogger.cs
+++ b/Engine/Source/Programs/Shared/EpicGames.UBA/ThreadedLogger.cs
@@ -38,11 +38,11 @@ namespace EpicGames.UBA
 			_running = 1;
 			_logQueueLock = new object();
 			_processLogQueue = new EventWaitHandle(false, EventResetMode.AutoReset);
-			_logQueue = new List<Action>();
+			_logQueue = [];
 			_token = new CancellationToken();
 			_logTask = Task.Factory.StartNew(() =>
 			{
-				List<Action> logQueue2 = new();
+				List<Action> logQueue2 = [];
 				while (Interlocked.Read(ref _running) == 1)
 				{
 					_processLogQueue.WaitOne();
diff --git a/Engine/Source/Programs/Shared/EpicGames.UBA/Utils.cs b/Engine/Source/Programs/Shared/EpicGames.UBA/Utils.cs
--- a/Engine/Source/Programs/Shared/EpicGames.UBA/Utils.cs
+++ b/Engine/Source/Programs/Shared/EpicGames.UBA/Utils.cs
@@ -25,7 +25,7 @@ namespace EpicGames.UBA
 		/// </summary>
 		/// <returns>enumerable of disallowed paths</returns>
 		public static IEnumerable<string> DisallowedPaths() => s_disallowedPaths;
-		static readonly SortedSet<string> s_disallowedPaths = new();
+		static readonly SortedSet<string> s_disallowedPaths = [];
 
 		/// <summary>
 		/// Registers a path that is not allowed to be transferred over the network for UBA remote agents.
