EpicGames.Core: .net8 suggestions pt1

#rnx

diff --git a/Engine/Source/Programs/Shared/EpicGames.Core/AsyncPipeline.cs b/Engine/Source/Programs/Shared/EpicGames.Core/AsyncPipeline.cs
--- a/Engine/Source/Programs/Shared/EpicGames.Core/AsyncPipeline.cs
+++ b/Engine/Source/Programs/Shared/EpicGames.Core/AsyncPipeline.cs
@@ -22,8 +22,7 @@ namespace EpicGames.Core
 		/// <summary>
 		/// Tests whether the pipeline has failed
 		/// </summary>
-		public bool IsFaulted
-			=> _exceptionDispatchInfo != null;
+		public bool IsFaulted => _exceptionDispatchInfo != null;
 
 		/// <summary>
 		/// Constructor
@@ -31,7 +30,7 @@ namespace EpicGames.Core
 		/// <param name="cancellationToken">Cancellation token for the pipeline</param>
 		public AsyncPipeline(CancellationToken cancellationToken)
 		{
-			_tasks = new List<Task>();
+			_tasks = [];
 			_cancellationToken = cancellationToken;
 			_cancellationSource = CancellationTokenSource.CreateLinkedTokenSource(cancellationToken);
 		}
@@ -109,10 +108,7 @@ namespace EpicGames.Core
 
 			_cancellationToken.ThrowIfCancellationRequested();
 
-			if (_exceptionDispatchInfo != null)
-			{
-				_exceptionDispatchInfo.Throw();
-			}
+			_exceptionDispatchInfo?.Throw();
 		}
 	}
 
diff --git a/Engine/Source/Programs/Shared/EpicGames.Core/AsyncQueue.cs b/Engine/Source/Programs/Shared/EpicGames.Core/AsyncQueue.cs
--- a/Engine/Source/Programs/Shared/EpicGames.Core/AsyncQueue.cs
+++ b/Engine/Source/Programs/Shared/EpicGames.Core/AsyncQueue.cs
@@ -13,7 +13,7 @@ namespace EpicGames.Core
 	/// </summary>
 	public sealed class AsyncQueue : IAsyncDisposable
 	{
-		readonly List<Task> _workerTasks = new List<Task>();
+		readonly List<Task> _workerTasks = [];
 		readonly Channel<Func<CancellationToken, Task>> _queuedTasks = Channel.CreateUnbounded<Func<CancellationToken, Task>>();
 		readonly CancellationTokenSource _cancellationSource = new CancellationTokenSource();
 
diff --git a/Engine/Source/Programs/Shared/EpicGames.Core/AsyncThreadPoolWorkQueue.cs b/Engine/Source/Programs/Shared/EpicGames.Core/AsyncThreadPoolWorkQueue.cs
--- a/Engine/Source/Programs/Shared/EpicGames.Core/AsyncThreadPoolWorkQueue.cs
+++ b/Engine/Source/Programs/Shared/EpicGames.Core/AsyncThreadPoolWorkQueue.cs
@@ -81,7 +81,7 @@ public sealed class AsyncThreadPoolWorkQueue : IDisposable
 			return;
 		}
 		
-		List<Task> tasks = new();
+		List<Task> tasks = [];
 		for (int i = 0; i < _numWorkers; i++)
 		{
 			tasks.Add(Task.Run(() => WorkerAsync(cancellationToken), cancellationToken));
diff --git a/Engine/Source/Programs/Shared/EpicGames.Core/AsyncUtils.cs b/Engine/Source/Programs/Shared/EpicGames.Core/AsyncUtils.cs
--- a/Engine/Source/Programs/Shared/EpicGames.Core/AsyncUtils.cs
+++ b/Engine/Source/Programs/Shared/EpicGames.Core/AsyncUtils.cs
@@ -101,7 +101,7 @@ namespace EpicGames.Core
 		/// <param name="tasks">List of tasks to remove tasks from</param>
 		public static void RemoveCompleteTasks(this List<Task> tasks)
 		{
-			List<Exception> exceptions = new List<Exception>();
+			List<Exception> exceptions = [];
 
 			int outIdx = 0;
 			for (int idx = 0; idx < tasks.Count; idx++)
@@ -138,7 +138,7 @@ namespace EpicGames.Core
 		/// <returns>Return values from the completed tasks</returns>
 		public static List<T> RemoveCompleteTasks<T>(this List<Task<T>> tasks)
 		{
-			List<T> results = new List<T>();
+			List<T> results = [];
 
 			int outIdx = 0;
 			for (int idx = 0; idx < tasks.Count; idx++)
diff --git a/Engine/Source/Programs/Shared/EpicGames.Core/BinaryArchive.cs b/Engine/Source/Programs/Shared/EpicGames.Core/BinaryArchive.cs
--- a/Engine/Source/Programs/Shared/EpicGames.Core/BinaryArchive.cs
+++ b/Engine/Source/Programs/Shared/EpicGames.Core/BinaryArchive.cs
@@ -38,7 +38,7 @@ namespace EpicGames.Core
 
 			static Func<BinaryArchiveReader, object> CreateBoxedReadMethod(Type type, MethodInfo methodInfo)
 			{
-				DynamicMethod readerMethod = new DynamicMethod($"Boxed_{methodInfo.Name}", typeof(object), new[] { typeof(BinaryArchiveReader) });
+				DynamicMethod readerMethod = new DynamicMethod($"Boxed_{methodInfo.Name}", typeof(object), [typeof(BinaryArchiveReader)]);
 
 				ILGenerator generator = readerMethod.GetILGenerator(64);
 				generator.Emit(OpCodes.Ldarg_0);
@@ -54,7 +54,7 @@ namespace EpicGames.Core
 
 			static Action<BinaryArchiveWriter, object> CreateBoxedWriteMethod(Type type, MethodInfo methodInfo)
 			{
-				DynamicMethod writerMethod = new DynamicMethod($"Boxed_{methodInfo.Name}", typeof(void), new[] { typeof(BinaryArchiveWriter), typeof(object) });
+				DynamicMethod writerMethod = new DynamicMethod($"Boxed_{methodInfo.Name}", typeof(void), [typeof(BinaryArchiveWriter), typeof(object)]);
 
 				ILGenerator generator = writerMethod.GetILGenerator(64);
 				generator.Emit(OpCodes.Ldarg_0);
@@ -262,12 +262,12 @@ namespace EpicGames.Core
 				InterfaceMapping interfaceMapping = converterType.GetInterfaceMap(interfaceType);
 
 				// Create a reader method
-				DynamicMethod readerMethod = new DynamicMethod($"BinaryArchiveReader_Dynamic_{type.Name}_With_{converterType.Name}", type, new[] { typeof(BinaryArchiveReader) });
+				DynamicMethod readerMethod = new DynamicMethod($"BinaryArchiveReader_Dynamic_{type.Name}_With_{converterType.Name}", type, [typeof(BinaryArchiveReader)]);
 				int readerIdx = Array.FindIndex(interfaceMapping.InterfaceMethods, x => x.Name.Equals(nameof(IBinaryConverter<object>.Read), StringComparison.Ordinal));
 				GenerateConverterReaderForwardingMethod(readerMethod, converterField, interfaceMapping.TargetMethods[readerIdx]);
 
 				// Create a writer method
-				DynamicMethod writerMethod = new DynamicMethod($"BinaryArchiveWriter_Dynamic_{type.Name}_With_{converterType.Name}", typeof(void), new[] { typeof(BinaryArchiveWriter), typeof(object) });
+				DynamicMethod writerMethod = new DynamicMethod($"BinaryArchiveWriter_Dynamic_{type.Name}_With_{converterType.Name}", typeof(void), [typeof(BinaryArchiveWriter), typeof(object)]);
 				int writerIdx = Array.FindIndex(interfaceMapping.InterfaceMethods, x => x.Name.Equals(nameof(IBinaryConverter<object>.Write), StringComparison.Ordinal));
 				GenerateConverterWriterForwardingMethod(writerMethod, converterField, interfaceMapping.TargetMethods[writerIdx]);
 
@@ -295,7 +295,7 @@ namespace EpicGames.Core
 			if (type.IsClass)
 			{
 				// Get all the class properties
-				List<PropertyInfo> properties = new List<PropertyInfo>();
+				List<PropertyInfo> properties = [];
 				foreach (PropertyInfo property in type.GetProperties(BindingFlags.Public | BindingFlags.NonPublic | BindingFlags.Instance))
 				{
 					MethodInfo? getMethod = property.GetGetMethod();
@@ -307,14 +307,14 @@ namespace EpicGames.Core
 				}
 
 				// Find the type constructor
-				ConstructorInfo? constructorInfo = type.GetConstructor(BindingFlags.Instance | BindingFlags.Public | BindingFlags.NonPublic, null, Array.Empty<Type>(), null)
+				ConstructorInfo? constructorInfo = type.GetConstructor(BindingFlags.Instance | BindingFlags.Public | BindingFlags.NonPublic, null, [], null)
 					?? throw new NotImplementedException($"Type '{type.Name}' does not have a parameterless constructor");
 
 				// Create the methods
-				DynamicMethod writerMethod = new DynamicMethod($"BinaryArchiveWriter_Dynamic_{type.Name}", typeof(void), new[] { typeof(BinaryArchiveWriter), type });
+				DynamicMethod writerMethod = new DynamicMethod($"BinaryArchiveWriter_Dynamic_{type.Name}", typeof(void), [typeof(BinaryArchiveWriter), type]);
 				GenerateClassWriterMethod(writerMethod, properties);
 
-				DynamicMethod readerMethod = new DynamicMethod($"BinaryArchiveReader_Dynamic_{type.Name}", type, new[] { typeof(BinaryArchiveReader) });
+				DynamicMethod readerMethod = new DynamicMethod($"BinaryArchiveReader_Dynamic_{type.Name}", type, [typeof(BinaryArchiveReader)]);
 				GenerateClassReaderMethod(readerMethod, constructorInfo, properties);
 
 				return new TypeSerializer(readerMethod, writerMethod);
diff --git a/Engine/Source/Programs/Shared/EpicGames.Core/BinaryArchiveReader.cs b/Engine/Source/Programs/Shared/EpicGames.Core/BinaryArchiveReader.cs
--- a/Engine/Source/Programs/Shared/EpicGames.Core/BinaryArchiveReader.cs
+++ b/Engine/Source/Programs/Shared/EpicGames.Core/BinaryArchiveReader.cs
@@ -32,7 +32,7 @@ namespace EpicGames.Core
 		/// <summary>
 		/// List of previously serialized objects
 		/// </summary>
-		readonly List<object?> _objects = new List<object?>();
+		readonly List<object?> _objects = [];
 
 		/// <summary>
 		/// Constructor
@@ -356,7 +356,7 @@ namespace EpicGames.Core
 			}
 			else
 			{
-				HashSet<T> result = new HashSet<T>();
+				HashSet<T> result = [];
 				for (int idx = 0; idx < count; idx++)
 				{
 					result.Add(readElement());
@@ -405,7 +405,7 @@ namespace EpicGames.Core
 			}
 			else
 			{
-				SortedSet<T> result = new SortedSet<T>();
+				SortedSet<T> result = [];
 				for (int idx = 0; idx < count; idx++)
 				{
 					result.Add(readElement());
diff --git a/Engine/Source/Programs/Shared/EpicGames.Core/BuzHash.cs b/Engine/Source/Programs/Shared/EpicGames.Core/BuzHash.cs
--- a/Engine/Source/Programs/Shared/EpicGames.Core/BuzHash.cs
+++ b/Engine/Source/Programs/Shared/EpicGames.Core/BuzHash.cs
@@ -10,8 +10,8 @@ namespace EpicGames.Core
 	/// </summary>
 	public class BuzHash
 	{
-		static readonly uint[] s_table = new uint[256]
-		{
+		static readonly uint[] s_table =
+		[
 			0x458be752, 0xc10748cc, 0xfbbcdbb8, 0x6ded5b68, 0xb10a82b5, 0x20d75648, 0xdfc5665f, 0xa8428801, 0x7ebf5191, 0x841135c7, 0x65cc53b3,
 			0x280a597c, 0x16f60255, 0xc78cbc3e, 0x294415f5, 0xb938d494, 0xec85c4e6, 0xb7d33edc, 0xe549b544, 0xfdeda5aa, 0x882bf287, 0x3116737c,
 			0x05569956, 0xe8cc1f68, 0x0806ac5e, 0x22a14443, 0x15297e10, 0x50d090e7, 0x4ba60f6f, 0xefd9f1a7, 0x5c5c885c, 0x82482f93, 0x9bfd7c64,
@@ -36,7 +36,7 @@ namespace EpicGames.Core
 			0xc6eb57bb, 0x30320350, 0x3f397e91, 0xa67791bc, 0x86bc0e2c, 0xefa0a7e2, 0xe9ff7543, 0xe733612c, 0xd185897b, 0x329e5388, 0x91dd236b,
 			0x2ecb0d93, 0xf4d82a3d, 0x35b5c03f, 0xe4e606f0, 0x05b21843, 0x37b45964, 0x5eff22f4, 0x6027f4cc, 0x77178b3c, 0xae507131, 0x7bf7cabc,
 			0xf9c18d66, 0x593ade65, 0xd95ddf11,
-		};
+		];
 
 		uint _state;
 
diff --git a/Engine/Source/Programs/Shared/EpicGames.Core/ByteArrayBuilder.cs b/Engine/Source/Programs/Shared/EpicGames.Core/ByteArrayBuilder.cs
--- a/Engine/Source/Programs/Shared/EpicGames.Core/ByteArrayBuilder.cs
+++ b/Engine/Source/Programs/Shared/EpicGames.Core/ByteArrayBuilder.cs
@@ -27,7 +27,7 @@ namespace EpicGames.Core
 			public ReadOnlyMemory<byte> WrittenMemory => Data.AsMemory(0, Length);
 		}
 
-		readonly List<Chunk> _chunks = new List<Chunk>();
+		readonly List<Chunk> _chunks = [];
 		readonly int _chunkSize;
 		Chunk _currentChunk;
 
diff --git a/Engine/Source/Programs/Shared/EpicGames.Core/CertificateUtils.cs b/Engine/Source/Programs/Shared/EpicGames.Core/CertificateUtils.cs
--- a/Engine/Source/Programs/Shared/EpicGames.Core/CertificateUtils.cs
+++ b/Engine/Source/Programs/Shared/EpicGames.Core/CertificateUtils.cs
@@ -31,7 +31,7 @@ namespace EpicGames.Core
 
 				request.CertificateExtensions.Add(new X509BasicConstraintsExtension(false, false, 0, true));
 				request.CertificateExtensions.Add(new X509KeyUsageExtension(X509KeyUsageFlags.KeyEncipherment | X509KeyUsageFlags.DigitalSignature, true));
-				request.CertificateExtensions.Add(new X509EnhancedKeyUsageExtension(new OidCollection { new Oid("1.3.6.1.5.5.7.3.1") }, true));
+				request.CertificateExtensions.Add(new X509EnhancedKeyUsageExtension([new Oid("1.3.6.1.5.5.7.3.1")], true));
 
 				SubjectAlternativeNameBuilder alternativeNameBuilder = new SubjectAlternativeNameBuilder();
 				alternativeNameBuilder.AddDnsName(dnsName);
diff --git a/Engine/Source/Programs/Shared/EpicGames.Core/CommandHost.cs b/Engine/Source/Programs/Shared/EpicGames.Core/CommandHost.cs
--- a/Engine/Source/Programs/Shared/EpicGames.Core/CommandHost.cs
+++ b/Engine/Source/Programs/Shared/EpicGames.Core/CommandHost.cs
@@ -92,7 +92,7 @@ namespace EpicGames.Core
 		/// <param name="description">Short description for display in the help text</param>
 		public CommandAttribute(string name, string description)
 		{
-			Names = new string[] { name };
+			Names = [name];
 			Description = description;
 		}
 
@@ -104,7 +104,7 @@ namespace EpicGames.Core
 		/// <param name="description">Short description for display in the help text</param>
 		public CommandAttribute(string category, string name, string description)
 		{
-			Names = new string[] { category, name };
+			Names = [category, name];
 			Description = description;
 		}
 	}
@@ -166,7 +166,7 @@ namespace EpicGames.Core
 		/// <param name="assembly"></param>
 		public static void AddCommandsFromAssembly(this IServiceCollection services, Assembly assembly)
 		{
-			List<(CommandAttribute, Type)> commands = new List<(CommandAttribute, Type)>();
+			List<(CommandAttribute, Type)> commands = [];
 			foreach (Type type in assembly.GetTypes())
 			{
 				if (typeof(ICommand).IsAssignableFrom(type) && !type.IsAbstract)
@@ -305,7 +305,7 @@ namespace EpicGames.Core
 		/// <param name="attributes">List of command attributes</param>
 		static void PrintCommands(IEnumerable<ICommandFactory> attributes)
 		{
-			List<KeyValuePair<string, string>> commands = new List<KeyValuePair<string, string>>();
+			List<KeyValuePair<string, string>> commands = [];
 			foreach (ICommandFactory attribute in attributes)
 			{
 				if (attribute.Advertise)
@@ -313,7 +313,7 @@ namespace EpicGames.Core
 					commands.Add(new KeyValuePair<string, string>(String.Join(" ", attribute.Names), attribute.Description));
 				}
 			}
-			HelpUtils.PrintTable(commands.OrderBy(x => x.Key).ToList(), 4, 20);
+			HelpUtils.PrintTable([.. commands.OrderBy(x => x.Key)], 4, 20);
 		}
 	}
 }
diff --git a/Engine/Source/Programs/Shared/EpicGames.Core/CommandLineArguments.cs b/Engine/Source/Programs/Shared/EpicGames.Core/CommandLineArguments.cs
--- a/Engine/Source/Programs/Shared/EpicGames.Core/CommandLineArguments.cs
+++ b/Engine/Source/Programs/Shared/EpicGames.Core/CommandLineArguments.cs
@@ -129,12 +129,12 @@ namespace EpicGames.Core
 		/// <summary>
 		/// List of positional arguments
 		/// </summary>
-		readonly List<int> _positionalArgumentIndices = new List<int>();
+		readonly List<int> _positionalArgumentIndices = [];
 
 		/// <summary>
 		/// Array of characters that separate argument names from values
 		/// </summary>
-		static readonly char[] s_valueSeparators = { '=', ':' };
+		static readonly char[] s_valueSeparators = ['=', ':'];
 
 		/// <summary>
 		/// Constructor
@@ -170,7 +170,7 @@ namespace EpicGames.Core
 					}
 					break;
 				}
-				else if (arguments[idx].StartsWith("-", StringComparison.Ordinal))
+				else if (arguments[idx].StartsWith('-'))
 				{
 					// Option argument
 					int separatorIdx = arguments[idx].IndexOfAny(s_valueSeparators);
@@ -703,7 +703,7 @@ namespace EpicGames.Core
 			}
 			else
 			{
-				if (valueType != typeof(bool) && attribute.Value == null && !prefix.EndsWith("=", StringComparison.Ordinal) && !prefix.EndsWith(":", StringComparison.Ordinal))
+				if (valueType != typeof(bool) && attribute.Value == null && !prefix.EndsWith('=') && !prefix.EndsWith(':'))
 				{
 					prefix += "=";
 				}
@@ -727,7 +727,7 @@ namespace EpicGames.Core
 		/// <param name="logger">Sink for error/warning messages</param>
 		public void ApplyTo(object targetObject, ILogger logger)
 		{
-			List<string> missingArguments = new List<string>();
+			List<string> missingArguments = [];
 
 			// Build a mapping from name to field and attribute for this object
 			List<ArgumentTarget> targets = GetArgumentTargetsForType(targetObject.GetType());
@@ -837,7 +837,7 @@ namespace EpicGames.Core
 
 		static List<ArgumentTarget> GetArgumentTargetsForType(Type? targetType)
 		{
-			List<ArgumentTarget> targets = new List<ArgumentTarget>();
+			List<ArgumentTarget> targets = [];
 			while (targetType != null && targetType != typeof(object))
 			{
 				foreach (FieldInfo fieldInfo in targetType.GetFields(BindingFlags.Instance | BindingFlags.GetField | BindingFlags.Public | BindingFlags.NonPublic | BindingFlags.DeclaredOnly))
@@ -868,7 +868,7 @@ namespace EpicGames.Core
 		/// <returns>List of parameters</returns>
 		public static List<KeyValuePair<string, string>> GetParameters(Type type)
 		{
-			List<KeyValuePair<string, string>> parameters = new List<KeyValuePair<string, string>>();
+			List<KeyValuePair<string, string>> parameters = [];
 
 			List<ArgumentTarget> targets = GetArgumentTargetsForType(type);
 			foreach (ArgumentTarget target in targets)
@@ -900,7 +900,7 @@ namespace EpicGames.Core
 				foreach (CommandLineAttribute attribute in target.Attributes)
 				{
 					string prefix = GetArgumentPrefix(target, attribute);
-					if(prefix.EndsWith("=", StringComparison.Ordinal))
+					if(prefix.EndsWith('='))
 					{
 						prefix += "...";
 					}
@@ -979,7 +979,7 @@ namespace EpicGames.Core
 		{
 			StringBuilder argument = new StringBuilder();
 
-			List<string> arguments = new List<string>();
+			List<string> arguments = [];
 			// First do a pass leaving all quotes in the arguments, they will be removed later
 			for(int idx = 0; idx < commandLine.Length; idx++)
 			{
@@ -1012,7 +1012,7 @@ namespace EpicGames.Core
 				}
 				arguments[idx] = arg.Replace("\"", "", StringComparison.OrdinalIgnoreCase);
 			}
-			return arguments.ToArray();
+			return [.. arguments];
 		}
 
 		/// <summary>
@@ -1041,7 +1041,7 @@ namespace EpicGames.Core
 		/// <returns>New argument list</returns>
 		public CommandLineArguments Remove(string prefix, out List<string> values)
 		{
-			values = new List<string>();
+			values = [];
 
 			// Split the arguments into the values array and an array of new arguments
 			int[] newArgumentIndex = new int[_arguments.Length];
@@ -1062,7 +1062,7 @@ namespace EpicGames.Core
 			}
 
 			// Create the new argument list, and mark the same arguments as used
-			CommandLineArguments newArguments = new CommandLineArguments(newArgumentList.ToArray());
+			CommandLineArguments newArguments = new CommandLineArguments([.. newArgumentList]);
 			for(int idx = 0; idx < _arguments.Length; idx++)
 			{
 				if(HasBeenUsed(idx) && newArgumentIndex[idx] != -1)
@@ -1087,7 +1087,7 @@ namespace EpicGames.Core
 		public void CheckAllArgumentsUsed(ILogger logger)
 		{
 			// Find all the unused arguments
-			List<string> remainingArguments = new List<string>();
+			List<string> remainingArguments = [];
 			for(int idx = 0; idx < _arguments.Length; idx++)
 			{
 				if(!_usedArguments[idx])
@@ -1218,7 +1218,7 @@ namespace EpicGames.Core
 			else
 			{
 				// Call the 'Add' method on the collection
-				collectionType.InvokeMember("Add", BindingFlags.InvokeMethod, null, target.GetValue(targetObject), new object[] { value });
+				collectionType.InvokeMember("Add", BindingFlags.InvokeMethod, null, target.GetValue(targetObject), [value]);
 				return true;
 			}
 		}
diff --git a/Engine/Source/Programs/Shared/EpicGames.Core/CommandLineAttribute.cs b/Engine/Source/Programs/Shared/EpicGames.Core/CommandLineAttribute.cs
--- a/Engine/Source/Programs/Shared/EpicGames.Core/CommandLineAttribute.cs
+++ b/Engine/Source/Programs/Shared/EpicGames.Core/CommandLineAttribute.cs
@@ -52,7 +52,7 @@ namespace EpicGames.Core
 
 			if(prefix != null)
 			{
-				if(!prefix.StartsWith("-", StringComparison.Ordinal))
+				if(!prefix.StartsWith('-'))
 				{
 					throw new Exception("Command-line arguments must begin with a '-' character");
 				}
diff --git a/Engine/Source/Programs/Shared/EpicGames.Core/ContentHash.cs b/Engine/Source/Programs/Shared/EpicGames.Core/ContentHash.cs
--- a/Engine/Source/Programs/Shared/EpicGames.Core/ContentHash.cs
+++ b/Engine/Source/Programs/Shared/EpicGames.Core/ContentHash.cs
@@ -34,7 +34,7 @@ namespace EpicGames.Core
 		/// <summary>
 		/// Retrieves an empty content hash
 		/// </summary>
-		public static ContentHash Empty { get; } = new ContentHash(Array.Empty<byte>());
+		public static ContentHash Empty { get; } = new ContentHash([]);
 
 		/// <summary>
 		/// The bytes compromising this hash
diff --git a/Engine/Source/Programs/Shared/EpicGames.Core/CsProjBuildRecord.cs b/Engine/Source/Programs/Shared/EpicGames.Core/CsProjBuildRecord.cs
--- a/Engine/Source/Programs/Shared/EpicGames.Core/CsProjBuildRecord.cs
+++ b/Engine/Source/Programs/Shared/EpicGames.Core/CsProjBuildRecord.cs
@@ -63,17 +63,17 @@ namespace EpicGames.Core
 		/// <summary>
 		/// Paths to referenced projects and the target build time in that project
 		/// </summary>
-		public List<CsProjBuildRecordRef> ProjectReferencesAndTimes { get; set; } = new List<CsProjBuildRecordRef>();
+		public List<CsProjBuildRecordRef> ProjectReferencesAndTimes { get; set; } = [];
 
 		/// <summary>
 		/// file dependencies from non-glob sources
 		/// </summary>
-		public HashSet<string> Dependencies { get; set; } = new HashSet<string>();
+		public HashSet<string> Dependencies { get; set; } = [];
 
 		/// <summary>
 		/// file dependencies from globs
 		/// </summary>
-		public HashSet<string> GlobbedDependencies { get; set; } = new HashSet<string>();
+		public HashSet<string> GlobbedDependencies { get; set; } = [];
 
 		/// <summary>
 		/// A glob pattern
@@ -104,7 +104,7 @@ namespace EpicGames.Core
 		/// <summary>
 		/// List of globs
 		/// </summary>
-		public List<Glob> Globs { get; set; } = new List<Glob>();
+		public List<Glob> Globs { get; set; } = [];
 	}
 
 	/// <summary>
diff --git a/Engine/Source/Programs/Shared/EpicGames.Core/CsProjectInfo.cs b/Engine/Source/Programs/Shared/EpicGames.Core/CsProjectInfo.cs
--- a/Engine/Source/Programs/Shared/EpicGames.Core/CsProjectInfo.cs
+++ b/Engine/Source/Programs/Shared/EpicGames.Core/CsProjectInfo.cs
@@ -34,22 +34,22 @@ namespace EpicGames.Core
 		/// <summary>
 		/// Mapping of referenced assemblies to their 'CopyLocal' (aka 'Private') setting.
 		/// </summary>
-		public Dictionary<FileReference, bool> References { get; } = new Dictionary<FileReference, bool>();
+		public Dictionary<FileReference, bool> References { get; } = [];
 
 		/// <summary>
 		/// Mapping of referenced projects to their 'CopyLocal' (aka 'Private') setting.
 		/// </summary>
-		public Dictionary<FileReference, bool> ProjectReferences { get; } = new Dictionary<FileReference, bool>();
+		public Dictionary<FileReference, bool> ProjectReferences { get; } = [];
 
 		/// <summary>
 		/// List of compile references in the project.
 		/// </summary>
-		public List<FileReference> CompileReferences { get; } = new List<FileReference>();
+		public List<FileReference> CompileReferences { get; } = [];
 
 		/// <summary>
 		/// Mapping of content IF they are flagged Always or Newer
 		/// </summary>
-		public Dictionary<FileReference, bool> ContentReferences { get; } = new Dictionary<FileReference, bool>();
+		public Dictionary<FileReference, bool> ContentReferences { get; } = [];
 
 		/// <summary>
 		/// Path to the CSProject file
@@ -118,7 +118,7 @@ namespace EpicGames.Core
 			}
 			else if (IsNetCoreProject())
 			{
-				string configuration = Properties.ContainsKey("Configuration") ? Properties["Configuration"] : "Development";
+				string configuration = Properties.TryGetValue("Configuration", out string? value) ? value : "Development";
 				return DirectoryReference.Combine(baseDirectory, "bin", configuration, Properties["TargetFramework"]);
 			}
 			else
@@ -644,10 +644,10 @@ namespace EpicGames.Core
 					// if we searched all directories regardless of depth, the rest of the components other than the last (file) are irrelevant
 					if (option == SearchOption.AllDirectories)
 					{
-						remainingComponents = new[] { remainingComponents.Last() };
+						remainingComponents = [remainingComponents.Last()];
 
 						// ** includes files in the current directory too
-						directories = directories.Concat(new[] { existingPath });
+						directories = directories.Concat([existingPath]);
 					}
 
 					foreach (DirectoryReference dir in directories)
@@ -689,13 +689,13 @@ namespace EpicGames.Core
 		/// <returns></returns>
 		static IEnumerable<FileReference> FindMatchingFiles(FileReference inPath)
 		{
-			List<FileReference> foundFiles = new List<FileReference>();		
+			List<FileReference> foundFiles = [];		
 
 			// split off the drive root
 			string driveRoot = Path.GetPathRoot(inPath.FullName)!;
 
 			// break the rest of the path into components
-			string[] pathComponents = inPath.FullName.Substring(driveRoot.Length).Split(new char[] { Path.DirectorySeparatorChar, Path.AltDirectorySeparatorChar });
+			string[] pathComponents = inPath.FullName.Substring(driveRoot.Length).Split([Path.DirectorySeparatorChar, Path.AltDirectorySeparatorChar]);
 
 			// Process all the components recursively
 			ProcessPathComponents(new DirectoryReference(driveRoot), pathComponents, foundFiles);
@@ -998,7 +998,7 @@ namespace EpicGames.Core
 					tokenIdx += 3;
 
 					// Parse any arguments
-					List<object> arguments = new List<object>();
+					List<object> arguments = [];
 					if(tokens[tokenIdx] != ")")
 					{
 						arguments.Add(ParseArgument(tokens[tokenIdx]));
@@ -1022,7 +1022,7 @@ namespace EpicGames.Core
 					// Execute the method
 					try
 					{
-						value = typeof(string).InvokeMember(methodName, System.Reflection.BindingFlags.Instance | System.Reflection.BindingFlags.Public | System.Reflection.BindingFlags.InvokeMethod, Type.DefaultBinder, value, arguments.ToArray())!.ToString()!;
+						value = typeof(string).InvokeMember(methodName, System.Reflection.BindingFlags.Instance | System.Reflection.BindingFlags.Public | System.Reflection.BindingFlags.InvokeMethod, Type.DefaultBinder, value, [.. arguments])!.ToString()!;
 					}
 					catch(Exception ex)
 					{
@@ -1081,7 +1081,7 @@ namespace EpicGames.Core
 		/// <returns>Array of the parsed tokens</returns>
 		static string[] Tokenize(string condition)
 		{
-			List<string> tokens = new List<string>();
+			List<string> tokens = [];
 			for (int idx = 0; idx < condition.Length; )
 			{
 				if(Char.IsWhiteSpace(condition[idx]))
@@ -1207,7 +1207,7 @@ namespace EpicGames.Core
 					tokens.Add(token);
 				}
 			}
-			return tokens.ToArray();
+			return [.. tokens];
 		}
 
 		/// <summary>
diff --git a/Engine/Source/Programs/Shared/EpicGames.Core/DefaultConsoleLogger.cs b/Engine/Source/Programs/Shared/EpicGames.Core/DefaultConsoleLogger.cs
--- a/Engine/Source/Programs/Shared/EpicGames.Core/DefaultConsoleLogger.cs
+++ b/Engine/Source/Programs/Shared/EpicGames.Core/DefaultConsoleLogger.cs
@@ -120,7 +120,7 @@ namespace EpicGames.Core
 			if (idx == -1 || idx + 1 == fmt.Length)
 			{
 				writer.Write(fmt);
-				return ReadOnlySpan<char>.Empty;
+				return [];
 			}
 			if (fmt[idx + 1] == '{')
 			{
@@ -135,11 +135,11 @@ namespace EpicGames.Core
 			if (endIdx == -1)
 			{
 				writer.Write(fmt);
-				return ReadOnlySpan<char>.Empty;
+				return [];
 			}
 
 			ReadOnlySpan<char> name = fmt.Slice(1, endIdx - 1);
-			ReadOnlySpan<char> spec = ReadOnlySpan<char>.Empty;
+			ReadOnlySpan<char> spec = [];
 
 			int specIdx = name.IndexOfAny(':', ',');
 			if (specIdx != -1)
diff --git a/Engine/Source/Programs/Shared/EpicGames.Core/ExceptionUtils.cs b/Engine/Source/Programs/Shared/EpicGames.Core/ExceptionUtils.cs
--- a/Engine/Source/Programs/Shared/EpicGames.Core/ExceptionUtils.cs
+++ b/Engine/Source/Programs/Shared/EpicGames.Core/ExceptionUtils.cs
@@ -28,7 +28,7 @@ namespace EpicGames.Core
 			List<string>? messages = ex.Data[ContextEntryName] as List<string>;
 			if (messages == null)
 			{
-				messages = new List<string>();
+				messages = [];
 				ex.Data[ContextEntryName] = messages;
 			}
 			messages.Add(message);
@@ -101,7 +101,7 @@ namespace EpicGames.Core
 		/// <returns>String containing the exception information. May be multiple lines.</returns>
 		public static string FormatExceptionDetails(Exception ex)
 		{
-			List<Exception> exceptionStack = new List<Exception>();
+			List<Exception> exceptionStack = [];
 			for (Exception? currentEx = ex; currentEx != null; currentEx = currentEx.InnerException)
 			{
 				exceptionStack.Add(currentEx);
diff --git a/Engine/Source/Programs/Shared/EpicGames.Core/FileFilter.cs b/Engine/Source/Programs/Shared/EpicGames.Core/FileFilter.cs
--- a/Engine/Source/Programs/Shared/EpicGames.Core/FileFilter.cs
+++ b/Engine/Source/Programs/Shared/EpicGames.Core/FileFilter.cs
@@ -44,7 +44,7 @@ namespace EpicGames.Core
 		/// <summary>
 		/// Terminating nodes for each rule added to the filter
 		/// </summary>
-		readonly List<FileFilterNode> _rules = new List<FileFilterNode>();
+		readonly List<FileFilterNode> _rules = [];
 
 		/// <summary>
 		/// Default constructor
@@ -90,7 +90,7 @@ namespace EpicGames.Core
 		/// <param name="rule">Pattern to match. See CreateRegex() for details.</param>
 		public void AddRule(string rule)
 		{
-			if (rule.StartsWith("-", StringComparison.Ordinal))
+			if (rule.StartsWith('-'))
 			{
 				Exclude(rule.Substring(1).TrimStart());
 			}
@@ -108,7 +108,7 @@ namespace EpicGames.Core
 		public void AddRule(string rule, params string[] allowTags)
 		{
 			string cleanRule = rule.Trim();
-			if(cleanRule.StartsWith("{", StringComparison.Ordinal))
+			if(cleanRule.StartsWith('{'))
 			{
 				// Find the end of the condition
 				int conditionEnd = cleanRule.IndexOf('}', StringComparison.Ordinal);
@@ -206,9 +206,9 @@ namespace EpicGames.Core
 			foreach(string line in File.ReadAllLines(fileName.FullName))
 			{
 				string trimLine = line.Trim();
-				if(!trimLine.StartsWith(";", StringComparison.Ordinal) && trimLine.Length > 0)
+				if(!trimLine.StartsWith(';') && trimLine.Length > 0)
 				{
-					if(trimLine.StartsWith("[", StringComparison.Ordinal))
+					if(trimLine.StartsWith('['))
 					{
 						bInSection = (trimLine == "[" + sectionName + "]");
 					}
@@ -291,7 +291,7 @@ namespace EpicGames.Core
 
 			// Remove the slash from the start of the pattern. Any exclude pattern that doesn't contain a directory separator is assumed to apply to any directory (eg. *.cpp), otherwise it's 
 			// taken relative to the root.
-			if (normalizedPattern.StartsWith("/", StringComparison.Ordinal))
+			if (normalizedPattern.StartsWith('/'))
 			{
 				normalizedPattern = normalizedPattern.Substring(1);
 			}
@@ -301,7 +301,7 @@ namespace EpicGames.Core
 			}
 
 			// All directories indicate a wildcard match
-			if (normalizedPattern.EndsWith("/", StringComparison.Ordinal))
+			if (normalizedPattern.EndsWith('/'))
 			{
 				normalizedPattern += "...";
 			}
@@ -382,7 +382,7 @@ namespace EpicGames.Core
 		{
 			string[] tokens = NormalizePath(folderName).Trim('/').Split('/');
 
-			FileFilterNode matchingNode = FindMatchingNode(_rootNode, tokens.Union(new string[]{ "" }).ToArray(), 0, _defaultNode);
+			FileFilterNode matchingNode = FindMatchingNode(_rootNode, tokens.Union([""]).ToArray(), 0, _defaultNode);
 
 			return matchingNode.Type == FileFilterType.Include || HighestPossibleIncludeMatch(_rootNode, tokens, 0, matchingNode.RuleNumber) > matchingNode.RuleNumber;
 		}
@@ -416,7 +416,7 @@ namespace EpicGames.Core
 		/// <returns>List of files that pass the filter</returns>
 		public List<FileReference> ApplyToDirectory(DirectoryReference directoryName, bool bIgnoreSymlinks)
 		{
-			List<FileReference> matchingFileNames = new List<FileReference>();
+			List<FileReference> matchingFileNames = [];
 			if (DirectoryReference.Exists(directoryName))
 			{
 				FindMatchesFromDirectory(new DirectoryInfo(directoryName.FullName), "", bIgnoreSymlinks, matchingFileNames);
@@ -433,7 +433,7 @@ namespace EpicGames.Core
 		/// <returns>List of files that pass the filter</returns>
 		public List<FileReference> ApplyToDirectory(DirectoryReference directoryName, string prefixPath, bool bIgnoreSymlinks)
 		{
-			List<FileReference> matchingFileNames = new List<FileReference>();
+			List<FileReference> matchingFileNames = [];
 			if (DirectoryReference.Exists(directoryName))
 			{
 				FindMatchesFromDirectory(new DirectoryInfo(directoryName.FullName), NormalizePath(prefixPath).TrimEnd('/') + "/", bIgnoreSymlinks, matchingFileNames);
diff --git a/Engine/Source/Programs/Shared/EpicGames.Core/FileFilterNode.cs b/Engine/Source/Programs/Shared/EpicGames.Core/FileFilterNode.cs
--- a/Engine/Source/Programs/Shared/EpicGames.Core/FileFilterNode.cs
+++ b/Engine/Source/Programs/Shared/EpicGames.Core/FileFilterNode.cs
@@ -33,7 +33,7 @@ namespace EpicGames.Core
 		/// <summary>
 		/// Child branches of this node, distinct by the the pattern for each.
 		/// </summary>
-		public List<FileFilterNode> Branches { get; set; } = new List<FileFilterNode>();
+		public List<FileFilterNode> Branches { get; set; } = [];
 
 		/// <summary>
 		/// If a path matches the sequence of nodes ending in this node, the number of the rule which added it. -1 if this was not the last node in a rule.
@@ -61,7 +61,7 @@ namespace EpicGames.Core
 		/// </summary>
 		public bool IsMatch(string token)
 		{
-			if(Pattern.EndsWith(".", StringComparison.Ordinal))
+			if(Pattern.EndsWith('.'))
 			{
 				return !token.Contains('.', StringComparison.Ordinal) && IsMatch(token, 0, Pattern.Substring(0, Pattern.Length - 1), 0);
 			}
diff --git a/Engine/Source/Programs/Shared/EpicGames.Core/FilePattern.cs b/Engine/Source/Programs/Shared/EpicGames.Core/FilePattern.cs
--- a/Engine/Source/Programs/Shared/EpicGames.Core/FilePattern.cs
+++ b/Engine/Source/Programs/Shared/EpicGames.Core/FilePattern.cs
@@ -80,7 +80,7 @@ namespace EpicGames.Core
 		public FilePattern(FileReference file)
 		{
 			BaseDirectory = file.Directory;
-			Tokens = new[] { file.GetFileName() };
+			Tokens = [file.GetFileName()];
 		}
 
 		/// <summary>
@@ -145,7 +145,7 @@ namespace EpicGames.Core
 			}
 
 			// Parse the tokens
-			List<string> tokens = new List<string>();
+			List<string> tokens = [];
 			int lastIdx = baseDirectoryLen;
 			for(int idx = baseDirectoryLen; idx < text.Length; idx++)
 			{
@@ -351,7 +351,7 @@ namespace EpicGames.Core
 			}
 
 			// Handle the case where source and target pattern are both individual files
-			Dictionary<FileReference, FileReference> targetFileToSourceFile = new Dictionary<FileReference, FileReference>();
+			Dictionary<FileReference, FileReference> targetFileToSourceFile = [];
 			if(sourcePattern.ContainsWildcards() || targetPattern.ContainsWildcards())
 			{
 				// Check the two patterns are compatible
@@ -435,7 +435,7 @@ namespace EpicGames.Core
 		/// <returns>List of valid files</returns>
 		public static List<FileReference> CheckInputFiles(IEnumerable<FileReference> inputFiles, DirectoryReference baseDirectory)
 		{
-			List<FileReference> files = new List<FileReference>();
+			List<FileReference> files = [];
 			foreach(FileReference inputFile in inputFiles)
 			{
 				if(!inputFile.IsUnderDirectory(baseDirectory))
diff --git a/Engine/Source/Programs/Shared/EpicGames.Core/FileSet.cs b/Engine/Source/Programs/Shared/EpicGames.Core/FileSet.cs
--- a/Engine/Source/Programs/Shared/EpicGames.Core/FileSet.cs
+++ b/Engine/Source/Programs/Shared/EpicGames.Core/FileSet.cs
@@ -17,7 +17,7 @@ namespace EpicGames.Core
 		/// <summary>
 		/// An empty fileset
 		/// </summary>
-		public static FileSet Empty { get; } = new FileSetFromFiles(Enumerable.Empty<(string, FileReference)>());
+		public static FileSet Empty { get; } = new FileSetFromFiles([]);
 		
 		/// <summary>
 		/// Path of this tree
@@ -53,7 +53,7 @@ namespace EpicGames.Core
 		/// <returns>Tree containing the given files</returns>
 		public static FileSet FromFile(DirectoryReference directory, string file)
 		{
-			return FromFiles(new[] { (File: file, FileReference.Combine(directory, file)) });
+			return FromFiles([(file, FileReference.Combine(directory, file))]);
 		}
 
 		/// <summary>
@@ -74,7 +74,7 @@ namespace EpicGames.Core
 		/// <returns>Tree containing the given files</returns>
 		public static FileSet FromFile(DirectoryReference directory, FileReference file)
 		{
-			return FromFiles(directory, new[] { file });
+			return FromFiles(directory, [file]);
 		}
 
 		/// <summary>
@@ -221,7 +221,7 @@ namespace EpicGames.Core
 		/// <returns></returns>
 		public Dictionary<FileReference, FileReference> Flatten(DirectoryReference outputDir)
 		{
-			Dictionary<FileReference, FileReference> targetToSourceFile = new Dictionary<FileReference, FileReference>();
+			Dictionary<FileReference, FileReference> targetToSourceFile = [];
 			foreach ((string path, FileReference sourceFile) in Flatten())
 			{
 				FileReference targetFile = FileReference.Combine(outputDir, path);
@@ -242,8 +242,8 @@ namespace EpicGames.Core
 	/// </summary>
 	class FileSetFromFiles : FileSet
 	{
-		readonly Dictionary<string, FileReference> _files = new Dictionary<string, FileReference>();
-		readonly Dictionary<string, FileSetFromFiles> _subTrees = new Dictionary<string, FileSetFromFiles>();
+		readonly Dictionary<string, FileReference> _files = [];
+		readonly Dictionary<string, FileSetFromFiles> _subTrees = [];
 
 		/// <summary>
 		/// Private constructor
diff --git a/Engine/Source/Programs/Shared/EpicGames.Core/FileUtils.cs b/Engine/Source/Programs/Shared/EpicGames.Core/FileUtils.cs
--- a/Engine/Source/Programs/Shared/EpicGames.Core/FileUtils.cs
+++ b/Engine/Source/Programs/Shared/EpicGames.Core/FileUtils.cs
@@ -543,7 +543,7 @@ namespace EpicGames.Core
 			{
 				dirName = RawPathPrefix + dirName;
 			}
-			if (!dirName.EndsWith("\\", StringComparison.Ordinal))
+			if (!dirName.EndsWith('\\'))
 			{
 				dirName += "\\";
 			}
@@ -767,7 +767,7 @@ namespace EpicGames.Core
 					throw new Win32Exception(result, "Unable to open restart manager session");
 				}
 
-				result = RmRegisterResources(sessionHandle, 1, new string[] { fileName }, 0, null, 0, null);
+				result = RmRegisterResources(sessionHandle, 1, [fileName], 0, null, 0, null);
 				if (result != 0)
 				{
 					throw new Win32Exception(result, "Unable to register resource with restart manager");
@@ -783,7 +783,7 @@ namespace EpicGames.Core
 					throw new Win32Exception(result, "Unable to query processes with file handle open");
 				}
 
-				List<FileLockInfoWin32> fileLocks = new List<FileLockInfoWin32>();
+				List<FileLockInfoWin32> fileLocks = [];
 				for (int idx = 0; idx < nProcInfo; idx++)
 				{
 					RM_PROCESS_INFO processInfo = processInfoArray[idx];
diff --git a/Engine/Source/Programs/Shared/EpicGames.Core/GuidUtils.cs b/Engine/Source/Programs/Shared/EpicGames.Core/GuidUtils.cs
--- a/Engine/Source/Programs/Shared/EpicGames.Core/GuidUtils.cs
+++ b/Engine/Source/Programs/Shared/EpicGames.Core/GuidUtils.cs
@@ -19,28 +19,25 @@ namespace EpicGames.Core
 		/// </summary>
 		public static Guid ReadGuidUnrealOrder(ReadOnlySpan<byte> span)
 		{
-			Span<byte> buffer = stackalloc byte[16];
-
-			buffer[0x0] = span[0x0];
-			buffer[0x1] = span[0x1];
-			buffer[0x2] = span[0x2];
-			buffer[0x3] = span[0x3];
-
-			buffer[0x4] = span[0x6];
-			buffer[0x5] = span[0x7];
-			buffer[0x6] = span[0x4];
-			buffer[0x7] = span[0x5];
-
-			buffer[0x8] = span[0xb];
-			buffer[0x9] = span[0xa];
-			buffer[0xa] = span[0x9];
-			buffer[0xb] = span[0x8];
-
-			buffer[0xc] = span[0xf];
-			buffer[0xd] = span[0xe];
-			buffer[0xe] = span[0xd];
-			buffer[0xf] = span[0xc];
-
+			Span<byte> buffer =
+			[
+				span[0x0],
+				span[0x1],
+				span[0x2],
+				span[0x3],
+				span[0x6],
+				span[0x7],
+				span[0x4],
+				span[0x5],
+				span[0xb],
+				span[0xa],
+				span[0x9],
+				span[0x8],
+				span[0xf],
+				span[0xe],
+				span[0xd],
+				span[0xc],
+			];
 			return new Guid(buffer);
 		}
 
diff --git a/Engine/Source/Programs/Shared/EpicGames.Core/HelpUtils.cs b/Engine/Source/Programs/Shared/EpicGames.Core/HelpUtils.cs
--- a/Engine/Source/Programs/Shared/EpicGames.Core/HelpUtils.cs
+++ b/Engine/Source/Programs/Shared/EpicGames.Core/HelpUtils.cs
@@ -117,7 +117,7 @@ namespace EpicGames.Core
 		/// <returns></returns>
 		public static void PrintTable(List<KeyValuePair<string, string>> items, int indent, int minFirstColumnWidth)
 		{
-			List<string> lines = new List<string>();
+			List<string> lines = [];
 			FormatTable(items, indent, minFirstColumnWidth, WindowWidth - 1, lines);
 
 			foreach (string line in lines)
@@ -136,7 +136,7 @@ namespace EpicGames.Core
 		/// <param name="logger"></param>
 		public static void PrintTable(List<KeyValuePair<string, string>> items, int indent, int minFirstColumnWidth, int maxWidth, ILogger logger)
 		{
-			List<string> lines = new List<string>();
+			List<string> lines = [];
 			FormatTable(items, indent, minFirstColumnWidth, maxWidth, lines);
 
 			foreach (string line in lines)
diff --git a/Engine/Source/Programs/Shared/EpicGames.Core/ILoggerIndent.cs b/Engine/Source/Programs/Shared/EpicGames.Core/ILoggerIndent.cs
--- a/Engine/Source/Programs/Shared/EpicGames.Core/ILoggerIndent.cs
+++ b/Engine/Source/Programs/Shared/EpicGames.Core/ILoggerIndent.cs
@@ -145,7 +145,7 @@ namespace EpicGames.Core
 		/// <summary>
 		/// Current list of indents
 		/// </summary>
-		readonly List<Scope> _scopes = new List<Scope>();
+		readonly List<Scope> _scopes = [];
 
 		/// <summary>
 		/// The inner logger
diff --git a/Engine/Source/Programs/Shared/EpicGames.Core/ILoggerProperty.cs b/Engine/Source/Programs/Shared/EpicGames.Core/ILoggerProperty.cs
--- a/Engine/Source/Programs/Shared/EpicGames.Core/ILoggerProperty.cs
+++ b/Engine/Source/Programs/Shared/EpicGames.Core/ILoggerProperty.cs
@@ -26,7 +26,7 @@ namespace EpicGames.Core
 			/// <param name="logger">Logger to wrap</param>
 			public LogScopedPropertyBuilder(ILogger logger)
 			{
-				_properties = new Dictionary<string, object>();
+				_properties = [];
 				_logger = logger;
 			}
 
diff --git a/Engine/Source/Programs/Shared/EpicGames.Core/JsonExtensions.cs b/Engine/Source/Programs/Shared/EpicGames.Core/JsonExtensions.cs
--- a/Engine/Source/Programs/Shared/EpicGames.Core/JsonExtensions.cs
+++ b/Engine/Source/Programs/Shared/EpicGames.Core/JsonExtensions.cs
@@ -32,7 +32,7 @@ namespace EpicGames.Core
 			}
 			else
 			{
-				propertyName = ReadOnlySpan<byte>.Empty;
+				propertyName = [];
 				return false;
 			}
 		}
diff --git a/Engine/Source/Programs/Shared/EpicGames.Core/JsonObject.cs b/Engine/Source/Programs/Shared/EpicGames.Core/JsonObject.cs
--- a/Engine/Source/Programs/Shared/EpicGames.Core/JsonObject.cs
+++ b/Engine/Source/Programs/Shared/EpicGames.Core/JsonObject.cs
@@ -204,14 +204,14 @@ namespace EpicGames.Core
 			switch (obj)
 			{
 				case OrderedDictionary objDictionary:
-					SystemJsonObject dictionaryJson= new SystemJsonObject();
+					SystemJsonObject dictionaryJson= [];
 					foreach (object? key in objDictionary.Keys)
 					{
 						dictionaryJson.Add((string) key, ToJsonNode(objDictionary[key]));
 					}
 					return dictionaryJson;
 				case Array objArray:
-					SystemJsonArray tempJsonArray = new SystemJsonArray();
+					SystemJsonArray tempJsonArray = [];
 					foreach (object? element in objArray)
 					{
 						tempJsonArray.Add(ToJsonNode(element));
@@ -229,7 +229,7 @@ namespace EpicGames.Core
 		/// <returns> Returns the System.Text.Json.Nodes.JsonObject representation of this object.</returns>
 		public SystemJsonObject ToSystemJsonObject()
 		{
-			SystemJsonObject returnObject = new SystemJsonObject();
+			SystemJsonObject returnObject = [];
 			foreach (DictionaryEntry entry in _rawOrderedObject )
 			{
 				returnObject.Add((string) entry.Key, ToJsonNode(entry.Value));
@@ -812,7 +812,7 @@ namespace EpicGames.Core
 				_rawOrderedObject[fieldName] = null;
 				return;
 			}
-			List<OrderedDictionary?> objList = new List<OrderedDictionary?>();
+			List<OrderedDictionary?> objList = [];
 			foreach (JsonObject? obj in value)
 			{
 				objList.Add(obj?._rawOrderedObject.CreateDeepCopy());
diff --git a/Engine/Source/Programs/Shared/EpicGames.Core/JsonSchema.cs b/Engine/Source/Programs/Shared/EpicGames.Core/JsonSchema.cs
--- a/Engine/Source/Programs/Shared/EpicGames.Core/JsonSchema.cs
+++ b/Engine/Source/Programs/Shared/EpicGames.Core/JsonSchema.cs
@@ -56,7 +56,7 @@ namespace EpicGames.Core
 			/// <summary>
 			/// Reference counts for each type (max of 2)
 			/// </summary>
-			public Dictionary<JsonSchemaType, int> TypeRefCount { get; } = new Dictionary<JsonSchemaType, int>();
+			public Dictionary<JsonSchemaType, int> TypeRefCount { get; } = [];
 
 			/// <inheritdoc/>
 			public void WriteBoolean(string name, bool value) { }
@@ -205,8 +205,8 @@ namespace EpicGames.Core
 			refCollector.WriteType(RootType);
 
 			// Assign names to each type definition
-			HashSet<string> definitionNames = new HashSet<string>();
-			Dictionary<JsonSchemaType, string> typeToDefinition = new Dictionary<JsonSchemaType, string>();
+			HashSet<string> definitionNames = [];
+			Dictionary<JsonSchemaType, string> typeToDefinition = [];
 			foreach ((JsonSchemaType type, int refCount) in refCollector.TypeRefCount)
 			{
 				if (refCount > 1)
diff --git a/Engine/Source/Programs/Shared/EpicGames.Core/JsonSchemaAttributes.cs b/Engine/Source/Programs/Shared/EpicGames.Core/JsonSchemaAttributes.cs
--- a/Engine/Source/Programs/Shared/EpicGames.Core/JsonSchemaAttributes.cs
+++ b/Engine/Source/Programs/Shared/EpicGames.Core/JsonSchemaAttributes.cs
@@ -55,7 +55,7 @@ namespace EpicGames.Core
 #pragma warning disable CA1019 // Define accessors for attribute arguments
 		public JsonSchemaCatalogAttribute(string name, string description, string? fileMatch)
 #pragma warning restore CA1019 // Define accessors for attribute arguments
-			: this(name, description, (fileMatch == null) ? (string[]?)null : new[] { fileMatch })
+			: this(name, description, (fileMatch == null) ? (string[]?)null : [fileMatch])
 		{
 		}
 
diff --git a/Engine/Source/Programs/Shared/EpicGames.Core/JsonSchemaFactory.cs b/Engine/Source/Programs/Shared/EpicGames.Core/JsonSchemaFactory.cs
--- a/Engine/Source/Programs/Shared/EpicGames.Core/JsonSchemaFactory.cs
+++ b/Engine/Source/Programs/Shared/EpicGames.Core/JsonSchemaFactory.cs
@@ -20,7 +20,7 @@ namespace EpicGames.Core
 		/// <summary>
 		/// Cache of known types
 		/// </summary>
-		public Dictionary<Type, JsonSchemaType> TypeCache { get; } = new Dictionary<Type, JsonSchemaType>();
+		public Dictionary<Type, JsonSchemaType> TypeCache { get; } = [];
 
 		/// <summary>
 		/// Constructor
@@ -194,7 +194,7 @@ namespace EpicGames.Core
 				if (attribute != null)
 				{
 					JsonSchemaObject knownObject = new JsonSchemaObject();
-					knownObject.Properties.Add(new JsonSchemaProperty("type", "Type discriminator", new JsonSchemaEnum(new[] { attribute.Name }, new[] { "Identifier for the derived type" })));
+					knownObject.Properties.Add(new JsonSchemaProperty("type", "Type discriminator", new JsonSchemaEnum([attribute.Name], ["Identifier for the derived type"])));
 					SetObjectProperties(knownObject, knownType);
 					obj.Types.Add(knownObject);
 				}
diff --git a/Engine/Source/Programs/Shared/EpicGames.Core/JsonSchemaType.cs b/Engine/Source/Programs/Shared/EpicGames.Core/JsonSchemaType.cs
--- a/Engine/Source/Programs/Shared/EpicGames.Core/JsonSchemaType.cs
+++ b/Engine/Source/Programs/Shared/EpicGames.Core/JsonSchemaType.cs
@@ -168,12 +168,12 @@ namespace EpicGames.Core
 		/// <summary>
 		/// Values for the enum
 		/// </summary>
-		public List<string> Values { get; } = new List<string>();
+		public List<string> Values { get; } = [];
 
 		/// <summary>
 		/// Descriptions for each enum value
 		/// </summary>
-		public List<string> Descriptions { get; } = new List<string>();
+		public List<string> Descriptions { get; } = [];
 
 		/// <summary>
 		/// Constructor
@@ -283,7 +283,7 @@ namespace EpicGames.Core
 		/// <summary>
 		/// Properties to allow in this object
 		/// </summary>
-		public List<JsonSchemaProperty> Properties { get; } = new List<JsonSchemaProperty>();
+		public List<JsonSchemaProperty> Properties { get; } = [];
 
 		/// <summary>
 		/// Type for additional properties
@@ -325,7 +325,7 @@ namespace EpicGames.Core
 		/// <summary>
 		/// List of valid types
 		/// </summary>
-		public List<JsonSchemaType> Types { get; } = new List<JsonSchemaType>();
+		public List<JsonSchemaType> Types { get; } = [];
 
 		/// <inheritdoc/>
 		public override void Write(IJsonSchemaWriter writer)
diff --git a/Engine/Source/Programs/Shared/EpicGames.Core/JsonTracer.cs b/Engine/Source/Programs/Shared/EpicGames.Core/JsonTracer.cs
--- a/Engine/Source/Programs/Shared/EpicGames.Core/JsonTracer.cs
+++ b/Engine/Source/Programs/Shared/EpicGames.Core/JsonTracer.cs
@@ -105,12 +105,12 @@ namespace EpicGames.Core
 			StartTimestamp = startTimestamp;
 			
 			_tags = initialTags == null
-				? new Dictionary<string, object>()
+				? []
 				: new Dictionary<string, object>(initialTags);
 
 			_references = references == null
-				? new List<Reference>()
-				: references.ToList();
+				? []
+				: [.. references];
 			
 			JsonTracerSpanContext? parentContext = FindPreferredParentRef(_references);
 			if (parentContext == null)
@@ -271,8 +271,8 @@ namespace EpicGames.Core
 		private readonly JsonTracer _tracer;
 		private readonly string _operationName;
 		private DateTimeOffset _startTimestamp = DateTimeOffset.MinValue;
-		private readonly List<JsonTracerSpan.Reference> _references = new List<JsonTracerSpan.Reference>();
-		private readonly Dictionary<string, object> _tags = new Dictionary<string, object>();
+		private readonly List<JsonTracerSpan.Reference> _references = [];
+		private readonly Dictionary<string, object> _tags = [];
 		private bool _ignoreActiveSpan;
 		
 		public JsonTracerSpanBuilder(JsonTracer tracer, string operationName)
@@ -405,7 +405,7 @@ namespace EpicGames.Core
 		public ISpan? ActiveSpan => ScopeManager.Active?.Span;
 
 		private readonly object _lock = new object();
-		private readonly List<JsonTracerSpan> _finishedSpans = new List<JsonTracerSpan>();
+		private readonly List<JsonTracerSpan> _finishedSpans = [];
 		private readonly DirectoryReference? _telemetryDir;
 		
 		public JsonTracer(DirectoryReference? telemetryDir = null)
diff --git a/Engine/Source/Programs/Shared/EpicGames.Core/ListExtensions.cs b/Engine/Source/Programs/Shared/EpicGames.Core/ListExtensions.cs
--- a/Engine/Source/Programs/Shared/EpicGames.Core/ListExtensions.cs
+++ b/Engine/Source/Programs/Shared/EpicGames.Core/ListExtensions.cs
@@ -143,11 +143,7 @@ namespace EpicGames.Core
 		/// <returns>Converted list</returns>
 		public static List<TOutput> ConvertAll<TInput, TOutput>(this List<TInput> input) where TInput : TOutput
 		{
-			List<TOutput> output = new List<TOutput>(input.Count);
-			foreach (TInput inputItem in input)
-			{
-				output.Add(inputItem);
-			}
+			List<TOutput> output = [.. input];
 			return output;
 		}
 
diff --git a/Engine/Source/Programs/Shared/EpicGames.Core/Log.cs b/Engine/Source/Programs/Shared/EpicGames.Core/Log.cs
--- a/Engine/Source/Programs/Shared/EpicGames.Core/Log.cs
+++ b/Engine/Source/Programs/Shared/EpicGames.Core/Log.cs
@@ -807,7 +807,7 @@ namespace EpicGames.Core
 				=> _outer._scopes.Remove(this);
 		}
 
-		readonly List<Scope> _scopes = new List<Scope>();
+		readonly List<Scope> _scopes = [];
 
 		/// <summary>
 		/// Tests whether the collection is empty
@@ -849,7 +849,7 @@ namespace EpicGames.Core
 		/// <summary>
 		/// List of captured events
 		/// </summary>
-		public List<LogEvent> Events { get; } = new List<LogEvent>();
+		public List<LogEvent> Events { get; } = [];
 
 		/// <summary>
 		/// Renders the captured events as a single string
@@ -1061,7 +1061,7 @@ namespace EpicGames.Core
 		/// <summary>
 		/// Output streams for structured log data
 		/// </summary>
-		private IReadOnlyList<FileStream> _jsonStreams = Array.Empty<FileStream>();
+		private IReadOnlyList<FileStream> _jsonStreams = [];
 
 		/// <summary>
 		/// Constructor
@@ -1111,7 +1111,7 @@ namespace EpicGames.Core
 		/// </summary>
 		async Task WriteFilesAsync()
 		{
-			byte[] newline = new byte[] { (byte)'\n' };
+			byte[] newline = [(byte)'\n'];
 			while (await _eventChannel.Reader.WaitToReadAsync())
 			{
 				IReadOnlyList<FileStream> streams = _jsonStreams;
diff --git a/Engine/Source/Programs/Shared/EpicGames.Core/LogBuffer.cs b/Engine/Source/Programs/Shared/EpicGames.Core/LogBuffer.cs
--- a/Engine/Source/Programs/Shared/EpicGames.Core/LogBuffer.cs
+++ b/Engine/Source/Programs/Shared/EpicGames.Core/LogBuffer.cs
@@ -24,7 +24,7 @@ namespace EpicGames.Core
 		{
 			_lineNumber = 1;
 			_history = new string?[historySize];
-			_nextLines = new List<string?>();
+			_nextLines = [];
 		}
 
 		/// <summary>
diff --git a/Engine/Source/Programs/Shared/EpicGames.Core/LogEvent.cs b/Engine/Source/Programs/Shared/EpicGames.Core/LogEvent.cs
--- a/Engine/Source/Programs/Shared/EpicGames.Core/LogEvent.cs
+++ b/Engine/Source/Programs/Shared/EpicGames.Core/LogEvent.cs
@@ -98,8 +98,8 @@ namespace EpicGames.Core
 
 		class MergedPropertyList : IEnumerable<KeyValuePair<string, object?>>
 		{
-			readonly HashSet<string> _names = new HashSet<string>();
-			readonly List<KeyValuePair<string, object?>> _properties = new List<KeyValuePair<string, object?>>();
+			readonly HashSet<string> _names = [];
+			readonly List<KeyValuePair<string, object?>> _properties = [];
 
 			public void AddRange(IEnumerable<KeyValuePair<string, object?>> properties)
 			{
@@ -153,7 +153,7 @@ namespace EpicGames.Core
 		/// <param name="name">Name of the property to add</param>
 		/// <param name="value">Value for the property</param>
 		public void AddProperty(string name, object? value)
-			=> AddProperties(new[] { KeyValuePair.Create(name, value) });
+			=> AddProperties([KeyValuePair.Create(name, value)]);
 
 		/// <summary>
 		/// Add new properties to this event
@@ -320,7 +320,7 @@ namespace EpicGames.Core
 
 		static Dictionary<string, object?> ReadProperties(ref Utf8JsonReader reader)
 		{
-			Dictionary<string, object?> properties = new Dictionary<string, object?>();
+			Dictionary<string, object?> properties = [];
 
 			ReadOnlySpan<byte> propertyName;
 			for (; JsonExtensions.TryReadNextPropertyName(ref reader, out propertyName); reader.Skip())
@@ -370,7 +370,7 @@ namespace EpicGames.Core
 
 		static object ReadArrayPropertyValue(ref Utf8JsonReader reader)
 		{
-			List<object?> result = new List<object?>();
+			List<object?> result = [];
 			while (reader.Read() && reader.TokenType != JsonTokenType.EndArray)
 			{
 				result.Add(ReadPropertyValue(ref reader));
@@ -381,7 +381,7 @@ namespace EpicGames.Core
 		static object ReadObjectPropertyValue(ref Utf8JsonReader reader)
 		{
 			// Read all the properties
-			Dictionary<Utf8String, object?>? properties = new Dictionary<Utf8String, object?>();
+			Dictionary<Utf8String, object?>? properties = [];
 			for (; JsonExtensions.TryReadNextPropertyName(ref reader, out ReadOnlySpan<byte> propertyName); reader.Skip())
 			{
 				properties.Add(new Utf8String(propertyName.ToArray()), ReadPropertyValue(ref reader));
@@ -444,7 +444,7 @@ namespace EpicGames.Core
 				{
 					if (!name.Equals(MessageTemplate.FormatPropertyName, StringComparison.Ordinal))
 					{
-						if (name.StartsWith("@", StringComparison.Ordinal))
+						if (name.StartsWith('@'))
 						{
 							writer.WritePropertyName(name[1..]);
 							JsonSerializer.Serialize(writer, value, value?.GetType() ?? typeof(object), s_jsonSerializerOptions);
@@ -516,7 +516,7 @@ namespace EpicGames.Core
 		/// </summary>
 		public static LogEvent Create(LogLevel level, EventId eventId, Exception? exception, string format, params object[] args)
 		{
-			Dictionary<string, object?> properties = new Dictionary<string, object?>();
+			Dictionary<string, object?> properties = [];
 			MessageTemplate.ParsePropertyValues(format, args, properties);
 
 			string message = MessageTemplate.Render(format, properties!);
@@ -636,7 +636,7 @@ namespace EpicGames.Core
 		/// <summary>
 		/// Multiple inner exceptions, in the case of an <see cref="AggregateException"/>
 		/// </summary>
-		public List<LogException> InnerExceptions { get; } = new List<LogException>();
+		public List<LogException> InnerExceptions { get; } = [];
 
 		/// <summary>
 		/// Constructor
@@ -702,7 +702,7 @@ namespace EpicGames.Core
 		/// </summary>
 		public IReadOnlyList<LogEvent> LogEvents => _logEvents;
 
-		private readonly List<LogEvent> _logEvents = new List<LogEvent>();
+		private readonly List<LogEvent> _logEvents = [];
 		private readonly int[] _includeEventIds;
 		private readonly Action<LogEvent>? _eventCallback;
 
diff --git a/Engine/Source/Programs/Shared/EpicGames.Core/LogEventBuilder.cs b/Engine/Source/Programs/Shared/EpicGames.Core/LogEventBuilder.cs
--- a/Engine/Source/Programs/Shared/EpicGames.Core/LogEventBuilder.cs
+++ b/Engine/Source/Programs/Shared/EpicGames.Core/LogEventBuilder.cs
@@ -87,7 +87,7 @@ namespace EpicGames.Core
 		public LogEventLine(string text)
 		{
 			Text = text;
-			Spans = new Dictionary<string, LogEventSpan>();
+			Spans = [];
 		}
 
 		/// <summary>
@@ -283,7 +283,7 @@ namespace EpicGames.Core
 			_properties ??= new Dictionary<string, object>(StringComparer.OrdinalIgnoreCase);
 			_properties.Add(name, span);
 
-			_spans ??= new List<LogSpan>();
+			_spans ??= [];
 			for(int insertIdx = _spans.Count; ;insertIdx--)
 			{
 				if (insertIdx == 0 || _spans[insertIdx - 1]._offset < offset)
@@ -386,7 +386,7 @@ namespace EpicGames.Core
 		/// </summary>
 		public void MoveNext()
 		{
-			_lines ??= new List<LogLine>();
+			_lines ??= [];
 			_lines.Add(CreateLine());
 
 			_spans = null;
@@ -441,7 +441,7 @@ namespace EpicGames.Core
 			Dictionary<string, object?>? properties = null;
 			if (line._properties != null)
 			{
-				properties = new Dictionary<string, object?>();
+				properties = [];
 				foreach ((string name, object value) in line._properties)
 				{
 					object newValue;
diff --git a/Engine/Source/Programs/Shared/EpicGames.Core/LogEventMarkup.cs b/Engine/Source/Programs/Shared/EpicGames.Core/LogEventMarkup.cs
--- a/Engine/Source/Programs/Shared/EpicGames.Core/LogEventMarkup.cs
+++ b/Engine/Source/Programs/Shared/EpicGames.Core/LogEventMarkup.cs
@@ -24,7 +24,7 @@ namespace EpicGames.Core
 		public static LogValue ScreenshotTest => new LogValue(LogValueType.ScreenshotTest, "");
 #pragma warning restore CS1591
 
-		static readonly char[] s_pathChars = { '/', '\\' };
+		static readonly char[] s_pathChars = ['/', '\\'];
 
 		static string RemoveRelativeDirs(string path)
 		{
@@ -69,7 +69,7 @@ namespace EpicGames.Core
 					try
 					{
 						string combinedPath = RemoveRelativeDirs(Path.Combine(baseDir, file));
-						properties = new Dictionary<Utf8String, object?>();
+						properties = [];
 						properties[LogEventPropertyName.File] = combinedPath;
 					}
 					catch
@@ -105,7 +105,7 @@ namespace EpicGames.Core
 			identifier = Regex.Replace(identifier, "^.* ", "");
 
 			// Add it to the list
-			Dictionary<Utf8String, object?> properties = new Dictionary<Utf8String, object?>();
+			Dictionary<Utf8String, object?> properties = [];
 			properties[LogEventPropertyName.Identifier] = identifier;
 			builder.Annotate(group, new LogValue(LogValueType.Symbol, "", properties));
 		}
diff --git a/Engine/Source/Programs/Shared/EpicGames.Core/LogEventParser.cs b/Engine/Source/Programs/Shared/EpicGames.Core/LogEventParser.cs
--- a/Engine/Source/Programs/Shared/EpicGames.Core/LogEventParser.cs
+++ b/Engine/Source/Programs/Shared/EpicGames.Core/LogEventParser.cs
@@ -79,7 +79,7 @@ namespace EpicGames.Core
 		public LogEventMatch(LogEventPriority priority, LogEvent logEvent)
 		{
 			Priority = priority;
-			Events = new List<LogEvent> { logEvent };
+			Events = [logEvent];
 		}
 
 		/// <summary>
@@ -113,12 +113,12 @@ namespace EpicGames.Core
 		/// <summary>
 		/// List of event matchers for this parser
 		/// </summary>
-		public List<ILogEventMatcher> Matchers { get; } = new List<ILogEventMatcher>();
+		public List<ILogEventMatcher> Matchers { get; } = [];
 
 		/// <summary>
 		/// List of patterns to ignore
 		/// </summary>
-		public List<Regex> IgnorePatterns { get; } = new List<Regex>();
+		public List<Regex> IgnorePatterns { get; } = [];
 
 		/// <summary>
 		/// Buffer of input lines
@@ -143,7 +143,7 @@ namespace EpicGames.Core
 		/// <summary>
 		/// Log events sinks in addition to <see cref="_logger" />
 		/// </summary>
-		readonly List<ILogEventSink> _logEventSinks = new List<ILogEventSink>();
+		readonly List<ILogEventSink> _logEventSinks = [];
 
 		/// <summary>
 		/// Timer for the parser being active
@@ -229,12 +229,11 @@ namespace EpicGames.Core
 		{
 			Stopwatch timer = Stopwatch.StartNew();
 
-			List<DirectoryReference> baseDirs = new List<DirectoryReference>();
-			baseDirs.Add(rootDir);
+			List<DirectoryReference> baseDirs = [rootDir];
 			AddRestrictedDirs(baseDirs, "Restricted");
 			AddRestrictedDirs(baseDirs, "Platforms");
 
-			List<(FileReference, Task)> tasks = new List<(FileReference, Task)>();
+			List<(FileReference, Task)> tasks = [];
 			foreach (DirectoryReference baseDir in baseDirs)
 			{
 				FileReference ignorePatternFile = FileReference.Combine(baseDir, "Build", "Horde", "IgnorePatterns.txt");
@@ -269,7 +268,7 @@ namespace EpicGames.Core
 		{
 			string[] lines = await FileReference.ReadAllLinesAsync(ignorePatternFile);
 
-			List<Regex> patterns = new List<Regex>();
+			List<Regex> patterns = [];
 			foreach (string line in lines)
 			{
 				string trimLine = line.Trim();
diff --git a/Engine/Source/Programs/Shared/EpicGames.Core/LogValue.cs b/Engine/Source/Programs/Shared/EpicGames.Core/LogValue.cs
--- a/Engine/Source/Programs/Shared/EpicGames.Core/LogValue.cs
+++ b/Engine/Source/Programs/Shared/EpicGames.Core/LogValue.cs
@@ -133,7 +133,7 @@ namespace EpicGames.Core
 				if (propertyInfo.GetCustomAttribute<JsonIgnoreAttribute>() == null)
 				{
 					string name = propertyInfo.GetCustomAttribute<JsonPropertyNameAttribute>()?.Name ?? propertyInfo.Name;
-					properties ??= new Dictionary<Utf8String, object?>();
+					properties ??= [];
 					properties[new Utf8String(name)] = propertyInfo.GetValue(obj)!;
 				}
 			}
diff --git a/Engine/Source/Programs/Shared/EpicGames.Core/ManagedProcess.cs b/Engine/Source/Programs/Shared/EpicGames.Core/ManagedProcess.cs
--- a/Engine/Source/Programs/Shared/EpicGames.Core/ManagedProcess.cs
+++ b/Engine/Source/Programs/Shared/EpicGames.Core/ManagedProcess.cs
@@ -485,7 +485,7 @@ namespace EpicGames.Core
 		/// </summary>
 		ManagedProcessGroup? _accountingProcessGroup;
 
-		byte[] _buffer = Array.Empty<byte>();
+		byte[] _buffer = [];
 		int _bufferPos = 0;
 		int _bufferLen = 0;
 		bool _bCarriageReturn = false;
@@ -637,7 +637,7 @@ namespace EpicGames.Core
 				if (environment != null)
 				{
 					// The native format for the environment block is a sequence of null terminated strings with a final null terminator.
-					List<byte> environmentBytes = new List<byte>();
+					List<byte> environmentBytes = [];
 					foreach (KeyValuePair<string, string> pair in environment)
 					{
 						environmentBytes.AddRange(Console.OutputEncoding.GetBytes(pair.Key));
@@ -735,7 +735,7 @@ namespace EpicGames.Core
 							}
 
 							// Under heavy load (ie. spawning large number of processes, typically Clang) we see CreateProcess very occasionally failing with ERROR_ACCESS_DENIED.
-							int[] retryDelay = { 100, 200, 1000, 5000 };
+							int[] retryDelay = [100, 200, 1000, 5000];
 							for(int attemptIdx = 0; ; attemptIdx++)
 							{
 								if (CreateProcess(null, new StringBuilder("\"" + fileName + "\" " + commandLine), IntPtr.Zero, IntPtr.Zero, true, flags, environmentBlock, workingDirectory, startupInfoEx, processInfo) != 0)
@@ -1182,7 +1182,7 @@ namespace EpicGames.Core
 		{
 			// Manually read all the output lines from the stream. Using ReadToEndAsync() or ReadAsync() on the StreamReader is abysmally slow, especially
 			// for high-volume processes. Manually picking out the lines via a buffered ReadAsync() call was found to be 6x faster on 'p4 -ztag have' calls.
-			List<string> outputLines = new List<string>();
+			List<string> outputLines = [];
 			byte[] buffer = new byte[32 * 1024];
 			byte lastCharacter = 0;
 			int numBytesInBuffer = 0;
@@ -1411,7 +1411,7 @@ namespace EpicGames.Core
 		/// <returns></returns>
 		public static Dictionary<string, string> GetCurrentEnvVars()
 		{
-			Dictionary<string, string> newEnvironment = new Dictionary<string, string>();
+			Dictionary<string, string> newEnvironment = [];
 			foreach (object? envVar in Environment.GetEnvironmentVariables())
 			{
 				System.Collections.DictionaryEntry entry = (System.Collections.DictionaryEntry)envVar!;
diff --git a/Engine/Source/Programs/Shared/EpicGames.Core/Md5Hash.cs b/Engine/Source/Programs/Shared/EpicGames.Core/Md5Hash.cs
--- a/Engine/Source/Programs/Shared/EpicGames.Core/Md5Hash.cs
+++ b/Engine/Source/Programs/Shared/EpicGames.Core/Md5Hash.cs
@@ -64,9 +64,9 @@ namespace EpicGames.Core
 		public static Md5Hash Compute(ReadOnlySpan<byte> data)
 		{
 			byte[] output = new byte[NumBytes];
-			using (MD5 hasher = MD5.Create())
+			if (!MD5.TryHashData(data, output, out int bytesWritten) || bytesWritten != NumBytes)
 			{
-				hasher.TryComputeHash(data, output, out _);
+				throw new Exception("Unable to hash data");
 			}
 			return new Md5Hash(output);
 		}
diff --git a/Engine/Source/Programs/Shared/EpicGames.Core/MemoryMappedFileCache.cs b/Engine/Source/Programs/Shared/EpicGames.Core/MemoryMappedFileCache.cs
--- a/Engine/Source/Programs/Shared/EpicGames.Core/MemoryMappedFileCache.cs
+++ b/Engine/Source/Programs/Shared/EpicGames.Core/MemoryMappedFileCache.cs
@@ -159,7 +159,7 @@ namespace EpicGames.Core
 		}
 
 		readonly object _lockObject = new object();
-		readonly Dictionary<FileReference, MappedFile> _pathToMappedFile = new Dictionary<FileReference, MappedFile>();
+		readonly Dictionary<FileReference, MappedFile> _pathToMappedFile = [];
 		readonly LinkedList<MappedFile> _mappedFiles = new LinkedList<MappedFile>();
 
 		long _mappedSize;
diff --git a/Engine/Source/Programs/Shared/EpicGames.Core/MemoryReader.cs b/Engine/Source/Programs/Shared/EpicGames.Core/MemoryReader.cs
--- a/Engine/Source/Programs/Shared/EpicGames.Core/MemoryReader.cs
+++ b/Engine/Source/Programs/Shared/EpicGames.Core/MemoryReader.cs
@@ -280,7 +280,7 @@ namespace EpicGames.Core
 		/// <param name="readItem">Delegate to write an individual item</param>
 		public static List<T> ReadList<T>(this IMemoryReader reader, Func<T> readItem)
 		{
-			List<T> list = new List<T>();
+			List<T> list = [];
 			ReadList(reader, list, readItem);
 			return list;
 		}
@@ -347,7 +347,7 @@ namespace EpicGames.Core
 		{
 			if (length == 0)
 			{
-				return Array.Empty<T>();
+				return [];
 			}
 
 			T[] array = new T[length];
@@ -368,7 +368,7 @@ namespace EpicGames.Core
 		{
 			if (length == 0)
 			{
-				return Array.Empty<T>();
+				return [];
 			}
 
 			T[] array = new T[length];
diff --git a/Engine/Source/Programs/Shared/EpicGames.Core/MemoryWriter.cs b/Engine/Source/Programs/Shared/EpicGames.Core/MemoryWriter.cs
--- a/Engine/Source/Programs/Shared/EpicGames.Core/MemoryWriter.cs
+++ b/Engine/Source/Programs/Shared/EpicGames.Core/MemoryWriter.cs
@@ -302,7 +302,7 @@ namespace EpicGames.Core
 			public ReadOnlyMemory<byte> WrittenMemory => Data.Slice(0, Length);
 		}
 
-		readonly List<Chunk> _chunks = new List<Chunk>();
+		readonly List<Chunk> _chunks = [];
 		readonly int _chunkSize;
 		Chunk _currentChunk;
 
diff --git a/Engine/Source/Programs/Shared/EpicGames.Core/MessageTemplate.cs b/Engine/Source/Programs/Shared/EpicGames.Core/MessageTemplate.cs
--- a/Engine/Source/Programs/Shared/EpicGames.Core/MessageTemplate.cs
+++ b/Engine/Source/Programs/Shared/EpicGames.Core/MessageTemplate.cs
@@ -204,7 +204,7 @@ namespace EpicGames.Core
 						{
 							break;
 						}
-						names ??= new List<(int, int)>();
+						names ??= [];
 
 						names.Add((startIdx, idx - startIdx));
 					}
diff --git a/Engine/Source/Programs/Shared/EpicGames.Core/ProcessUtils.cs b/Engine/Source/Programs/Shared/EpicGames.Core/ProcessUtils.cs
--- a/Engine/Source/Programs/Shared/EpicGames.Core/ProcessUtils.cs
+++ b/Engine/Source/Programs/Shared/EpicGames.Core/ProcessUtils.cs
@@ -85,7 +85,7 @@ namespace EpicGames.Core
 		[SupportedOSPlatform("windows")]
 		static bool TerminateProcessesWin32(Predicate<FileReference> predicate, ILogger logger, CancellationToken cancellationToken)
 		{
-			Dictionary<int, int> processToCount = new Dictionary<int, int>();
+			Dictionary<int, int> processToCount = [];
 			for(; ;)
 			{
 				cancellationToken.ThrowIfCancellationRequested();
@@ -110,7 +110,7 @@ namespace EpicGames.Core
 				}
 
 				// Find the processes to terminate
-				List<SafeProcessHandle> waitHandles = new List<SafeProcessHandle>();
+				List<SafeProcessHandle> waitHandles = [];
 				try
 				{
 					// Open each process in turn
@@ -180,7 +180,7 @@ namespace EpicGames.Core
 		static bool TerminateProcessesGenericPlatform(Predicate<FileReference> predicate, ILogger logger, CancellationToken cancellationToken)
 		{
 			bool result = true;
-			Dictionary<(int, DateTime), int> processToCount = new Dictionary<(int, DateTime), int>();
+			Dictionary<(int, DateTime), int> processToCount = [];
 			for (; ; )
 			{
 				cancellationToken.ThrowIfCancellationRequested();
diff --git a/Engine/Source/Programs/Shared/EpicGames.Core/ReadOnlySequenceBuilder.cs b/Engine/Source/Programs/Shared/EpicGames.Core/ReadOnlySequenceBuilder.cs
--- a/Engine/Source/Programs/Shared/EpicGames.Core/ReadOnlySequenceBuilder.cs
+++ b/Engine/Source/Programs/Shared/EpicGames.Core/ReadOnlySequenceBuilder.cs
@@ -39,7 +39,7 @@ namespace EpicGames.Core
 			}
 		}
 
-		readonly List<ReadOnlyMemory<T>> _segments = new List<ReadOnlyMemory<T>>();
+		readonly List<ReadOnlyMemory<T>> _segments = [];
 
 		/// <summary>
 		/// Current length of the sequence
diff --git a/Engine/Source/Programs/Shared/EpicGames.Core/RefCountedHandle.cs b/Engine/Source/Programs/Shared/EpicGames.Core/RefCountedHandle.cs
--- a/Engine/Source/Programs/Shared/EpicGames.Core/RefCountedHandle.cs
+++ b/Engine/Source/Programs/Shared/EpicGames.Core/RefCountedHandle.cs
@@ -114,10 +114,7 @@ namespace EpicGames.Core
 		/// <inheritdoc/>
 		public IRefCountedHandle<T> AddRef()
 		{
-			if (_refCountedDisposer == null)
-			{
-				throw new ObjectDisposedException(typeof(RefCountedHandle<T>).Name);
-			}
+			ObjectDisposedException.ThrowIf(_refCountedDisposer == null, typeof(RefCountedHandle<T>).Name);
 
 			_refCountedDisposer.AddRef();
 			return new RefCountedHandle<T>(Target, _refCountedDisposer);
diff --git a/Engine/Source/Programs/Shared/EpicGames.Core/RegistryConfigurationSource.cs b/Engine/Source/Programs/Shared/EpicGames.Core/RegistryConfigurationSource.cs
--- a/Engine/Source/Programs/Shared/EpicGames.Core/RegistryConfigurationSource.cs
+++ b/Engine/Source/Programs/Shared/EpicGames.Core/RegistryConfigurationSource.cs
@@ -65,7 +65,7 @@ namespace EpicGames.Core
 		{
 			if (RuntimeInformation.IsOSPlatform(OSPlatform.Windows))
 			{
-				Dictionary<string, string?> data = new Dictionary<string, string?>();
+				Dictionary<string, string?> data = [];
 				GetValues(_baseKey, _keyPath, _baseConfigName, _filter, data);
 				Data = data;
 			}
diff --git a/Engine/Source/Programs/Shared/EpicGames.Core/Sha1Hash.cs b/Engine/Source/Programs/Shared/EpicGames.Core/Sha1Hash.cs
--- a/Engine/Source/Programs/Shared/EpicGames.Core/Sha1Hash.cs
+++ b/Engine/Source/Programs/Shared/EpicGames.Core/Sha1Hash.cs
@@ -67,13 +67,9 @@ namespace EpicGames.Core
 		public static Sha1Hash Compute(ReadOnlySpan<byte> data)
 		{
 			byte[] output = new byte[20];
-			using (SHA1 sha1 = SHA1.Create())
+			if (!SHA1.TryHashData(data, output, out int bytesWritten) || bytesWritten != NumBytes)
 			{
-				int bytesWritten;
-				if (!sha1.TryComputeHash(data, output, out bytesWritten) || bytesWritten != NumBytes)
-				{
-					throw new Exception($"Unable to hash data");
-				}
+				throw new Exception("Unable to hash data");
 			}
 			return new Sha1Hash(output);
 		}
diff --git a/Engine/Source/Programs/Shared/EpicGames.Core/SingleInstanceMutex.cs b/Engine/Source/Programs/Shared/EpicGames.Core/SingleInstanceMutex.cs
--- a/Engine/Source/Programs/Shared/EpicGames.Core/SingleInstanceMutex.cs
+++ b/Engine/Source/Programs/Shared/EpicGames.Core/SingleInstanceMutex.cs
@@ -56,10 +56,7 @@ namespace EpicGames.Core
 					// Waiting for multiple handles is only supported on Windows
 					if (RuntimeInformation.IsOSPlatform(OSPlatform.Windows))
 					{
-						WaitHandle[] handles = new WaitHandle[2];
-						handles[0] = _disposing;
-						handles[1] = mutex;
-
+						WaitHandle[] handles = [_disposing, mutex];
 						if (WaitHandle.WaitAny(handles) == 0)
 						{
 							return;
diff --git a/Engine/Source/Programs/Shared/EpicGames.Core/StringUtils.cs b/Engine/Source/Programs/Shared/EpicGames.Core/StringUtils.cs
--- a/Engine/Source/Programs/Shared/EpicGames.Core/StringUtils.cs
+++ b/Engine/Source/Programs/Shared/EpicGames.Core/StringUtils.cs
@@ -26,7 +26,7 @@ namespace EpicGames.Core
 		/// <summary>
 		/// Array mapping human readable size of bytes, 1024^x. long max is within the range of Exabytes.
 		/// </summary>
-		static readonly string[] s_byteSizes = { "B", "KB", "MB", "GB", "TB", "PB", "EB" };
+		static readonly string[] s_byteSizes = ["B", "KB", "MB", "GB", "TB", "PB", "EB"];
 
 		/// <summary>
 		/// Static constructor. Initializes the HexDigits array.
diff --git a/Engine/Source/Programs/Shared/EpicGames.Core/Telemetry/DataRouterTelemetryService.cs b/Engine/Source/Programs/Shared/EpicGames.Core/Telemetry/DataRouterTelemetryService.cs
--- a/Engine/Source/Programs/Shared/EpicGames.Core/Telemetry/DataRouterTelemetryService.cs
+++ b/Engine/Source/Programs/Shared/EpicGames.Core/Telemetry/DataRouterTelemetryService.cs
@@ -124,7 +124,7 @@ namespace EpicGames.Core.Telemetry
 		public List<object> Events
 		{
 			get;
-		} = new List<object>();
+		} = [];
 	}
 
 	/// <summary>
@@ -138,7 +138,7 @@ namespace EpicGames.Core.Telemetry
 		public List<string> EventsSerialized
 		{
 			get;
-		} = new List<string>();
+		} = [];
 
 	}
 
@@ -458,7 +458,7 @@ namespace EpicGames.Core.Telemetry
 		/// <param name="outPostsList"></param>
 		private void CreateBatchedSerializedDataRouterPosts(DataRouterPost inPost, string postURI, out List<SerializedDataRouterPost> outPostsList)
 		{
-			outPostsList = new List<SerializedDataRouterPost>();
+			outPostsList = [];
 
 			SerializedDataRouterPost currentBatch = new SerializedDataRouterPost();
 
@@ -555,7 +555,7 @@ namespace EpicGames.Core.Telemetry
 				}
 			};
 
-			List<Task<HttpResponseMessage>> postRequests = new List<Task<HttpResponseMessage>>();
+			List<Task<HttpResponseMessage>> postRequests = [];
 			foreach (KeyValuePair<string, DataRouterPost> postData in posts)
 			{
 				// Break our post into multiple if the DataRouterPost serialized would go over our max event size
diff --git a/Engine/Source/Programs/Shared/EpicGames.Core/ThreadPoolWorkQueue.cs b/Engine/Source/Programs/Shared/EpicGames.Core/ThreadPoolWorkQueue.cs
--- a/Engine/Source/Programs/Shared/EpicGames.Core/ThreadPoolWorkQueue.cs
+++ b/Engine/Source/Programs/Shared/EpicGames.Core/ThreadPoolWorkQueue.cs
@@ -32,7 +32,7 @@ namespace EpicGames.Core
 		/// <summary>
 		/// Exceptions which occurred while executing tasks
 		/// </summary>
-		readonly ConcurrentBag<Exception> _exceptions = new();
+		readonly ConcurrentBag<Exception> _exceptions = [];
 
 		/// <summary>
 		/// Default constructor
diff --git a/Engine/Source/Programs/Shared/EpicGames.Core/TraceSpan.cs b/Engine/Source/Programs/Shared/EpicGames.Core/TraceSpan.cs
--- a/Engine/Source/Programs/Shared/EpicGames.Core/TraceSpan.cs
+++ b/Engine/Source/Programs/Shared/EpicGames.Core/TraceSpan.cs
@@ -54,7 +54,7 @@ namespace EpicGames.Core
 			public string? Service { get; }
 			public DateTimeOffset StartTime { get; }
 			public DateTimeOffset? FinishTime { get; private set; }
-			public Dictionary<string, string> Metadata { get; } = new Dictionary<string, string>();
+			public Dictionary<string, string> Metadata { get; } = [];
 
 			public TraceSpanImpl(string name, string? resource, string? service)
 			{
@@ -86,7 +86,7 @@ namespace EpicGames.Core
 		/// <summary>
 		/// The current scope provider
 		/// </summary>
-		readonly List<TraceSpanImpl> _spans = new List<TraceSpanImpl>();
+		readonly List<TraceSpanImpl> _spans = [];
 
 		/// <summary>
 		/// Constructor
@@ -198,7 +198,7 @@ namespace EpicGames.Core
 		/// <returns></returns>
 		static List<ITraceSink> GetDefaultSinks()
 		{
-			List<ITraceSink> sinks = new List<ITraceSink>();
+			List<ITraceSink> sinks = [];
 
 			string? telemetryDir = Environment.GetEnvironmentVariable("UE_TELEMETRY_DIR");
 			if (telemetryDir != null)
@@ -234,7 +234,7 @@ namespace EpicGames.Core
 		{
 			if (s_sinks.Count == 0)
 			{
-				return new CombinedTraceSpan(Array.Empty<ITraceSpan>());
+				return new CombinedTraceSpan([]);
 			}
 			else if (s_sinks.Count == 1)
 			{
@@ -242,7 +242,7 @@ namespace EpicGames.Core
 			}
 			else
 			{
-				return new CombinedTraceSpan(s_sinks.ConvertAll(x => x.Create(operation, resource, service)).ToArray());
+				return new CombinedTraceSpan([.. s_sinks.ConvertAll(x => x.Create(operation, resource, service))]);
 			}
 		}
 
diff --git a/Engine/Source/Programs/Shared/EpicGames.Core/UHTTypes.cs b/Engine/Source/Programs/Shared/EpicGames.Core/UHTTypes.cs
--- a/Engine/Source/Programs/Shared/EpicGames.Core/UHTTypes.cs
+++ b/Engine/Source/Programs/Shared/EpicGames.Core/UHTTypes.cs
@@ -110,11 +110,11 @@ namespace EpicGames.Core
 			public string BaseDirectory { get; set; } = string.Empty;
 			public string IncludeBase { get; set; } = string.Empty; // The include path which all UHT-generated includes should be relative to
 			public string OutputDirectory { get; set; } = string.Empty;
-			public List<string> ClassesHeaders { get; set; } = new List<string>();
-			public List<string> PublicHeaders { get; set; } = new List<string>();
-			public List<string> InternalHeaders { get; set; } = new List<string>();
-			public List<string> PrivateHeaders { get; set; } = new List<string>();
-			public List<string> PublicDefines { get; set; } = new List<string>();
+			public List<string> ClassesHeaders { get; set; } = [];
+			public List<string> PublicHeaders { get; set; } = [];
+			public List<string> InternalHeaders { get; set; } = [];
+			public List<string> PrivateHeaders { get; set; } = [];
+			public List<string> PublicDefines { get; set; } = [];
 			public string? GeneratedCPPFilenameBase { get; set; } = null;
 			public bool SaveExportedHeaders { get; set; } = false;
 			[JsonConverter(typeof(JsonStringEnumConverter))]
@@ -189,12 +189,12 @@ namespace EpicGames.Core
 		/// <summary>
 		/// List of modules
 		/// </summary>
-		public List<Module> Modules { get; set; } = new List<Module>();
+		public List<Module> Modules { get; set; } = [];
 
 		/// <summary>
 		/// List of active UHT plugins.  Only used by the C# version of UHT.  This
 		/// list contains plugins from only modules listed above.
 		/// </summary>
-		public List<string> UhtPlugins { get; set; } = new List<string>();
+		public List<string> UhtPlugins { get; set; } = [];
 	}
 }
diff --git a/Engine/Source/Programs/Shared/EpicGames.Core/UnrealEngineTypes.cs b/Engine/Source/Programs/Shared/EpicGames.Core/UnrealEngineTypes.cs
--- a/Engine/Source/Programs/Shared/EpicGames.Core/UnrealEngineTypes.cs
+++ b/Engine/Source/Programs/Shared/EpicGames.Core/UnrealEngineTypes.cs
@@ -1293,7 +1293,7 @@ namespace EpicGames.Core
 	{
 		private static readonly Lazy<List<string>> BitNames = new Lazy<List<string>>(() =>
 		{
-			List<string> Out = new List<string>();
+			List<string> Out = [];
 			FieldInfo[] Fields = typeof(EPropertyFlags).GetFields();
 			for (int Bit = 0; Bit < 64; ++Bit)
 			{
@@ -1343,7 +1343,7 @@ namespace EpicGames.Core
 		{
 			List<string> LocalBitNames = BitNames.Value;
 			ulong IntFlags = (ulong)InFlags;
-			List<string> Out = new List<string>();
+			List<string> Out = [];
 			for (int Bit = 0; Bit < 64; ++Bit)
 			{
 				ulong Mask = (ulong)1 << Bit;
diff --git a/Engine/Source/Programs/Shared/EpicGames.Core/Utf8String.cs b/Engine/Source/Programs/Shared/EpicGames.Core/Utf8String.cs
--- a/Engine/Source/Programs/Shared/EpicGames.Core/Utf8String.cs
+++ b/Engine/Source/Programs/Shared/EpicGames.Core/Utf8String.cs
@@ -161,7 +161,7 @@ namespace EpicGames.Core
 			}
 			else
 			{
-				return Span.IndexOf(Encoding.UTF8.GetBytes(new[] { character }));
+				return Span.IndexOf(Encoding.UTF8.GetBytes([character]));
 			}
 		}
 
@@ -178,7 +178,7 @@ namespace EpicGames.Core
 			}
 			else
 			{
-				result = Span.Slice(index, count).IndexOf(Encoding.UTF8.GetBytes(new[] { character }));
+				result = Span.Slice(index, count).IndexOf(Encoding.UTF8.GetBytes([character]));
 			}
 			return (result == -1) ? -1 : result + index;
 		}
@@ -217,7 +217,7 @@ namespace EpicGames.Core
 			}
 			else
 			{
-				return Span.LastIndexOf(Encoding.UTF8.GetBytes(new[] { character }));
+				return Span.LastIndexOf(Encoding.UTF8.GetBytes([character]));
 			}
 		}
 
diff --git a/Engine/Source/Programs/Shared/EpicGames.Core/Utf8StringBuilder.cs b/Engine/Source/Programs/Shared/EpicGames.Core/Utf8StringBuilder.cs
--- a/Engine/Source/Programs/Shared/EpicGames.Core/Utf8StringBuilder.cs
+++ b/Engine/Source/Programs/Shared/EpicGames.Core/Utf8StringBuilder.cs
@@ -25,7 +25,7 @@ namespace EpicGames.Core
 
 				if (value > prevLength)
 				{
-					_writer.WrittenSpan.Slice(prevLength).Fill(0);
+					_writer.WrittenSpan.Slice(prevLength).Clear();
 				}
 			}
 		}
@@ -86,13 +86,11 @@ namespace EpicGames.Core
 		{
 			if (ch < 127)
 			{
-				_writer.WriteUInt8((byte)ch);
+				Append((byte)ch);
 			}
 			else
 			{
-				Span<char> buffer = stackalloc char[1];
-				buffer[0] = ch;
-				Append(buffer);
+				Append([ch]);
 			}
 		}
 
diff --git a/Engine/Source/Programs/Shared/EpicGames.Core/VersionNumber.cs b/Engine/Source/Programs/Shared/EpicGames.Core/VersionNumber.cs
--- a/Engine/Source/Programs/Shared/EpicGames.Core/VersionNumber.cs
+++ b/Engine/Source/Programs/Shared/EpicGames.Core/VersionNumber.cs
@@ -16,7 +16,7 @@ namespace EpicGames.Core
 		/// <summary>
 		/// Set of delimiters for version numbers
 		/// </summary>
-		static readonly char[] s_delimiters = { '.', ',' };
+		static readonly char[] s_delimiters = ['.', ','];
 
 		/// <summary>
 		/// The individual version components
@@ -191,12 +191,12 @@ namespace EpicGames.Core
 		/// <returns>A version number object</returns>
 		public static VersionNumber Parse(string text)
 		{
-			List<int> components = new List<int>();
+			List<int> components = [];
 			foreach (string textElement in text.Split(s_delimiters))
 			{
 				components.Add(Int32.Parse(textElement));
 			}
-			return new VersionNumber(components.ToArray());
+			return new VersionNumber([.. components]);
 		}
 
 		/// <summary>
@@ -207,7 +207,7 @@ namespace EpicGames.Core
 		/// <returns>A version number object</returns>
 		public static bool TryParse(string text, [NotNullWhen(true)] out VersionNumber? outNumber)
 		{
-			List<int> components = new List<int>();
+			List<int> components = [];
 			foreach (string textElement in text.Split(s_delimiters))
 			{
 				int component;
@@ -219,7 +219,7 @@ namespace EpicGames.Core
 				components.Add(component);
 			}
 
-			outNumber = new VersionNumber(components.ToArray());
+			outNumber = new VersionNumber([.. components]);
 			return true;
 		}
 
