EpicGames.Core: Revert change causing issue with rider

#rnx

diff --git a/Engine/Source/Programs/Shared/EpicGames.Core/Utf8String.cs b/Engine/Source/Programs/Shared/EpicGames.Core/Utf8String.cs
--- a/Engine/Source/Programs/Shared/EpicGames.Core/Utf8String.cs
+++ b/Engine/Source/Programs/Shared/EpicGames.Core/Utf8String.cs
@@ -161,7 +161,7 @@ namespace EpicGames.Core
 			}
 			else
 			{
-				return Span.IndexOf(Encoding.UTF8.GetBytes([character]));
+				return Span.IndexOf(Encoding.UTF8.GetBytes(new[] { character }));
 			}
 		}
 
@@ -178,7 +178,7 @@ namespace EpicGames.Core
 			}
 			else
 			{
-				result = Span.Slice(index, count).IndexOf(Encoding.UTF8.GetBytes([character]));
+				result = Span.Slice(index, count).IndexOf(Encoding.UTF8.GetBytes(new[] { character }));
 			}
 			return (result == -1) ? -1 : result + index;
 		}
@@ -217,7 +217,7 @@ namespace EpicGames.Core
 			}
 			else
 			{
-				return Span.LastIndexOf(Encoding.UTF8.GetBytes([character]));
+				return Span.LastIndexOf(Encoding.UTF8.GetBytes(new[] { character }));
 			}
 		}
 
