Horde: Improve throughput for bundle downloads.

* No longer reads through the shared cache. This prevents having to paginate reads and reduces the number of separate read calls that are required.
* Increased length of the request queue used for determining read requests.
* Pipeline for reading data now uses a shared BatchReader implementation.

#rnx

diff --git a/Engine/Source/Programs/Horde/Plugins/Storage/HordeServer.Storage/Storage/StorageService.cs b/Engine/Source/Programs/Horde/Plugins/Storage/HordeServer.Storage/Storage/StorageService.cs
--- a/Engine/Source/Programs/Horde/Plugins/Storage/HordeServer.Storage/Storage/StorageService.cs
+++ b/Engine/Source/Programs/Horde/Plugins/Storage/HordeServer.Storage/Storage/StorageService.cs
@@ -1202,7 +1202,7 @@ namespace HordeServer.Storage
 				Channel<SortedSetEntry<ObjectId>> channel = Channel.CreateBounded<SortedSetEntry<ObjectId>>(new BoundedChannelOptions(128));
 
 				await using AsyncPipeline pipeline = new AsyncPipeline(cancellationSource.Token);
-				pipeline.AddTask(ctx => FindBlobsForReachabilityCheckAsync(namespaceInfo, channel.Writer, queueChangeEvent, ctx));
+				_ = pipeline.AddTask(ctx => FindBlobsForReachabilityCheckAsync(namespaceInfo, channel.Writer, queueChangeEvent, ctx));
 				pipeline.AddTasks(8, channel.Reader, (entry, ctx) => CheckReachabilityAsync(namespaceInfo, entry, lastImportBlobInfoId, sweepState, storageConfig, queueChangeEvent, ctx));
 				await pipeline.WaitForCompletionAsync();
 
@@ -1263,7 +1263,7 @@ namespace HordeServer.Storage
 		}
 
 		// Checks whether an individual blob can be removed
-		async Task CheckReachabilityAsync(NamespaceInfo namespaceInfo, SortedSetEntry<ObjectId> entry, ObjectId lastImportBlobInfoId, GcSweepState state, StorageConfig storageConfig, AsyncEvent queueChangeEvent, CancellationToken cancellationToken)
+		async ValueTask CheckReachabilityAsync(NamespaceInfo namespaceInfo, SortedSetEntry<ObjectId> entry, ObjectId lastImportBlobInfoId, GcSweepState state, StorageConfig storageConfig, AsyncEvent queueChangeEvent, CancellationToken cancellationToken)
 		{
 			ObjectId blobInfoId = new ObjectId(((byte[]?)entry.ElementValue)!);
 
@@ -1412,7 +1412,7 @@ namespace HordeServer.Storage
 			writer.Complete();
 		}
 
-		async Task FindBlobLengthsAsync(BlobInfo blobInfo, State state, CancellationToken cancellationToken)
+		async ValueTask FindBlobLengthsAsync(BlobInfo blobInfo, State state, CancellationToken cancellationToken)
 		{
 			NamespaceInfo? namespaceInfo;
 			if (state.Namespaces.TryGetValue(blobInfo.NamespaceId, out namespaceInfo))
diff --git a/Engine/Source/Programs/Shared/EpicGames.Core/AsyncPipeline.cs b/Engine/Source/Programs/Shared/EpicGames.Core/AsyncPipeline.cs
--- a/Engine/Source/Programs/Shared/EpicGames.Core/AsyncPipeline.cs
+++ b/Engine/Source/Programs/Shared/EpicGames.Core/AsyncPipeline.cs
@@ -60,9 +60,11 @@ namespace EpicGames.Core
 		/// Adds a new task to the pipeline.
 		/// </summary>
 		/// <param name="taskFunc">Method to execute</param>
-		public void AddTask(Func<CancellationToken, Task> taskFunc)
+		public Task AddTask(Func<CancellationToken, Task> taskFunc)
 		{
-			_tasks.Add(Task.Run(() => RunGuardedAsync(taskFunc), _cancellationSource.Token));
+			Task task = Task.Run(() => RunGuardedAsync(taskFunc), _cancellationSource.Token);
+			_tasks.Add(task);
+			return task;
 		}
 
 		/// <summary>
@@ -70,12 +72,14 @@ namespace EpicGames.Core
 		/// </summary>
 		/// <param name="count">Number of tasks to run</param>
 		/// <param name="taskFunc">Method to execute</param>
-		public void AddTasks(int count, Func<CancellationToken, Task> taskFunc)
+		public Task[] AddTasks(int count, Func<CancellationToken, Task> taskFunc)
 		{
+			Task[] tasks = new Task[count];
 			for (int idx = 0; idx < count; idx++)
 			{
-				AddTask(taskFunc);
+				tasks[idx] = AddTask(taskFunc);
 			}
+			return tasks;
 		}
 
 		async Task RunGuardedAsync(Func<CancellationToken, Task> taskFunc)
@@ -124,28 +128,60 @@ namespace EpicGames.Core
 		/// <param name="pipeline">Pipeline to add the worker to</param>
 		/// <param name="reader">Source for the items</param>
 		/// <param name="taskFunc">Action to execute for each item</param>
-		public static void AddTask<T>(this AsyncPipeline pipeline, ChannelReader<T> reader, Func<T, CancellationToken, Task> taskFunc)
+		public static Task AddTask<T>(this AsyncPipeline pipeline, ChannelReader<T> reader, Func<T, CancellationToken, ValueTask> taskFunc)
+			=> pipeline.AddTask(ctx => ProcessItemsAsync(reader, taskFunc, ctx));
+
+		/// <summary>
+		/// Adds a worker to process items from a channel
+		/// </summary>
+		/// <typeparam name="TInput">Input item type</typeparam>
+		/// <typeparam name="TOutput">Output item type</typeparam>
+		/// <param name="pipeline">Pipeline to add the worker to</param>
+		/// <param name="reader">Reader for input items</param>
+		/// <param name="writer">Writer for output items</param>
+		/// <param name="taskFunc">Action to execute for each item</param>
+		public static Task AddTask<TInput, TOutput>(this AsyncPipeline pipeline, ChannelReader<TInput> reader, ChannelWriter<TOutput> writer, Func<TInput, CancellationToken, ValueTask<TOutput>> taskFunc)
+			=> pipeline.AddTask(ctx => ProcessItemsAsync(reader, writer, taskFunc, ctx));
+
+		/// <summary>
+		/// Adds a worker to process items from a channel
+		/// </summary>
+		/// <typeparam name="T">Input item type</typeparam>
+		/// <param name="pipeline">Pipeline to add the worker to</param>
+		/// <param name="count">Number of workers to add</param>
+		/// <param name="reader">Reader for input items</param>
+		/// <param name="taskFunc">Action to execute for each item</param>
+		public static Task[] AddTasks<T>(this AsyncPipeline pipeline, int count, ChannelReader<T> reader, Func<T, CancellationToken, ValueTask> taskFunc)
 		{
-			pipeline.AddTask(ctx => ProcessItemsAsync(reader, taskFunc, ctx));
+			Task[] tasks = new Task[count];
+			for (int idx = 0; idx < count; idx++)
+			{
+				tasks[idx] = AddTask(pipeline, reader, taskFunc);
+			}
+			return tasks;
 		}
 
 		/// <summary>
 		/// Adds a worker to process items from a channel
 		/// </summary>
-		/// <typeparam name="T">Item type</typeparam>
+		/// <typeparam name="TInput">Input item type</typeparam>
+		/// <typeparam name="TOutput">Output item type</typeparam>
 		/// <param name="pipeline">Pipeline to add the worker to</param>
 		/// <param name="count">Number of workers to add</param>
-		/// <param name="reader">Source for the items</param>
+		/// <param name="reader">Reader for input items</param>
+		/// <param name="writer">Writer for output items</param>
 		/// <param name="taskFunc">Action to execute for each item</param>
-		public static void AddTasks<T>(this AsyncPipeline pipeline, int count, ChannelReader<T> reader, Func<T, CancellationToken, Task> taskFunc)
+		public static Task[] AddTasks<TInput, TOutput>(this AsyncPipeline pipeline, int count, ChannelReader<TInput> reader, ChannelWriter<TOutput> writer, Func<TInput, CancellationToken, ValueTask<TOutput>> taskFunc)
 		{
+			Task[] tasks = new Task[count];
 			for (int idx = 0; idx < count; idx++)
 			{
-				AddTask(pipeline, reader, taskFunc);
+				tasks[idx] = AddTask(pipeline, reader, writer, taskFunc);
 			}
+			return tasks;
 		}
 
-		static async Task ProcessItemsAsync<T>(ChannelReader<T> reader, Func<T, CancellationToken, Task> taskFunc, CancellationToken cancellationToken)
+		static async Task ProcessItemsAsync<T>(ChannelReader<T> reader, Func<T, CancellationToken, ValueTask> taskFunc, CancellationToken cancellationToken)
 		{
 			while (await reader.WaitToReadAsync(cancellationToken))
 			{
@@ -156,5 +192,18 @@ namespace EpicGames.Core
 				}
 			}
 		}
+
+		static async Task ProcessItemsAsync<TInput, TOutput>(ChannelReader<TInput> reader, ChannelWriter<TOutput> writer, Func<TInput, CancellationToken, ValueTask<TOutput>> taskFunc, CancellationToken cancellationToken)
+		{
+			while (await reader.WaitToReadAsync(cancellationToken))
+			{
+				TInput? input;
+				if (reader.TryRead(out input))
+				{
+					TOutput output = await taskFunc(input, cancellationToken);
+					await writer.WriteAsync(output, cancellationToken);
+				}
+			}
+		}
 	}
 }
diff --git a/Engine/Source/Programs/Shared/EpicGames.Horde/Storage/BatchReader.cs b/Engine/Source/Programs/Shared/EpicGames.Horde/Storage/BatchReader.cs
new file mode 100644
--- /dev/null
+++ b/Engine/Source/Programs/Shared/EpicGames.Horde/Storage/BatchReader.cs
@@ -0,0 +1,402 @@
+// Copyright Epic Games, Inc. All Rights Reserved.
+
+using System;
+using System.Collections;
+using System.Collections.Generic;
+using System.Diagnostics;
+using System.Diagnostics.CodeAnalysis;
+using System.Linq;
+using System.Threading;
+using System.Threading.Channels;
+using System.Threading.Tasks;
+using EpicGames.Core;
+using EpicGames.Horde.Storage.Bundles;
+using EpicGames.Horde.Storage.Bundles.V2;
+using OpenTracing.Tag;
+
+namespace EpicGames.Horde.Storage
+{
+	/// <summary>
+	/// Request for a blob to be read
+	/// </summary>
+	/// <typeparam name="TContext">User defined context type for the request</typeparam>
+	/// <param name="BlobRef">Reference to the blob data</param>
+	/// <param name="Context">User defined context</param>
+	public record class BatchReadRequest<TContext>(IBlobRef BlobRef, TContext Context);
+
+	/// <summary>
+	/// Individual response for a blob that was read
+	/// </summary>
+	/// <typeparam name="TContext">User defined context type for the request</typeparam>
+	/// <param name="BlobData">Data for the blob</param>
+	/// <param name="Context">User defined context</param>
+	public sealed record class BatchReadResponseItem<TContext>(BlobData BlobData, TContext Context) : IDisposable
+	{
+		/// <inheritdoc/>
+		public void Dispose()
+			=> BlobData.Dispose();
+	}
+
+	/// <summary>
+	/// Response containing a batch of items that were read
+	/// </summary>
+	/// <typeparam name="TContext">User defined context type</typeparam>
+	/// <param name="Items">Items for this response</param>
+	public sealed record class BatchReadResponse<TContext>(BatchReadResponseItem<TContext>[] Items) : IDisposable
+	{
+		/// <inheritdoc/>
+		public void Dispose()
+		{
+			foreach (BatchReadResponseItem<TContext> item in Items)
+			{
+				item.Dispose();
+			}
+		}
+	}
+
+	/// <summary>
+	/// Stats for a batch read
+	/// </summary>
+	public record class BatchReaderStats(int NumRequests, int NumReads, int NumBundles, int NumPackets, long BytesRead, long BytesDecoded);
+
+	/// <summary>
+	/// Implements an efficient pipeline for streaming blob data
+	/// </summary>
+	/// <typeparam name="TContext">User-defined context type</typeparam>
+	public sealed class BatchReader<TContext> : IDisposable
+	{
+		record class BundleRequest(BundleHandle Handle, int MinOffset, int MaxOffset, List<PacketRequest> Packets);
+
+		record class BundleResponse(IReadOnlyMemoryOwner<byte> Data, int MinOffset, int MaxOffset, List<PacketRequest> Packets) : IDisposable
+		{
+			public void Dispose()
+				=> Data.Dispose();
+		}
+
+		record class PacketRequest(FlushedPacketHandle Handle, List<ExportRequest> Exports)
+		{
+			public int MinOffset => Handle.PacketOffset;
+			public int MaxOffset => Handle.PacketOffset + Handle.PacketLength;
+		}
+
+		record class ExportRequest(ExportHandle Handle, TContext Context);
+
+		record class PacketResponse(IReadOnlyMemoryOwner<byte> Data, List<ExportRequest> Exports) : IDisposable
+		{
+			public void Dispose()
+				=> Data.Dispose();
+		}
+
+		// Number of items to read from the input queue before partitioning into batches
+		const int MinQueueLength = 2000;
+
+		// Maximum gap between reads that should be coalesced
+		const long CoalesceReadsBelowSize = 2 * 1024 * 1024;
+
+		readonly Channel<BatchReadRequest<TContext>> _blobRequests;
+		readonly Channel<BatchReadRequest<TContext>> _otherRequests;
+		readonly Channel<BundleRequest> _bundleRequests;
+
+		readonly Channel<BundleResponse> _bundleResponses;
+		readonly Channel<PacketResponse> _packetResponses;
+		readonly Channel<BatchReadResponse<TContext>> _blobResponses;
+
+		int _numRequests;
+		int _numReads;
+		int _numBundles;
+		int _numPackets;
+		long _bytesRead;
+		long _bytesDecoded;
+
+		/// <summary>
+		/// Access to the request channel
+		/// </summary>
+		public ChannelWriter<BatchReadRequest<TContext>> RequestWriter
+			=> _blobRequests.Writer;
+
+		/// <summary>
+		/// Access to the response channel
+		/// </summary>
+		public ChannelReader<BatchReadResponse<TContext>> ResponseReader
+			=> _blobResponses.Reader;
+
+		/// <summary>
+		/// Constructor
+		/// </summary>
+		public BatchReader()
+		{
+			_blobRequests = Channel.CreateUnbounded<BatchReadRequest<TContext>>();
+			_otherRequests = Channel.CreateUnbounded<BatchReadRequest<TContext>>();
+			_bundleRequests = Channel.CreateUnbounded<BundleRequest>();
+			_bundleResponses = Channel.CreateUnbounded<BundleResponse>();
+			_packetResponses = Channel.CreateUnbounded<PacketResponse>();
+			_blobResponses = Channel.CreateUnbounded<BatchReadResponse<TContext>>();
+		}
+
+		/// <inheritdoc/>
+		public void Dispose()
+		{
+			while (_bundleResponses.Reader.TryRead(out BundleResponse? item))
+			{
+				item.Dispose();
+			}
+
+			while (_packetResponses.Reader.TryRead(out PacketResponse? item))
+			{
+				item.Dispose();
+			}
+
+			while (_blobResponses.Reader.TryRead(out BatchReadResponse<TContext>? item))
+			{
+				item.Dispose();
+			}
+		}
+
+		/// <summary>
+		/// Gets stats for the reader
+		/// </summary>
+		public BatchReaderStats GetStats()
+			=> new BatchReaderStats(_numRequests, _numReads, _numBundles, _numPackets, _bytesRead, _bytesDecoded);
+
+		/// <summary>
+		/// Adds a batch reader to the given pipeline
+		/// </summary>
+		internal Task AddToPipeline(AsyncPipeline pipeline, int numReadTasks, int numDecodeTasks)
+		{
+			_ = pipeline.AddTask(CreateBundleRequestsAsync);
+			Task otherTask = pipeline.AddTask(HandleOtherRequestsAsync);
+
+			// Read bundles and mark the bundle response channel as complete once we're finished
+			Task[] bundleTasks = pipeline.AddTasks(numReadTasks, ReadBundlesAsync);
+			_ = Task.WhenAll(bundleTasks).ContinueWith(_ => _bundleResponses.Writer.TryComplete(), TaskScheduler.Default);
+
+			// Read packets
+			Task[] packetTasks = pipeline.AddTasks(numDecodeTasks, ReadPacketsAsync);
+
+			// Read all the output blobs and mark the blob response channel as complete
+			IEnumerable<Task> blobTasks = packetTasks.Append(otherTask);
+			return Task.WhenAll(blobTasks).ContinueWith(_ => _blobResponses.Writer.TryComplete(), TaskScheduler.Default);
+		}
+
+		// Read a blob using the naive read pipeline
+		async Task HandleOtherRequestsAsync(CancellationToken cancellationToken)
+		{
+			await foreach (BatchReadRequest<TContext> request in _otherRequests.Reader.ReadAllAsync(cancellationToken))
+			{
+				BlobData? blobData = null;
+				try
+				{
+					blobData = await request.BlobRef.ReadBlobDataAsync(cancellationToken);
+#pragma warning disable CA2000
+					BatchReadResponse<TContext> response = new BatchReadResponse<TContext>([new BatchReadResponseItem<TContext>(blobData, request.Context)]);
+					await _blobResponses.Writer.WriteAsync(response, cancellationToken);
+#pragma warning restore CA2000
+				}
+				catch
+				{
+					blobData?.Dispose();
+				}
+			}
+		}
+
+		#region Grouping requests
+
+		record class PacketExport(FlushedPacketHandle PacketHandle, ExportHandle ExportHandle, TContext Context);
+
+		async Task CreateBundleRequestsAsync(CancellationToken cancellationToken)
+		{
+			int queueLength = 0;
+			Queue<BundleHandle> bundleQueue = new Queue<BundleHandle>();
+			Dictionary<BundleHandle, List<PacketExport>> bundleHandleToExportBatch = new Dictionary<BundleHandle, List<PacketExport>>();
+
+			for (; ; )
+			{
+				// Fill the queue up to the max length
+				for(;;)
+				{
+					BatchReadRequest<TContext>? request;
+					if (!_blobRequests.Reader.TryRead(out request))
+					{
+						if(queueLength >= MinQueueLength)
+						{
+							break;
+						}
+						if (!await _blobRequests.Reader.WaitToReadAsync(cancellationToken))
+						{
+							break;
+						}
+					}
+					else
+					{
+						Interlocked.Increment(ref _numRequests);
+
+						PacketExport? outputExport;
+						if (!TryGetPacketExport(request, out outputExport))
+						{
+							await _otherRequests.Writer.WriteAsync(request, cancellationToken);
+						}
+						else
+						{
+							BundleHandle bundleHandle = outputExport.PacketHandle.Bundle;
+							if (!bundleHandleToExportBatch.TryGetValue(bundleHandle, out List<PacketExport>? existingExportBatch))
+							{
+								existingExportBatch = new List<PacketExport>();
+								bundleHandleToExportBatch.Add(bundleHandle, existingExportBatch);
+								bundleQueue.Enqueue(bundleHandle);
+							}
+
+							existingExportBatch.Add(outputExport);
+							queueLength++;
+						}
+					}
+				}
+
+				// Exit once we've processed everything and can't get any more items to read.
+				if (queueLength == 0)
+				{
+					_bundleRequests.Writer.TryComplete();
+					_otherRequests.Writer.TryComplete();
+					break;
+				}
+
+				// Flush the first queue
+				{
+					BundleHandle bundleHandle = bundleQueue.Dequeue();
+					List<PacketExport> exportBatch = bundleHandleToExportBatch[bundleHandle];
+					queueLength -= exportBatch.Count;
+					bundleHandleToExportBatch.Remove(bundleHandle);
+
+					await WriteBundleRequestsAsync(bundleHandle, exportBatch, cancellationToken);
+				}
+			}
+		}
+
+		async Task WriteBundleRequestsAsync(BundleHandle bundleHandle, List<PacketExport> exports, CancellationToken cancellationToken)
+		{
+			Interlocked.Increment(ref _numBundles);
+
+			// Group the reads by packet
+			List<PacketRequest> packetRequests = new List<PacketRequest>();
+			foreach (IGrouping<int, PacketExport> group in exports.GroupBy(x => x.PacketHandle.PacketOffset))
+			{
+				List<ExportRequest> exportRequests = group.Select(x => new ExportRequest(x.ExportHandle, x.Context)).ToList();
+				PacketRequest packetRequest = new PacketRequest(group.First().PacketHandle, exportRequests);
+				packetRequests.Add(packetRequest);
+			}
+			packetRequests.SortBy(x => x.MinOffset);
+
+			// Split the packet requests into contiguous bundle reads
+			for (int maxIdx = 0; maxIdx < packetRequests.Count; maxIdx++)
+			{
+				int minIdx = maxIdx;
+				while (maxIdx + 1 < packetRequests.Count && packetRequests[maxIdx + 1].MinOffset < packetRequests[maxIdx].MaxOffset + CoalesceReadsBelowSize)
+				{
+					maxIdx++;
+				}
+
+				int minOffset = packetRequests[minIdx].MinOffset;
+				int maxOffset = packetRequests[maxIdx].MaxOffset;
+
+				BundleRequest request = new BundleRequest(bundleHandle, minOffset, maxOffset, packetRequests.GetRange(minIdx, (maxIdx + 1) - minIdx));
+				await _bundleRequests.Writer.WriteAsync(request, cancellationToken);
+			}
+		}
+
+		static bool TryGetPacketExport(BatchReadRequest<TContext> chunk, [NotNullWhen(true)] out PacketExport? export)
+		{
+			if (chunk.BlobRef.Innermost is ExportHandle exportHandle && exportHandle.Packet is FlushedPacketHandle packetHandle)
+			{
+				export = new PacketExport(packetHandle, exportHandle, chunk.Context);
+				return true;
+			}
+			else
+			{
+				export = null;
+				return false;
+			}
+		}
+
+		#endregion
+		#region Bundle Reads
+
+		async Task ReadBundlesAsync(CancellationToken cancellationToken)
+		{
+			await foreach (BundleRequest request in _bundleRequests.Reader.ReadAllAsync(cancellationToken))
+			{
+				IReadOnlyMemoryOwner<byte>? data = null;
+				try
+				{
+					data = await request.Handle.ReadAsync(request.MinOffset, request.MaxOffset - request.MinOffset, cancellationToken);
+					Interlocked.Increment(ref _numReads);
+					Interlocked.Add(ref _bytesRead, data.Memory.Length);
+					BundleResponse response = new BundleResponse(data, request.MinOffset, request.MaxOffset, request.Packets);
+					await _bundleResponses.Writer.WriteAsync(response, cancellationToken);
+				}
+				catch
+				{
+					data?.Dispose();
+					throw;
+				}
+			}
+		}
+
+		#endregion
+		#region Packet reads
+
+		async Task ReadPacketsAsync(CancellationToken cancellationToken)
+		{
+			await foreach(BundleResponse bundleResponse in _bundleResponses.Reader.ReadAllAsync(cancellationToken))
+			{
+				using IDisposable lifetime = bundleResponse;
+				foreach (PacketRequest packetRequest in bundleResponse.Packets)
+				{
+					// Decode the data for this packet
+					ReadOnlyMemory<byte> memory = bundleResponse.Data.Memory.Slice(packetRequest.Handle.PacketOffset - bundleResponse.MinOffset, packetRequest.Handle.PacketLength);
+
+					using PacketReader packetReader = packetRequest.Handle.CreatePacketReader(memory);
+					Interlocked.Increment(ref _numPackets);
+					Interlocked.Add(ref _bytesDecoded, packetReader.Packet.Length);
+
+					// Create responses for each blob
+					BatchReadResponseItem<TContext>[] items = new BatchReadResponseItem<TContext>[packetRequest.Exports.Count];
+					try
+					{
+						for (int idx = 0; idx < items.Length; idx++)
+						{
+							ExportRequest exportRequest = packetRequest.Exports[idx];
+#pragma warning disable CA2000
+							BlobData blobData = packetReader.ReadExport(exportRequest.Handle.ExportIdx);
+							items[idx] = new BatchReadResponseItem<TContext>(blobData, exportRequest.Context);
+#pragma warning restore CA2000
+						}
+
+						BatchReadResponse<TContext> response = new BatchReadResponse<TContext>(items);
+						await _blobResponses.Writer.WriteAsync(response, cancellationToken);
+					}
+					catch
+					{
+						for (int idx = 0; idx < items.Length; idx++)
+						{
+							items[idx]?.Dispose();
+						}
+						throw;
+					}
+				}
+			}
+		}
+
+		#endregion
+	}
+
+	/// <summary>
+	/// Extension methods for <see cref="BatchReader{TContext}"/>
+	/// </summary>
+	public static class BatchReaderExtensions
+	{
+		/// <summary>
+		/// Add tasks for reading to the given async pipeline 
+		/// </summary>
+		public static Task AddBatchReaderTasks<T>(this AsyncPipeline pipeline, int numReadTasks, int numDecodeTasks, BatchReader<T> reader)
+			=> reader.AddToPipeline(pipeline, numReadTasks, numDecodeTasks);
+	}
+}
diff --git a/Engine/Source/Programs/Shared/EpicGames.Horde/Storage/Bundles/V2/PacketHandle.cs b/Engine/Source/Programs/Shared/EpicGames.Horde/Storage/Bundles/V2/PacketHandle.cs
--- a/Engine/Source/Programs/Shared/EpicGames.Horde/Storage/Bundles/V2/PacketHandle.cs
+++ b/Engine/Source/Programs/Shared/EpicGames.Horde/Storage/Bundles/V2/PacketHandle.cs
@@ -72,6 +72,11 @@ namespace EpicGames.Horde.Storage.Bundles.V2
 		/// </summary>
 		public int PacketOffset => _packetOffset;
 
+		/// <summary>
+		/// Length of the packet within the bundle
+		/// </summary>
+		public int PacketLength => _packetLength;
+
 		/// <summary>
 		/// Constructor
 		/// </summary>
@@ -215,6 +220,13 @@ namespace EpicGames.Horde.Storage.Bundles.V2
 			return await _cache.PacketReaderCache.ScopedGetOrAddAsync(cacheKey, CreatePacketReaderAsync, cancellationToken);
 		}
 
+		internal PacketReader CreatePacketReader(ReadOnlyMemory<byte> encodedData)
+		{
+			PacketReaderCacheKey cacheKey = new PacketReaderCacheKey(_outer, _packetOffset);
+			IRefCountedHandle<Packet> packet = Packet.Decode(encodedData, _cache.Allocator, cacheKey);
+			return new PacketReader(_storageNamespace, _cache, _outer, this, packet.Target, packet);
+		}
+
 		async Task<Scoped<PacketReader>> CreatePacketReaderAsync(PacketReaderCacheKey cacheKey, CancellationToken cancellationToken)
 		{
 			using IReadOnlyMemoryOwner<byte> encodedData = await ReadEncodedPacketAsync(cancellationToken);
diff --git a/Engine/Source/Programs/Shared/EpicGames.Horde/Storage/Nodes/DirectoryNode.Extract.cs b/Engine/Source/Programs/Shared/EpicGames.Horde/Storage/Nodes/DirectoryNode.Extract.cs
--- a/Engine/Source/Programs/Shared/EpicGames.Horde/Storage/Nodes/DirectoryNode.Extract.cs
+++ b/Engine/Source/Programs/Shared/EpicGames.Horde/Storage/Nodes/DirectoryNode.Extract.cs
@@ -8,7 +8,9 @@ using System.Diagnostics.CodeAnalysis;
 using System.IO;
 using System.IO.MemoryMappedFiles;
 using System.Linq;
+using System.Net.Security;
 using System.Runtime.ExceptionServices;
+using System.Text;
 using System.Threading;
 using System.Threading.Channels;
 using System.Threading.Tasks;
@@ -27,82 +29,27 @@ namespace EpicGames.Horde.Storage.Nodes
 		/// <summary>
 		/// Number of files that have been copied
 		/// </summary>
-		int Count { get; }
+		int NumFiles { get; }
 
 		/// <summary>
 		/// Total size of data to be copied
 		/// </summary>
-		long Size { get; }
+		long ExtractSize { get; }
 
 		/// <summary>
 		/// Processing speed, in bytes per second
 		/// </summary>
-		double Rate { get; }
-	}
-
-	/// <summary>
-	/// Reports progress info back to callers
-	/// </summary>
-	class ExtractStats : IExtractStats
-	{
-		readonly object _lockObject = new object();
-		readonly Stopwatch _timer = Stopwatch.StartNew();
-		readonly IProgress<IExtractStats>? _progress;
-		readonly TimeSpan _frequency;
-		long _lastTotalSize;
-		readonly Queue<double> _rateSamples = new Queue<double>();
-
-		public int Count { get; set; }
-		public long Size { get; set; }
-		public double Rate { get; set; }
-
-		public ExtractStats(IProgress<IExtractStats>? progress, TimeSpan frequency)
-		{
-			_progress = progress;
-			_frequency = frequency;
-		}
-
-		public void Update(int count, long size)
-		{
-			if (_progress != null)
-			{
-				lock (_lockObject)
-				{
-					Count += count;
-					Size += size;
-					if (_timer.Elapsed > _frequency)
-					{
-						FlushInternal();
-					}
-				}
-			}
-		}
-
-		public void Flush()
-		{
-			if (_progress != null)
-			{
-				lock (_lockObject)
-				{
-					FlushInternal();
-				}
-			}
-		}
+		double ExtractRate { get; }
 
-		void FlushInternal()
-		{
-			_rateSamples.Enqueue((Size - _lastTotalSize) / _timer.Elapsed.TotalSeconds);
-			while (_rateSamples.Count > Math.Max(1, 10.0 / _frequency.TotalSeconds))
-			{
-				_rateSamples.Dequeue();
-			}
-
-			Rate = _rateSamples.Average();
-			_lastTotalSize = Size;
+		/// <summary>
+		/// Total size of the data downloaded
+		/// </summary>
+		long DownloadSize { get; }
 
-			_progress!.Report(this);
-			_timer.Restart();
-		}
+		/// <summary>
+		/// Download speed, in bytes per second
+		/// </summary>
+		double DownloadRate { get; }
 	}
 
 	/// <summary>
@@ -114,6 +61,11 @@ namespace EpicGames.Horde.Storage.Nodes
 		readonly long _totalSize;
 		readonly ILogger _logger;
 
+		/// <summary>
+		/// Whether to print out separate stats for download speed
+		/// </summary>
+		public bool ShowDownloadStats { get; set; }
+		
 		/// <summary>
 		/// Constructor
 		/// </summary>
@@ -133,21 +85,27 @@ namespace EpicGames.Horde.Storage.Nodes
 		/// <inheritdoc/>
 		public void Report(IExtractStats stats)
 		{
+			if (ShowDownloadStats)
+			{
+				double downloadRate = (stats.DownloadRate * 8.0) / (1024.0 * 1024.0);
+				_logger.LogInformation("Downloaded {TotalSize:n1}mb, {Rate:n0} mbps", stats.DownloadSize / (1024.0 * 1024.0), downloadRate);
+			}
+
 			if (_totalCount > 0 && _totalSize > 0)
 			{
-				_logger.LogInformation("Copied {NumFiles:n0}/{TotalFiles:n0} files ({Size:n1}/{TotalSize:n1}mb, {Rate:n1}mb/s, {Pct}%)", stats.Count, _totalCount, stats.Size / (1024.0 * 1024.0), _totalSize / (1024.0 * 1024.0), stats.Rate / (1024.0 * 1024.0), (int)((Math.Max(stats.Size, 1) * 100) / Math.Max(_totalSize, 1)));
+				_logger.LogInformation("Written {NumFiles:n0}/{TotalFiles:n0} files ({Size:n1}/{TotalSize:n1}mb, {Rate:n1}mb/s, {Pct}%). Downloaded {DownloadSize:n1} mb ({DownloadRate:n0} mpbs).", stats.NumFiles, _totalCount, stats.ExtractSize / (1024.0 * 1024.0), _totalSize / (1024.0 * 1024.0), stats.ExtractRate / (1024.0 * 1024.0), (int)((Math.Max(stats.ExtractSize, 1) * 100) / Math.Max(_totalSize, 1)), stats.DownloadSize / (1024.0 * 1024.0));
 			}
 			else if (_totalCount > 0)
 			{
-				_logger.LogInformation("Copied {NumFiles:n0}/{TotalFiles:n0} files ({Size:n1}mb, {Rate:n1}mb/s)", stats.Count, _totalCount, stats.Size / (1024.0 * 1024.0), stats.Rate / (1024.0 * 1024.0));
+				_logger.LogInformation("Written {NumFiles:n0}/{TotalFiles:n0} files ({Size:n1}mb, {Rate:n1}mb/s)", stats.NumFiles, _totalCount, stats.ExtractSize / (1024.0 * 1024.0), stats.ExtractRate / (1024.0 * 1024.0));
 			}
 			else if (_totalSize > 0)
 			{
-				_logger.LogInformation("Copied {NumFiles:n0} files ({Size:n1}/{TotalSize:n1}mb, {Rate:n1}mb/s, {Pct}%)", stats.Count, stats.Size / (1024.0 * 1024.0), _totalSize / (1024.0 * 1024.0), stats.Rate / (1024.0 * 1024.0), (int)((Math.Max(stats.Size, 1) * 100) / Math.Max(_totalSize, 1)));
+				_logger.LogInformation("Written {NumFiles:n0} files ({Size:n1}/{TotalSize:n1}mb, {Rate:n1}mb/s, {Pct}%)", stats.NumFiles, stats.ExtractSize / (1024.0 * 1024.0), _totalSize / (1024.0 * 1024.0), stats.ExtractRate / (1024.0 * 1024.0), (int)((Math.Max(stats.ExtractSize, 1) * 100) / Math.Max(_totalSize, 1)));
 			}
 			else
 			{
-				_logger.LogInformation("Copied {NumFiles:n0} files ({Size:n1}mb, {Rate:n1}mb/s)", stats.Count, stats.Size / (1024.0 * 1024.0), stats.Rate / (1024.0 * 1024.0));
+				_logger.LogInformation("Written {NumFiles:n0} files ({Size:n1}mb, {Rate:n1}mb/s)", stats.NumFiles, stats.ExtractSize / (1024.0 * 1024.0), stats.ExtractRate / (1024.0 * 1024.0));
 			}
 		}
 	}
@@ -167,6 +125,12 @@ namespace EpicGames.Horde.Storage.Nodes
 		public static Task ExtractAsync(this IBlobRef<DirectoryNode> directoryRef, DirectoryInfo directoryInfo, ILogger logger, CancellationToken cancellationToken) 
 			=> ExtractAsync(directoryRef, directoryInfo, null, logger, cancellationToken);
 
+		class OutputStats
+		{
+			public int _writtenFiles;
+			public long _writtenBytes;
+		}
+
 		class OutputFile
 		{
 			public string Path { get; }
@@ -255,10 +219,10 @@ namespace EpicGames.Horde.Storage.Nodes
 		{
 			public OutputFile OutputFile { get; }
 
-			readonly ChannelWriter<OutputChunk> _chunkWriter;
-			OutputChunk? _bufferedChunk;
+			readonly ChannelWriter<BatchReadRequest<OutputChunk>> _chunkWriter;
+			BatchReadRequest<OutputChunk>? _bufferedChunk;
 
-			public OutputChunkWriter(OutputFile file, ChannelWriter<OutputChunk> chunkWriter)
+			public OutputChunkWriter(OutputFile file, ChannelWriter<BatchReadRequest<OutputChunk>> chunkWriter)
 			{
 				OutputFile = file;
 				_chunkWriter = chunkWriter;
@@ -272,7 +236,7 @@ namespace EpicGames.Horde.Storage.Nodes
 				}
 
 				OutputFile.IncrementRemaining();
-				_bufferedChunk = new OutputChunk(OutputFile, offset, length, handle);
+				_bufferedChunk = new BatchReadRequest<OutputChunk>(handle, new OutputChunk(OutputFile, offset, length, handle));
 			}
 
 			public async Task FlushAsync(CancellationToken cancellationToken)
@@ -341,90 +305,41 @@ namespace EpicGames.Horde.Storage.Nodes
 		/// <param name="cancellationToken">Cancellation token for the operation</param>
 		public static async Task ExtractAsync(this DirectoryNode directoryNode, DirectoryInfo directoryInfo, IProgress<IExtractStats>? progress, TimeSpan frequency, ILogger logger, CancellationToken cancellationToken)
 		{
-			int numTasks = Math.Min(1 + (int)(directoryNode.Length / (16 * 1024 * 1024)), 16);
-			logger.LogInformation("Splitting read into {NumThreads} threads", numTasks);
+			int numReadTasks = 16;
+			int numDecodeTasks = Math.Min(Environment.ProcessorCount, 16);
+			int numWriteTasks = Math.Min(1 + (int)(directoryNode.Length / (16 * 1024 * 1024)), 16);
+			logger.LogInformation("Using {NumReadTasks} read tasks, {NumDecodeTasks} decode tasks, {NumWriteTasks} write tasks", numReadTasks, numDecodeTasks, numWriteTasks);
 
-			ExtractStats extractStats = new ExtractStats(progress, frequency);
-			using (CancellationTokenSource cancellationSource = CancellationTokenSource.CreateLinkedTokenSource(cancellationToken))
+			await using (AsyncPipeline pipeline = new AsyncPipeline(cancellationToken))
 			{
-				ExceptionDispatchInfo? exceptionDispatchInfo = null;
-
-				// Helper method to run a background task and set a cancellation source on error
-				async Task RunBackgroundTask(Func<CancellationToken, Task> taskFunc)
-				{
-					try
-					{
-						await Task.Run(() => taskFunc(cancellationSource.Token), cancellationSource.Token);
-					}
-					catch (OperationCanceledException)
-					{
-						// Ignore
-					}
-					catch (Exception ex)
-					{
-						logger.LogError(ex, "Error while extracting data: {Message}", ex.Message);
-
-						if (exceptionDispatchInfo == null)
-						{
-							ExceptionDispatchInfo dispatchInfo = ExceptionDispatchInfo.Capture(ex);
-							Interlocked.CompareExchange(ref exceptionDispatchInfo, dispatchInfo, null);
-						}
-
-#if NET8_0_OR_GREATER
-						await cancellationSource.CancelAsync();
-#else
-						cancellationSource.Cancel();
-#endif
-					}
-				}
-
-				List<Task> tasks = new List<Task>();
-
-				Channel<OutputChunk> chunks = Channel.CreateBounded<OutputChunk>(new BoundedChannelOptions(128 * 1024) { FullMode = BoundedChannelFullMode.Wait });
-				tasks.Add(RunBackgroundTask(ctx => FindOutputChunksRootAsync(directoryInfo, directoryNode, chunks.Writer, logger, ctx)));
+				using BatchReader<OutputChunk> batchReader = new BatchReader<OutputChunk>();
 
-				Channel<OutputBatch> batches = Channel.CreateUnbounded<OutputBatch>();
-				tasks.Add(RunBackgroundTask(ctx => ReadBatchesAsync(chunks.Reader, batches.Writer, ctx)));
-
-				Channel<OutputFetchedBatch> prefetchBatches = Channel.CreateBounded<OutputFetchedBatch>(new BoundedChannelOptions(128) { FullMode = BoundedChannelFullMode.Wait });
-				try
-				{
-					tasks.Add(RunBackgroundTask(ctx => FetchAsync(batches.Reader, prefetchBatches.Writer, numTasks, ctx)));
+				_ = pipeline.AddTask(ctx => FindOutputChunksRootAsync(directoryInfo, directoryNode, batchReader.RequestWriter, logger, ctx));
+				_ = pipeline.AddBatchReaderTasks(numReadTasks, numDecodeTasks, batchReader);
 
-					for (int idx = 0; idx < numTasks; idx++)
-					{
-						tasks.Add(RunBackgroundTask(ctx => WriteAsync(prefetchBatches.Reader, extractStats, logger, ctx)));
-					}
+				OutputStats outputStats = new OutputStats();
+				Task[] writeTasks = pipeline.AddTasks(numWriteTasks, ctx => WriteAsync(batchReader.ResponseReader, outputStats, logger, ctx));
 
-					await Task.WhenAll(tasks);
-				}
-				finally
+				if (progress != null)
 				{
-					while (prefetchBatches.Reader.TryRead(out OutputFetchedBatch? fetchedBatch))
-					{
-						fetchedBatch.Dispose();
-					}
+					_ = pipeline.AddTask(ctx => UpdateStatsAsync(batchReader, outputStats, Task.WhenAll(writeTasks), progress, ctx));
 				}
-				extractStats.Flush();
 
-				cancellationToken.ThrowIfCancellationRequested();
+				await pipeline.WaitForCompletionAsync();
 
-				if (exceptionDispatchInfo != null)
-				{
-					exceptionDispatchInfo.Throw();
-				}
+				cancellationToken.ThrowIfCancellationRequested();
 			}
 		}
 
 		#region Enumerate chunks
 
-		static async Task FindOutputChunksRootAsync(DirectoryInfo rootDir, DirectoryNode node, ChannelWriter<OutputChunk> chunks, ILogger logger, CancellationToken cancellationToken)
+		static async Task FindOutputChunksRootAsync(DirectoryInfo rootDir, DirectoryNode node, ChannelWriter<BatchReadRequest<OutputChunk>> chunks, ILogger logger, CancellationToken cancellationToken)
 		{
 			await FindOutputChunksForDirectoryAsync(rootDir, "", node, chunks, logger, cancellationToken);
 			chunks.Complete();
 		}
 
-		static async Task FindOutputChunksForDirectoryAsync(DirectoryInfo rootDir, string path, DirectoryNode node, ChannelWriter<OutputChunk> chunks, ILogger logger, CancellationToken cancellationToken)
+		static async Task FindOutputChunksForDirectoryAsync(DirectoryInfo rootDir, string path, DirectoryNode node, ChannelWriter<BatchReadRequest<OutputChunk>> chunks, ILogger logger, CancellationToken cancellationToken)
 		{
 			foreach (FileEntry fileEntry in node.Files)
 			{
@@ -445,7 +360,7 @@ namespace EpicGames.Horde.Storage.Nodes
 			}
 		}
 
-		static async Task FindOutputChunksForFileAsync(OutputFile outputFile, ChannelWriter<OutputChunk> chunks, ILogger logger, CancellationToken cancellationToken)
+		static async Task FindOutputChunksForFileAsync(OutputFile outputFile, ChannelWriter<BatchReadRequest<OutputChunk>> chunks, ILogger logger, CancellationToken cancellationToken)
 		{
 			OutputChunkWriter outputWriter = new OutputChunkWriter(outputFile, chunks);
 			await FindOutputChunksAsync(outputWriter, 0, outputFile.FileEntry.Target, logger, cancellationToken);
@@ -492,194 +407,22 @@ namespace EpicGames.Horde.Storage.Nodes
 
 		#endregion
 
-		#region Group requests by bundles
-
-		/// <summary>
-		/// Compare two blob handles to allow sorting for optimized reads
-		/// </summary>
-		public static int CompareBlobs(IBlobRef a, IBlobRef b)
-		{
-			if (a.Innermost is not ExportHandle exportHandleA || exportHandleA.Packet is not FlushedPacketHandle packetHandleA)
-			{
-				return Comparer.Default.Compare(a, b);
-			}
-
-			if (b.Innermost is not ExportHandle exportHandleB || exportHandleB.Packet is not FlushedPacketHandle packetHandleB)
-			{
-				return Comparer.Default.Compare(a, b);
-			}
-
-			if (packetHandleA.Bundle != packetHandleB.Bundle)
-			{
-				return Comparer.Default.Compare(packetHandleA.Bundle, packetHandleB.Bundle);
-			}
-			else if (packetHandleA.PacketOffset != packetHandleB.PacketOffset)
-			{
-				return Comparer.Default.Compare(packetHandleA.PacketOffset, packetHandleB.PacketOffset);
-			}
-			else
-			{
-				return Comparer.Default.Compare(exportHandleA.ExportIdx, exportHandleB.ExportIdx);
-			}
-		}
-
-		record class OutputExport(BundleHandle BundleHandle, int PacketOffset, int ExportIdx, OutputChunk Chunk);
-		record class OutputExportBatch(BundleHandle BundleHandle, List<OutputExport> Exports);
-
-		static async Task ReadBatchesAsync(ChannelReader<OutputChunk> chunkReader, ChannelWriter<OutputBatch> batchWriter, CancellationToken cancellationToken)
-		{
-			const int MaxQueueLength = 200;
-
-			int queueLength = 0;
-			Queue<OutputExportBatch> exportBatchQueue = new Queue<OutputExportBatch>();
-			Dictionary<BundleHandle, OutputExportBatch> bundleHandleToExportBatch = new Dictionary<BundleHandle, OutputExportBatch>();
-
-			for (; ; )
-			{
-				// Fill the queue up to the max length
-				while (queueLength < MaxQueueLength)
-				{
-					OutputChunk? chunk;
-					if (!chunkReader.TryRead(out chunk))
-					{
-						if (!await chunkReader.WaitToReadAsync(cancellationToken))
-						{
-							break;
-						}
-					}
-					else
-					{
-						OutputExport? outputExport;
-						if (!TryGetOutputExport(chunk, out outputExport))
-						{
-							OutputBatch batch = new OutputBatch(new List<OutputChunk> { chunk });
-							await batchWriter.WriteAsync(batch, cancellationToken);
-						}
-						else
-						{
-							BundleHandle bundleHandle = outputExport.BundleHandle;
-							if (!bundleHandleToExportBatch.TryGetValue(bundleHandle, out OutputExportBatch? existingExportBatch))
-							{
-								existingExportBatch = new OutputExportBatch(bundleHandle, new List<OutputExport>());
-								exportBatchQueue.Enqueue(existingExportBatch);
-								bundleHandleToExportBatch.Add(bundleHandle, existingExportBatch);
-							}
-
-							existingExportBatch.Exports.Add(outputExport);
-							queueLength++;
-						}
-					}
-				}
-
-				// Exit once we've processed everything and can't get any more items to read.
-				if (queueLength == 0)
-				{
-					batchWriter.TryComplete();
-					break;
-				}
-
-				// Flush the first queue
-				OutputExportBatch exportBatch = exportBatchQueue.Dequeue();
-				queueLength -= exportBatch.Exports.Count;
-				bundleHandleToExportBatch.Remove(exportBatch.BundleHandle);
-
-				List<OutputChunk> chunkBatch = exportBatch.Exports.OrderBy(x => x.PacketOffset).ThenBy(x => x.ExportIdx).Select(x => x.Chunk).ToList();
-				await batchWriter.WriteAsync(new OutputBatch(chunkBatch), cancellationToken);
-			}
-		}
-
-		static bool TryGetOutputExport(OutputChunk chunk, [NotNullWhen(true)] out OutputExport? export)
-		{
-			if (chunk.Handle.Innermost is ExportHandle exportHandle && exportHandle.Packet is FlushedPacketHandle packetHandle)
-			{
-				export = new OutputExport(packetHandle.Bundle, packetHandle.PacketOffset, exportHandle.ExportIdx, chunk);
-				return true;
-			}
-			else
-			{
-				export = null;
-				return false;
-			}
-		}
-
-		#endregion
-
-		#region Fetch data
-
-		record class OutputFetchedChunk(OutputChunk Chunk, BlobData BlobData) : OutputChunk(Chunk), IDisposable
-		{
-			public void Dispose() => BlobData.Dispose();
-		}
-
-		record class OutputFetchedBatch(List<OutputFetchedChunk> Chunks) : IDisposable
-		{
-			public OutputFetchedBatch() : this(new List<OutputFetchedChunk>()) { }
-			public void Dispose() => Chunks.DisposeAll();
-		}
-
-		static async Task FetchAsync(ChannelReader<OutputBatch> batchReader, ChannelWriter<OutputFetchedBatch> batchWriter, int numParallel, CancellationToken cancellationToken)
-		{
-			List<Task> tasks = new List<Task>();
-			for (int idx = 0; idx < numParallel; idx++)
-			{
-				tasks.Add(Task.Run(() => FetchWorkerAsync(batchReader, batchWriter, cancellationToken), cancellationToken));
-			}
-
-			try
-			{
-				await Task.WhenAll(tasks);
-			}
-			finally
-			{
-				batchWriter.Complete();
-			}
-		}
-
-		static async Task FetchWorkerAsync(ChannelReader<OutputBatch> batchReader, ChannelWriter<OutputFetchedBatch> batchWriter, CancellationToken cancellationToken)
-		{
-			while (await batchReader.WaitToReadAsync(cancellationToken))
-			{
-				OutputBatch? batch;
-				if (batchReader.TryRead(out batch))
-				{
-#pragma warning disable CA2000 // fetchedBatch may be pushed onto the output channel, which assumes its ownership.
-					OutputFetchedBatch fetchedBatch = new OutputFetchedBatch();
-					try
-					{
-						foreach (OutputChunk chunk in batch.Chunks)
-						{
-							BlobData blobData = await chunk.Handle.ReadBlobDataAsync(cancellationToken);
-							fetchedBatch.Chunks.Add(new OutputFetchedChunk(chunk, blobData));
-						}
-						await batchWriter.WriteAsync(fetchedBatch, cancellationToken);
-					}
-					catch
-					{
-						fetchedBatch.Dispose();
-					}
-#pragma warning restore CA2000
-				}
-			}
-		}
-
-		#endregion
-
 		#region Write to disk
 
-		static async Task WriteAsync(ChannelReader<OutputFetchedBatch> batchReader, ExtractStats stats, ILogger logger, CancellationToken cancellationToken)
+		static async Task WriteAsync(ChannelReader<BatchReadResponse<OutputChunk>> batchReader, OutputStats outputStats, ILogger logger, CancellationToken cancellationToken)
 		{
 			const int WriteBatchSize = 64;
 			while (await batchReader.WaitToReadAsync(cancellationToken))
 			{
-				OutputFetchedBatch? batch;
+				BatchReadResponse<OutputChunk>? batch;
 				while (batchReader.TryRead(out batch))
 				{
 					try
 					{
 						List<Task> tasks = new List<Task>();
-						foreach (IReadOnlyList<OutputFetchedChunk> group in batch.Chunks.Batch(WriteBatchSize))
+						foreach (IReadOnlyList<BatchReadResponseItem<OutputChunk>> group in batch.Items.Batch(WriteBatchSize))
 						{
-							tasks.Add(WriteChunksAsync(group.ToArray(), stats, logger, cancellationToken));
+							tasks.Add(WriteChunksAsync(group.ToArray(), outputStats, logger, cancellationToken));
 						}
 						await Task.WhenAll(tasks);
 					}
@@ -691,22 +434,22 @@ namespace EpicGames.Horde.Storage.Nodes
 			}
 		}
 
-		static async Task WriteChunksAsync(ArraySegment<OutputFetchedChunk> chunks, ExtractStats stats, ILogger logger, CancellationToken cancellationToken)
+		static async Task WriteChunksAsync(ArraySegment<BatchReadResponseItem<OutputChunk>> chunks, OutputStats outputStats, ILogger logger, CancellationToken cancellationToken)
 		{
 			for (int chunkIdx = 0; chunkIdx < chunks.Count;)
 			{
-				OutputFile file = chunks[chunkIdx].File;
+				OutputFile file = chunks[chunkIdx].Context.File;
 
 				int maxChunkIdx = chunkIdx + 1;
-				while (maxChunkIdx < chunks.Count && chunks[maxChunkIdx].File == file)
+				while (maxChunkIdx < chunks.Count && chunks[maxChunkIdx].Context.File == file)
 				{
 					maxChunkIdx++;
 				}
 
 				try
 				{
-					await ExtractChunksToFileAsync(file, chunks.Slice(chunkIdx, maxChunkIdx - chunkIdx), stats, logger, cancellationToken);
-					//await ExtractChunksToNullAsync(file, chunks.Slice(chunkIdx, maxChunkIdx - chunkIdx), stats, logger, cancellationToken);
+					await ExtractChunksToFileAsync(file, chunks.Slice(chunkIdx, maxChunkIdx - chunkIdx), outputStats, logger, cancellationToken);
+					//await ExtractChunksToNullAsync(file, chunks.Slice(chunkIdx, maxChunkIdx - chunkIdx), outputStats, logger, cancellationToken);
 				}
 				catch (OperationCanceledException)
 				{
@@ -721,7 +464,7 @@ namespace EpicGames.Horde.Storage.Nodes
 			}
 		}
 
-		static async Task ExtractChunksToFileAsync(OutputFile file, ArraySegment<OutputFetchedChunk> chunks, ExtractStats stats, ILogger logger, CancellationToken cancellationToken)
+		static async Task ExtractChunksToFileAsync(OutputFile file, ArraySegment<BatchReadResponseItem<OutputChunk>> chunks, OutputStats stats, ILogger logger, CancellationToken cancellationToken)
 		{
 			// Open the file for the current chunk
 			int remainingChunks = 0;
@@ -743,16 +486,16 @@ namespace EpicGames.Horde.Storage.Nodes
 
 					for (int chunkIdx = 0; chunkIdx < chunks.Count; chunkIdx++)
 					{
-						OutputFetchedChunk chunk = chunks[chunkIdx];
+						BatchReadResponseItem<OutputChunk> chunk = chunks[chunkIdx];
 						cancellationToken.ThrowIfCancellationRequested();
 
 						// Write this chunk
-						TraceBlobRead("Leaf", chunk.File.Path, chunk.Handle, logger);
-						chunk.BlobData.Data.CopyTo(memoryMappedView!.GetMemory(chunk.Offset, chunk.BlobData.Data.Length));
+						TraceBlobRead("Leaf", chunk.Context.File.Path, chunk.Context.Handle, logger);
+						chunk.BlobData.Data.CopyTo(memoryMappedView!.GetMemory(chunk.Context.Offset, chunk.BlobData.Data.Length));
 
 						// Update the stats
 						remainingChunks = file.DecrementRemaining();
-						stats.Update(0, chunk.Length);
+						Interlocked.Add(ref stats._writtenBytes, chunk.Context.Length);
 					}
 				}
 			}
@@ -768,35 +511,82 @@ namespace EpicGames.Horde.Storage.Nodes
 					file.FileInfo.LastWriteTimeUtc = file.FileEntry.ModTime;
 				}
 
-				stats.Update(1, 0);
+				Interlocked.Increment(ref stats._writtenFiles);
 			}
 		}
 
 #pragma warning disable IDE0051
 		// Update counters for extracting chunks without writing any data. Useful for profiling bottlenecks in other stages of the pipeline.
-		static Task ExtractChunksToNullAsync(OutputFile file, ArraySegment<OutputFetchedChunk> chunks, ExtractStats stats, ILogger logger, CancellationToken cancellationToken)
+		static Task ExtractChunksToNullAsync(OutputFile file, ArraySegment<BatchReadResponseItem<OutputChunk>> chunks, OutputStats stats, ILogger logger, CancellationToken cancellationToken)
 		{
 			_ = logger;
 			_ = cancellationToken;
 
-			foreach (OutputFetchedChunk chunk in chunks)
+			foreach (BatchReadResponseItem<OutputChunk> chunk in chunks)
 			{
+				Interlocked.Add(ref stats._writtenBytes, chunk.Context.Length);
+
 				int remainingChunks = file.DecrementRemaining();
 				if (remainingChunks == 0)
 				{
-					stats.Update(1, chunk.Length);
-				}
-				else
-				{
-					stats.Update(0, chunk.Length);
+					Interlocked.Increment(ref stats._writtenFiles);
 				}
 			}
+
 			return Task.CompletedTask;
 		}
 #pragma warning restore IDE0051
 
 		#endregion
 
+		#region Stats
+
+		class ExtractStats : IExtractStats
+		{
+			public int NumFiles { get; set; }
+			public long ExtractSize { get; set; }
+			public double ExtractRate { get; set; }
+			public long DownloadSize { get; set; }
+			public double DownloadRate { get; set; }
+		}
+
+		static async Task UpdateStatsAsync(BatchReader<OutputChunk> batchReader, OutputStats outputStats, Task writeTask, IProgress<IExtractStats> progress, CancellationToken cancellationToken)
+		{
+			Stopwatch timer = Stopwatch.StartNew();
+
+			long lastExtractSize = 0;
+			long lastDownloadSize = 0;
+
+			Task? completeTask = null;
+			while (completeTask != writeTask)
+			{
+				completeTask = await Task.WhenAny(writeTask, Task.Delay(TimeSpan.FromSeconds(5.0)));
+
+				ExtractStats stats = new ExtractStats();
+				stats.NumFiles = Interlocked.CompareExchange(ref outputStats._writtenFiles, 0, 0);
+				stats.ExtractSize = Interlocked.CompareExchange(ref outputStats._writtenBytes, 0, 0);
+
+				BatchReaderStats batchStats = batchReader.GetStats();
+				stats.DownloadSize = batchStats.BytesRead;
+
+				double elapsedSeconds = timer.Elapsed.TotalSeconds;
+				if (elapsedSeconds > 1.0)
+				{
+					stats.ExtractRate = (stats.ExtractSize - lastExtractSize) / elapsedSeconds;
+					stats.DownloadRate = (stats.DownloadSize - lastDownloadSize) / elapsedSeconds;
+
+					lastExtractSize = stats.ExtractSize;
+					lastDownloadSize = stats.DownloadSize;
+
+					timer.Restart();
+				}
+
+				progress.Report(stats);
+			}
+		}
+
+		#endregion
+
 		static string CombinePaths(string basePath, string nextPath)
 		{
 			if (basePath.Length > 0)
