Horde: Fix static analysis warnings.

#rnx

diff --git a/Engine/Source/Programs/Horde/Plugins/Storage/HordeServer.Storage/Storage/StorageService.cs b/Engine/Source/Programs/Horde/Plugins/Storage/HordeServer.Storage/Storage/StorageService.cs
--- a/Engine/Source/Programs/Horde/Plugins/Storage/HordeServer.Storage/Storage/StorageService.cs
+++ b/Engine/Source/Programs/Horde/Plugins/Storage/HordeServer.Storage/Storage/StorageService.cs
@@ -1363,7 +1363,7 @@ namespace HordeServer.Storage
 			Channel<BlobInfo> channel = Channel.CreateBounded<BlobInfo>(new BoundedChannelOptions(1000));
 
 			await using AsyncPipeline pipeline = new AsyncPipeline(cancellationToken);
-			pipeline.AddTask(ctx => FindBlobsForLengthScanAsync(channel.Writer, ctx));
+			_ = pipeline.AddTask(ctx => FindBlobsForLengthScanAsync(channel.Writer, ctx));
 			pipeline.AddTasks(8, channel.Reader, (entry, ctx) => FindBlobLengthsAsync(entry, state, ctx));
 			await pipeline.WaitForCompletionAsync();
 		}
diff --git a/Engine/Source/Programs/Shared/EpicGames.Horde/Storage/BatchReader.cs b/Engine/Source/Programs/Shared/EpicGames.Horde/Storage/BatchReader.cs
--- a/Engine/Source/Programs/Shared/EpicGames.Horde/Storage/BatchReader.cs
+++ b/Engine/Source/Programs/Shared/EpicGames.Horde/Storage/BatchReader.cs
@@ -1,9 +1,7 @@
 // Copyright Epic Games, Inc. All Rights Reserved.
 
 using System;
-using System.Collections;
 using System.Collections.Generic;
-using System.Diagnostics;
 using System.Diagnostics.CodeAnalysis;
 using System.Linq;
 using System.Threading;
@@ -12,7 +10,6 @@ using System.Threading.Tasks;
 using EpicGames.Core;
 using EpicGames.Horde.Storage.Bundles;
 using EpicGames.Horde.Storage.Bundles.V2;
-using OpenTracing.Tag;
 
 namespace EpicGames.Horde.Storage
 {
@@ -161,7 +158,7 @@ namespace EpicGames.Horde.Storage
 		/// <summary>
 		/// Adds a batch reader to the given pipeline
 		/// </summary>
-		internal Task AddToPipeline(AsyncPipeline pipeline, int numReadTasks, int numDecodeTasks)
+		internal void AddToPipeline(AsyncPipeline pipeline, int numReadTasks, int numDecodeTasks)
 		{
 			_ = pipeline.AddTask(CreateBundleRequestsAsync);
 			Task otherTask = pipeline.AddTask(HandleOtherRequestsAsync);
@@ -175,7 +172,7 @@ namespace EpicGames.Horde.Storage
 
 			// Read all the output blobs and mark the blob response channel as complete
 			IEnumerable<Task> blobTasks = packetTasks.Append(otherTask);
-			return Task.WhenAll(blobTasks).ContinueWith(_ => _blobResponses.Writer.TryComplete(), TaskScheduler.Default);
+			_ = Task.WhenAll(blobTasks).ContinueWith(_ => _blobResponses.Writer.TryComplete(), TaskScheduler.Default);
 		}
 
 		// Read a blob using the naive read pipeline
@@ -212,12 +209,12 @@ namespace EpicGames.Horde.Storage
 			for (; ; )
 			{
 				// Fill the queue up to the max length
-				for(;;)
+				for (; ; )
 				{
 					BatchReadRequest<TContext>? request;
 					if (!_blobRequests.Reader.TryRead(out request))
 					{
-						if(queueLength >= MinQueueLength)
+						if (queueLength >= MinQueueLength)
 						{
 							break;
 						}
@@ -329,8 +326,10 @@ namespace EpicGames.Horde.Storage
 					data = await request.Handle.ReadAsync(request.MinOffset, request.MaxOffset - request.MinOffset, cancellationToken);
 					Interlocked.Increment(ref _numReads);
 					Interlocked.Add(ref _bytesRead, data.Memory.Length);
+#pragma warning disable CA2000
 					BundleResponse response = new BundleResponse(data, request.MinOffset, request.MaxOffset, request.Packets);
 					await _bundleResponses.Writer.WriteAsync(response, cancellationToken);
+#pragma warning restore CA2000
 				}
 				catch
 				{
@@ -345,7 +344,7 @@ namespace EpicGames.Horde.Storage
 
 		async Task ReadPacketsAsync(CancellationToken cancellationToken)
 		{
-			await foreach(BundleResponse bundleResponse in _bundleResponses.Reader.ReadAllAsync(cancellationToken))
+			await foreach (BundleResponse bundleResponse in _bundleResponses.Reader.ReadAllAsync(cancellationToken))
 			{
 				using IDisposable lifetime = bundleResponse;
 				foreach (PacketRequest packetRequest in bundleResponse.Packets)
@@ -361,17 +360,17 @@ namespace EpicGames.Horde.Storage
 					BatchReadResponseItem<TContext>[] items = new BatchReadResponseItem<TContext>[packetRequest.Exports.Count];
 					try
 					{
+#pragma warning disable CA2000
 						for (int idx = 0; idx < items.Length; idx++)
 						{
 							ExportRequest exportRequest = packetRequest.Exports[idx];
-#pragma warning disable CA2000
 							BlobData blobData = packetReader.ReadExport(exportRequest.Handle.ExportIdx);
 							items[idx] = new BatchReadResponseItem<TContext>(blobData, exportRequest.Context);
-#pragma warning restore CA2000
 						}
 
 						BatchReadResponse<TContext> response = new BatchReadResponse<TContext>(items);
 						await _blobResponses.Writer.WriteAsync(response, cancellationToken);
+#pragma warning restore CA2000
 					}
 					catch
 					{
@@ -396,7 +395,7 @@ namespace EpicGames.Horde.Storage
 		/// <summary>
 		/// Add tasks for reading to the given async pipeline 
 		/// </summary>
-		public static Task AddBatchReaderTasks<T>(this AsyncPipeline pipeline, int numReadTasks, int numDecodeTasks, BatchReader<T> reader)
+		public static void AddBatchReaderTasks<T>(this AsyncPipeline pipeline, int numReadTasks, int numDecodeTasks, BatchReader<T> reader)
 			=> reader.AddToPipeline(pipeline, numReadTasks, numDecodeTasks);
 	}
 }
diff --git a/Engine/Source/Programs/Shared/EpicGames.Horde/Storage/Nodes/DirectoryNode.Extract.cs b/Engine/Source/Programs/Shared/EpicGames.Horde/Storage/Nodes/DirectoryNode.Extract.cs
--- a/Engine/Source/Programs/Shared/EpicGames.Horde/Storage/Nodes/DirectoryNode.Extract.cs
+++ b/Engine/Source/Programs/Shared/EpicGames.Horde/Storage/Nodes/DirectoryNode.Extract.cs
@@ -1,22 +1,15 @@
 // Copyright Epic Games, Inc. All Rights Reserved.
 
 using System;
-using System.Collections;
 using System.Collections.Generic;
 using System.Diagnostics;
-using System.Diagnostics.CodeAnalysis;
 using System.IO;
 using System.IO.MemoryMappedFiles;
 using System.Linq;
-using System.Net.Security;
-using System.Runtime.ExceptionServices;
-using System.Text;
 using System.Threading;
 using System.Threading.Channels;
 using System.Threading.Tasks;
 using EpicGames.Core;
-using EpicGames.Horde.Storage.Bundles;
-using EpicGames.Horde.Storage.Bundles.V2;
 using Microsoft.Extensions.Logging;
 
 namespace EpicGames.Horde.Storage.Nodes
@@ -65,7 +58,7 @@ namespace EpicGames.Horde.Storage.Nodes
 		/// Whether to print out separate stats for download speed
 		/// </summary>
 		public bool ShowDownloadStats { get; set; }
-		
+
 		/// <summary>
 		/// Constructor
 		/// </summary>
@@ -93,7 +86,7 @@ namespace EpicGames.Horde.Storage.Nodes
 
 			if (_totalCount > 0 && _totalSize > 0)
 			{
-				_logger.LogInformation("Written {NumFiles:n0}/{TotalFiles:n0} files ({Size:n1}/{TotalSize:n1}mb, {Rate:n1}mb/s, {Pct}%). Downloaded {DownloadSize:n1} mb ({DownloadRate:n0} mpbs).", stats.NumFiles, _totalCount, stats.ExtractSize / (1024.0 * 1024.0), _totalSize / (1024.0 * 1024.0), stats.ExtractRate / (1024.0 * 1024.0), (int)((Math.Max(stats.ExtractSize, 1) * 100) / Math.Max(_totalSize, 1)), stats.DownloadSize / (1024.0 * 1024.0));
+				_logger.LogInformation("Written {NumFiles:n0}/{TotalFiles:n0} files ({Size:n1}/{TotalSize:n1}mb, {Rate:n1}mb/s, {Pct}%)", stats.NumFiles, _totalCount, stats.ExtractSize / (1024.0 * 1024.0), _totalSize / (1024.0 * 1024.0), stats.ExtractRate / (1024.0 * 1024.0), (int)((Math.Max(stats.ExtractSize, 1) * 100) / Math.Max(_totalSize, 1)));
 			}
 			else if (_totalCount > 0)
 			{
@@ -122,7 +115,7 @@ namespace EpicGames.Horde.Storage.Nodes
 		/// <param name="directoryInfo"></param>
 		/// <param name="logger"></param>
 		/// <param name="cancellationToken"></param>
-		public static Task ExtractAsync(this IBlobRef<DirectoryNode> directoryRef, DirectoryInfo directoryInfo, ILogger logger, CancellationToken cancellationToken) 
+		public static Task ExtractAsync(this IBlobRef<DirectoryNode> directoryRef, DirectoryInfo directoryInfo, ILogger logger, CancellationToken cancellationToken)
 			=> ExtractAsync(directoryRef, directoryInfo, null, logger, cancellationToken);
 
 		class OutputStats
@@ -315,14 +308,14 @@ namespace EpicGames.Horde.Storage.Nodes
 				using BatchReader<OutputChunk> batchReader = new BatchReader<OutputChunk>();
 
 				_ = pipeline.AddTask(ctx => FindOutputChunksRootAsync(directoryInfo, directoryNode, batchReader.RequestWriter, logger, ctx));
-				_ = pipeline.AddBatchReaderTasks(numReadTasks, numDecodeTasks, batchReader);
+				pipeline.AddBatchReaderTasks(numReadTasks, numDecodeTasks, batchReader);
 
 				OutputStats outputStats = new OutputStats();
 				Task[] writeTasks = pipeline.AddTasks(numWriteTasks, ctx => WriteAsync(batchReader.ResponseReader, outputStats, logger, ctx));
 
 				if (progress != null)
 				{
-					_ = pipeline.AddTask(ctx => UpdateStatsAsync(batchReader, outputStats, Task.WhenAll(writeTasks), progress, ctx));
+					_ = pipeline.AddTask(ctx => UpdateStatsAsync(batchReader, outputStats, Task.WhenAll(writeTasks), progress, frequency, ctx));
 				}
 
 				await pipeline.WaitForCompletionAsync();
@@ -506,7 +499,7 @@ namespace EpicGames.Horde.Storage.Nodes
 				file.FileInfo.Refresh();
 				FileEntry.SetPermissions(file.FileInfo!, file.FileEntry.Flags);
 
-				if((file.FileEntry.Flags & FileEntryFlags.HasModTime) != 0)
+				if ((file.FileEntry.Flags & FileEntryFlags.HasModTime) != 0)
 				{
 					file.FileInfo.LastWriteTimeUtc = file.FileEntry.ModTime;
 				}
@@ -550,7 +543,7 @@ namespace EpicGames.Horde.Storage.Nodes
 			public double DownloadRate { get; set; }
 		}
 
-		static async Task UpdateStatsAsync(BatchReader<OutputChunk> batchReader, OutputStats outputStats, Task writeTask, IProgress<IExtractStats> progress, CancellationToken cancellationToken)
+		static async Task UpdateStatsAsync(BatchReader<OutputChunk> batchReader, OutputStats outputStats, Task writeTask, IProgress<IExtractStats> progress, TimeSpan frequency, CancellationToken cancellationToken)
 		{
 			Stopwatch timer = Stopwatch.StartNew();
 
@@ -560,7 +553,7 @@ namespace EpicGames.Horde.Storage.Nodes
 			Task? completeTask = null;
 			while (completeTask != writeTask)
 			{
-				completeTask = await Task.WhenAny(writeTask, Task.Delay(TimeSpan.FromSeconds(5.0)));
+				completeTask = await Task.WhenAny(writeTask, Task.Delay(frequency, cancellationToken));
 
 				ExtractStats stats = new ExtractStats();
 				stats.NumFiles = Interlocked.CompareExchange(ref outputStats._writtenFiles, 0, 0);
diff --git a/Engine/Source/Programs/UnrealGameSync/UnrealGameSync/ArtifactDownload.cs b/Engine/Source/Programs/UnrealGameSync/UnrealGameSync/ArtifactDownload.cs
--- a/Engine/Source/Programs/UnrealGameSync/UnrealGameSync/ArtifactDownload.cs
+++ b/Engine/Source/Programs/UnrealGameSync/UnrealGameSync/ArtifactDownload.cs
@@ -31,7 +31,7 @@ namespace UnrealGameSync
 
 			public void Report(IExtractStats value)
 			{
-				_inner.Report($"Copied {value.Count} files ({value.Size / (1024.0 * 1024.0):n1}mb, {value.Rate / (1024.0 * 1024.0):n1}mb/s)...");
+				_inner.Report($"Copied {value.NumFiles} files ({value.ExtractSize / (1024.0 * 1024.0):n1}mb, {value.ExtractRate / (1024.0 * 1024.0):n1}mb/s)...");
 			}
 		}
 
