[UHT]
* Fixed static analysis compile warnings

diff --git a/Engine/Source/Programs/Shared/EpicGames.UHT/Exporters/CodeGen/UhtHeaderCodeGeneratorHFile.cs b/Engine/Source/Programs/Shared/EpicGames.UHT/Exporters/CodeGen/UhtHeaderCodeGeneratorHFile.cs
--- a/Engine/Source/Programs/Shared/EpicGames.UHT/Exporters/CodeGen/UhtHeaderCodeGeneratorHFile.cs
+++ b/Engine/Source/Programs/Shared/EpicGames.UHT/Exporters/CodeGen/UhtHeaderCodeGeneratorHFile.cs
@@ -433,9 +433,9 @@ namespace EpicGames.UHT.Exporters.CodeGen
 
 						builder.Append("\tuint32 ModifiedFlags = 0; \\\r\n");
 
-						builder.Append("\t").Append(scriptStruct.SourceName).Append("() = default; \\\r\n");
+						builder.Append('\t').Append(scriptStruct.SourceName).Append("() = default; \\\r\n");
 
-						builder.Append("\t").Append(scriptStruct.SourceName).Append("(FStateStreamCopyContext& Context, const ").Append(scriptStruct.SourceName).Append("& Other) \\\r\n");
+						builder.Append('\t').Append(scriptStruct.SourceName).Append("(FStateStreamCopyContext& Context, const ").Append(scriptStruct.SourceName).Append("& Other) \\\r\n");
 						bool isFirst = true;
 
 						builder.Append(isFirst ? "\t:\t" : "\t,\t").Append("ModifiedFlags(Other.ModifiedFlags) \\\r\n");
@@ -445,7 +445,7 @@ namespace EpicGames.UHT.Exporters.CodeGen
 
 						foreach (UhtProperty property in scriptStruct.Properties)
 						{
-							builder.Append(isFirst ? "\t:\t" : "\t,\t").Append(property.SourceName).Append("(");
+							builder.Append(isFirst ? "\t:\t" : "\t,\t").Append(property.SourceName).Append('(');
 							if (property is UhtStructProperty structProperty)
 							{
 								if (structProperty.ScriptStruct.MetaData.ContainsKey(UhtNames.StateStreamHandle))
