Horde: Prevent missing refs causing deduped artifact upload to fail.

#rnx

diff --git a/Engine/Source/Programs/AutomationTool/BuildGraph/Tasks/CreateArtifactTask.cs b/Engine/Source/Programs/AutomationTool/BuildGraph/Tasks/CreateArtifactTask.cs
--- a/Engine/Source/Programs/AutomationTool/BuildGraph/Tasks/CreateArtifactTask.cs
+++ b/Engine/Source/Programs/AutomationTool/BuildGraph/Tasks/CreateArtifactTask.cs
@@ -205,20 +205,30 @@ namespace AutomationTool.Tasks
 				if (_parameters.Dedupe)
 				{
 					await using DedupeBlobWriter dedupeWriter = new DedupeBlobWriter(blobWriter);
-
-					IArtifact? prevArtifact = await hordeClient.Artifacts.FindAsync(streamId, maxCommitId: commitId, name: name, type: type).FirstOrDefaultAsync(x => x.CommitId != commitId);
-					if (prevArtifact == null)
+					try
 					{
-						Logger.LogInformation("Could not find a previous artifact; dedupe will be disabled");
+						IArtifact? prevArtifact = await hordeClient.Artifacts.FindAsync(streamId, maxCommitId: commitId, name: name, type: type).FirstOrDefaultAsync(x => x.CommitId != commitId);
+						if (prevArtifact == null)
+						{
+							Logger.LogInformation("Could not find a previous artifact; dedupe will be disabled");
+						}
+						else if (prevArtifact.NamespaceId != artifact.NamespaceId)
+						{
+							Logger.LogInformation("Previous artifact {ArtifactId} is from namespace {PrevNamespaceId}, not {NamespaceId}. Will disable dedupe.", artifact.Id, prevArtifact.NamespaceId, artifact.NamespaceId);
+						}
+						else
+						{
+							Logger.LogInformation("Caching blobs from previous artifact ({ArtifactId})", prevArtifact.Id);
+							await dedupeWriter.AddToCacheAsync(prevArtifact.Content);
+						}
 					}
-					else if (prevArtifact.NamespaceId != artifact.NamespaceId)
+					catch (RefNameNotFoundException ex)
 					{
-						Logger.LogInformation("Previous artifact {ArtifactId} is from namespace {PrevNamespaceId}, not {NamespaceId}. Will disable dedupe.", artifact.Id, prevArtifact.NamespaceId, artifact.NamespaceId);
+						Logger.LogInformation(ex, "Unable to read ref for previous artifact ({RefName}).", ex.Name);
 					}
-					else
+					catch (Exception ex)
 					{
-						Logger.LogInformation("Caching blobs from previous artifact ({ArtifactId})", prevArtifact.Id);
-						await dedupeWriter.AddToCacheAsync(prevArtifact.Content);
+						Logger.LogInformation(ex, "Unable to read previous artifact: {Message}", ex.Message);
 					}
 
 					rootRef = await dedupeWriter.WriteFilesAsync(baseDir, files);
diff --git a/Engine/Source/Programs/Shared/EpicGames.Horde/Storage/Backends/HttpStorageBackend.cs b/Engine/Source/Programs/Shared/EpicGames.Horde/Storage/Backends/HttpStorageBackend.cs
--- a/Engine/Source/Programs/Shared/EpicGames.Horde/Storage/Backends/HttpStorageBackend.cs
+++ b/Engine/Source/Programs/Shared/EpicGames.Horde/Storage/Backends/HttpStorageBackend.cs
@@ -97,7 +97,7 @@ namespace EpicGames.Horde.Storage.Backends
 						}
 
 						HttpResponseMessage response = await httpClient.SendAsync(request, cancellationToken);
-						response.EnsureSuccessStatusCode();
+						await EnsureSuccessAsync(response, cancellationToken);
 
 						Stream stream = await response.Content.ReadAsStreamAsync(cancellationToken);
 
@@ -265,7 +265,7 @@ namespace EpicGames.Horde.Storage.Backends
 				request.AddAliases.Add(new AddAliasRequest { Name = name, Target = target, Rank = rank, Data = data.ToArray() });
 
 				using HttpResponseMessage response = await httpClient.PostAsJsonAsync($"{_basePath}", request, cancellationToken: cancellationToken);
-				response.EnsureSuccessStatusCode();
+				await EnsureSuccessAsync(response, cancellationToken);
 			}
 		}
 
@@ -278,7 +278,7 @@ namespace EpicGames.Horde.Storage.Backends
 				request.RemoveAliases.Add(new RemoveAliasRequest { Name = name, Target = target });
 
 				using HttpResponseMessage response = await httpClient.PostAsJsonAsync($"{_basePath}", request, cancellationToken: cancellationToken);
-				response.EnsureSuccessStatusCode();
+				await EnsureSuccessAsync(response, cancellationToken);
 			}
 		}
 
@@ -298,7 +298,7 @@ namespace EpicGames.Horde.Storage.Backends
 				{
 					using (HttpResponseMessage response = await httpClient.SendAsync(request, cancellationToken))
 					{
-						response.EnsureSuccessStatusCode();
+						await EnsureSuccessAsync(response, cancellationToken);
 
 						FindNodesResponse? message = await response.Content.ReadFromJsonAsync<FindNodesResponse>(cancellationToken: cancellationToken);
 
@@ -338,7 +338,7 @@ namespace EpicGames.Horde.Storage.Backends
 							return false;
 						}
 
-						response.EnsureSuccessStatusCode();
+						await EnsureSuccessAsync(response, cancellationToken);
 						return false;
 					}
 				}
@@ -371,7 +371,8 @@ namespace EpicGames.Horde.Storage.Backends
 						}
 						else
 						{
-							response.EnsureSuccessStatusCode();
+							await EnsureSuccessAsync(response, cancellationToken);
+
 							ReadRefResponse? data = await response.Content.ReadFromJsonAsync<ReadRefResponse>(cancellationToken: cancellationToken);
 							_logger.LogDebug("Read ref {RefName} -> {Hash} / {Locator}", name, data!.Hash, data!.Target);
 
@@ -395,7 +396,7 @@ namespace EpicGames.Horde.Storage.Backends
 
 				using (HttpResponseMessage response = await httpClient.PutAsync($"{_basePath}/refs/{name}", request, cancellationToken))
 				{
-					response.EnsureSuccessStatusCode();
+					await EnsureSuccessAsync(response, cancellationToken);
 				}
 			}
 		}
@@ -413,6 +414,15 @@ namespace EpicGames.Horde.Storage.Backends
 				stats.Add("backend.http.concurrency_ratio", (long)(_sequentialReadTime.TotalSeconds * 100.0 / _readTime.Elapsed.TotalSeconds));
 			}
 		}
+
+		static async Task EnsureSuccessAsync(HttpResponseMessage message, CancellationToken cancellationToken)
+		{
+			if (!message.IsSuccessStatusCode)
+			{
+				string response = await message.Content.ReadAsStringAsync(cancellationToken);
+				throw new StorageException($"Http response {message.StatusCode} (Content: {response})");
+			}
+		}
 	}
 
 	/// <summary>
