[Horde]
* Changed text from "Error" to "Exception" to prevent that exception from failing submittool (the build still passes with this "error")

diff --git a/Engine/Source/Programs/Shared/EpicGames.Horde/Compute/ComputeSocket.cs b/Engine/Source/Programs/Shared/EpicGames.Horde/Compute/ComputeSocket.cs
--- a/Engine/Source/Programs/Shared/EpicGames.Horde/Compute/ComputeSocket.cs
+++ b/Engine/Source/Programs/Shared/EpicGames.Horde/Compute/ComputeSocket.cs
@@ -538,7 +538,7 @@ namespace EpicGames.Horde.Compute
 			}
 			catch (Exception e)
 			{
-				_logger.LogInformation(e, "Error in background receive");
+				_logger.LogInformation(e, "Exception in background receive");
 				throw;
 			}
 		}
