OidcToken: Allow HTTP for discovery endpoints

When Horde runs with local accounts, it uses the built-in OIDC sign on. This auth mode is often deployed without HTTPS which is why removing this requirement is useful.

diff --git a/Engine/Source/Programs/Shared/EpicGames.OIDC/OidcTokenManager.cs b/Engine/Source/Programs/Shared/EpicGames.OIDC/OidcTokenManager.cs
--- a/Engine/Source/Programs/Shared/EpicGames.OIDC/OidcTokenManager.cs
+++ b/Engine/Source/Programs/Shared/EpicGames.OIDC/OidcTokenManager.cs
@@ -611,7 +611,8 @@ namespace EpicGames.OIDC
 				Address = discoUrl,
 				Policy =
 				{
-					ValidateEndpoints = false
+					ValidateEndpoints = false,
+					RequireHttps = _authorityUri.Scheme == "https"
 				}
 			};
 			DiscoveryDocumentResponse disco = await client.GetDiscoveryDocumentAsync(doc, cancellationToken);
