EpicGames.Perforce: Assume and mark workspace having untracked files prior to starting a sync

diff --git a/Engine/Source/Programs/Shared/EpicGames.Perforce.Managed/ManagedWorkspace.cs b/Engine/Source/Programs/Shared/EpicGames.Perforce.Managed/ManagedWorkspace.cs
--- a/Engine/Source/Programs/Shared/EpicGames.Perforce.Managed/ManagedWorkspace.cs
+++ b/Engine/Source/Programs/Shared/EpicGames.Perforce.Managed/ManagedWorkspace.cs
@@ -971,6 +971,10 @@ namespace EpicGames.Perforce.Managed
 					{
 						await CleanInternalAsync(removeUntracked, cancellationToken);
 					}
+					
+					// Immediately after a potential clean, mark repository having untracked files.
+					// This ensure a clean is run next time even if the sync below is aborted or fails
+					await SaveUntrackedStateAsync(perforce, true, cancellationToken);
 
 					// Wait for the have table update to finish
 					await updateHaveTableTask;
@@ -1000,9 +1004,6 @@ namespace EpicGames.Perforce.Managed
 				await AddFilesToWorkspaceAsync(perforce, contents, fakeSync, cancellationToken);
 			}
 			
-			// Once sync is complete and about to be handed off, mark repository having untracked files
-			await SaveUntrackedStateAsync(perforce, true, cancellationToken);
-			
 			_logger.LogInformation("Completed in {ElapsedTime}s", $"{timer.Elapsed.TotalSeconds:0.0}");
 		}
 
