Added helper function to CaptureLogger that sends the captured log to another logger.

#rnx

diff --git a/Engine/Source/Programs/Shared/EpicGames.Core/Log.cs b/Engine/Source/Programs/Shared/EpicGames.Core/Log.cs
--- a/Engine/Source/Programs/Shared/EpicGames.Core/Log.cs
+++ b/Engine/Source/Programs/Shared/EpicGames.Core/Log.cs
@@ -869,6 +869,12 @@ namespace EpicGames.Core
 		/// <returns>List of rendered log lines</returns>
 		public List<string> RenderLines() => Events.ConvertAll(x => x.ToString());
 
+		/// <summary>
+		/// Renders all the captured events to another logger
+		/// </summary>
+		/// <param name="Logger"></param>
+		public void RenderTo(ILogger logger) => Events.ForEach( x => logger.Log(x.Level, x.Id, x, null, (s, e) => s.ToString()) );
+
 		/// <inheritdoc/>
 		public IDisposable BeginScope<TState>(TState state) => _scopeCollection.BeginScope(state);
 
