[Trivial] fix xml comment

#rnx

diff --git a/Engine/Source/Programs/Shared/EpicGames.Core/Log.cs b/Engine/Source/Programs/Shared/EpicGames.Core/Log.cs
--- a/Engine/Source/Programs/Shared/EpicGames.Core/Log.cs
+++ b/Engine/Source/Programs/Shared/EpicGames.Core/Log.cs
@@ -872,7 +872,7 @@ namespace EpicGames.Core
 		/// <summary>
 		/// Renders all the captured events to another logger
 		/// </summary>
-		/// <param name="Logger"></param>
+		/// <param name="logger"></param>
 		public void RenderTo(ILogger logger) => Events.ForEach( x => logger.Log(x.Level, x.Id, x, null, (s, e) => s.ToString()) );
 
 		/// <inheritdoc/>
