EpicGames.UBA: Add an event handler for when a disallowed path is registered

#rnx

diff --git a/Engine/Source/Programs/Shared/EpicGames.UBA/EpicGames.UBA.csproj b/Engine/Source/Programs/Shared/EpicGames.UBA/EpicGames.UBA.csproj
--- a/Engine/Source/Programs/Shared/EpicGames.UBA/EpicGames.UBA.csproj
+++ b/Engine/Source/Programs/Shared/EpicGames.UBA/EpicGames.UBA.csproj
@@ -37,6 +37,12 @@
     <DocumentationFile></DocumentationFile>
   </PropertyGroup>
 
+  <ItemGroup>
+    <Compile Include="../../../../Restricted/*/Source/Programs/Shared/EpicGames.UBA/**/*.cs">
+      <Link>$([System.Text.RegularExpressions.Regex]::Replace(%(FullPath), '^.+?[\\/]EpicGames.UBA[\\/]', ''))</Link>
+    </Compile>
+  </ItemGroup>
+
   <ItemGroup>
     <Compile Include="..\MetaData.cs">
       <Link>Properties\MetaData.cs</Link>
diff --git a/Engine/Source/Programs/Shared/EpicGames.UBA/Utils.cs b/Engine/Source/Programs/Shared/EpicGames.UBA/Utils.cs
--- a/Engine/Source/Programs/Shared/EpicGames.UBA/Utils.cs
+++ b/Engine/Source/Programs/Shared/EpicGames.UBA/Utils.cs
@@ -11,9 +11,8 @@ namespace EpicGames.UBA
 	/// <summary>
 	/// Utils
 	/// </summary>
-	public static class Utils
+	public static partial class Utils
 	{
-
 		/// <summary>
 		/// Is UBA available?
 		/// </summary>
@@ -30,8 +29,24 @@ namespace EpicGames.UBA
 		/// <summary>
 		/// Registers a path that is not allowed to be transferred over the network for UBA remote agents.
 		/// </summary>
-		/// <param name="paths">The paths to add to thie disallowed list</param>
-		public static void RegisterDisallowedPaths(params string[] paths) => s_disallowedPaths.UnionWith(paths);
+		/// <param name="paths">The paths to add to the disallowed list</param>
+		public static void RegisterDisallowedPaths(params string[] paths)
+		{
+			s_disallowedPaths.UnionWith(paths);
+			DisallowedPathRegistered?.Invoke(s_disallowedPaths, new(paths));
+		}
+
+		/// <summary>
+		/// Delegate for registering a remote disallowed path
+		/// </summary>
+		/// <param name="sender">collection that is being changed</param>
+		/// <param name="e">event args containing which paths were added</param>
+		public delegate void DisallowedPathRegisteredEventHandler(SortedSet<string> sender, DisallowedPathRegisteredEventArgs e);
+
+		/// <summary>
+		/// Remote disallowed path registered event handler
+		/// </summary>
+		public static event DisallowedPathRegisteredEventHandler? DisallowedPathRegistered;
 
 		/// <summary>
 		/// Get the path to the p/invoke library that would be loaded 
@@ -58,4 +73,15 @@ namespace EpicGames.UBA
 			throw new PlatformNotSupportedException();
 		}
 	}
+
+	/// <summary>
+	/// Event args for registering a remote disallowed path
+	/// </summary>
+	public class DisallowedPathRegisteredEventArgs(params string[] paths) : EventArgs
+	{
+		/// <summary>
+		/// The paths being registered
+		/// </summary>
+		public IEnumerable<string> Paths { get; } = paths;
+	}
 }
diff --git a/Engine/Source/Programs/UnrealBuildTool/Executors/UnrealBuildAccelerator/UBAAgentCoordinatorHorde.cs b/Engine/Source/Programs/UnrealBuildTool/Executors/UnrealBuildAccelerator/UBAAgentCoordinatorHorde.cs
--- a/Engine/Source/Programs/UnrealBuildTool/Executors/UnrealBuildAccelerator/UBAAgentCoordinatorHorde.cs
+++ b/Engine/Source/Programs/UnrealBuildTool/Executors/UnrealBuildAccelerator/UBAAgentCoordinatorHorde.cs
@@ -553,7 +553,7 @@ namespace UnrealBuildTool
 						//arguments.Add("-NoStore");
 						//arguments.Add("-KillRandom"); // For debugging
 
-						if (System.OperatingSystem.IsMacOS())
+						if (OperatingSystem.IsMacOS() && EpicGames.UBA.Utils.DisallowedPaths().Any())
 						{
 							// we need to populate the cas with all known xcodes so we can serve all the ones we have installed
 							string xcodeVersion = Utils.RunLocalProcessAndReturnStdOut("/bin/sh", "-c '/usr/bin/defaults read $(xcode-select -p)/../version.plist ProductBuildVersion");
