EpicGames.Core: Pin System.Text.Json to 8.0.5

#rnx

diff --git a/Engine/Source/Programs/Shared/EpicGames.Core/EpicGames.Core.csproj b/Engine/Source/Programs/Shared/EpicGames.Core/EpicGames.Core.csproj
--- a/Engine/Source/Programs/Shared/EpicGames.Core/EpicGames.Core.csproj
+++ b/Engine/Source/Programs/Shared/EpicGames.Core/EpicGames.Core.csproj
@@ -59,5 +59,6 @@
     <PackageReference Include="Microsoft.Extensions.ObjectPool" Version="6.0.8" />
     <PackageReference Include="OpenTracing" Version="0.12.1" />
     <PackageReference Include="System.Memory" Version="4.5.5" />
+    <PackageReference Include="System.Text.Json" Version="8.0.5" />
   </ItemGroup>
 </Project>
\ No newline at end of file
