Update Logging nuget packages

#rnx

diff --git a/Engine/Source/Programs/Horde/Drivers/JobDriver.Tests/JobDriver.Tests.csproj b/Engine/Source/Programs/Horde/Drivers/JobDriver.Tests/JobDriver.Tests.csproj
--- a/Engine/Source/Programs/Horde/Drivers/JobDriver.Tests/JobDriver.Tests.csproj
+++ b/Engine/Source/Programs/Horde/Drivers/JobDriver.Tests/JobDriver.Tests.csproj
@@ -19,7 +19,7 @@
   </PropertyGroup>
 
   <ItemGroup>
-    <PackageReference Include="Microsoft.Extensions.Logging.Abstractions" Version="8.0.0" />
+    <PackageReference Include="Microsoft.Extensions.Logging.Abstractions" Version="8.0.2" />
     <PackageReference Include="Microsoft.NET.Test.Sdk" Version="17.11.1" />
     <PackageReference Include="Moq" Version="4.20.70" />
     <PackageReference Include="MSTest.TestAdapter" Version="3.6.0" />
diff --git a/Engine/Source/Programs/Horde/Drivers/JobDriver/JobDriver.csproj b/Engine/Source/Programs/Horde/Drivers/JobDriver/JobDriver.csproj
--- a/Engine/Source/Programs/Horde/Drivers/JobDriver/JobDriver.csproj
+++ b/Engine/Source/Programs/Horde/Drivers/JobDriver/JobDriver.csproj
@@ -71,7 +71,7 @@
     <PackageReference Include="Microsoft.Extensions.Hosting.WindowsServices" Version="8.0.0" />
     <PackageReference Include="Microsoft.Extensions.Http" Version="8.0.0" />
     <PackageReference Include="Microsoft.Extensions.Http.Polly" Version="8.0.0" />
-    <PackageReference Include="Microsoft.Extensions.Logging.Abstractions" Version="8.0.0" />
+    <PackageReference Include="Microsoft.Extensions.Logging.Abstractions" Version="8.0.2" />
     <PackageReference Include="Microsoft.Extensions.Options.DataAnnotations" Version="8.0.0" />
     <PackageReference Include="Microsoft.Management.Infrastructure" Version="3.0.0" />
     <PackageReference Include="Microsoft.VisualStudio.Threading.Analyzers" Version="17.8.14">
diff --git a/Engine/Source/Programs/Horde/HordeAgent.Tests/HordeAgent.Tests.csproj b/Engine/Source/Programs/Horde/HordeAgent.Tests/HordeAgent.Tests.csproj
--- a/Engine/Source/Programs/Horde/HordeAgent.Tests/HordeAgent.Tests.csproj
+++ b/Engine/Source/Programs/Horde/HordeAgent.Tests/HordeAgent.Tests.csproj
@@ -20,7 +20,7 @@
   </PropertyGroup>
 
   <ItemGroup>
-    <PackageReference Include="Microsoft.Extensions.Logging.Abstractions" Version="8.0.0" />
+    <PackageReference Include="Microsoft.Extensions.Logging.Abstractions" Version="8.0.2" />
     <PackageReference Include="Microsoft.NET.Test.Sdk" Version="17.11.1" />
     <PackageReference Include="Moq" Version="4.20.70" />
     <PackageReference Include="MSTest.TestAdapter" Version="3.6.0" />
diff --git a/Engine/Source/Programs/Horde/HordeAgent/HordeAgent.csproj b/Engine/Source/Programs/Horde/HordeAgent/HordeAgent.csproj
--- a/Engine/Source/Programs/Horde/HordeAgent/HordeAgent.csproj
+++ b/Engine/Source/Programs/Horde/HordeAgent/HordeAgent.csproj
@@ -53,7 +53,7 @@
     <PackageReference Include="Microsoft.Extensions.Hosting.WindowsServices" Version="8.0.0" />
     <PackageReference Include="Microsoft.Extensions.Http" Version="8.0.0" />
     <PackageReference Include="Microsoft.Extensions.Http.Polly" Version="8.0.0" />
-    <PackageReference Include="Microsoft.Extensions.Logging.Abstractions" Version="8.0.0" />
+    <PackageReference Include="Microsoft.Extensions.Logging.Abstractions" Version="8.0.2" />
     <PackageReference Include="Microsoft.Extensions.Options.DataAnnotations" Version="8.0.0" />
     <PackageReference Include="Microsoft.Management.Infrastructure" Version="3.0.0" />
     <PackageReference Include="Microsoft.VisualStudio.Threading.Analyzers" Version="17.8.14">
diff --git a/Engine/Source/Programs/Horde/HordeServer/HordeServer.csproj b/Engine/Source/Programs/Horde/HordeServer/HordeServer.csproj
--- a/Engine/Source/Programs/Horde/HordeServer/HordeServer.csproj
+++ b/Engine/Source/Programs/Horde/HordeServer/HordeServer.csproj
@@ -53,7 +53,7 @@
     <PackageReference Include="Microsoft.Extensions.Hosting.WindowsServices" Version="8.0.0" />
     <PackageReference Include="Microsoft.Extensions.Http" Version="8.0.0" />
     <PackageReference Include="Microsoft.Extensions.Http.Polly" Version="8.0.0" />
-    <PackageReference Include="Microsoft.Extensions.Logging" Version="8.0.0" />
+    <PackageReference Include="Microsoft.Extensions.Logging" Version="8.0.1" />
     <PackageReference Include="Microsoft.VisualStudio.Threading.Analyzers" Version="17.8.14">
       <PrivateAssets>all</PrivateAssets>
       <IncludeAssets>runtime; build; native; contentfiles; analyzers; buildtransitive</IncludeAssets>
diff --git a/Engine/Source/Programs/Horde/Plugins/Build/HordeServer.Build/Jobs/JobTaskSource.cs b/Engine/Source/Programs/Horde/Plugins/Build/HordeServer.Build/Jobs/JobTaskSource.cs
--- a/Engine/Source/Programs/Horde/Plugins/Build/HordeServer.Build/Jobs/JobTaskSource.cs
+++ b/Engine/Source/Programs/Horde/Plugins/Build/HordeServer.Build/Jobs/JobTaskSource.cs
@@ -768,7 +768,7 @@ namespace HordeServer.Jobs
 			IAgent agent = waiter.Agent;
 			LeaseId leaseId = new(BinaryIdUtils.CreateNew()); // Generate a new unique id for the lease
 
-			using IDisposable logScope = _logger
+			using IDisposable? logScope = _logger
 				.WithProperty("JobId", job.Id.ToString())
 				.WithProperty("BatchId", batch.Id.ToString())
 				.WithProperty("AgentId", agent.Id.ToString())
diff --git a/Engine/Source/Programs/Horde/Tools/ReleaseNotes/ReleaseNotes.csproj b/Engine/Source/Programs/Horde/Tools/ReleaseNotes/ReleaseNotes.csproj
--- a/Engine/Source/Programs/Horde/Tools/ReleaseNotes/ReleaseNotes.csproj
+++ b/Engine/Source/Programs/Horde/Tools/ReleaseNotes/ReleaseNotes.csproj
@@ -9,7 +9,7 @@
   </PropertyGroup>
 
   <ItemGroup>
-    <PackageReference Include="Microsoft.Extensions.Logging.Console" Version="8.0.0" />
+    <PackageReference Include="Microsoft.Extensions.Logging.Console" Version="8.0.1" />
     <PackageReference Include="System.Text.Json" Version="8.0.5" />
   </ItemGroup>
 
diff --git a/Engine/Source/Programs/Shared/EpicGames.Core.Tests/ILoggerPropertyTests.cs b/Engine/Source/Programs/Shared/EpicGames.Core.Tests/ILoggerPropertyTests.cs
--- a/Engine/Source/Programs/Shared/EpicGames.Core.Tests/ILoggerPropertyTests.cs
+++ b/Engine/Source/Programs/Shared/EpicGames.Core.Tests/ILoggerPropertyTests.cs
@@ -33,7 +33,7 @@ public class SpyLogger : ILogger
 		return true;
 	}
 
-	public IDisposable BeginScope<TState>(TState state)
+	public IDisposable? BeginScope<TState>(TState state) where TState : notnull
 	{
 		StubLogScope scope = new (state);
 		ScopesStarted.Add(scope);
diff --git a/Engine/Source/Programs/Shared/EpicGames.Core/DefaultConsoleLogger.cs b/Engine/Source/Programs/Shared/EpicGames.Core/DefaultConsoleLogger.cs
--- a/Engine/Source/Programs/Shared/EpicGames.Core/DefaultConsoleLogger.cs
+++ b/Engine/Source/Programs/Shared/EpicGames.Core/DefaultConsoleLogger.cs
@@ -48,7 +48,7 @@ namespace EpicGames.Core
 		public bool IsEnabled(LogLevel logLevel) => logLevel >= _minLevel;
 
 		/// <inheritdoc/>
-		public IDisposable BeginScope<TState>(TState state) => new Disposable();
+		public IDisposable? BeginScope<TState>(TState state) where TState : notnull => new Disposable();
 
 		/// <inheritdoc/>
 		public void Log<TState>(LogLevel logLevel, EventId eventId, TState state, Exception? exception, Func<TState, Exception?, string> formatter)
diff --git a/Engine/Source/Programs/Shared/EpicGames.Core/EpicGames.Core.csproj b/Engine/Source/Programs/Shared/EpicGames.Core/EpicGames.Core.csproj
--- a/Engine/Source/Programs/Shared/EpicGames.Core/EpicGames.Core.csproj
+++ b/Engine/Source/Programs/Shared/EpicGames.Core/EpicGames.Core.csproj
@@ -55,7 +55,7 @@
     <PackageReference Include="JetBrains.Annotations" Version="2022.1.0" />
     <PackageReference Include="Microsoft.Extensions.Configuration" Version="8.0.0" />
     <PackageReference Include="Microsoft.Extensions.Http.Polly" Version="6.0.26" />
-    <PackageReference Include="Microsoft.Extensions.Logging.Console" Version="6.0.0" />
+    <PackageReference Include="Microsoft.Extensions.Logging.Console" Version="8.0.1" />
     <PackageReference Include="Microsoft.Extensions.ObjectPool" Version="6.0.8" />
     <PackageReference Include="OpenTracing" Version="0.12.1" />
     <PackageReference Include="System.Memory" Version="4.5.5" />
diff --git a/Engine/Source/Programs/Shared/EpicGames.Core/ForwardingLogger.cs b/Engine/Source/Programs/Shared/EpicGames.Core/ForwardingLogger.cs
--- a/Engine/Source/Programs/Shared/EpicGames.Core/ForwardingLogger.cs
+++ b/Engine/Source/Programs/Shared/EpicGames.Core/ForwardingLogger.cs
@@ -44,7 +44,7 @@ namespace EpicGames.Core
 		}
 
 		/// <inheritdoc/>
-		public IDisposable BeginScope<TState>(TState state)
+		public IDisposable? BeginScope<TState>(TState state) where TState : notnull
 		{
 			DisposableList scope = new DisposableList();
 			foreach (ILogger logger in Loggers)
diff --git a/Engine/Source/Programs/Shared/EpicGames.Core/ILoggerIndent.cs b/Engine/Source/Programs/Shared/EpicGames.Core/ILoggerIndent.cs
--- a/Engine/Source/Programs/Shared/EpicGames.Core/ILoggerIndent.cs
+++ b/Engine/Source/Programs/Shared/EpicGames.Core/ILoggerIndent.cs
@@ -172,7 +172,7 @@ namespace EpicGames.Core
 		}
 
 		/// <inheritdoc/>
-		public IDisposable BeginScope<TState>(TState state)
+		public IDisposable? BeginScope<TState>(TState state) where TState : notnull
 		{
 			ILoggerIndent? indent = state as ILoggerIndent;
 			if (indent != null)
@@ -234,7 +234,7 @@ namespace EpicGames.Core
 		/// <param name="logger">Logger interface</param>
 		/// <param name="indent">The indent to apply</param>
 		/// <returns>Disposable object</returns>
-		public static IDisposable BeginIndentScope(this ILogger logger, string indent)
+		public static IDisposable? BeginIndentScope(this ILogger logger, string indent)
 		{
 			return logger.BeginScope(new LoggerIndent(indent));
 		}
diff --git a/Engine/Source/Programs/Shared/EpicGames.Core/ILoggerProperty.cs b/Engine/Source/Programs/Shared/EpicGames.Core/ILoggerProperty.cs
--- a/Engine/Source/Programs/Shared/EpicGames.Core/ILoggerProperty.cs
+++ b/Engine/Source/Programs/Shared/EpicGames.Core/ILoggerProperty.cs
@@ -46,7 +46,7 @@ namespace EpicGames.Core
 			/// Starts new scope on underlying logger, passing previously set properties
 			/// </summary>
 			/// <returns>Original IDisposable from ILogger.BeginScope()</returns>
-			public IDisposable BeginScope()
+			public IDisposable? BeginScope()
 			{
 				return _logger.BeginScope(_properties);
 			}
diff --git a/Engine/Source/Programs/Shared/EpicGames.Core/Log.cs b/Engine/Source/Programs/Shared/EpicGames.Core/Log.cs
--- a/Engine/Source/Programs/Shared/EpicGames.Core/Log.cs
+++ b/Engine/Source/Programs/Shared/EpicGames.Core/Log.cs
@@ -876,7 +876,7 @@ namespace EpicGames.Core
 		public void RenderTo(ILogger logger) => Events.ForEach( x => logger.Log(x.Level, x.Id, x, null, (s, e) => s.ToString()) );
 
 		/// <inheritdoc/>
-		public IDisposable BeginScope<TState>(TState state) => _scopeCollection.BeginScope(state);
+		public IDisposable? BeginScope<TState>(TState state) where TState : notnull => _scopeCollection.BeginScope(state);
 
 		/// <inheritdoc/>
 		public bool IsEnabled(LogLevel logLevel) => true;
@@ -922,7 +922,7 @@ namespace EpicGames.Core
 			}
 		}
 
-		public IDisposable BeginScope<TState>(TState state) => _inner.BeginScope(state);
+		public IDisposable? BeginScope<TState>(TState state) where TState : notnull => _inner.BeginScope(state);
 
 		public bool IsEnabled(LogLevel logLevel) => _inner.IsEnabled(logLevel);
 
@@ -1230,7 +1230,7 @@ namespace EpicGames.Core
 		}
 
 		/// <inheritdoc/>
-		public IDisposable BeginScope<TState>(TState state)
+		public IDisposable? BeginScope<TState>(TState state) where TState : notnull
 			=> _scopeCollection.BeginScope(state);
 
 		/// <inheritdoc/>
diff --git a/Engine/Source/Programs/Shared/EpicGames.Core/PrefixLogger.cs b/Engine/Source/Programs/Shared/EpicGames.Core/PrefixLogger.cs
--- a/Engine/Source/Programs/Shared/EpicGames.Core/PrefixLogger.cs
+++ b/Engine/Source/Programs/Shared/EpicGames.Core/PrefixLogger.cs
@@ -26,7 +26,7 @@ namespace EpicGames.Core
 		}
 
 		/// <inheritdoc/>
-		public IDisposable BeginScope<TState>(TState state) => _inner.BeginScope<TState>(state);
+		public IDisposable? BeginScope<TState>(TState state) where TState : notnull => _inner.BeginScope<TState>(state);
 
 		/// <inheritdoc/>
 		public bool IsEnabled(LogLevel logLevel) => _inner.IsEnabled(logLevel);
diff --git a/Engine/Source/Programs/Shared/EpicGames.Horde.Tests/ComputeSocketTests.cs b/Engine/Source/Programs/Shared/EpicGames.Horde.Tests/ComputeSocketTests.cs
--- a/Engine/Source/Programs/Shared/EpicGames.Horde.Tests/ComputeSocketTests.cs
+++ b/Engine/Source/Programs/Shared/EpicGames.Horde.Tests/ComputeSocketTests.cs
@@ -57,7 +57,7 @@ namespace EpicGames.Horde.Tests
 
 		class TestLogger : ILogger
 		{
-			public IDisposable BeginScope<TState>(TState state) => null!;
+			public IDisposable? BeginScope<TState>(TState state) where TState : notnull => null!;
 
 			public bool IsEnabled(LogLevel logLevel) => true;
 
diff --git a/Engine/Source/Programs/Shared/EpicGames.Horde.Tests/ServerLogPacketBuilderTests.cs b/Engine/Source/Programs/Shared/EpicGames.Horde.Tests/ServerLogPacketBuilderTests.cs
--- a/Engine/Source/Programs/Shared/EpicGames.Horde.Tests/ServerLogPacketBuilderTests.cs
+++ b/Engine/Source/Programs/Shared/EpicGames.Horde.Tests/ServerLogPacketBuilderTests.cs
@@ -32,7 +32,7 @@ namespace EpicGames.Horde.Tests
 		{
 			public List<LogEvent> _lines = new List<LogEvent>();
 
-			public IDisposable BeginScope<TState>(TState state) => null!;
+			public IDisposable? BeginScope<TState>(TState state) where TState : notnull => null!;
 
 			public bool IsEnabled(LogLevel logLevel) => true;
 
diff --git a/Engine/Source/Programs/Shared/EpicGames.Horde/Logs/IServerLogger.cs b/Engine/Source/Programs/Shared/EpicGames.Horde/Logs/IServerLogger.cs
--- a/Engine/Source/Programs/Shared/EpicGames.Horde/Logs/IServerLogger.cs
+++ b/Engine/Source/Programs/Shared/EpicGames.Horde/Logs/IServerLogger.cs
@@ -34,7 +34,7 @@ namespace EpicGames.Horde.Logs
 			}
 
 			/// <inheritdoc/>
-			public IDisposable BeginScope<TState>(TState state)
+			public IDisposable? BeginScope<TState>(TState state) where TState : notnull
 				=> _localLogger.BeginScope(state);
 
 			/// <inheritdoc/>
diff --git a/Engine/Source/Programs/Shared/EpicGames.Horde/Logs/ServerLogger.cs b/Engine/Source/Programs/Shared/EpicGames.Horde/Logs/ServerLogger.cs
--- a/Engine/Source/Programs/Shared/EpicGames.Horde/Logs/ServerLogger.cs
+++ b/Engine/Source/Programs/Shared/EpicGames.Horde/Logs/ServerLogger.cs
@@ -75,7 +75,7 @@ namespace EpicGames.Horde.Logs
 		public bool IsEnabled(LogLevel logLevel) => logLevel >= _minimumLevel;
 
 		/// <inheritdoc/>
-		public IDisposable BeginScope<TState>(TState state)
+		public IDisposable? BeginScope<TState>(TState state) where TState : notnull
 			=> null!;
 
 		private void WriteFormattedEvent(JsonLogEvent jsonLogEvent)
diff --git a/Engine/Source/Programs/Shared/EpicGames.OIDC/EpicGames.OIDC.csproj b/Engine/Source/Programs/Shared/EpicGames.OIDC/EpicGames.OIDC.csproj
--- a/Engine/Source/Programs/Shared/EpicGames.OIDC/EpicGames.OIDC.csproj
+++ b/Engine/Source/Programs/Shared/EpicGames.OIDC/EpicGames.OIDC.csproj
@@ -40,6 +40,7 @@
     <PackageReference Include="IdentityModel.OidcClient" Version="4.0.0" />
     <PackageReference Include="Microsoft.Extensions.Configuration.Binder" Version="6.0.0" />
     <PackageReference Include="Microsoft.Extensions.Configuration.Json" Version="6.0.0" />
+    <PackageReference Include="System.Text.Json" Version="8.0.5" />
   </ItemGroup>
 
 </Project>
diff --git a/Engine/Source/Programs/Shared/EpicGames.Perforce.Fixture/EpicGames.Perforce.Fixture.csproj b/Engine/Source/Programs/Shared/EpicGames.Perforce.Fixture/EpicGames.Perforce.Fixture.csproj
--- a/Engine/Source/Programs/Shared/EpicGames.Perforce.Fixture/EpicGames.Perforce.Fixture.csproj
+++ b/Engine/Source/Programs/Shared/EpicGames.Perforce.Fixture/EpicGames.Perforce.Fixture.csproj
@@ -32,7 +32,7 @@
   </ItemGroup>
 
   <ItemGroup>
-    <PackageReference Include="Microsoft.Extensions.Logging.Console" Version="8.0.0" />
+    <PackageReference Include="Microsoft.Extensions.Logging.Console" Version="8.0.1" />
     <PackageReference Include="Microsoft.NET.Test.Sdk" Version="17.11.1" />
     <PackageReference Include="MSTest.TestAdapter" Version="3.6.0">
       <TreatAsUsed>true</TreatAsUsed>
diff --git a/Engine/Source/Programs/Shared/EpicGames.Perforce/PerforceMetadataLogger.cs b/Engine/Source/Programs/Shared/EpicGames.Perforce/PerforceMetadataLogger.cs
--- a/Engine/Source/Programs/Shared/EpicGames.Perforce/PerforceMetadataLogger.cs
+++ b/Engine/Source/Programs/Shared/EpicGames.Perforce/PerforceMetadataLogger.cs
@@ -86,7 +86,7 @@ namespace EpicGames.Perforce
 		public bool IsEnabled(LogLevel logLevel) => _inner.IsEnabled(logLevel);
 
 		/// <inheritdoc/>
-		public IDisposable BeginScope<TState>(TState state) => _inner.BeginScope(state);
+		public IDisposable? BeginScope<TState>(TState state) where TState : notnull => _inner.BeginScope(state);
 
 		static bool ReadFirstLogProperty(ref Utf8JsonReader reader)
 		{
diff --git a/Engine/Source/Programs/Shared/EpicGames.Slack/EpicGames.Slack.csproj b/Engine/Source/Programs/Shared/EpicGames.Slack/EpicGames.Slack.csproj
--- a/Engine/Source/Programs/Shared/EpicGames.Slack/EpicGames.Slack.csproj
+++ b/Engine/Source/Programs/Shared/EpicGames.Slack/EpicGames.Slack.csproj
@@ -38,6 +38,6 @@
     <Folder Include="Blocks\" />
   </ItemGroup>
   <ItemGroup>
-    <PackageReference Include="Microsoft.Extensions.Logging.Abstractions" Version="8.0.0" />
+    <PackageReference Include="Microsoft.Extensions.Logging.Abstractions" Version="8.0.2" />
   </ItemGroup>
 </Project>
\ No newline at end of file
diff --git a/Engine/Source/Programs/Shared/EpicGames.UBA/EpicGames.UBA.csproj b/Engine/Source/Programs/Shared/EpicGames.UBA/EpicGames.UBA.csproj
--- a/Engine/Source/Programs/Shared/EpicGames.UBA/EpicGames.UBA.csproj
+++ b/Engine/Source/Programs/Shared/EpicGames.UBA/EpicGames.UBA.csproj
@@ -53,6 +53,6 @@
   <Import Project="Library.props" />
   <ItemGroup>
     <PackageReference Include="Microsoft.CSharp" Version="4.7.0" />
-    <PackageReference Include="Microsoft.Extensions.Logging" Version="6.0.0" />
+    <PackageReference Include="Microsoft.Extensions.Logging" Version="8.0.1" />
   </ItemGroup>
 </Project>
\ No newline at end of file
diff --git a/Engine/Source/Programs/Shared/EpicGames.UBA/ThreadedLogger.cs b/Engine/Source/Programs/Shared/EpicGames.UBA/ThreadedLogger.cs
--- a/Engine/Source/Programs/Shared/EpicGames.UBA/ThreadedLogger.cs
+++ b/Engine/Source/Programs/Shared/EpicGames.UBA/ThreadedLogger.cs
@@ -130,7 +130,7 @@ namespace EpicGames.UBA
 		}
 
 		/// <inheritdoc/>
-		public IDisposable BeginScope<TState>(TState state)
+		public IDisposable? BeginScope<TState>(TState state) where TState : notnull
 		{
 			return _internalLogger.BeginScope(state);
 		}
diff --git a/Engine/Source/Programs/Shared/UnrealBuildTool.Tests/LogEventMatcherTests.cs b/Engine/Source/Programs/Shared/UnrealBuildTool.Tests/LogEventMatcherTests.cs
--- a/Engine/Source/Programs/Shared/UnrealBuildTool.Tests/LogEventMatcherTests.cs
+++ b/Engine/Source/Programs/Shared/UnrealBuildTool.Tests/LogEventMatcherTests.cs
@@ -32,7 +32,7 @@ namespace UnrealBuildTool.Tests
 
 			public List<LogEvent> _events = [];
 
-			public IDisposable BeginScope<TState>(TState state) => new Scope();
+			public IDisposable? BeginScope<TState>(TState state) where TState : notnull => new Scope();
 
 			public bool IsEnabled(LogLevel logLevel) => true;
 
diff --git a/Engine/Source/Programs/UnrealBuildTool/Modes/ProfileUnitySizesMode.cs b/Engine/Source/Programs/UnrealBuildTool/Modes/ProfileUnitySizesMode.cs
--- a/Engine/Source/Programs/UnrealBuildTool/Modes/ProfileUnitySizesMode.cs
+++ b/Engine/Source/Programs/UnrealBuildTool/Modes/ProfileUnitySizesMode.cs
@@ -54,7 +54,7 @@ namespace UnrealBuildTool
 			}
 
 			/// <inheritdoc/>
-			public IDisposable BeginScope<TState>(TState State)
+			public IDisposable? BeginScope<TState>(TState State) where TState : notnull
 			{
 				return Inner.BeginScope(State);
 			}
diff --git a/Engine/Source/Programs/UnrealBuildTool/System/DynamicCompilation.cs b/Engine/Source/Programs/UnrealBuildTool/System/DynamicCompilation.cs
--- a/Engine/Source/Programs/UnrealBuildTool/System/DynamicCompilation.cs
+++ b/Engine/Source/Programs/UnrealBuildTool/System/DynamicCompilation.cs
@@ -216,6 +216,7 @@ namespace UnrealBuildTool
 			MetadataReferences.Add(MetadataReference.CreateFromFile(Assembly.Load("System.Private.Xml").Location));
 			MetadataReferences.Add(MetadataReference.CreateFromFile(Assembly.Load("System.Private.Xml.Linq").Location));
 			MetadataReferences.Add(MetadataReference.CreateFromFile(Assembly.Load("System.Text.RegularExpressions").Location));
+			MetadataReferences.Add(MetadataReference.CreateFromFile(Assembly.Load("Microsoft.CodeAnalysis.CSharp").Location));
 			MetadataReferences.Add(MetadataReference.CreateFromFile(Assembly.Load("System.Console").Location));
 			MetadataReferences.Add(MetadataReference.CreateFromFile(Assembly.Load("System.Runtime.Extensions").Location));
 			MetadataReferences.Add(MetadataReference.CreateFromFile(Assembly.Load("Microsoft.Extensions.Logging.Abstractions").Location));
diff --git a/Engine/Source/Programs/UnrealBuildTool/UnrealBuildTool.csproj b/Engine/Source/Programs/UnrealBuildTool/UnrealBuildTool.csproj
--- a/Engine/Source/Programs/UnrealBuildTool/UnrealBuildTool.csproj
+++ b/Engine/Source/Programs/UnrealBuildTool/UnrealBuildTool.csproj
@@ -64,7 +64,7 @@
     </Compile>
   </ItemGroup>
   <ItemGroup>
-     <PackageReference Include="Microsoft.Extensions.Logging.Console" Version="6.0.0" />
+     <PackageReference Include="Microsoft.Extensions.Logging.Console" Version="8.0.1" />
   </ItemGroup>
 
   <!-- The following will glob all platform extensions' UnrealBuildTool source
diff --git a/Engine/Source/Programs/UnrealGameSync/UnrealGameSync/Controls/LogControl.cs b/Engine/Source/Programs/UnrealGameSync/UnrealGameSync/Controls/LogControl.cs
--- a/Engine/Source/Programs/UnrealGameSync/UnrealGameSync/Controls/LogControl.cs
+++ b/Engine/Source/Programs/UnrealGameSync/UnrealGameSync/Controls/LogControl.cs
@@ -956,6 +956,6 @@ namespace UnrealGameSync
 
 		public bool IsEnabled(LogLevel logLevel) => true;
 
-		public IDisposable BeginScope<TState>(TState state) => new NullDisposable();
+		public IDisposable? BeginScope<TState>(TState state) where TState : notnull => new NullDisposable();
 	}
 }
diff --git a/Engine/Source/Programs/UnrealGameSync/UnrealGameSyncShared/OpenProjectInfo.cs b/Engine/Source/Programs/UnrealGameSync/UnrealGameSyncShared/OpenProjectInfo.cs
--- a/Engine/Source/Programs/UnrealGameSync/UnrealGameSyncShared/OpenProjectInfo.cs
+++ b/Engine/Source/Programs/UnrealGameSync/UnrealGameSyncShared/OpenProjectInfo.cs
@@ -59,7 +59,7 @@ namespace UnrealGameSync
 
 		public static async Task<OpenProjectInfo> CreateAsync(IPerforceConnection defaultConnection, UserSelectedProjectSettings selectedProject, UserSettings userSettings, bool generateP4Config, ILogger<OpenProjectInfo> logger, CancellationToken cancellationToken)
 		{
-			using IDisposable loggerScope = logger.BeginScope("Project {SelectedProject}", selectedProject.ToString());
+			using IDisposable? loggerScope = logger.BeginScope("Project {SelectedProject}", selectedProject.ToString());
 			logger.LogInformation("Detecting settings for {Project}", selectedProject);
 
 			// Use the cached client path to the file if it's available; it's much quicker than trying to find the correct workspace.
diff --git a/Engine/Source/Programs/UnrealGameSync/UnrealGameSyncShared/OutputAdapters.cs b/Engine/Source/Programs/UnrealGameSync/UnrealGameSyncShared/OutputAdapters.cs
--- a/Engine/Source/Programs/UnrealGameSync/UnrealGameSyncShared/OutputAdapters.cs
+++ b/Engine/Source/Programs/UnrealGameSync/UnrealGameSyncShared/OutputAdapters.cs
@@ -20,7 +20,7 @@ namespace UnrealGameSync
 			_inner = inInner;
 		}
 
-		public IDisposable BeginScope<TState>(TState state) => _inner.BeginScope(state);
+		public IDisposable? BeginScope<TState>(TState state) where TState : notnull => _inner.BeginScope(state);
 
 		public bool IsEnabled(LogLevel logLevel) => _inner.IsEnabled(logLevel);
 
