Remove unneeded .NET preprocessor conditions now that we are on net8.0

#rnx

diff --git a/Engine/Source/Programs/Shared/EpicGames.BuildGraph/BgGraphDef.cs b/Engine/Source/Programs/Shared/EpicGames.BuildGraph/BgGraphDef.cs
--- a/Engine/Source/Programs/Shared/EpicGames.BuildGraph/BgGraphDef.cs
+++ b/Engine/Source/Programs/Shared/EpicGames.BuildGraph/BgGraphDef.cs
@@ -6,7 +6,6 @@ using System.Diagnostics.CodeAnalysis;
 using System.Linq;
 using System.Runtime.CompilerServices;
 using System.Text;
-using System.Xml.Linq;
 using EpicGames.Core;
 using Microsoft.Extensions.Logging;
 
diff --git a/Engine/Source/Programs/Shared/EpicGames.Core/AsyncQueue.cs b/Engine/Source/Programs/Shared/EpicGames.Core/AsyncQueue.cs
--- a/Engine/Source/Programs/Shared/EpicGames.Core/AsyncQueue.cs
+++ b/Engine/Source/Programs/Shared/EpicGames.Core/AsyncQueue.cs
@@ -42,12 +42,7 @@ namespace EpicGames.Core
 		/// <inheritdoc/>
 		public async ValueTask DisposeAsync()
 		{
-#if NET8_0_OR_GREATER
 			await _cancellationSource.CancelAsync();
-#else
-			_cancellationSource.Cancel();
-#endif
-
 			await StopAsync(CancellationToken.None).ConfigureAwait(false);
 
 			_cancellationSource.Dispose();
diff --git a/Engine/Source/Programs/Shared/EpicGames.Core/BackgroundTask.cs b/Engine/Source/Programs/Shared/EpicGames.Core/BackgroundTask.cs
--- a/Engine/Source/Programs/Shared/EpicGames.Core/BackgroundTask.cs
+++ b/Engine/Source/Programs/Shared/EpicGames.Core/BackgroundTask.cs
@@ -82,11 +82,7 @@ namespace EpicGames.Core
 			{
 				try
 				{
-#if NET8_0_OR_GREATER
 					await _cancellationTokenSource.CancelAsync();
-#else
-					_cancellationTokenSource.Cancel();
-#endif
 					await _task.WaitAsync(cancellationToken).ConfigureAwait(false);
 				}
 				catch (OperationCanceledException)
@@ -161,11 +157,7 @@ namespace EpicGames.Core
 			{
 				try
 				{
-#if NET8_0_OR_GREATER
 					await _cancellationTokenSource.CancelAsync();
-#else
-					_cancellationTokenSource.Cancel();
-#endif
 					await _task.WaitAsync(cancellationToken);
 				}
 				catch (OperationCanceledException)
diff --git a/Engine/Source/Programs/Shared/EpicGames.Core/ByteString.cs b/Engine/Source/Programs/Shared/EpicGames.Core/ByteString.cs
--- a/Engine/Source/Programs/Shared/EpicGames.Core/ByteString.cs
+++ b/Engine/Source/Programs/Shared/EpicGames.Core/ByteString.cs
@@ -117,18 +117,9 @@ namespace EpicGames.Core
 		/// <inheritdoc/>
 		public override int GetHashCode()
 		{
-#if NETCOREAPP3_1
-			HashCode hashCode = new HashCode();
-			foreach (byte value in Data.Span)
-			{
-				hashCode.Add(value);
-			}
-			return hashCode.ToHashCode();
-#else
 			HashCode hashCode = new HashCode();
 			hashCode.AddBytes(Data.Span);
 			return hashCode.ToHashCode();
-#endif
 		}
 
 		/// <inheritdoc/>
diff --git a/Engine/Source/Programs/Shared/EpicGames.Core/JsonExtensions.cs b/Engine/Source/Programs/Shared/EpicGames.Core/JsonExtensions.cs
--- a/Engine/Source/Programs/Shared/EpicGames.Core/JsonExtensions.cs
+++ b/Engine/Source/Programs/Shared/EpicGames.Core/JsonExtensions.cs
@@ -19,11 +19,7 @@ namespace EpicGames.Core
 		/// <param name="reader">Token reader</param>
 		/// <param name="propertyName">Receives the property name on success</param>
 		/// <returns>True if the read succeeded</returns>
-#if NET7_0_OR_GREATER
 		public static bool TryReadNextPropertyName(scoped ref Utf8JsonReader reader, out ReadOnlySpan<byte> propertyName)
-#else
-		public static bool TryReadNextPropertyName(ref Utf8JsonReader reader, out ReadOnlySpan<byte> propertyName)
-#endif
 		{
 			if (reader.Read() && reader.TokenType == JsonTokenType.PropertyName)
 			{
diff --git a/Engine/Source/Programs/Shared/EpicGames.Horde.Tests/ComputeSocketTests.cs b/Engine/Source/Programs/Shared/EpicGames.Horde.Tests/ComputeSocketTests.cs
--- a/Engine/Source/Programs/Shared/EpicGames.Horde.Tests/ComputeSocketTests.cs
+++ b/Engine/Source/Programs/Shared/EpicGames.Horde.Tests/ComputeSocketTests.cs
@@ -121,11 +121,7 @@ namespace EpicGames.Horde.Tests
 		private static async Task<(Socket client, Socket server)> CreateSocketsAsync(CancellationToken cancellationToken)
 		{
 			int port = GetAvailablePort();
-#if NET8_0_OR_GREATER
 			using TcpListener listener = new (IPAddress.Loopback, port);
-#else
-			TcpListener listener = new(IPAddress.Loopback, port);
-#endif
 			listener.Start();
 			Socket clientSocket = new(SocketType.Stream, ProtocolType.Tcp);
 			Task clientConnectTask = clientSocket.ConnectAsync(IPAddress.Loopback, port, cancellationToken).AsTask();
@@ -227,11 +223,7 @@ namespace EpicGames.Horde.Tests
 
 		static int GetAvailablePort()
 		{
-#if NET8_0_OR_GREATER
 			using TcpListener listener = new(IPAddress.Loopback, 0);
-#else
-			TcpListener listener = new(IPAddress.Loopback, 0);
-#endif
 			listener.Start();
 			int port = ((IPEndPoint)listener.LocalEndpoint).Port;
 			listener.Stop();
diff --git a/Engine/Source/Programs/Shared/EpicGames.Horde/Compute/AgentMessageHandler.cs b/Engine/Source/Programs/Shared/EpicGames.Horde/Compute/AgentMessageHandler.cs
--- a/Engine/Source/Programs/Shared/EpicGames.Horde/Compute/AgentMessageHandler.cs
+++ b/Engine/Source/Programs/Shared/EpicGames.Horde/Compute/AgentMessageHandler.cs
@@ -67,11 +67,7 @@ namespace EpicGames.Horde.Compute
 			{
 				// Capture stack from call site
 				Interlocked.CompareExchange(ref exceptionInfo, ExceptionDispatchInfo.Capture(ex), null);
-#if NET8_0_OR_GREATER
 				await cancellationSource.CancelAsync();
-#else
-				await Task.Run(() => cancellationSource.Cancel(), CancellationToken.None);
-#endif
 			}
 
 			await RunAsync(socket, 0, 4 * 1024 * 1024, PostException, cancellationSource.Token);
diff --git a/Engine/Source/Programs/Shared/EpicGames.Horde/Compute/Transports/IdleTimeoutTransport.cs b/Engine/Source/Programs/Shared/EpicGames.Horde/Compute/Transports/IdleTimeoutTransport.cs
--- a/Engine/Source/Programs/Shared/EpicGames.Horde/Compute/Transports/IdleTimeoutTransport.cs
+++ b/Engine/Source/Programs/Shared/EpicGames.Horde/Compute/Transports/IdleTimeoutTransport.cs
@@ -83,11 +83,7 @@ public class IdleTimeoutTransport : ComputeTransport
 			if (remainingTime < TimeSpan.Zero)
 			{
 				logger.LogWarning("Terminating compute transport due to timeout (last tick at {Time})", DateTime.UtcNow - TimeSinceActivity);
-#if NET8_0_OR_GREATER
 				await cts.CancelAsync();
-#else
-				cts.Cancel();
-#endif
 				break;
 			}
 			await Task.Delay(remainingTime + TimeSpan.FromSeconds(0.2), cancellationToken);
diff --git a/Engine/Source/Programs/Shared/EpicGames.Horde/HordeHttpAuthHandler.cs b/Engine/Source/Programs/Shared/EpicGames.Horde/HordeHttpAuthHandler.cs
--- a/Engine/Source/Programs/Shared/EpicGames.Horde/HordeHttpAuthHandler.cs
+++ b/Engine/Source/Programs/Shared/EpicGames.Horde/HordeHttpAuthHandler.cs
@@ -138,11 +138,7 @@ namespace EpicGames.Horde
 		{
 			if (_authStateTask != null && !_authStateTask.IsCompleted)
 			{
-#if NET8_0_OR_GREATER
 				await _cancellationTokenSource.CancelAsync();
-#else
-				_cancellationTokenSource.Cancel();
-#endif
 				try
 				{
 #pragma warning disable VSTHRD003
diff --git a/Engine/Source/Programs/Shared/EpicGames.Horde/Storage/BlobPipeline.cs b/Engine/Source/Programs/Shared/EpicGames.Horde/Storage/BlobPipeline.cs
--- a/Engine/Source/Programs/Shared/EpicGames.Horde/Storage/BlobPipeline.cs
+++ b/Engine/Source/Programs/Shared/EpicGames.Horde/Storage/BlobPipeline.cs
@@ -107,11 +107,7 @@ namespace EpicGames.Horde.Storage
 		{
 			if (_tasks.Count > 0)
 			{
-#if NET8_0_OR_GREATER
 				await _cancellationTokenSource.CancelAsync();
-#else
-				_cancellationTokenSource.Cancel();
-#endif
 				try
 				{
 					await Task.WhenAll(_tasks);
diff --git a/Engine/Source/Programs/Shared/EpicGames.OIDC/OidcTokenManager.cs b/Engine/Source/Programs/Shared/EpicGames.OIDC/OidcTokenManager.cs
--- a/Engine/Source/Programs/Shared/EpicGames.OIDC/OidcTokenManager.cs
+++ b/Engine/Source/Programs/Shared/EpicGames.OIDC/OidcTokenManager.cs
@@ -342,11 +342,7 @@ namespace EpicGames.OIDC
 								{
 									// timed out
 									loginResult = new LoginResult($"Login timed out after: {_loginTimeout.TotalMinutes} minutes");
-#if NET8_0_OR_GREATER
 									await tokenSource.CancelAsync();
-#else
-									tokenSource.Cancel();
-#endif
 									// we do not await the processHttpTask as the GetContext method does not provide any way to cancel, as such this task will be stuck until httplistener is disposed
 								}
 							}
@@ -444,11 +440,7 @@ namespace EpicGames.OIDC
 
 							await using Stream body = request.InputStream;
 							using StreamReader reader = new StreamReader(body, request.ContentEncoding);
-#if NET7_0_OR_GREATER
 							responseData = await reader.ReadToEndAsync(cancellationToken);
-#else
-							responseData = await reader.ReadToEndAsync();
-#endif
 
 							loginResult = await oidcClient!.ProcessResponseAsync(responseData, loginState, cancellationToken: cancellationToken);
 							break;
diff --git a/Engine/Source/Programs/Shared/EpicGames.Redis/Utility/RedisLock.cs b/Engine/Source/Programs/Shared/EpicGames.Redis/Utility/RedisLock.cs
--- a/Engine/Source/Programs/Shared/EpicGames.Redis/Utility/RedisLock.cs
+++ b/Engine/Source/Programs/Shared/EpicGames.Redis/Utility/RedisLock.cs
@@ -53,11 +53,7 @@ namespace EpicGames.Redis.Utility
 		{
 			if (_cancellationSource != null)
 			{
-#if NET8_0_OR_GREATER
 				await _cancellationSource.CancelAsync();
-#else
-				_cancellationSource.Cancel();
-#endif
 			}
 			if (_backgroundTask != null)
 			{
diff --git a/Engine/Source/Programs/Shared/EpicGames.UHT/Exporters/CodeGen/UhtHeaderCodeGeneratorHFile.cs b/Engine/Source/Programs/Shared/EpicGames.UHT/Exporters/CodeGen/UhtHeaderCodeGeneratorHFile.cs
--- a/Engine/Source/Programs/Shared/EpicGames.UHT/Exporters/CodeGen/UhtHeaderCodeGeneratorHFile.cs
+++ b/Engine/Source/Programs/Shared/EpicGames.UHT/Exporters/CodeGen/UhtHeaderCodeGeneratorHFile.cs
@@ -465,7 +465,7 @@ namespace EpicGames.UHT.Exporters.CodeGen
 						{
 							string propName = property.SourceName;
 							bool isBoolWithB = (property is UhtBoolProperty && propName.StartsWith("b", StringComparison.Ordinal));
-							var cleanPropName = isBoolWithB ? propName[1..] : propName;
+							string cleanPropName = isBoolWithB ? propName[1..] : propName;
 							builder.Append("\tinline bool ").Append(cleanPropName).Append("Modified() const { return !!(ModifiedFlags & Field_"). Append(propName).Append("); } \\\r\n");
 
 							builder.Append("\tinline void Set").Append(cleanPropName).Append("(const ");
@@ -509,7 +509,7 @@ namespace EpicGames.UHT.Exporters.CodeGen
 						{
 							string propName = property.SourceName;
 							bool isBoolWithB = (property is UhtBoolProperty && propName.StartsWith("b", StringComparison.Ordinal));
-							var cleanPropName = isBoolWithB ? propName[1..] : propName;
+							string cleanPropName = isBoolWithB ? propName[1..] : propName;
 
 							builder.Append("\t\tif (Other.").Append(cleanPropName).Append("Modified()");
 							builder.Append(" && !StateStreamEquals(").Append(propName).Append(", Other.").Append(propName);
@@ -542,7 +542,7 @@ namespace EpicGames.UHT.Exporters.CodeGen
 							string propName = property.SourceName;
 							bool isBool = property is UhtBoolProperty;
 							bool isBoolWithB = isBool && propName.StartsWith("b", StringComparison.Ordinal);
-							var cleanPropName = isBoolWithB ? propName[1..] : propName;
+							string cleanPropName = isBoolWithB ? propName[1..] : propName;
 
 							builder.Append("\t\tif (To.").Append(cleanPropName).Append("Modified()) \\\r\n");
 							builder.Append("\t\t{ \\\r\n");
diff --git a/Engine/Source/Programs/UnrealBuildTool/Executors/UnrealBuildAccelerator/UBAAgentCoordinatorHorde.cs b/Engine/Source/Programs/UnrealBuildTool/Executors/UnrealBuildAccelerator/UBAAgentCoordinatorHorde.cs
--- a/Engine/Source/Programs/UnrealBuildTool/Executors/UnrealBuildAccelerator/UBAAgentCoordinatorHorde.cs
+++ b/Engine/Source/Programs/UnrealBuildTool/Executors/UnrealBuildAccelerator/UBAAgentCoordinatorHorde.cs
@@ -106,11 +106,7 @@ namespace UnrealBuildTool
 		{
 			if (_workers.Count > 0) // Should handle double-dispose, prevent cancelling twice
 			{
-#if NET8_0_OR_GREATER
 				await _cancellationTokenSource.CancelAsync();
-#else
-				_cancellationTokenSource.Cancel();
-#endif
 
 				for (int idx = _workers.Count - 1; idx >= 0; idx--)
 				{
