[UBA]
* Added api code to use virtual files from UBT
* Added event to UBAExecutor that is invoked when session server is created (so virtual files can be added)

diff --git a/Engine/Source/Programs/Shared/EpicGames.UBA/ISessionServer.cs b/Engine/Source/Programs/Shared/EpicGames.UBA/ISessionServer.cs
--- a/Engine/Source/Programs/Shared/EpicGames.UBA/ISessionServer.cs
+++ b/Engine/Source/Programs/Shared/EpicGames.UBA/ISessionServer.cs
@@ -184,6 +184,15 @@ namespace EpicGames.UBA
 		/// <param name="files">The files to register</param>
 		public abstract void RegisterNewFiles(params string[] files);
 
+		/// <summary>
+		/// Registers virtual files that are part of other files (mapped into sections of existing file)
+		/// </summary>
+		/// <param name="name">Name of the virtual file</param>
+		/// <param name="sourceFile">Name of the source file</param>
+		/// <param name="sourceOffset">Offset inside the source file</param>
+		/// <param name="sourceSize">Size of the segment inside the source file</param>
+		public abstract bool RegisterVirtualFile(string name, string sourceFile, ulong sourceOffset, ulong sourceSize);
+
 		/// <summary>
 		/// Registers the start of an external process
 		/// </summary>
diff --git a/Engine/Source/Programs/Shared/EpicGames.UBA/Impl/SessionServerImpl.cs b/Engine/Source/Programs/Shared/EpicGames.UBA/Impl/SessionServerImpl.cs
--- a/Engine/Source/Programs/Shared/EpicGames.UBA/Impl/SessionServerImpl.cs
+++ b/Engine/Source/Programs/Shared/EpicGames.UBA/Impl/SessionServerImpl.cs
@@ -88,6 +88,9 @@ namespace EpicGames.UBA.Impl
 		[DllImport("UbaHost", CharSet = CharSet.Auto)]
 		static extern void SessionServer_RegisterNewFile(nint server, string filename);
 
+		[DllImport("UbaHost", CharSet = CharSet.Auto)]
+		static extern bool SessionServer_RegisterVirtualFile(nint server, string filename, string sourceFile, ulong sourceOffset, ulong sourceSize);
+
 		[DllImport("UbaHost", CharSet = CharSet.Auto)]
 		static extern uint SessionServer_BeginExternalProcess(nint server, string description);
 
@@ -205,6 +208,8 @@ namespace EpicGames.UBA.Impl
 		public void RefreshDirectories(params string[] directories) => Array.ForEach(directories, (directory) => SessionServer_RefreshDirectory(_handle, directory));
 
 		public void RegisterNewFiles(params string[] files) => Array.ForEach(files, (file) => SessionServer_RegisterNewFile(_handle, file));
+		public bool RegisterVirtualFile(string name, string sourceFile, ulong sourceOffset, ulong sourceSize) => SessionServer_RegisterVirtualFile(_handle, name, sourceFile, sourceOffset, sourceSize);
+
 
 		public uint BeginExternalProcess(string description) => SessionServer_BeginExternalProcess(_handle, description);
 
diff --git a/Engine/Source/Programs/UnrealBuildTool/Executors/UnrealBuildAccelerator/UBAExecutor.cs b/Engine/Source/Programs/UnrealBuildTool/Executors/UnrealBuildAccelerator/UBAExecutor.cs
--- a/Engine/Source/Programs/UnrealBuildTool/Executors/UnrealBuildAccelerator/UBAExecutor.cs
+++ b/Engine/Source/Programs/UnrealBuildTool/Executors/UnrealBuildAccelerator/UBAExecutor.cs
@@ -81,6 +81,11 @@ namespace UnrealBuildTool
 			return EpicGames.UBA.Utils.IsAvailable();
 		}
 
+		public delegate void SessionCreatedDelegate(ISessionServer Session);
+
+		// Event invoked when ISessionServer is created. Can be used to for example register virtual files etc
+		public event SessionCreatedDelegate? OnSessionCreated;
+
 		public static DirectoryReference UbaBinariesDir
 		{
 			get
@@ -389,6 +394,9 @@ namespace UnrealBuildTool
 						try
 						{
 							_session = ISessionServer.CreateSessionServer(serverCreateInfo);
+
+							OnSessionCreated?.Invoke(_session);
+
 							_cacheClient = ICacheClient.CreateCacheClient(_session, UBAConfig.bReportCacheMissReason);
 							if (!String.IsNullOrEmpty(UBAConfig.CacheServer))
 							{
