Implementation of a Jupiter storage backend in the storage library.

- Blobs are framed inside a CbObject with its imports written as attachments and uploaded to the remote Jupiter store.
- Server address and auth headers are assumed to be set in the HttpClient instances returned by the createHttpClient() delegate passed to the constructor.
- Does not currently support aliases.

#rnx

diff --git a/Engine/Source/Programs/Shared/EpicGames.Horde/Storage/Backends/JupiterStorageBackend.cs b/Engine/Source/Programs/Shared/EpicGames.Horde/Storage/Backends/JupiterStorageBackend.cs
new file mode 100644
--- /dev/null
+++ b/Engine/Source/Programs/Shared/EpicGames.Horde/Storage/Backends/JupiterStorageBackend.cs
@@ -0,0 +1,341 @@
+// Copyright Epic Games, Inc. All Rights Reserved.
+
+using System;
+using System.Collections.Generic;
+using System.IO;
+using System.Net;
+using System.Net.Http;
+using System.Net.Http.Headers;
+using System.Net.Http.Json;
+using System.Net.Mime;
+using System.Threading;
+using System.Threading.Tasks;
+using EpicGames.Core;
+using EpicGames.Serialization;
+using Microsoft.Extensions.Logging;
+
+namespace EpicGames.Horde.Storage.Backends
+{
+	/// <summary>
+	/// Storage backend which communicates with Jupiter over HTTP
+	/// </summary>
+	public sealed class JupiterStorageBackend : IStorageBackend
+	{
+		const string CompactBinaryMimeType = "application/x-ue-cb";
+
+		static class HeaderNames
+		{
+			public const string JupiterIoHash = "X-Jupiter-IoHash";
+		}
+
+		static class BlobFieldNames
+		{
+			public const string Data = "data";
+			public const string Attachments = "attachments";
+		}
+
+		record class WriteBlobResponse(string Identifier);
+
+		readonly NamespaceId _namespaceId;
+		readonly string _bucketId;
+		readonly Func<HttpClient> _createClient;
+		readonly ILogger _logger;
+
+		/// <inheritdoc/>
+		public bool SupportsRedirects => false;
+
+		/// <summary>
+		/// Constructor
+		/// </summary>
+		public JupiterStorageBackend(NamespaceId namespaceId, string bucketId, Func<HttpClient> createClient, ILogger logger)
+		{
+			_namespaceId = namespaceId;
+			_bucketId = bucketId;
+			_createClient = createClient;
+			_logger = logger;
+		}
+
+		#region Blobs
+
+		/// <inheritdoc/>
+		public async Task<Stream> OpenBlobAsync(BlobLocator locator, int offset, int? length, CancellationToken cancellationToken = default)
+		{
+			ReadOnlyMemory<byte> data = await ReadBlobDataAsync(locator, offset, length, cancellationToken);
+			return new ReadOnlyMemoryStream(data);
+		}
+
+		/// <inheritdoc/>
+		public async Task<IReadOnlyMemoryOwner<byte>> ReadBlobAsync(BlobLocator locator, int offset, int? length, CancellationToken cancellationToken = default)
+		{
+			ReadOnlyMemory<byte> data = await ReadBlobDataAsync(locator, offset, length, cancellationToken);
+			return ReadOnlyMemoryOwner.Create(data);
+		}
+
+		async Task<ReadOnlyMemory<byte>> ReadBlobDataAsync(BlobLocator locator, int offset, int? length, CancellationToken cancellationToken = default)
+		{
+			if (offset == 0 && length == null)
+			{
+				_logger.LogDebug("Reading {Locator}", locator);
+			}
+			else if (length == null)
+			{
+				_logger.LogDebug("Reading {Locator} ({Offset}..)", locator, offset);
+			}
+			else
+			{
+				_logger.LogDebug("Reading {Locator} ({Offset}+{Length})", locator, offset, length);
+			}
+
+			ReadOnlyMemory<byte> data = ReadOnlyMemory<byte>.Empty;
+			if (!length.HasValue || length.Value > 0)
+			{
+				byte[] encodedData = await ReadEncodedBlobDataAsync(locator, cancellationToken);
+				data = UnwrapDataFromCbObject(encodedData);
+			}
+
+			if (offset > 0)
+			{
+				data = data.Slice(offset);
+			}
+			if (length.HasValue)
+			{
+				data = data.Slice(0, length.Value);
+			}
+
+			return data;
+		}
+
+		// Extract the data field from an encoded cb object
+		static ReadOnlyMemory<byte> UnwrapDataFromCbObject(ReadOnlyMemory<byte> encodedData)
+		{
+			CbObject obj = new CbObject(encodedData);
+			return obj[BlobFieldNames.Data].AsBinary();
+		}
+
+		// Reads a blob and returns the raw encoded CB object
+		async Task<byte[]> ReadEncodedBlobDataAsync(BlobLocator locator, CancellationToken cancellationToken = default)
+		{
+			using (HttpClient httpClient = _createClient())
+			{
+				using (HttpRequestMessage request = new HttpRequestMessage(HttpMethod.Get, $"api/v1/blobs/{_namespaceId}/{locator}"))
+				{
+					HttpResponseMessage response = await httpClient.SendAsync(request, cancellationToken);
+					await EnsureSuccessAsync(response, cancellationToken);
+
+					byte[] data = await response.Content.ReadAsByteArrayAsync(cancellationToken);
+					return data;
+				}
+			}
+		}
+
+		/// <inheritdoc/>
+		public async Task WriteBlobAsync(BlobLocator locator, Stream stream, IReadOnlyCollection<BlobLocator>? imports, CancellationToken cancellationToken = default)
+		{
+			BlobLocator result = await WriteBlobAsync(stream, imports, null, cancellationToken);
+			if (locator != result)
+			{
+				throw new InvalidOperationException("Incorrect locator passed to Jupiter backend for upload.");
+			}
+		}
+
+		/// <inheritdoc/>
+		public async Task<BlobLocator> WriteBlobAsync(Stream stream, IReadOnlyCollection<BlobLocator>? imports, string? prefix = null, CancellationToken cancellationToken = default)
+		{
+			// Read all the input data
+			byte[] data = await stream.ReadAllBytesAsync(cancellationToken);
+
+			// Frame it within a CB object
+			CbWriter writer = new CbWriter();
+			writer.BeginObject();
+			writer.WriteBinaryArray(BlobFieldNames.Data, data);
+			if (imports != null && imports.Count > 0)
+			{
+				writer.BeginArray(BlobFieldNames.Attachments);
+				foreach (BlobLocator import in imports)
+				{
+					writer.WriteObjectAttachmentValue(IoHash.Parse(import.Path));
+				}
+				writer.EndArray();
+			}
+			writer.EndObject();
+
+			// Send the message
+			using (HttpClient httpClient = _createClient())
+			{
+				using ByteArrayContent content = new ByteArrayContent(writer.ToByteArray());
+				content.Headers.ContentType = new MediaTypeHeaderValue(MediaTypeNames.Application.Octet);
+
+				using (HttpRequestMessage request = new HttpRequestMessage(HttpMethod.Post, $"api/v1/blobs/{_namespaceId}"))
+				{
+					request.Content = content;
+
+					using (HttpResponseMessage response = await httpClient.SendAsync(request, cancellationToken))
+					{
+						if (!response.IsSuccessStatusCode)
+						{
+							string responseText = await response.Content.ReadAsStringAsync(cancellationToken);
+							throw new StorageException($"Upload to {request.RequestUri} failed ({response.StatusCode}). Response: {responseText}");
+						}
+
+						WriteBlobResponse responseBody = await response.Content.ReadFromJsonAsync<WriteBlobResponse>(cancellationToken: cancellationToken)
+							?? throw new InvalidDataException("Expected non-null response");
+
+						_logger.LogDebug("Written blob {BlobId}", responseBody.Identifier);
+						return new BlobLocator(responseBody.Identifier);
+					}
+				}
+			}
+		}
+
+		/// <inheritdoc/>
+		public ValueTask<Uri?> TryGetBlobReadRedirectAsync(BlobLocator locator, CancellationToken cancellationToken = default)
+			=> default;
+
+		/// <inheritdoc/>
+		public ValueTask<Uri?> TryGetBlobWriteRedirectAsync(BlobLocator locator, IReadOnlyCollection<BlobLocator> imports, CancellationToken cancellationToken = default)
+			=> default;
+
+		/// <inheritdoc/>
+		public ValueTask<(BlobLocator, Uri)?> TryGetBlobWriteRedirectAsync(IReadOnlyCollection<BlobLocator>? imports = null, string? prefix = null, CancellationToken cancellationToken = default)
+			=> default;
+
+		#endregion
+
+		#region Aliases
+
+		/// <inheritdoc/>
+		public Task AddAliasAsync(string name, BlobLocator target, int rank = 0, ReadOnlyMemory<byte> data = default, CancellationToken cancellationToken = default)
+			=> throw new NotSupportedException("Aliases are not currently supported for the Jupiter backend.");
+
+		/// <inheritdoc/>
+		public Task RemoveAliasAsync(string name, BlobLocator target, CancellationToken cancellationToken = default)
+			=> throw new NotSupportedException("Aliases are not currently supported for the Jupiter backend.");
+
+		/// <inheritdoc/>
+		public Task<BlobAliasLocator[]> FindAliasesAsync(string alias, int? maxResults = null, CancellationToken cancellationToken = default)
+			=> throw new NotSupportedException("Aliases are not currently supported for the Jupiter backend.");
+
+		#endregion
+
+		#region Refs
+
+		/// <inheritdoc/>
+		public async Task<bool> DeleteRefAsync(RefName name, CancellationToken cancellationToken)
+		{
+			_logger.LogDebug("Deleting ref {RefName}", name);
+			using (HttpClient httpClient = _createClient())
+			{
+				IoHash refId = IoHash.Compute(name.Text);
+				using (HttpRequestMessage request = new HttpRequestMessage(HttpMethod.Delete, $"api/v1/refs/{_namespaceId}/{_bucketId}/{refId}"))
+				{
+					using (HttpResponseMessage response = await httpClient.SendAsync(request, cancellationToken))
+					{
+						if (response.IsSuccessStatusCode)
+						{
+							return true;
+						}
+						if (response.StatusCode == HttpStatusCode.NotFound)
+						{
+							return false;
+						}
+
+						await EnsureSuccessAsync(response, cancellationToken);
+						return false;
+					}
+				}
+			}
+		}
+
+		/// <inheritdoc/>
+		public async Task<HashedBlobRefValue?> TryReadRefAsync(RefName name, RefCacheTime cacheTime = default, CancellationToken cancellationToken = default)
+		{
+			using (HttpClient httpClient = _createClient())
+			{
+				IoHash refId = IoHash.Compute(name.Text.Span);
+
+				using (HttpRequestMessage request = new HttpRequestMessage(HttpMethod.Get, $"api/v1/refs/{_namespaceId}/{_bucketId}/{refId}"))
+				{
+					if (cacheTime.IsSet())
+					{
+						request.Headers.CacheControl = new CacheControlHeaderValue { MaxAge = cacheTime.MaxAge };
+					}
+
+					using (HttpResponseMessage response = await httpClient.SendAsync(request, cancellationToken))
+					{
+						if (response.StatusCode == HttpStatusCode.NotFound)
+						{
+							_logger.LogDebug("Read ref {RefName} -> None", name);
+							return null;
+						}
+						else if (!response.IsSuccessStatusCode)
+						{
+							_logger.LogError("Unable to read ref {RefName} (status: {StatusCode}, body: {Body})", name, response.StatusCode, await response.Content.ReadAsStringAsync(cancellationToken));
+							throw new StorageException($"Unable to read ref '{name}'");
+						}
+						else
+						{
+							await EnsureSuccessAsync(response, cancellationToken);
+
+							if (!String.Equals(response.Content.Headers.ContentType?.MediaType, CompactBinaryMimeType, StringComparison.Ordinal))
+							{
+								throw new NotSupportedException("Unknown response type. Expected compact binary.");
+							}
+
+							byte[] encodedData = await response.Content.ReadAsByteArrayAsync(cancellationToken);
+
+							IoHash hash = IoHash.Compute(encodedData);
+							_logger.LogDebug("Read ref {RefId} -> {BlobHash}", refId, hash);
+
+							return new HashedBlobRefValue(hash, new BlobLocator(hash.ToUtf8String()));
+						}
+					}
+				}
+			}
+		}
+
+		/// <inheritdoc/>
+		public async Task WriteRefAsync(RefName name, HashedBlobRefValue value, RefOptions? options = null, CancellationToken cancellationToken = default)
+		{
+			_logger.LogDebug("Writing ref {RefName} -> {Locator}", name, value.Locator);
+			using (HttpClient httpClient = _createClient())
+			{
+				IoHash refId = IoHash.Compute(name.Text.Span);
+				using (HttpRequestMessage request = new HttpRequestMessage(HttpMethod.Put, $"api/v1/refs/{_namespaceId}/{_bucketId}/{refId}"))
+				{
+					byte[] data = await ReadEncodedBlobDataAsync(value.Locator, cancellationToken);
+
+					IoHash hash = IoHash.Compute(data);
+					request.Headers.Add(HeaderNames.JupiterIoHash, hash.ToString());
+
+					if (!String.Equals(value.Locator.Path.ToString(), hash.ToString(), StringComparison.OrdinalIgnoreCase))
+					{
+						throw new InvalidDataException();
+					}
+
+					request.Content = new ByteArrayContent(data);
+					request.Content.Headers.ContentType = new MediaTypeHeaderValue(CompactBinaryMimeType);
+
+					using (HttpResponseMessage response = await httpClient.SendAsync(request, cancellationToken))
+					{
+						await EnsureSuccessAsync(response, cancellationToken);
+					}
+				}
+			}
+		}
+
+		#endregion
+
+		/// <inheritdoc/>
+		public void GetStats(StorageStats stats)
+		{ }
+
+		static async Task EnsureSuccessAsync(HttpResponseMessage message, CancellationToken cancellationToken)
+		{
+			if (!message.IsSuccessStatusCode)
+			{
+				string response = await message.Content.ReadAsStringAsync(cancellationToken);
+				throw new StorageException($"Http response {message.StatusCode} (Content: {response})");
+			}
+		}
+	}
+}
