Horde: Jupiter backend changes.

Upload blobs without any references as binary attachments, and blobs with references as object attachments. The type of attachment is determined by a ".bin" or ".obj" extension on the BlobLocator. All blobs would previously be uploaded as objects.

#rnx

diff --git a/Engine/Source/Programs/Shared/EpicGames.Horde/Storage/Backends/JupiterStorageBackend.cs b/Engine/Source/Programs/Shared/EpicGames.Horde/Storage/Backends/JupiterStorageBackend.cs
--- a/Engine/Source/Programs/Shared/EpicGames.Horde/Storage/Backends/JupiterStorageBackend.cs
+++ b/Engine/Source/Programs/Shared/EpicGames.Horde/Storage/Backends/JupiterStorageBackend.cs
@@ -28,12 +28,26 @@ namespace EpicGames.Horde.Storage.Backends
 			public const string JupiterIoHash = "X-Jupiter-IoHash";
 		}
 
-		static class BlobFieldNames
+		// Names of fields in ref objects
+		static class RefFieldNames
+		{
+			public const string Target = "tgt";
+		}
+
+		// Names of fields when blobs are encoded as cb objects
+		static class ObjectFieldNames
 		{
 			public const string Data = "data";
-			public const string Attachments = "attachments";
+			public const string BinaryAttachments = "bin";
+			public const string ObjectAttachments = "obj";
 		}
 
+		// Extensions for locators to raw data
+		static Utf8String BinaryExtension { get; } = new Utf8String(".bin");
+
+		// Extensions for locators that point to cb objects
+		static Utf8String ObjectExtension { get; } = new Utf8String(".obj");
+
 		record class WriteBlobResponse(string Identifier);
 
 		readonly NamespaceId _namespaceId;
@@ -55,6 +69,55 @@ namespace EpicGames.Horde.Storage.Backends
 			_logger = logger;
 		}
 
+		static bool TryGetHash(Utf8String path, out IoHash hash)
+		{
+			if (TryGetBinaryAttachment(path, out CbBinaryAttachment binaryAttachment))
+			{
+				hash = binaryAttachment.Hash;
+				return true;
+			}
+			else if (TryGetObjectAttachment(path, out CbObjectAttachment objectAttachment))
+			{
+				hash = objectAttachment.Hash;
+				return true;
+			}
+			else
+			{
+				hash = IoHash.Zero;
+				return false;
+			}
+		}
+
+		static bool TryGetBinaryAttachment(Utf8String path, out CbBinaryAttachment attachment)
+		{
+			IoHash hash;
+			if (path.EndsWith(BinaryExtension) && IoHash.TryParse(path.Substring(0, path.Length - BinaryExtension.Length), out hash))
+			{
+				attachment = new CbBinaryAttachment(hash);
+				return true;
+			}
+			else
+			{
+				attachment = default;
+				return false;
+			}
+		}
+
+		static bool TryGetObjectAttachment(Utf8String path, out CbObjectAttachment attachment)
+		{
+			IoHash hash;
+			if (path.EndsWith(ObjectExtension) && IoHash.TryParse(path.Substring(0, path.Length - ObjectExtension.Length), out hash))
+			{
+				attachment = new CbObjectAttachment(hash);
+				return true;
+			}
+			else
+			{
+				attachment = default;
+				return false;
+			}
+		}
+
 		#region Blobs
 
 		/// <inheritdoc/>
@@ -89,8 +152,20 @@ namespace EpicGames.Horde.Storage.Backends
 			ReadOnlyMemory<byte> data = ReadOnlyMemory<byte>.Empty;
 			if (!length.HasValue || length.Value > 0)
 			{
-				byte[] encodedData = await ReadEncodedBlobDataAsync(locator, cancellationToken);
-				data = UnwrapDataFromCbObject(encodedData);
+				if (TryGetBinaryAttachment(locator.Path, out CbBinaryAttachment binaryAttachment))
+				{
+					byte[] encodedData = await ReadEncodedBlobDataAsync(binaryAttachment.Hash, cancellationToken);
+					data = encodedData;
+				}
+				else if (TryGetObjectAttachment(locator.Path, out CbObjectAttachment objectAttachment))
+				{
+					byte[] encodedData = await ReadEncodedBlobDataAsync(objectAttachment.Hash, cancellationToken);
+					data = UnwrapDataFromCbObject(encodedData);
+				}
+				else
+				{
+					throw new InvalidOperationException($"Locator is not a valid form for Jupiter ({locator})");
+				}
 			}
 
 			if (offset > 0)
@@ -109,15 +184,15 @@ namespace EpicGames.Horde.Storage.Backends
 		static ReadOnlyMemory<byte> UnwrapDataFromCbObject(ReadOnlyMemory<byte> encodedData)
 		{
 			CbObject obj = new CbObject(encodedData);
-			return obj[BlobFieldNames.Data].AsBinary();
+			return obj[ObjectFieldNames.Data].AsBinary();
 		}
 
 		// Reads a blob and returns the raw encoded CB object
-		async Task<byte[]> ReadEncodedBlobDataAsync(BlobLocator locator, CancellationToken cancellationToken = default)
+		async Task<byte[]> ReadEncodedBlobDataAsync(IoHash hash, CancellationToken cancellationToken = default)
 		{
 			using (HttpClient httpClient = _createClient())
 			{
-				using (HttpRequestMessage request = new HttpRequestMessage(HttpMethod.Get, $"api/v1/blobs/{_namespaceId}/{locator}"))
+				using (HttpRequestMessage request = new HttpRequestMessage(HttpMethod.Get, $"api/v1/blobs/{_namespaceId}/{hash}"))
 				{
 					HttpResponseMessage response = await httpClient.SendAsync(request, cancellationToken);
 					await EnsureSuccessAsync(response, cancellationToken);
@@ -138,31 +213,73 @@ namespace EpicGames.Horde.Storage.Backends
 			}
 		}
 
-		/// <inheritdoc/>
-		public async Task<BlobLocator> WriteBlobAsync(Stream stream, IReadOnlyCollection<BlobLocator>? imports, string? prefix = null, CancellationToken cancellationToken = default)
+		static byte[] CreateCbObject(ReadOnlyMemory<byte> payload, IReadOnlyCollection<BlobLocator> imports)
 		{
-			// Read all the input data
-			byte[] data = await stream.ReadAllBytesAsync(cancellationToken);
-
-			// Frame it within a CB object
 			CbWriter writer = new CbWriter();
 			writer.BeginObject();
-			writer.WriteBinaryArray(BlobFieldNames.Data, data);
-			if (imports != null && imports.Count > 0)
+			writer.WriteBinaryArray(ObjectFieldNames.Data, payload);
+
+			List<CbBinaryAttachment> binaryAttachments = new List<CbBinaryAttachment>();
+			List<CbObjectAttachment> objectAttachments = new List<CbObjectAttachment>();
+			foreach (BlobLocator import in imports)
 			{
-				writer.BeginArray(BlobFieldNames.Attachments);
-				foreach (BlobLocator import in imports)
+				if (TryGetBinaryAttachment(import.Path, out CbBinaryAttachment binaryAttachment))
+				{
+					binaryAttachments.Add(binaryAttachment);
+				}
+				else if (TryGetObjectAttachment(import.Path, out CbObjectAttachment objectAttachment))
 				{
-					writer.WriteObjectAttachmentValue(IoHash.Parse(import.Path));
+					objectAttachments.Add(objectAttachment);
+				}
+				else
+				{
+					throw new InvalidOperationException("Cannot reference objects that are not in Jupiter store.");
+				}
+			}
+
+			if (binaryAttachments.Count > 0)
+			{
+				writer.BeginArray(ObjectFieldNames.BinaryAttachments);
+				foreach (CbBinaryAttachment binaryAttachment in binaryAttachments)
+				{
+					writer.WriteBinaryAttachmentValue(binaryAttachment);
 				}
 				writer.EndArray();
 			}
+
+			if (objectAttachments.Count > 0)
+			{
+				writer.BeginArray(ObjectFieldNames.ObjectAttachments);
+				foreach (CbObjectAttachment objectAttachment in objectAttachments)
+				{
+					writer.WriteObjectAttachmentValue(objectAttachment);
+				}
+				writer.EndArray();
+			}
+
 			writer.EndObject();
 
+			return writer.ToByteArray();
+		}
+
+		/// <inheritdoc/>
+		public async Task<BlobLocator> WriteBlobAsync(Stream stream, IReadOnlyCollection<BlobLocator>? imports, string? prefix = null, CancellationToken cancellationToken = default)
+		{
+			// Read all the input data
+			byte[] data = await stream.ReadAllBytesAsync(cancellationToken);
+
+			// Frame it within a CB object if there are any imports
+			Utf8String extension = BinaryExtension;
+			if (imports != null && imports.Count > 0)
+			{
+				data = CreateCbObject(data, imports);
+				extension = ObjectExtension;
+			}
+
 			// Send the message
 			using (HttpClient httpClient = _createClient())
 			{
-				using ByteArrayContent content = new ByteArrayContent(writer.ToByteArray());
+				using ByteArrayContent content = new ByteArrayContent(data);
 				content.Headers.ContentType = new MediaTypeHeaderValue(MediaTypeNames.Application.Octet);
 
 				using (HttpRequestMessage request = new HttpRequestMessage(HttpMethod.Post, $"api/v1/blobs/{_namespaceId}"))
@@ -181,7 +298,7 @@ namespace EpicGames.Horde.Storage.Backends
 							?? throw new InvalidDataException("Expected non-null response");
 
 						_logger.LogDebug("Written blob {BlobId}", responseBody.Identifier);
-						return new BlobLocator(responseBody.Identifier);
+						return new BlobLocator(responseBody.Identifier + extension);
 					}
 				}
 			}
@@ -283,10 +400,30 @@ namespace EpicGames.Horde.Storage.Backends
 
 							byte[] encodedData = await response.Content.ReadAsByteArrayAsync(cancellationToken);
 
-							IoHash hash = IoHash.Compute(encodedData);
-							_logger.LogDebug("Read ref {RefId} -> {BlobHash}", refId, hash);
+							CbObject objectData = new CbObject(encodedData);
+							CbField target = objectData[RefFieldNames.Target];
+
+							IoHash hash;
+							BlobLocator locator;
+
+							if (target.IsBinaryAttachment())
+							{
+								hash = target.AsBinaryAttachment().Hash;
+								locator = new BlobLocator(hash.ToUtf8String() + BinaryExtension);
+							}
+							else if (target.IsObjectAttachment())
+							{
+								hash = target.AsObjectAttachment().Hash;
+								locator = new BlobLocator(hash.ToUtf8String() + ObjectExtension);
+							}
+							else
+							{
+								_logger.LogDebug("Ref {RefName} exists but does not contain a target field", name);
+								return null;
+							}
 
-							return new HashedBlobRefValue(hash, new BlobLocator(hash.ToUtf8String()));
+							_logger.LogDebug("Read ref {RefName} -> {Locator}", name, locator);
+							return new HashedBlobRefValue(hash, locator);
 						}
 					}
 				}
@@ -302,15 +439,27 @@ namespace EpicGames.Horde.Storage.Backends
 				IoHash refId = IoHash.Compute(name.Text.Span);
 				using (HttpRequestMessage request = new HttpRequestMessage(HttpMethod.Put, $"api/v1/refs/{_namespaceId}/{_bucketId}/{refId}"))
 				{
-					byte[] data = await ReadEncodedBlobDataAsync(value.Locator, cancellationToken);
+					CbWriter writer = new CbWriter();
+					writer.BeginObject();
 
-					IoHash hash = IoHash.Compute(data);
-					request.Headers.Add(HeaderNames.JupiterIoHash, hash.ToString());
-
-					if (!String.Equals(value.Locator.Path.ToString(), hash.ToString(), StringComparison.OrdinalIgnoreCase))
+					if (TryGetBinaryAttachment(value.Locator.Path, out CbBinaryAttachment binaryAttachment))
 					{
-						throw new InvalidDataException();
+						writer.WriteBinaryAttachment(RefFieldNames.Target, binaryAttachment);
 					}
+					else if (TryGetObjectAttachment(value.Locator.Path, out CbObjectAttachment objectAttachment))
+					{
+						writer.WriteObjectAttachment(RefFieldNames.Target, objectAttachment);
+					}
+					else
+					{
+						throw new InvalidOperationException($"Cannot write ref to non-Jupiter blob ({value.Locator})");
+					}
+
+					writer.EndObject();
+
+					byte[] data = writer.ToByteArray();
+					IoHash hash = IoHash.Compute(data);
+					request.Headers.Add(HeaderNames.JupiterIoHash, hash.ToString());
 
 					request.Content = new ByteArrayContent(data);
 					request.Content.Headers.ContentType = new MediaTypeHeaderValue(CompactBinaryMimeType);
diff --git a/Engine/Source/Programs/Shared/EpicGames.Serialization/CbWriter.cs b/Engine/Source/Programs/Shared/EpicGames.Serialization/CbWriter.cs
--- a/Engine/Source/Programs/Shared/EpicGames.Serialization/CbWriter.cs
+++ b/Engine/Source/Programs/Shared/EpicGames.Serialization/CbWriter.cs
@@ -1111,7 +1111,7 @@ namespace EpicGames.Serialization
 		/// </summary>
 		/// <param name="writer">Writer for output data</param>
 		/// <param name="value">Value to be written</param>
-		public static void WriteBinaryArrayValue(this ICbWriter writer, byte[] value) => writer.WriteBinarySpanValue(value.AsSpan());
+		public static void WriteBinaryArrayValue(this ICbWriter writer, ReadOnlyMemory<byte> value) => writer.WriteBinarySpanValue(value.Span);
 
 		/// <summary>
 		/// Writes a named binary value
@@ -1119,6 +1119,6 @@ namespace EpicGames.Serialization
 		/// <param name="writer">Writer for output data</param>
 		/// <param name="name">Name of the field</param>
 		/// <param name="value">Value to be written</param>
-		public static void WriteBinaryArray(this ICbWriter writer, CbFieldName name, byte[] value) => writer.WriteBinarySpan(name, value.AsSpan());
+		public static void WriteBinaryArray(this ICbWriter writer, CbFieldName name, ReadOnlyMemory<byte> value) => writer.WriteBinarySpan(name, value.Span);
 	}
 }
\ No newline at end of file
