Horde: Remove unused function causing static analysis warning.

#rnx

diff --git a/Engine/Source/Programs/Shared/EpicGames.Horde/Storage/Backends/JupiterStorageBackend.cs b/Engine/Source/Programs/Shared/EpicGames.Horde/Storage/Backends/JupiterStorageBackend.cs
--- a/Engine/Source/Programs/Shared/EpicGames.Horde/Storage/Backends/JupiterStorageBackend.cs
+++ b/Engine/Source/Programs/Shared/EpicGames.Horde/Storage/Backends/JupiterStorageBackend.cs
@@ -69,25 +69,6 @@ namespace EpicGames.Horde.Storage.Backends
 			_logger = logger;
 		}
 
-		static bool TryGetHash(Utf8String path, out IoHash hash)
-		{
-			if (TryGetBinaryAttachment(path, out CbBinaryAttachment binaryAttachment))
-			{
-				hash = binaryAttachment.Hash;
-				return true;
-			}
-			else if (TryGetObjectAttachment(path, out CbObjectAttachment objectAttachment))
-			{
-				hash = objectAttachment.Hash;
-				return true;
-			}
-			else
-			{
-				hash = IoHash.Zero;
-				return false;
-			}
-		}
-
 		static bool TryGetBinaryAttachment(Utf8String path, out CbBinaryAttachment attachment)
 		{
 			IoHash hash;
