Horde: Instrument ManagedWorkspace with OpenTelemetry tracing

diff --git a/Engine/Source/Programs/Horde/Drivers/JobDriver/Execution/ManagedWorkspaceMaterializer.cs b/Engine/Source/Programs/Horde/Drivers/JobDriver/Execution/ManagedWorkspaceMaterializer.cs
--- a/Engine/Source/Programs/Horde/Drivers/JobDriver/Execution/ManagedWorkspaceMaterializer.cs
+++ b/Engine/Source/Programs/Horde/Drivers/JobDriver/Execution/ManagedWorkspaceMaterializer.cs
@@ -85,7 +85,7 @@ public sealed class ManagedWorkspaceMaterializer : IWorkspaceMaterializer
 	public static async Task<ManagedWorkspaceMaterializer> CreateAsync(RpcAgentWorkspace agentWorkspace, DirectoryReference workingDir, bool useCacheFile, bool cleanDuringFinalize, Tracer tracer, ILogger logger, CancellationToken cancellationToken)
 	{
 		ManagedWorkspaceOptions options = WorkspaceInfo.GetMwOptions(agentWorkspace);
-		WorkspaceInfo workspace = await WorkspaceInfo.CreateWorkspaceInfoAsync(agentWorkspace, workingDir, options, logger, cancellationToken);
+		WorkspaceInfo workspace = await WorkspaceInfo.CreateWorkspaceInfoAsync(agentWorkspace, workingDir, options, tracer, logger, cancellationToken);
 		return new ManagedWorkspaceMaterializer(agentWorkspace, useCacheFile, cleanDuringFinalize, workspace, tracer, logger);
 	}
 
diff --git a/Engine/Source/Programs/Horde/Drivers/JobDriver/Execution/PerforceExecutor.cs b/Engine/Source/Programs/Horde/Drivers/JobDriver/Execution/PerforceExecutor.cs
--- a/Engine/Source/Programs/Horde/Drivers/JobDriver/Execution/PerforceExecutor.cs
+++ b/Engine/Source/Programs/Horde/Drivers/JobDriver/Execution/PerforceExecutor.cs
@@ -59,7 +59,7 @@ namespace JobDriver.Execution
 				autoSdkSpan.SetAttribute("horde.workspace", "AutoSDK");
 
 				ManagedWorkspaceOptions options = WorkspaceInfo.GetMwOptions(_autoSdkWorkspaceInfo);
-				_autoSdkWorkspace = await WorkspaceInfo.CreateWorkspaceInfoAsync(_autoSdkWorkspaceInfo, _rootDir, options, logger, cancellationToken);
+				_autoSdkWorkspace = await WorkspaceInfo.CreateWorkspaceInfoAsync(_autoSdkWorkspaceInfo, _rootDir, options, Tracer, logger, cancellationToken);
 
 				using IPerforceConnection autoSdkPerforce = await PerforceConnection.CreateAsync(_autoSdkWorkspace.PerforceSettings, logger);
 				await _autoSdkWorkspace.SetupWorkspaceAsync(autoSdkPerforce, cancellationToken);
@@ -110,7 +110,7 @@ namespace JobDriver.Execution
 				
 				// Sync the regular workspace
 				ManagedWorkspaceOptions options = WorkspaceInfo.GetMwOptions(_workspaceInfo);
-				_workspace = await WorkspaceInfo.CreateWorkspaceInfoAsync(_workspaceInfo, _rootDir, options, logger, cancellationToken);
+				_workspace = await WorkspaceInfo.CreateWorkspaceInfoAsync(_workspaceInfo, _rootDir, options, Tracer, logger, cancellationToken);
 
 				using IPerforceConnection perforce = await PerforceConnection.CreateAsync(_workspace.PerforceSettings, logger);
 				await _workspace.SetupWorkspaceAsync(perforce, cancellationToken);
@@ -270,7 +270,7 @@ namespace JobDriver.Execution
 				{
 					ManagedWorkspaceOptions options = WorkspaceInfo.GetMwOptions(pendingWorkspace);
 
-					WorkspaceInfo workspace = await WorkspaceInfo.CreateWorkspaceInfoAsync(pendingWorkspace, rootDir, options, logger, cancellationToken);
+					WorkspaceInfo workspace = await WorkspaceInfo.CreateWorkspaceInfoAsync(pendingWorkspace, rootDir, options, tracer, logger, cancellationToken);
 					workspaces.Add(workspace);
 				}
 
diff --git a/Engine/Source/Programs/Horde/Drivers/JobDriver/Execution/WorkspaceExecutor.cs b/Engine/Source/Programs/Horde/Drivers/JobDriver/Execution/WorkspaceExecutor.cs
--- a/Engine/Source/Programs/Horde/Drivers/JobDriver/Execution/WorkspaceExecutor.cs
+++ b/Engine/Source/Programs/Horde/Drivers/JobDriver/Execution/WorkspaceExecutor.cs
@@ -49,23 +49,15 @@ namespace JobDriver.Execution
 			// Setup and sync the AutoSDK workspace
 			if (_autoSdkWorkspace != null)
 			{
-				using TelemetrySpan autoSdkSpan = Tracer.StartActiveSpan("Workspace");
-				autoSdkSpan.SetAttribute("horde.workspace", "AutoSDK");
-				// TODO: Set type of workspace materializer as scope tag.
-
 				SyncOptions syncOptions = new();
 				await _autoSdkWorkspace.SyncAsync(IWorkspaceMaterializer.LatestChangeNumber, -1, syncOptions, cancellationToken);
 			}
 
 			// Sync the regular workspace
 			{
-				using TelemetrySpan workspaceSpan = Tracer.StartActiveSpan("Workspace");
-				workspaceSpan.SetAttribute("horde.workspace", _workspace.Identifier);
-
 				await _workspace.SyncAsync(Batch.Change, Batch.PreflightChange, new SyncOptions(), cancellationToken);
 
 				// TODO: Purging of cache for ManagedWorkspace did happen here in WorkspaceInfo
-
 				DeleteCachedBuildGraphManifests(_workspace.DirectoryPath, logger);
 			}
 
diff --git a/Engine/Source/Programs/Horde/Drivers/JobDriver/Utility/WorkspaceInfo.cs b/Engine/Source/Programs/Horde/Drivers/JobDriver/Utility/WorkspaceInfo.cs
--- a/Engine/Source/Programs/Horde/Drivers/JobDriver/Utility/WorkspaceInfo.cs
+++ b/Engine/Source/Programs/Horde/Drivers/JobDriver/Utility/WorkspaceInfo.cs
@@ -10,6 +10,7 @@ using EpicGames.Perforce;
 using EpicGames.Perforce.Managed;
 using HordeCommon.Rpc.Messages;
 using Microsoft.Extensions.Logging;
+using OpenTelemetry.Trace;
 
 namespace JobDriver.Utility
 {
@@ -115,10 +116,11 @@ namespace JobDriver.Utility
 		/// <param name="workspace">The workspace definition</param>
 		/// <param name="rootDir">Root directory for storing the workspace</param>
 		/// <param name="options">Extra options for ManagedWorkspace</param>
+		/// <param name="tracer">Tracer</param>
 		/// <param name="logger">Logger output</param>
 		/// <param name="cancellationToken">Cancellation token</param>
 		/// <returns>New workspace info</returns>
-		public static async Task<WorkspaceInfo> CreateWorkspaceInfoAsync(RpcAgentWorkspace workspace, DirectoryReference rootDir, ManagedWorkspaceOptions options, ILogger logger, CancellationToken cancellationToken)
+		public static async Task<WorkspaceInfo> CreateWorkspaceInfoAsync(RpcAgentWorkspace workspace, DirectoryReference rootDir, ManagedWorkspaceOptions options, Tracer tracer, ILogger logger, CancellationToken cancellationToken)
 		{
 			// Fill in the default credentials iff they are not set
 			string? serverAndPort = String.IsNullOrEmpty(workspace.ServerAndPort) ? null : workspace.ServerAndPort;
@@ -209,7 +211,7 @@ namespace JobDriver.Utility
 			DirectoryReference workspaceDir = DirectoryReference.Combine(metadataDir, "Sync");
 
 			// Create the repository
-			ManagedWorkspace newRepository = await ManagedWorkspace.LoadOrCreateAsync(hostName, metadataDir, true, options, logger, cancellationToken);
+			ManagedWorkspace newRepository = await ManagedWorkspace.LoadOrCreateAsync(hostName, metadataDir, true, options, tracer, logger, cancellationToken);
 			return new WorkspaceInfo(perforceClientSettings, hostName, workspace.Stream, streamView, metadataDir, workspaceDir, workspace.View, !workspace.Incremental, newRepository);
 		}
 		
diff --git a/Engine/Source/Programs/Horde/HordeCmd/Commands/Perforce/PerforceBase.cs b/Engine/Source/Programs/Horde/HordeCmd/Commands/Perforce/PerforceBase.cs
--- a/Engine/Source/Programs/Horde/HordeCmd/Commands/Perforce/PerforceBase.cs
+++ b/Engine/Source/Programs/Horde/HordeCmd/Commands/Perforce/PerforceBase.cs
@@ -5,6 +5,7 @@ using EpicGames.Core;
 using EpicGames.Perforce;
 using EpicGames.Perforce.Managed;
 using Microsoft.Extensions.Logging;
+using OpenTelemetry.Trace;
 
 namespace Horde.Commands.Perforce
 {
@@ -64,7 +65,7 @@ namespace Horde.Commands.Perforce
 			bool useHaveTable = UseHaveTable.Equals("true", StringComparison.Ordinal);
 
 			ManagedWorkspaceOptions options = new() { UseHaveTable = useHaveTable };
-			ManagedWorkspace repo = await ManagedWorkspace.LoadOrCreateAsync(info.ClientHost!, BaseDir, Overwrite, options, logger, CancellationToken.None);
+			ManagedWorkspace repo = await ManagedWorkspace.LoadOrCreateAsync(info.ClientHost!, BaseDir, Overwrite, options, TracerProvider.Default.GetTracer("HordeCmd"), logger, CancellationToken.None);
 			await ExecuteAsync(perforce, repo, logger);
 			return 0;
 		}
diff --git a/Engine/Source/Programs/Shared/EpicGames.Perforce.Managed.Tests/ManagedWorkspaceTest.cs b/Engine/Source/Programs/Shared/EpicGames.Perforce.Managed.Tests/ManagedWorkspaceTest.cs
--- a/Engine/Source/Programs/Shared/EpicGames.Perforce.Managed.Tests/ManagedWorkspaceTest.cs
+++ b/Engine/Source/Programs/Shared/EpicGames.Perforce.Managed.Tests/ManagedWorkspaceTest.cs
@@ -10,6 +10,7 @@ using EpicGames.Core;
 using EpicGames.Perforce.Fixture;
 using Microsoft.Extensions.Logging;
 using Microsoft.VisualStudio.TestTools.UnitTesting;
+using OpenTelemetry.Trace;
 
 namespace EpicGames.Perforce.Managed.Tests;
 
@@ -352,7 +353,7 @@ public class ManagedWorkspaceTest : BasePerforceFixtureTest
 	private async Task<ManagedWorkspace> CreateManagedWorkspaceAsync(bool useHaveTable)
 	{
 		ManagedWorkspaceOptions options = new() { UseHaveTable = useHaveTable };
-		ManagedWorkspace ws = await ManagedWorkspace.CreateAsync(Environment.MachineName, TempDir, options, _mwLogger, CancellationToken.None);
+		ManagedWorkspace ws = await ManagedWorkspace.CreateAsync(Environment.MachineName, TempDir, options, TracerProvider.Default.GetTracer("Test"), _mwLogger, CancellationToken.None);
 		await ws.SetupAsync(PerforceConnection, StreamName, CancellationToken.None);
 		return ws;
 	}
diff --git a/Engine/Source/Programs/Shared/EpicGames.Perforce.Managed/EpicGames.Perforce.Managed.csproj b/Engine/Source/Programs/Shared/EpicGames.Perforce.Managed/EpicGames.Perforce.Managed.csproj
--- a/Engine/Source/Programs/Shared/EpicGames.Perforce.Managed/EpicGames.Perforce.Managed.csproj
+++ b/Engine/Source/Programs/Shared/EpicGames.Perforce.Managed/EpicGames.Perforce.Managed.csproj
@@ -28,6 +28,10 @@
     <DocumentationFile></DocumentationFile>
   </PropertyGroup>
 
+  <ItemGroup>
+    <PackageReference Include="OpenTelemetry" Version="1.8.0" />
+  </ItemGroup>
+  
   <ItemGroup>
     <ProjectReference Include="..\EpicGames.Core\EpicGames.Core.csproj" />
     <ProjectReference Include="..\EpicGames.IoHash\EpicGames.IoHash.csproj" />
diff --git a/Engine/Source/Programs/Shared/EpicGames.Perforce.Managed/ManagedWorkspace.cs b/Engine/Source/Programs/Shared/EpicGames.Perforce.Managed/ManagedWorkspace.cs
--- a/Engine/Source/Programs/Shared/EpicGames.Perforce.Managed/ManagedWorkspace.cs
+++ b/Engine/Source/Programs/Shared/EpicGames.Perforce.Managed/ManagedWorkspace.cs
@@ -14,6 +14,7 @@ using System.Threading.Tasks;
 using EpicGames.Core;
 using Microsoft.Extensions.Logging;
 using Microsoft.Extensions.Logging.Abstractions;
+using OpenTelemetry.Trace;
 
 namespace EpicGames.Perforce.Managed
 {
@@ -207,7 +208,12 @@ namespace EpicGames.Perforce.Managed
 		/// Whether a repair operation should be run on this workspace. Set whenever the state may be inconsistent.
 		/// </summary>
 		bool _requiresRepair;
-
+		
+		/// <summary>
+		/// Tracer
+		/// </summary>
+		readonly Tracer _tracer;
+		
 		/// <summary>
 		/// The log output device
 		/// </summary>
@@ -265,12 +271,14 @@ namespace EpicGames.Perforce.Managed
 		/// <param name="nextSequenceNumber">The next sequence number for operations</param>
 		/// <param name="baseDir">The root directory for the stash</param>
 		/// <param name="options">Extra options</param>
+		/// <param name="tracer">Tracer</param>
 		/// <param name="logger">The log output device</param>
-		private ManagedWorkspace(string hostName, uint nextSequenceNumber, DirectoryReference baseDir, ManagedWorkspaceOptions options, ILogger logger)
+		private ManagedWorkspace(string hostName, uint nextSequenceNumber, DirectoryReference baseDir, ManagedWorkspaceOptions options, Tracer tracer, ILogger logger)
 		{
 			// Save the Perforce settings
 			_hostName = hostName;
 			_nextSequenceNumber = nextSequenceNumber;
+			_tracer = tracer;
 			_logger = logger;
 			_options = options;
 
@@ -295,16 +303,17 @@ namespace EpicGames.Perforce.Managed
 		/// <param name="baseDir">The base directory for the repository</param>
 		/// <param name="overwrite">Whether to allow overwriting a repository that's not up to date</param>
 		/// <param name="options">Extra options</param>
+		/// <param name="tracer">Tracer</param>
 		/// <param name="logger">The logging interface</param>
 		/// <param name="cancellationToken">Cancellation token for this operation</param>
 		/// <returns></returns>
-		public static async Task<ManagedWorkspace> LoadOrCreateAsync(string hostName, DirectoryReference baseDir, bool overwrite, ManagedWorkspaceOptions options, ILogger logger, CancellationToken cancellationToken)
+		public static async Task<ManagedWorkspace> LoadOrCreateAsync(string hostName, DirectoryReference baseDir, bool overwrite, ManagedWorkspaceOptions options, Tracer tracer, ILogger logger, CancellationToken cancellationToken)
 		{
 			if (Exists(baseDir))
 			{
 				try
 				{
-					return await LoadAsync(hostName, baseDir, options, logger, cancellationToken);
+					return await LoadAsync(hostName, baseDir, options, tracer, logger, cancellationToken);
 				}
 				catch (Exception ex)
 				{
@@ -319,7 +328,7 @@ namespace EpicGames.Perforce.Managed
 				}
 			}
 
-			return await CreateAsync(hostName, baseDir, options, logger, cancellationToken);
+			return await CreateAsync(hostName, baseDir, options, tracer, logger, cancellationToken);
 		}
 		/*
 				public static PerforceConnection GetPerforceConnection(PerforceConnection Perforce)
@@ -353,18 +362,20 @@ namespace EpicGames.Perforce.Managed
 		/// <param name="hostName">Name of the current machine.</param>
 		/// <param name="baseDir">The base directory for the repository</param>
 		/// <param name="options">Extra options</param>
+		/// <param name="tracer">Tracer</param>
 		/// <param name="logger">The log output device</param>
 		/// <param name="cancellationToken">Cancellation token for this operation</param>
 		/// <returns>New repository instance</returns>
-		public static async Task<ManagedWorkspace> CreateAsync(string hostName, DirectoryReference baseDir, ManagedWorkspaceOptions options, ILogger logger, CancellationToken cancellationToken)
+		public static async Task<ManagedWorkspace> CreateAsync(string hostName, DirectoryReference baseDir, ManagedWorkspaceOptions options, Tracer tracer, ILogger logger, CancellationToken cancellationToken)
 		{
+			using TelemetrySpan span = tracer.StartActiveSpan($"{nameof(ManagedWorkspace)}.{nameof(CreateAsync)}");
 			logger.LogInformation("Creating repository at {Location}...", baseDir);
 
 			// Make sure all the fields are valid
 			DirectoryReference.CreateDirectory(baseDir);
 			FileUtils.ForceDeleteDirectoryContents(baseDir);
 
-			ManagedWorkspace repo = new ManagedWorkspace(hostName, 1, baseDir, options, logger);
+			ManagedWorkspace repo = new ManagedWorkspace(hostName, 1, baseDir, options, tracer, logger);
 			await repo.SaveAsync(TransactionState.Clean, cancellationToken);
 			repo.CreateCacheHierarchy();
 
@@ -405,10 +416,13 @@ namespace EpicGames.Perforce.Managed
 		/// <param name="hostName">Name of the current host. Will be obtained from a 'p4 info' call if not specified</param>
 		/// <param name="baseDir">The base directory for the repository</param>
 		/// <param name="options">Extra options</param> 
+		/// <param name="tracer">Tracer</param>
 		/// <param name="logger">The log output device</param>
 		/// <param name="cancellationToken">Cancellation token for this command</param>
-		public static async Task<ManagedWorkspace> LoadAsync(string hostName, DirectoryReference baseDir, ManagedWorkspaceOptions options, ILogger logger, CancellationToken cancellationToken)
+		public static async Task<ManagedWorkspace> LoadAsync(string hostName, DirectoryReference baseDir, ManagedWorkspaceOptions options, Tracer tracer, ILogger logger, CancellationToken cancellationToken)
 		{
+			using TelemetrySpan span = tracer.StartActiveSpan($"{nameof(ManagedWorkspace)}.{nameof(LoadAsync)}");
+			
 			if (!Exists(baseDir))
 			{
 				throw new FatalErrorException("No valid repository found at {0}", baseDir);
@@ -429,7 +443,7 @@ namespace EpicGames.Perforce.Managed
 			bool requiresRepair = reader.ReadBoolean();
 			uint nextSequenceNumber = reader.ReadUInt32();
 
-			ManagedWorkspace repo = new(hostName, nextSequenceNumber, baseDir, options, logger);
+			ManagedWorkspace repo = new(hostName, nextSequenceNumber, baseDir, options, tracer, logger);
 			repo._requiresRepair = requiresRepair;
 
 			int numTrackedFiles = reader.ReadInt32();
@@ -451,6 +465,8 @@ namespace EpicGames.Perforce.Managed
 		/// </summary>
 		private async Task SaveAsync(TransactionState state, CancellationToken cancellationToken)
 		{
+			using TelemetrySpan span = _tracer.StartActiveSpan($"{nameof(ManagedWorkspace)}.{nameof(SaveAsync)}");
+			
 			// Allocate the buffer for writing
 			int serializedSize = sizeof(int) + sizeof(byte) + sizeof(int) + sizeof(int) + _contentIdToTrackedFile.Values.Sum(x => x.GetSerializedSize()) + _workspace.GetSerializedSize();
 			byte[] buffer = new byte[serializedSize];
@@ -567,16 +583,21 @@ namespace EpicGames.Perforce.Managed
 		/// <param name="cancellationToken">Cancellation token</param>
 		private async Task CleanInternalAsync(bool removeUntracked, CancellationToken cancellationToken)
 		{
+			using TelemetrySpan span = _tracer.StartActiveSpan($"{nameof(ManagedWorkspace)}.{nameof(CleanInternalAsync)}");
+			
 			FileInfo[] filesToDelete;
 			DirectoryInfo[] directoriesToDelete;
 			using (Trace("FindFilesToClean"))
 			using (ILoggerProgress status = _logger.BeginProgressScope("Finding files to clean..."))
 			{
+				using TelemetrySpan findFilesSpan = _tracer.StartActiveSpan($"{nameof(ManagedWorkspace)}.FindFilesToClean");
 				Stopwatch timer = Stopwatch.StartNew();
 
 				(FileInfo[] filesToDelete, DirectoryInfo[] directoriesToDelete) refreshResult = await _workspace.RefreshAsync(removeUntracked, _options.MaxFileConcurrency);
 				filesToDelete = refreshResult.filesToDelete;
 				directoriesToDelete = refreshResult.directoriesToDelete;
+				findFilesSpan.SetAttribute("horde.mw.num_delete_files", filesToDelete.Length);
+				findFilesSpan.SetAttribute("horde.mw.num_delete_dirs", directoriesToDelete.Length);
 
 				status.Progress = $"({timer.Elapsed.TotalSeconds:0.0}s)";
 			}
@@ -600,20 +621,27 @@ namespace EpicGames.Perforce.Managed
 				using (Trace("CleanFiles"))
 				using (ILoggerProgress scope = _logger.BeginProgressScope("Cleaning files..."))
 				{
+					using TelemetrySpan cleanSpan = _tracer.StartActiveSpan($"{nameof(ManagedWorkspace)}.CleaningFilesAndDirs");
 					Stopwatch timer = Stopwatch.StartNew();
 
 					ParallelOptions options = new() { MaxDegreeOfParallelism = _options.MaxFileConcurrency, CancellationToken = cancellationToken };
-					await Parallel.ForEachAsync(filesToDelete, options, (fileToDelete, ct) =>
 					{
-						FileUtils.ForceDeleteFile(fileToDelete);
-						return ValueTask.CompletedTask;
-					});
-
-					await Parallel.ForEachAsync(directoriesToDelete, options, (directoryToDelete, ct) =>
+						using TelemetrySpan cleanFilesSpan = _tracer.StartActiveSpan($"{nameof(ManagedWorkspace)}.CleaningFiles");
+						await Parallel.ForEachAsync(filesToDelete, options, (fileToDelete, ct) =>
+						{
+							FileUtils.ForceDeleteFile(fileToDelete);
+							return ValueTask.CompletedTask;
+						});
+					}
+					
 					{
-						FileUtils.ForceDeleteDirectory(directoryToDelete);
-						return ValueTask.CompletedTask;
-					});
+						using TelemetrySpan cleanDirsSpan = _tracer.StartActiveSpan($"{nameof(ManagedWorkspace)}.CleaningDirs");
+						await Parallel.ForEachAsync(directoriesToDelete, options, (directoryToDelete, ct) =>
+						{
+							FileUtils.ForceDeleteDirectory(directoryToDelete);
+							return ValueTask.CompletedTask;
+						});
+					}
 
 					scope.Progress = $"({timer.Elapsed.TotalSeconds:0.0}s)";
 				}
@@ -679,6 +707,8 @@ namespace EpicGames.Perforce.Managed
 		/// </summary>
 		public async Task RepairCacheAsync(CancellationToken cancellationToken)
 		{
+			using TelemetrySpan span = _tracer.StartActiveSpan($"{nameof(ManagedWorkspace)}.{nameof(RepairCacheAsync)}");
+			
 			using (Trace("Repair"))
 			using (ILoggerProgress status = _logger.BeginProgressScope("Checking cache..."))
 			{
@@ -742,6 +772,7 @@ namespace EpicGames.Perforce.Managed
 		/// <param name="cancellationToken">Cancellation token</param>
 		public async Task PurgeAsync(long maxSize, CancellationToken cancellationToken)
 		{
+			using TelemetrySpan span = _tracer.StartActiveSpan($"{nameof(ManagedWorkspace)}.{nameof(PurgeAsync)}");
 			_logger.LogInformation("Purging cache (limit {MaxSize:n0} bytes)...", maxSize);
 			using (Trace("Purge"))
 			using (_logger.BeginIndentScope("  "))
@@ -783,6 +814,8 @@ namespace EpicGames.Perforce.Managed
 		/// </summary>
 		public async Task StatsAsync(IPerforceConnection perforceClient, List<string> streamNames, List<string> view, CancellationToken cancellationToken)
 		{
+			using TelemetrySpan span = _tracer.StartActiveSpan($"{nameof(ManagedWorkspace)}.{nameof(StatsAsync)}");
+			
 			_logger.LogInformation("Finding stats for {NumStreams} streams", streamNames.Count);
 			using (_logger.BeginIndentScope("  "))
 			{
@@ -889,6 +922,8 @@ namespace EpicGames.Perforce.Managed
 		/// </summary>
 		public async Task StatusAsync()
 		{
+			using TelemetrySpan span = _tracer.StartActiveSpan($"{nameof(ManagedWorkspace)}.{nameof(StatusAsync)}");
+			
 			// Print size stats
 			_logger.LogInformation("Cache contains {NumFiles:n0} files, {TotalSize:n1}mb", _contentIdToTrackedFile.Count, _contentIdToTrackedFile.Values.Sum(x => x.Length) / (1024.0 * 1024.0));
 			_logger.LogInformation("Stage contains {NumFiles:n0} files, {TotalSize:n1}mb", _workspace.GetFiles().Count, _workspace.GetFiles().Sum(x => x._length) / (1024.0 * 1024.0));
@@ -935,6 +970,9 @@ namespace EpicGames.Perforce.Managed
 		/// <param name="cancellationToken">Cancellation token</param>
 		public async Task SyncAsync(IPerforceConnection perforce, string streamName, int changeNumber, IReadOnlyList<string> view, bool removeUntracked, bool fakeSync, FileReference? cacheFile, CancellationToken cancellationToken)
 		{
+			using TelemetrySpan span = _tracer.StartActiveSpan($"{nameof(ManagedWorkspace)}.{nameof(SyncAsync)}");
+			span.SetAttribute("horde.mw.use_have_table", _options.UseHaveTable);
+			
 			Stopwatch timer = Stopwatch.StartNew();
 			if (changeNumber == -1)
 			{
@@ -1013,6 +1051,8 @@ namespace EpicGames.Perforce.Managed
 		/// <returns>Async task</returns>
 		public async Task UnshelveAsync(IPerforceConnection perforce, int unshelveChangelist, CancellationToken cancellationToken)
 		{
+			using TelemetrySpan span = _tracer.StartActiveSpan($"{nameof(ManagedWorkspace)}.{nameof(UnshelveAsync)}");
+			
 			// Need to mark those files as dirty - update the workspace with those files 
 			// Delete is fine, but need to flag anything added
 
@@ -1122,6 +1162,8 @@ namespace EpicGames.Perforce.Managed
 		/// </summary>
 		public async Task<Tuple<int, StreamSnapshot>[]> PopulateCleanAsync(List<PopulateRequest> requests, CancellationToken cancellationToken)
 		{
+			using TelemetrySpan span = _tracer.StartActiveSpan($"{nameof(ManagedWorkspace)}.{nameof(PopulateCleanAsync)}");
+			
 			// Revert all changes in each of the unique clients
 			foreach (PopulateRequest request in requests)
 			{
@@ -1225,6 +1267,8 @@ namespace EpicGames.Perforce.Managed
 		/// </summary>
 		public async Task PopulateSyncAsync(List<PopulateRequest> requests, Tuple<int, StreamSnapshot>[] streamState, bool fakeSync, CancellationToken cancellationToken)
 		{
+			using TelemetrySpan span = _tracer.StartActiveSpan($"{nameof(ManagedWorkspace)}.{nameof(PopulateSyncAsync)}");
+			
 			// Sync all the new files
 			for (int idx = 0; idx < requests.Count; idx++)
 			{
@@ -1265,6 +1309,7 @@ namespace EpicGames.Perforce.Managed
 		/// <returns>Async task</returns>
 		public async Task DeleteClientAsync(IPerforceConnection perforceClient, CancellationToken cancellationToken)
 		{
+			using TelemetrySpan span = _tracer.StartActiveSpan($"{nameof(ManagedWorkspace)}.{nameof(DeleteClientAsync)}");
 			async Task DeleteAsync(string clientName)
 			{
 				PerforceResponse response = await perforceClient.TryDeleteClientAsync(DeleteClientOptions.None, clientName, cancellationToken);
@@ -1285,6 +1330,7 @@ namespace EpicGames.Perforce.Managed
 
 		private async Task<ClientRecord> GetOrCreateClientAsync(IPerforceConnection perforceClient, string streamName, CancellationToken cancellationToken)
 		{
+			using TelemetrySpan span = _tracer.StartActiveSpan($"{nameof(ManagedWorkspace)}.{nameof(GetOrCreateClientAsync)}");
 			string clientName = perforceClient.Settings.ClientName!;//GetClientName(perforceClient, useHaveTable);
 			if (_createdClients.TryGetValue(clientName, out ClientRecord? client) && client.Stream == streamName)
 			{
@@ -1375,6 +1421,7 @@ namespace EpicGames.Perforce.Managed
 		private async Task<int> GetLatestClientChangeAsync(IPerforceConnection perforceClient, CancellationToken cancellationToken)
 		{
 			int changeNumber;
+			using TelemetrySpan span = _tracer.StartActiveSpan($"{nameof(ManagedWorkspace)}.{nameof(GetLatestClientChangeAsync)}");
 			using (Trace("FindChange"))
 			using (ILoggerProgress status = _logger.BeginProgressScope("Finding latest change..."))
 			{
@@ -1399,6 +1446,7 @@ namespace EpicGames.Perforce.Managed
 		/// <param name="cancellationToken">Cancellation token</param>
 		private async Task RevertInternalAsync(IPerforceConnection perforceClient, CancellationToken cancellationToken)
 		{
+			using TelemetrySpan span = _tracer.StartActiveSpan($"{nameof(ManagedWorkspace)}.{nameof(RevertInternalAsync)}");
 			using (Trace("Revert"))
 			using (ILoggerProgress status = _logger.BeginProgressScope("Reverting changes..."))
 			{
@@ -1444,6 +1492,7 @@ namespace EpicGames.Perforce.Managed
 		/// <param name="cancellationToken">The cancellation token</param>
 		private async Task ClearClientHaveTableAsync(IPerforceConnection perforceClient, CancellationToken cancellationToken)
 		{
+			using TelemetrySpan span = _tracer.StartActiveSpan($"{nameof(ManagedWorkspace)}.{nameof(ClearClientHaveTableAsync)}");
 			using (Trace("ClearHaveTable"))
 			using (ILoggerProgress scope = _logger.BeginProgressScope("Clearing have table..."))
 			{
@@ -1462,6 +1511,7 @@ namespace EpicGames.Perforce.Managed
 		/// <param name="cancellationToken">Cancellation token</param>
 		private async Task UpdateClientHaveTableAsync(IPerforceConnection perforceClient, int changeNumber, IReadOnlyList<string> view, CancellationToken cancellationToken)
 		{
+			using TelemetrySpan span = _tracer.StartActiveSpan($"{nameof(ManagedWorkspace)}.{nameof(UpdateClientHaveTableAsync)}");
 			using (Trace("UpdateHaveTable"))
 			using (ILoggerProgress scope = _logger.BeginProgressScope("Updating have table..."))
 			{
@@ -1562,7 +1612,8 @@ namespace EpicGames.Perforce.Managed
 		private async Task<StreamSnapshotFromMemory> FindClientContentsAsync(IPerforceConnection perforceClient, int changeNumber, CancellationToken cancellationToken)
 		{
 			StreamTreeBuilder builder = new StreamTreeBuilder();
-
+			
+			using TelemetrySpan span = _tracer.StartActiveSpan($"{nameof(ManagedWorkspace)}.{nameof(FindClientContentsAsync)}");
 			using (Trace("FetchMetadata"))
 			using (ILoggerProgress scope = _logger.BeginProgressScope("Fetching metadata..."))
 			{
@@ -1724,7 +1775,8 @@ namespace EpicGames.Perforce.Managed
 		public async Task<StreamSnapshotFromMemory> FindClientContentsWithoutHaveTableAsync(IPerforceConnection perforceClient, string streamName, IReadOnlyList<string> view, int changeNumber, CancellationToken cancellationToken)
 		{
 			DepotStreamTreeBuilder builder = new();
-
+			
+			using TelemetrySpan span = _tracer.StartActiveSpan($"{nameof(ManagedWorkspace)}.{nameof(FindClientContentsWithoutHaveTableAsync)}");
 			using (Trace("FetchMetadata"))
 			using (ILoggerProgress scope = _logger.BeginProgressScope("Fetching metadata (without have table)..."))
 			{
@@ -1795,6 +1847,7 @@ namespace EpicGames.Perforce.Managed
 			StreamSnapshot? contents = null;
 			if (FileReference.Exists(cacheFile))
 			{
+				using TelemetrySpan span = _tracer.StartActiveSpan($"{nameof(ManagedWorkspace)}.{nameof(TryLoadClientContentsAsync)}");
 				using (Trace("ReadMetadata"))
 				using (ILoggerProgress scope = _logger.BeginProgressScope($"Reading cached metadata from {cacheFile}..."))
 				{
@@ -1821,7 +1874,8 @@ namespace EpicGames.Perforce.Managed
 			StreamSnapshotFromMemory contents = _options.UseHaveTable
 				? await FindClientContentsAsync(perforceClient, changeNumber, cancellationToken)
 				: await FindClientContentsWithoutHaveTableAsync(perforceClient, basePath.ToString(), view, changeNumber, cancellationToken);
-
+			
+			using TelemetrySpan span = _tracer.StartActiveSpan($"{nameof(ManagedWorkspace)}.{nameof(FindAndSaveClientContentsAsync)}");
 			using (Trace("WriteMetadata"))
 			using (ILoggerProgress scope = _logger.BeginProgressScope($"Saving metadata to {cacheFile}..."))
 			{
@@ -1862,6 +1916,7 @@ namespace EpicGames.Perforce.Managed
 
 			// Figure out what to remove
 			RemoveTransaction transaction;
+			using TelemetrySpan span = _tracer.StartActiveSpan($"{nameof(ManagedWorkspace)}.{nameof(RemoveFilesFromWorkspaceAsync)}");
 			using (Trace("GatherFilesToRemove"))
 			using (ILoggerProgress scope = _logger.BeginProgressScope("Gathering files to remove..."))
 			{
@@ -2014,6 +2069,7 @@ namespace EpicGames.Perforce.Managed
 
 			// Figure out what we need to do
 			AddTransaction transaction;
+			using TelemetrySpan span = _tracer.StartActiveSpan($"{nameof(ManagedWorkspace)}.{nameof(AddFilesToWorkspaceAsync)}");
 			using (Trace("GatherFilesToAdd"))
 			using (ILoggerProgress status = _logger.BeginProgressScope("Gathering files to add..."))
 			{
@@ -2197,6 +2253,7 @@ namespace EpicGames.Perforce.Managed
 		/// <returns>Async task</returns>
 		async Task SyncBatchAsync(IPerforceConnection client, WorkspaceFileToSync[] filesToSync, int beginIdx, int endIdx, bool fakeSync, Stack<IPerforceConnection> connectionPool, CancellationToken cancellationToken)
 		{
+			using TelemetrySpan span = _tracer.StartActiveSpan($"{nameof(ManagedWorkspace)}.{nameof(SyncBatchAsync)}");
 			if (fakeSync)
 			{
 				for (int idx = beginIdx; idx < endIdx; idx++)
