Cloud DDC - Changed server timing metrics into a concurrent bag instead of list with locks to reduce lock time a bit.

diff --git a/Engine/Source/Programs/Shared/EpicGames.AspNet/ServerTimings.cs b/Engine/Source/Programs/Shared/EpicGames.AspNet/ServerTimings.cs
--- a/Engine/Source/Programs/Shared/EpicGames.AspNet/ServerTimings.cs
+++ b/Engine/Source/Programs/Shared/EpicGames.AspNet/ServerTimings.cs
@@ -1,6 +1,7 @@
 // Copyright Epic Games, Inc. All Rights Reserved.
 
 using System;
+using System.Collections.Concurrent;
 using System.Collections.Generic;
 using System.Globalization;
 using System.Linq;
@@ -85,14 +86,11 @@ namespace EpicGames.AspNet
 
 	public class ServerTiming : IServerTiming
 	{
-		private readonly List<ServerTimingMetric> _metrics = new List<ServerTimingMetric>();
+		private readonly ConcurrentBag<ServerTimingMetric> _metrics = new ConcurrentBag<ServerTimingMetric>();
 
 		public void AddServerTimingMetric(string metricName, double? duration, string? description)
 		{
-			lock (_metrics)
-			{
-				_metrics.Add(new ServerTimingMetric(metricName, duration, description));
-			}
+			_metrics.Add(new ServerTimingMetric(metricName, duration, description));
 		}
 
 		public ServerTimingMetricScoped CreateServerTimingMetricScope(string metricName, string? description)
