Horde: Do not retry HTTP requests which throw an HttpRequestException. These generally represent a non-transient error, unlike HTTP status responses.

In the case of bad server addresses or ports, requests may take a long time to timeout. Retrying them does not make sense, and creates very long delays before the result is reported.

diff --git a/Engine/Source/Programs/Shared/EpicGames.Horde/HordeHttpMessageHandler.cs b/Engine/Source/Programs/Shared/EpicGames.Horde/HordeHttpMessageHandler.cs
--- a/Engine/Source/Programs/Shared/EpicGames.Horde/HordeHttpMessageHandler.cs
+++ b/Engine/Source/Programs/Shared/EpicGames.Horde/HordeHttpMessageHandler.cs
@@ -79,8 +79,7 @@ namespace EpicGames.Horde
 			// Policy for transient errors is the same as HttpPolicyExtensions.HandleTransientHttpError(), but excludes HttpStatusCode.ServiceUnavailable (which is used as a response
 			// when allocating compute resources when none are available). This pathway is handled explicitly on the application side.
 			return Policy<HttpResponseMessage>
-				.Handle<HttpRequestException>()
-				.OrResult(x => (x.StatusCode >= HttpStatusCode.InternalServerError && x.StatusCode != HttpStatusCode.ServiceUnavailable) || x.StatusCode == HttpStatusCode.RequestTimeout)
+				.HandleResult(x => (x.StatusCode >= HttpStatusCode.InternalServerError && x.StatusCode != HttpStatusCode.ServiceUnavailable) || x.StatusCode == HttpStatusCode.RequestTimeout)
 				.WaitAndRetryAsync(retryTimes, OnTimeoutAsync);
 		}
 	}
