Horde: Improvements to connection handling and authorization.

- Authentication state is now determined via a call to /api/v1/dashboard/challenge. This will actually attempt to use the current access token, which will prompt a token refresh on auth failure using HordeHttpAuthHandler.
- HordeHttpAuthHandler now checks a HTTP request option allowing the use of an interactive login, allowing interactive logins to be configured on a per-request basis.
- HordeHttpAuthHandler has a Reset() method which will discard the state of any in-flight auth requests, allowing any subsequent requests to start from scratch.

#rnx

diff --git a/Engine/Source/Programs/Shared/EpicGames.Horde/HordeClient.cs b/Engine/Source/Programs/Shared/EpicGames.Horde/HordeClient.cs
--- a/Engine/Source/Programs/Shared/EpicGames.Horde/HordeClient.cs
+++ b/Engine/Source/Programs/Shared/EpicGames.Horde/HordeClient.cs
@@ -108,7 +108,7 @@ namespace EpicGames.Horde
 		}
 
 		/// <inheritdoc/>
-		public abstract Task<bool> LoginAsync(bool allowLogin, CancellationToken cancellationToken);
+		public abstract Task<bool> LoginAsync(bool interactive, CancellationToken cancellationToken);
 
 		/// <inheritdoc/>
 		public abstract bool HasValidAccessToken();
@@ -291,8 +291,15 @@ namespace EpicGames.Horde
 		}
 
 		/// <inheritdoc/>
-		public override Task<bool> LoginAsync(bool allowLogin, CancellationToken cancellationToken)
-			=> Task.FromResult(true);
+		public override async Task<bool> LoginAsync(bool allowLogin, CancellationToken cancellationToken)
+		{
+			using (HttpClient httpClient = CreateAuthenticatedHttpClient())
+			{
+				using HttpRequestMessage request = new HttpRequestMessage(HttpMethod.Get, "api/v1/dashboard/challenge");
+				using HttpResponseMessage response = await httpClient.SendAsync(request, cancellationToken);
+				return response.IsSuccessStatusCode;
+			}
+		}
 
 		/// <inheritdoc/>
 		public override bool HasValidAccessToken()
@@ -359,7 +366,18 @@ namespace EpicGames.Horde
 		/// <inheritdoc/>
 		public override async Task<bool> LoginAsync(bool allowLogin, CancellationToken cancellationToken)
 		{
-			return await _authHandlerState.LoginAsync(allowLogin, cancellationToken);
+			// Reset any cached state in the auth handler
+			_authHandlerState.Reset();
+
+			// Send a request to the server to log in automatically
+			using (HttpClient httpClient = CreateAuthenticatedHttpClient())
+			{
+				using HttpRequestMessage request = new HttpRequestMessage(HttpMethod.Get, "api/v1/dashboard/challenge");
+				request.Options.Set(HordeHttpAuthHandler.AllowInteractiveLogin, true);
+
+				using HttpResponseMessage response = await httpClient.SendAsync(request, cancellationToken);
+				return response.IsSuccessStatusCode;
+			}
 		}
 
 		/// <inheritdoc/>
diff --git a/Engine/Source/Programs/Shared/EpicGames.Horde/HordeHttpAuthHandler.cs b/Engine/Source/Programs/Shared/EpicGames.Horde/HordeHttpAuthHandler.cs
--- a/Engine/Source/Programs/Shared/EpicGames.Horde/HordeHttpAuthHandler.cs
+++ b/Engine/Source/Programs/Shared/EpicGames.Horde/HordeHttpAuthHandler.cs
@@ -2,6 +2,7 @@
 
 using System;
 using System.Collections.Generic;
+using System.Diagnostics.CodeAnalysis;
 using System.Linq;
 using System.Net;
 using System.Net.Http;
@@ -27,6 +28,11 @@ namespace EpicGames.Horde
 		readonly HordeHttpAuthHandlerState _authState;
 		readonly IOptions<HordeOptions> _options;
 
+		/// <summary>
+		/// Option for HTTP requests that can override the default behavior for whether to enable interactive auth prompts
+		/// </summary>
+		public static HttpRequestOptionsKey<bool> AllowInteractiveLogin { get; } = new HttpRequestOptionsKey<bool>("Horde-AllowInteractiveLogin");
+
 		/// <summary>
 		/// Constructor
 		/// </summary>
@@ -46,12 +52,19 @@ namespace EpicGames.Horde
 				return await base.SendAsync(request, cancellationToken);
 			}
 
+			// Check whether the request specifically allows interactive auth, otherwise fall back to the default
+			bool allowInteractiveLogin;
+			if (!request.Options.TryGetValue(AllowInteractiveLogin, out allowInteractiveLogin))
+			{
+				allowInteractiveLogin = _options.Value.AllowAuthPrompt;
+			}
+
 			// Get the current access token and send the request with that
-			string? accessToken = await _authState.GetAccessTokenAsync(_options.Value.AllowAuthPrompt, cancellationToken);
+			string? accessToken = await _authState.GetAccessTokenAsync(allowInteractiveLogin, cancellationToken);
 			for (int attempt = 0; ; attempt++)
 			{
 				// Attempt to perform the request with this access token
-				request.Headers.Authorization = (accessToken == null)? null : new AuthenticationHeaderValue("Bearer", accessToken);
+				request.Headers.Authorization = (accessToken == null) ? null : new AuthenticationHeaderValue("Bearer", accessToken);
 				HttpResponseMessage response = await base.SendAsync(request, cancellationToken);
 
 				const int MaxAttempts = 3;
@@ -67,7 +80,7 @@ namespace EpicGames.Horde
 				}
 
 				// Get the next token, and quit out if it's the same
-				string? nextAccessToken = await _authState.GetAccessTokenAsync(_options.Value.AllowAuthPrompt, cancellationToken);
+				string? nextAccessToken = await _authState.GetAccessTokenAsync(allowInteractiveLogin, cancellationToken);
 				if (String.Equals(accessToken, nextAccessToken, StringComparison.Ordinal))
 				{
 					return response;
@@ -95,13 +108,17 @@ namespace EpicGames.Horde
 		}
 
 		readonly object _lockObject = new object();
-		readonly CancellationTokenSource _cancellationTokenSource = new CancellationTokenSource();
-		Task<AuthState>? _authStateTask = null;
 		readonly HttpMessageHandler _httpMessageHandler;
 		readonly Uri _serverUrl;
 		readonly IOptions<HordeOptions> _options;
 		readonly ILogger _logger;
-		
+
+		// Background auth process
+		int _currentAuthTaskId;
+		bool _currentAuthInteractive;
+		BackgroundTask? _currentAuthWorker;
+		TaskCompletionSource<AuthState>? _currentAuthResult;
+
 		// Allow these to be overridden in tests
 		readonly ITokenStore? _tokenStore;
 		readonly IOidcTokenManager? _oidcTokenManager;
@@ -136,21 +153,11 @@ namespace EpicGames.Horde
 		/// <inheritdoc/>
 		public async ValueTask DisposeAsync()
 		{
-			if (_authStateTask != null && !_authStateTask.IsCompleted)
+			if (_currentAuthWorker != null)
 			{
-				await _cancellationTokenSource.CancelAsync();
-				try
-				{
-#pragma warning disable VSTHRD003
-					await _authStateTask;
-#pragma warning restore VSTHRD003
-				}
-				catch (OperationCanceledException)
-				{
-				}
+				await _currentAuthWorker.DisposeAsync();
+				_currentAuthWorker = null;
 			}
-
-			_cancellationTokenSource.Dispose();
 		}
 
 		/// <summary>
@@ -162,27 +169,81 @@ namespace EpicGames.Horde
 			{
 				return true;
 			}
-			if (_authStateTask != null && _authStateTask.IsCompletedSuccessfully && _authStateTask.TryGetResult(out AuthState? authState) && authState.IsAuthorized())
+			if (TryGetCurrentAuthState(out AuthState? authState) && authState.IsAuthorized())
 			{
 				return true;
 			}
 			return false;
 		}
 
+		/// <summary>
+		/// Gets the current access token
+		/// </summary>
+		public string? GetCurrentAccessToken()
+		{
+			string? accessToken = GetAccessTokenFromConfig();
+			if (accessToken != null)
+			{
+				return accessToken;
+			}
+			if (TryGetCurrentAuthState(out AuthState? authState) && authState.IsAuthorized())
+			{
+				return authState.TokenInfo?.AccessToken;
+			}
+			return null;
+		}
+
+		/// <summary>
+		/// Gets the current auth state instance. Fails if the current auth task has not finished.
+		/// </summary>
+		bool TryGetCurrentAuthState([NotNullWhen(true)] out AuthState? authState)
+		{
+			AuthState? state;
+			if (_currentAuthResult != null && _currentAuthResult.Task.IsCompletedSuccessfully && _currentAuthResult.Task.TryGetResult(out state))
+			{
+				authState = state;
+				return true;
+			}
+			else
+			{
+				authState = null;
+				return false;
+			}
+		}
+
+		/// <summary>
+		/// Resets the current auth state
+		/// </summary>
+		public void Reset()
+		{
+			if (GetAccessTokenFromConfig() != null)
+			{
+				return;
+			}
+
+			lock (_lockObject)
+			{
+				_currentAuthResult = null;
+			}
+		}
+
 		/// <summary>
 		/// Marks the given access token as invalid, having attempted to use it and got an unauthorized response
 		/// </summary>
 		/// <param name="accessToken">The access  header to invalidate</param>
 		public void Invalidate(string? accessToken)
 		{
+			if (GetAccessTokenFromConfig() != null)
+			{
+				return;
+			}
+
 			lock (_lockObject)
 			{
-#pragma warning disable VSTHRD002
-				if (_authStateTask != null && _authStateTask.IsCompleted && Object.Equals(_authStateTask.Result?.TokenInfo?.AccessToken, accessToken))
+				if (TryGetCurrentAuthState(out AuthState? authState) && Object.Equals(authState.TokenInfo?.AccessToken, accessToken))
 				{
-					_authStateTask = null;
+					_currentAuthResult = null;
 				}
-#pragma warning restore VSTHRD002
 			}
 		}
 
@@ -216,62 +277,100 @@ namespace EpicGames.Horde
 			return null;
 		}
 
-		/// <summary>
-		/// Refresh the auth state
-		/// </summary>
-		/// <param name="interactive">Whether to allow logging in interactively</param>
-		/// <param name="cancellationToken">Cancellation token for the operation</param>
-		/// <returns></returns>
-		public async Task<bool> LoginAsync(bool interactive, CancellationToken cancellationToken)
-		{
-			if (GetAccessTokenFromConfig() != null)
-			{
-				return true;
-			}
-
-			AuthState? state = await GetAuthStateAsync(interactive, cancellationToken);
-			return state?.IsAuthorized() ?? false;
-		}
-
 		/// <summary>
 		/// Gets the current access token
 		/// </summary>
-		public async Task<string?> GetAccessTokenAsync(bool allowAuthPrompt, CancellationToken cancellationToken)
+		public async Task<string?> GetAccessTokenAsync(bool interactive, CancellationToken cancellationToken)
 		{
+			// If there's a static access token configured, use that
 			string? accessToken = GetAccessTokenFromConfig();
 			if (accessToken != null)
 			{
 				return accessToken;
 			}
 
-			AuthState? authState = await GetAuthStateAsync(allowAuthPrompt, cancellationToken);
+			// Otherwise check if we need to start a background task to figure it out
+			Task<AuthState> authStateTask;
+			lock (_lockObject)
+			{
+				if (_currentAuthResult == null || (interactive && !_currentAuthInteractive))
+				{
+					int authTaskId = ++_currentAuthTaskId;
+
+					if (_currentAuthResult == null || _currentAuthResult.Task.IsCompleted)
+					{
+						_currentAuthResult = new TaskCompletionSource<AuthState>(TaskCreationOptions.RunContinuationsAsynchronously);
+					}
+
+					BackgroundTask? prevAuthWorker = _currentAuthWorker;
+					_currentAuthWorker = BackgroundTask.StartNew(ctx => GetAuthStateHandlerAsync(authTaskId, interactive, prevAuthWorker, ctx));
+
+					_currentAuthInteractive = interactive;
+				}
+				authStateTask = _currentAuthResult.Task;
+			}
+
+			// Wait for the task to complete
+			AuthState authState = await authStateTask.WaitAsync(cancellationToken);
 			return authState?.TokenInfo?.AccessToken;
 		}
 
-		async Task<AuthState?> GetAuthStateAsync(bool interactive, CancellationToken cancellationToken)
+		async Task GetAuthStateHandlerAsync(int authTaskId, bool interactive, BackgroundTask? prevAuthTask, CancellationToken cancellationToken)
 		{
-			if (GetAccessTokenFromConfig() != null)
+			Task? disposeTask = null;
+			try
 			{
-				return null;
-			}
+				// Start disposing of the previous auth task asynchronously
+				if (prevAuthTask != null)
+				{
+					disposeTask = Task.Run(() => prevAuthTask.DisposeAsync().AsTask(), CancellationToken.None);
+				}
 
-			Task<AuthState>? authStateTask = null;
-			for (; ; )
-			{
-				lock (_lockObject)
+				// Get the new auth state
+				bool stateHasChanged = false;
+				try
+				{
+					AuthState authState = await GetAuthStateInternalAsync(interactive, cancellationToken);
+					lock (_lockObject)
+					{
+						if (_currentAuthTaskId == authTaskId)
+						{
+							_logger.LogInformation("Auth task complete (interactive: {Interactive}, authorized: {Authorized})", authState.Interactive, authState.IsAuthorized());
+							stateHasChanged = _currentAuthResult?.TrySetResult(authState) ?? false;
+						}
+					}
+				}
+				catch (Exception ex)
 				{
-					if (_authStateTask == null || _authStateTask == authStateTask)
+					lock (_lockObject)
 					{
-						_authStateTask = Task.Run(() => GetAuthStateInternalAsync(interactive, _cancellationTokenSource.Token), _cancellationTokenSource.Token);
-						_authStateTask.ContinueWith(_ => OnStateChanged?.Invoke(), CancellationToken.None, TaskContinuationOptions.None, TaskScheduler.Default);
+						if (_currentAuthTaskId == authTaskId)
+						{
+							_logger.LogError(ex, "Exception while attempting auth: {Message}", ex.Message);
+							stateHasChanged = _currentAuthResult?.TrySetException(ex) ?? false;
+						}
 					}
-					authStateTask = _authStateTask;
 				}
 
-				AuthState authState = await authStateTask.WaitAsync(cancellationToken);
-				if (authState.IsAuthorized() || !interactive || authState.Interactive)
+				// Send notifications about the updated state
+				if (stateHasChanged)
 				{
-					return authState;
+					try
+					{
+						OnStateChanged?.Invoke();
+					}
+					catch (Exception ex)
+					{
+						_logger.LogError(ex, "Exception while sending state change notifications: {Message}", ex.Message);
+					}
+				}
+			}
+			finally
+			{
+				// Wait for the child task to finish being disposed
+				if (disposeTask != null)
+				{
+					await disposeTask;
 				}
 			}
 		}
@@ -330,7 +429,7 @@ namespace EpicGames.Horde
 
 			return new AuthState(_clock, authConfig.Method, result, interactive);
 		}
-		
+
 		private ITokenStore CreateTokenStore()
 		{
 			return _tokenStore ?? TokenStoreFactory.CreateTokenStore();
@@ -342,19 +441,19 @@ namespace EpicGames.Horde
 			{
 				return _oidcTokenManager;
 			}
-			
+
 			string oidcProvider = authConfig.ProfileName ?? "Horde";
 
-			Dictionary<string, string?> values = new ();
+			Dictionary<string, string?> values = new();
 			values[$"Providers:{oidcProvider}:DisplayName"] = oidcProvider;
 			values[$"Providers:{oidcProvider}:ServerUri"] = authConfig.ServerUrl;
 			values[$"Providers:{oidcProvider}:ClientId"] = authConfig.ClientId;
 			values[$"Providers:{oidcProvider}:RedirectUri"] = localRedirectUrl;
-			
-			ConfigurationBuilder builder = new ();
+
+			ConfigurationBuilder builder = new();
 			builder.AddInMemoryCollection(values);
 			IConfiguration configuration = builder.Build();
-			
+
 			return OidcTokenManager.CreateTokenManager(configuration, tokenStore, new List<string>() { oidcProvider });
 		}
 	}
diff --git a/Engine/Source/Programs/Shared/EpicGames.Horde/IHordeClient.cs b/Engine/Source/Programs/Shared/EpicGames.Horde/IHordeClient.cs
--- a/Engine/Source/Programs/Shared/EpicGames.Horde/IHordeClient.cs
+++ b/Engine/Source/Programs/Shared/EpicGames.Horde/IHordeClient.cs
@@ -58,10 +58,10 @@ namespace EpicGames.Horde
 		/// <summary>
 		/// Connect to the Horde server
 		/// </summary>
-		/// <param name="allowPrompt">Whether to allow prompting for credentials</param>
+		/// <param name="interactive">Whether to allow prompting for credentials</param>
 		/// <param name="cancellationToken">Cancellation token for the operation</param>
 		/// <returns>True if the connection succeded</returns>
-		Task<bool> LoginAsync(bool allowPrompt, CancellationToken cancellationToken);
+		Task<bool> LoginAsync(bool interactive, CancellationToken cancellationToken);
 
 		/// <summary>
 		/// Gets the current connection state
diff --git a/Engine/Source/Programs/UnrealGameSync/UnrealGameSync/ProgramApplicationContext.cs b/Engine/Source/Programs/UnrealGameSync/UnrealGameSync/ProgramApplicationContext.cs
--- a/Engine/Source/Programs/UnrealGameSync/UnrealGameSync/ProgramApplicationContext.cs
+++ b/Engine/Source/Programs/UnrealGameSync/UnrealGameSync/ProgramApplicationContext.cs
@@ -260,7 +260,7 @@ namespace UnrealGameSync
 
 			if (hordeAuthState != null)
 			{
-				await hordeAuthState.LoginAsync(false, CancellationToken.None);
+				await hordeAuthState.GetAccessTokenAsync(false, CancellationToken.None);
 			}
 		}
 
