Update nuget packages

#rnx

diff --git a/Engine/Source/Programs/Shared/EpicGames.Core/EpicGames.Core.csproj b/Engine/Source/Programs/Shared/EpicGames.Core/EpicGames.Core.csproj
--- a/Engine/Source/Programs/Shared/EpicGames.Core/EpicGames.Core.csproj
+++ b/Engine/Source/Programs/Shared/EpicGames.Core/EpicGames.Core.csproj
@@ -52,7 +52,7 @@
     </When>
   </Choose>
   <ItemGroup>
-    <PackageReference Include="JetBrains.Annotations" Version="2022.1.0" />
+    <PackageReference Include="JetBrains.Annotations" Version="2024.2.0" />
     <PackageReference Include="Microsoft.Extensions.Configuration" Version="8.0.0" />
     <PackageReference Include="Microsoft.Extensions.Http.Polly" Version="8.0.10" />
     <PackageReference Include="Microsoft.Extensions.Logging.Console" Version="8.0.1" />
diff --git a/Engine/Source/Programs/Shared/EpicGames.OIDC/EpicGames.OIDC.csproj b/Engine/Source/Programs/Shared/EpicGames.OIDC/EpicGames.OIDC.csproj
--- a/Engine/Source/Programs/Shared/EpicGames.OIDC/EpicGames.OIDC.csproj
+++ b/Engine/Source/Programs/Shared/EpicGames.OIDC/EpicGames.OIDC.csproj
@@ -37,9 +37,10 @@
   </ItemGroup>
   
   <ItemGroup>
-    <PackageReference Include="IdentityModel.OidcClient" Version="4.0.0" />
+    <PackageReference Include="IdentityModel.OidcClient" Version="6.0.0" />
     <PackageReference Include="Microsoft.Extensions.Configuration.Binder" Version="8.0.2" />
     <PackageReference Include="Microsoft.Extensions.Configuration.Json" Version="8.0.1" />
+    <PackageReference Include="Microsoft.Extensions.Options" Version="8.0.2" />
     <PackageReference Include="System.Text.Json" Version="8.0.5" />
   </ItemGroup>
 
