Update the helper properties in UnrealBuildBase.Unreal to all be lazy loaded

#rnx

diff --git a/Engine/Source/Programs/Shared/EpicGames.Build/Unreal.cs b/Engine/Source/Programs/Shared/EpicGames.Build/Unreal.cs
--- a/Engine/Source/Programs/Shared/EpicGames.Build/Unreal.cs
+++ b/Engine/Source/Programs/Shared/EpicGames.Build/Unreal.cs
@@ -112,9 +112,9 @@ namespace UnrealBuildBase
 			return UnrealBuildToolDllPath;
 		}
 
-		static private string DotnetVersionDirectory = "8.0.300";
+		private static string DotnetVersionDirectory = "8.0.300";
 
-		static private string FindRelativeDotnetDirectory(RuntimePlatform.Type HostPlatform)
+		private static string FindRelativeDotnetDirectory(RuntimePlatform.Type HostPlatform)
 		{
 			string platform;
 			string architecture;
@@ -137,159 +137,123 @@ namespace UnrealBuildBase
 			return Path.Combine("Binaries", "ThirdParty", "DotNet", DotnetVersionDirectory, $"{platform}-{architecture}");
 		}
 
-		static private string FindRelativeDotnetDirectory() => FindRelativeDotnetDirectory(RuntimePlatform.Current);
+		private static string FindRelativeDotnetDirectory() => FindRelativeDotnetDirectory(RuntimePlatform.Current);
 
 		/// <summary>
 		/// Relative path to the dotnet executable from EngineDir
 		/// </summary>
 		/// <returns></returns>
-		static public string RelativeDotnetDirectory = FindRelativeDotnetDirectory();
+		public static string RelativeDotnetDirectory => _RelativeDotnetDirectory.Value;
+		private static readonly Lazy<string> _RelativeDotnetDirectory = new(FindRelativeDotnetDirectory);
 
-		static private DirectoryReference FindDotnetDirectory() => DirectoryReference.Combine(EngineDirectory, RelativeDotnetDirectory);
+		private static DirectoryReference FindDotnetDirectory() => DirectoryReference.Combine(EngineDirectory, RelativeDotnetDirectory);
 
 		/// <summary>
 		/// The full name of the root UE directory
 		/// </summary>
-		public static readonly DirectoryReference RootDirectory = FindRootDirectory();
+		public static DirectoryReference RootDirectory => _RootDirectory.Value;
+		private static readonly Lazy<DirectoryReference> _RootDirectory = new(FindRootDirectory);
 
 		/// <summary>
 		/// The full name of the Engine directory
 		/// </summary>
-		public static readonly DirectoryReference EngineDirectory = DirectoryReference.Combine(RootDirectory, "Engine");
+		public static DirectoryReference EngineDirectory => _EngineDirectory.Value;
+		private static readonly Lazy<DirectoryReference> _EngineDirectory = new(() => DirectoryReference.Combine(RootDirectory, "Engine"));
 
 		/// <summary>
 		/// The full name of the Engine/Source directory
 		/// </summary>
-		public static readonly DirectoryReference EngineSourceDirectory = DirectoryReference.Combine(EngineDirectory, "Source");
+		public static DirectoryReference EngineSourceDirectory => _EngineSourceDirectory.Value;
+		private static readonly Lazy<DirectoryReference> _EngineSourceDirectory = new(() => DirectoryReference.Combine(EngineDirectory, "Source"));
 
 		/// <summary>
 		/// Returns the User Settings Directory path. This matches FPlatformProcess::UserSettingsDir().
 		/// </summary>
-		static public readonly DirectoryReference UserSettingDirectory = GetUserSettingDirectory();
+		public static DirectoryReference UserSettingDirectory => _UserSettingDirectory.Value;
+		private static readonly Lazy<DirectoryReference> _UserSettingDirectory = new(GetUserSettingDirectory);
 
 		/// <summary>
 		/// Returns the User Directory path. This matches FPlatformProcess::UserDir().
 		/// </summary>
-		static public readonly DirectoryReference? UserDirectory = GetUserDirectory();
+		public static DirectoryReference? UserDirectory => _UserDirectory.Value;
+		private static readonly Lazy<DirectoryReference?> _UserDirectory = new(GetUserDirectory);
 
 		/// <summary>
 		/// Writable engine directory. Uses the user's settings folder for installed builds.
 		/// </summary>
-		public static readonly DirectoryReference WritableEngineDirectory = IsEngineInstalled() ? DirectoryReference.Combine(UserSettingDirectory, "UnrealEngine") : EngineDirectory;
+		public static DirectoryReference WritableEngineDirectory => _WritableEngineDirectory.Value;
+		private static readonly Lazy<DirectoryReference> _WritableEngineDirectory = new(() => IsEngineInstalled() ? DirectoryReference.Combine(UserSettingDirectory, "UnrealEngine") : EngineDirectory);
 
 		/// <summary>
 		/// The engine saved programs directory
 		/// </summary>
-		public static readonly DirectoryReference EngineProgramSavedDirectory = IsEngineInstalled() ? UserSettingDirectory : DirectoryReference.Combine(EngineDirectory, "Programs");
+		public static DirectoryReference EngineProgramSavedDirectory => _EngineProgramSavedDirectory.Value;
+		private static readonly Lazy<DirectoryReference> _EngineProgramSavedDirectory = new(() => IsEngineInstalled() ? UserSettingDirectory : DirectoryReference.Combine(EngineDirectory, "Programs"));
 
 		/// <summary>
 		/// The path to UBT
 		/// </summary>
 		[Obsolete("Deprecated in UE5.1; to launch UnrealBuildTool, use this dll as the first argument with DonetPath")]
-		public static readonly FileReference UnrealBuildToolPath = FindUnrealBuildToolDll().ChangeExtension(RuntimePlatform.ExeExtension);
+		public static FileReference UnrealBuildToolPath => _UnrealBuildToolDllPath.Value.ChangeExtension(RuntimePlatform.ExeExtension);
 
 		/// <summary>
 		/// The path to UBT
 		/// </summary>
-		public static readonly FileReference UnrealBuildToolDllPath = FindUnrealBuildToolDll();
+		public static FileReference UnrealBuildToolDllPath => _UnrealBuildToolDllPath.Value;
+		private static readonly Lazy<FileReference> _UnrealBuildToolDllPath = new(FindUnrealBuildToolDll);
 
 		/// <summary>
 		/// The directory containing the bundled .NET installation
 		/// </summary>
-		static public readonly DirectoryReference DotnetDirectory = FindDotnetDirectory();
+		public static DirectoryReference DotnetDirectory => _DotnetDirectory.Value;
+		private static readonly Lazy<DirectoryReference> _DotnetDirectory = new(FindDotnetDirectory);
 
 		/// <summary>
 		/// The path of the bundled dotnet executable
 		/// </summary>
-		static public readonly FileReference DotnetPath = FileReference.Combine(DotnetDirectory, "dotnet" + RuntimePlatform.ExeExtension);
+		public static FileReference DotnetPath => _DotnetPath.Value;
+		private static readonly Lazy<FileReference> _DotnetPath = new(() => FileReference.Combine(DotnetDirectory, "dotnet" + RuntimePlatform.ExeExtension));
 
 		/// <summary>
-		/// Whether we're running on a build machine
-		/// </summary>
-		static private bool? bIsBuildMachine;
-
-		/// <summary>
-		/// Whether we're running with engine installed
+		/// Returns true if the application is running on a build machine
 		/// </summary>
-		static private bool? bIsEngineInstalled;
+		/// <returns>True if running on a build machine</returns>
+		public static bool IsBuildMachine() => _IsBuildMachine.Value;
+		private static readonly Lazy<bool> _IsBuildMachine = new Lazy<bool>(() => Environment.GetEnvironmentVariable("IsBuildMachine")?.Trim() == "1");
 
 		/// <summary>
-		/// Whether we're running with an installed project (ie. a mod kit)
+		/// Returns true if the application is running using installed Engine components
 		/// </summary>
-		static private bool? bIsProjectInstalled;
+		/// <returns>True if running using installed Engine components</returns>
+		public static bool IsEngineInstalled() => _IsEngineInstalled.Value;
+		private static readonly Lazy<bool> _IsEngineInstalled = new(() => FileReference.Exists(FileReference.Combine(EngineDirectory, "Build", "InstalledBuild.txt")));
 
 		/// <summary>
 		/// If we are running with an installed project, specifies the path to it
 		/// </summary>
-		static private FileReference? InstalledProjectFile;
+		private static readonly Lazy<FileReference?> _InstalledProjectFile = new(() => {
+			FileReference installedProjectLocationFile = FileReference.Combine(EngineDirectory, "Build", "InstalledProjectBuild.txt");
+			return FileReference.Exists(installedProjectLocationFile) ? FileReference.Combine(RootDirectory, FileReference.ReadAllText(installedProjectLocationFile).Trim()) : null;
+		});
 
 		/// <summary>
 		/// Returns where another platform's Dotnet is located
 		/// </summary>
 		/// <param name="HostPlatform"></param>
 		/// <returns></returns>
-		static public DirectoryReference FindDotnetDirectoryForPlatform(RuntimePlatform.Type HostPlatform)
-		{
-			return DirectoryReference.Combine(EngineDirectory, FindRelativeDotnetDirectory(HostPlatform));
-		}
-
-		/// <summary>
-		/// Returns true if the application is running on a build machine
-		/// </summary>
-		/// <returns>True if running on a build machine</returns>
-		static public bool IsBuildMachine()
-		{
-			if (!bIsBuildMachine.HasValue)
-			{
-				bIsBuildMachine = Environment.GetEnvironmentVariable("IsBuildMachine")?.Trim() == "1";
-			}
-			return bIsBuildMachine.Value;
-		}
-
-		/// <summary>
-		/// Returns true if the application is running using installed Engine components
-		/// </summary>
-		/// <returns>True if running using installed Engine components</returns>
-		static public bool IsEngineInstalled()
-		{
-			if (!bIsEngineInstalled.HasValue)
-			{
-				bIsEngineInstalled = FileReference.Exists(FileReference.Combine(EngineDirectory, "Build", "InstalledBuild.txt"));
-			}
-			return bIsEngineInstalled.Value;
-		}
+		public static DirectoryReference FindDotnetDirectoryForPlatform(RuntimePlatform.Type HostPlatform) => DirectoryReference.Combine(EngineDirectory, FindRelativeDotnetDirectory(HostPlatform));
 
 		/// <summary>
 		/// Returns true if the application is running using an installed project (ie. a mod kit)
 		/// </summary>
 		/// <returns>True if running using an installed project</returns>
-		public static bool IsProjectInstalled()
-		{
-			if (!bIsProjectInstalled.HasValue)
-			{
-				FileReference InstalledProjectLocationFile = FileReference.Combine(EngineDirectory, "Build", "InstalledProjectBuild.txt");
-				if (FileReference.Exists(InstalledProjectLocationFile))
-				{
-					InstalledProjectFile = FileReference.Combine(RootDirectory, FileReference.ReadAllText(InstalledProjectLocationFile).Trim());
-					bIsProjectInstalled = true;
-				}
-				else
-				{
-					InstalledProjectFile = null;
-					bIsProjectInstalled = false;
-				}
-			}
-			return bIsProjectInstalled.Value;
-		}
+		public static bool IsProjectInstalled() => _InstalledProjectFile.Value != null;
 
 		/// <summary>
 		/// Gets the installed project file
 		/// </summary>
 		/// <returns>Location of the installed project file</returns>
-		public static FileReference? GetInstalledProjectFile()
-		{
-			return IsProjectInstalled() ? InstalledProjectFile : null;
-		}
+		public static FileReference? GetInstalledProjectFile() => _InstalledProjectFile.Value;
 
 		/// <summary>
 		/// Checks whether the given file is under an installed directory, and should not be overridden
@@ -302,7 +266,7 @@ namespace UnrealBuildBase
 			{
 				return true;
 			}
-			if (IsProjectInstalled() && File.IsUnderDirectory(InstalledProjectFile!.Directory))
+			if (IsProjectInstalled() && File.IsUnderDirectory(_InstalledProjectFile.Value!.Directory))
 			{
 				return true;
 			}
@@ -319,7 +283,7 @@ namespace UnrealBuildBase
 
 		// A subset of the functionality in DataDrivenPlatformInfo.GetAllPlatformInfos() - finds the DataDrivenPlatformInfo.ini files and records their existence, but does not parse them
 		// (perhaps DataDrivenPlatformInfo.GetAllPlatformInfos() could be modified to use this data to avoid an additional search through the filesystem)
-		private static Lazy<HashSet<string>> IniPresentForPlatform = new Lazy<HashSet<string>>(() =>
+		private static readonly Lazy<HashSet<string>> IniPresentForPlatform = new Lazy<HashSet<string>>(() =>
 		{
 			HashSet<string> Set = new HashSet<string>(StringComparer.OrdinalIgnoreCase);
 			// find all platform directories (skipping NFL/NoRedist)
@@ -352,13 +316,11 @@ namespace UnrealBuildBase
 
 			return Set;
 		});
-		private static bool DataDrivenPlatformInfoIniIsPresent(string PlatformName)
-		{
-			return IniPresentForPlatform.Value.Contains(PlatformName);
-		}
+
+		private static bool DataDrivenPlatformInfoIniIsPresent(string PlatformName) => IniPresentForPlatform.Value.Contains(PlatformName);
 
 		// cached dictionary of BaseDir to extension directories
-		private static ConcurrentDictionary<DirectoryReference, Tuple<List<DirectoryReference>, List<DirectoryReference>>> CachedExtensionDirectories = new ConcurrentDictionary<DirectoryReference, Tuple<List<DirectoryReference>, List<DirectoryReference>>>();
+		private static ConcurrentDictionary<DirectoryReference, Tuple<List<DirectoryReference>, List<DirectoryReference>>> CachedExtensionDirectories = new();
 
 		/// <summary>
 		/// Finds all the extension directories for the given base directory. This includes platform extensions and restricted folders.
@@ -407,7 +369,7 @@ namespace UnrealBuildBase
 			});
 
 			// now return what the caller wanted (always include BaseDir)
-			List<DirectoryReference> ExtensionDirs = new List<DirectoryReference>();
+			List<DirectoryReference> ExtensionDirs = [];
 			if (bIncludeBaseDirectory)
 			{
 				ExtensionDirs.Add(BaseDir);
@@ -424,7 +386,7 @@ namespace UnrealBuildBase
 		}
 
 		/// <summary>
-		/// Finds all the extension directories for the given base directory. This includes platform extensions and restricted folders.
+		/// Finds all the extension directories for the given base directory. This includes platform extensions and restricted folders.f
 		/// </summary>
 		/// <param name="BaseDir">Location of the base directory</param>
 		/// <param name="SubDir">The subdirectory to find</param>
@@ -434,7 +396,10 @@ namespace UnrealBuildBase
 		/// <returns>List of extension directories, including the given base directory</returns>
 		public static List<DirectoryReference> GetExtensionDirs(DirectoryReference BaseDir, string SubDir, bool bIncludePlatformDirectories = true, bool bIncludeRestrictedDirectories = true, bool bIncludeBaseDirectory = true)
 		{
-			return GetExtensionDirs(BaseDir, bIncludePlatformDirectories, bIncludeRestrictedDirectories, bIncludeBaseDirectory).Select(x => DirectoryReference.Combine(x, SubDir)).Where(x => DirectoryReference.Exists(x)).ToList();
+			return GetExtensionDirs(BaseDir, bIncludePlatformDirectories, bIncludeRestrictedDirectories, bIncludeBaseDirectory)
+				.Select(x => DirectoryReference.Combine(x, SubDir))
+				.Where(x => DirectoryReference.Exists(x))
+				.ToList();
 		}
 
 		/// <summary>
@@ -453,10 +418,10 @@ namespace UnrealBuildBase
 			else
 			{
 				// Not all user accounts have a local application data directory (eg. SYSTEM, used by Jenkins for builds).
-				List<Environment.SpecialFolder> DataFolders = new List<Environment.SpecialFolder>() {
+				List<Environment.SpecialFolder> DataFolders = [
 					Environment.SpecialFolder.LocalApplicationData,
 					Environment.SpecialFolder.CommonApplicationData
-				};
+				];
 				foreach (Environment.SpecialFolder DataFolder in DataFolders)
 				{
 					string DirectoryName = Environment.GetFolderPath(DataFolder);
@@ -491,30 +456,27 @@ namespace UnrealBuildBase
 			return null;
 		}
 
-		private static string? _MachineName;
-		public static string MachineName
+		/// <summary>
+		/// The current Machine name
+		/// </summary>
+		public static string MachineName => _MachineName.Value;
+		private static readonly Lazy<string> _MachineName = new(() =>
 		{
-			get
+			try
 			{
-				if (_MachineName == null)
-				{
-					try
-					{
-						// this likely can't fail, but just in case, fallback to preview implementation
-						_MachineName = System.Net.Dns.GetHostName();
+				// this likely can't fail, but just in case, fallback to preview implementation
+				string machineName = System.Net.Dns.GetHostName();
 
-						if (RuntimePlatform.IsMac && _MachineName.EndsWith(".local")) 
-						{
-							_MachineName = _MachineName.Replace(".local", "");
-						}
-					}
-					catch (Exception)
-					{
-						_MachineName = System.Environment.MachineName;
-					}
+				if (RuntimePlatform.IsMac && machineName.EndsWith(".local"))
+				{
+					machineName = machineName.Replace(".local", "");
 				}
-				return _MachineName!;
+				return machineName;
 			}
-		}
+			catch (Exception)
+			{
+			}
+			return Environment.MachineName;
+		});
 	}
 }
