Shared: intellisense warnings

#rnx

diff --git a/Engine/Source/Programs/Shared/EpicGames.Build/.editorconfig b/Engine/Source/Programs/Shared/EpicGames.Build/.editorconfig
--- a/Engine/Source/Programs/Shared/EpicGames.Build/.editorconfig
+++ b/Engine/Source/Programs/Shared/EpicGames.Build/.editorconfig
@@ -92,3 +92,6 @@ dotnet_diagnostic.CA1829.severity = none
 
 # CA2249: Consider using 'string.Contains' instead of 'string.IndexOf'
 dotnet_diagnostic.CA2249.severity = none
+
+# IDE0130: Namespace does not match folder structure
+dotnet_diagnostic.IDE0130.severity = none
diff --git a/Engine/Source/Programs/Shared/EpicGames.Build/UEBuildPlatformSDK.cs b/Engine/Source/Programs/Shared/EpicGames.Build/UEBuildPlatformSDK.cs
--- a/Engine/Source/Programs/Shared/EpicGames.Build/UEBuildPlatformSDK.cs
+++ b/Engine/Source/Programs/Shared/EpicGames.Build/UEBuildPlatformSDK.cs
@@ -10,7 +10,6 @@ using System.Text;
 using System.Text.Json;
 using Microsoft.Extensions.Logging;
 using UnrealBuildBase;
-using EpicGames.Core;
 
 namespace EpicGames.Core
 {
@@ -857,7 +856,7 @@ namespace EpicGames.Core
 		/// <param name="OutValue"></param>
 		/// <param name="Hint">A platform specific hint that can help guide conversion (usually SDKName or device type)</param>
 		/// <returns>If the StringValue was able to be be converted to an integer</returns>
-		public abstract bool TryConvertVersionToInt(string? StringValue, out UInt64 OutValue, string? Hint = null);
+		public abstract bool TryConvertVersionToInt(string? StringValue, out ulong OutValue, string? Hint = null);
 
 		/// <summary>
 		/// Like TryConvertVersionToInt, but will throw an exception on failure
@@ -865,7 +864,7 @@ namespace EpicGames.Core
 		/// <param name="StringValue">Version that comes from the PlatformSDK class (should not be used with Manifest or other user supplied versions)</param>
 		/// <param name="Hint">A platform specific hint that can help guide conversion (usually SDKName or device type)</param>
 		/// <returns>The integer version of StringValue, can be used to compare against a valid range</returns>
-		public UInt64 ConvertVersionToInt(string? StringValue, string? Hint = null)
+		public ulong ConvertVersionToInt(string? StringValue, string? Hint = null)
 		{
 			// quickly handle null input, which is not necessarily an error case
 			if (StringValue == null)
@@ -873,7 +872,7 @@ namespace EpicGames.Core
 				return 0;
 			}
 
-			UInt64 Result;
+			ulong Result;
 			if (!TryConvertVersionToInt(StringValue, out Result, Hint))
 			{
 				throw new Exception($"Unable to convert {GetType()} version '{StringValue}' to an integer. Likely this version was supplied by code, and is expected to be valid.");
@@ -891,7 +890,7 @@ namespace EpicGames.Core
 		/// <returns>Comparison integer</returns>
 		public int SdkVersionsCompare(string? StringValueA, string? StringValueB, string? Hint = null)
 		{
-			UInt64 ValueA, ValueB;
+			ulong ValueA, ValueB;
 			TryConvertVersionToInt(StringValueA, out ValueA, Hint);
 			TryConvertVersionToInt(StringValueB, out ValueB, Hint);
 			if (ValueA == ValueB)
@@ -1533,9 +1532,9 @@ namespace EpicGames.Core
 		/// <param name="InstalledSDKVersionString">SDK version string to set</param>
 		/// <param name="InstalledSDKLevelString"></param>
 		/// <returns>true if was able to set it</returns>
-		protected bool SetCurrentlyInstalledAutoSDKString(String InstalledSDKVersionString, String InstalledSDKLevelString)
+		protected bool SetCurrentlyInstalledAutoSDKString(string InstalledSDKVersionString, string InstalledSDKLevelString)
 		{
-			String PlatformSDKRoot = GetPathToPlatformAutoSDKs();
+			string PlatformSDKRoot = GetPathToPlatformAutoSDKs();
 			if (Directory.Exists(PlatformSDKRoot))
 			{
 				string VersionFilename = Path.Combine(PlatformSDKRoot, CurrentlyInstalledSDKStringManifest);
@@ -1560,8 +1559,8 @@ namespace EpicGames.Core
 		{
 			if (PlatformSupportsAutoSDKs() && HasAutoSDKSystemEnabled())
 			{
-				String InstalledSDKVersionString = GetAutoSDKDirectoryForMainVersion();
-				String PlatformSDKRoot = GetPathToPlatformAutoSDKs();
+				string InstalledSDKVersionString = GetAutoSDKDirectoryForMainVersion();
+				string PlatformSDKRoot = GetPathToPlatformAutoSDKs();
 				if (!Directory.Exists(PlatformSDKRoot))
 				{
 					Directory.CreateDirectory(PlatformSDKRoot);
@@ -1686,7 +1685,7 @@ namespace EpicGames.Core
 
 		protected bool SetLastRunAutoSDKScriptVersion(string LastRunScriptVersion)
 		{
-			String PlatformSDKRoot = GetPathToPlatformAutoSDKs();
+			string PlatformSDKRoot = GetPathToPlatformAutoSDKs();
 			if (Directory.Exists(PlatformSDKRoot))
 			{
 				string VersionFilename = Path.Combine(PlatformSDKRoot, LastRunScriptVersionManifest);
@@ -1915,7 +1914,7 @@ namespace EpicGames.Core
 					List<string> EnvVarValues = [];
 
 					bool bNeedsToWriteAutoSetupEnvVar = true;
-					String PlatformSetupEnvVar = GetPlatformAutoSDKSetupEnvVar();
+					string PlatformSetupEnvVar = GetPlatformAutoSDKSetupEnvVar();
 					for (; ; )
 					{
 						string? VariableString = Reader.ReadLine();
@@ -1965,9 +1964,9 @@ namespace EpicGames.Core
 
 
 					// actually perform the PATH stripping / adding.
-					String? OrigPathVar = Environment.GetEnvironmentVariable("PATH");
-					String PathDelimiter = RuntimePlatform.IsWindows ? ";" : ":";
-					String[] PathVars = [];
+					string? OrigPathVar = Environment.GetEnvironmentVariable("PATH");
+					string PathDelimiter = RuntimePlatform.IsWindows ? ";" : ":";
+					string[] PathVars = [];
 					if (!String.IsNullOrEmpty(OrigPathVar))
 					{
 						PathVars = OrigPathVar.Split(PathDelimiter.ToCharArray());
@@ -1977,12 +1976,12 @@ namespace EpicGames.Core
 						Logger.LogDebug("Path environment variable is null during AutoSDK");
 					}
 
-					List<String> ModifiedPathVars = [.. PathVars];
+					List<string> ModifiedPathVars = [.. PathVars];
 
 					// perform removes first, in case they overlap with any adds.
-					foreach (String PathRemove in PathRemoves)
+					foreach (string PathRemove in PathRemoves)
 					{
-						foreach (String PathVar in PathVars)
+						foreach (string PathVar in PathVars)
 						{
 							if (PathVar.IndexOf(PathRemove, StringComparison.OrdinalIgnoreCase) >= 0)
 							{
@@ -1994,9 +1993,9 @@ namespace EpicGames.Core
 
 					// remove all the of ADDs so that if this function is executed multiple times, the paths will be guaranteed to be in the same order after each run.
 					// If we did not do this, a 'remove' that matched some, but not all, of our 'adds' would cause the order to change.
-					foreach (String PathAdd in PathAdds)
+					foreach (string PathAdd in PathAdds)
 					{
-						foreach (String PathVar in PathVars)
+						foreach (string PathVar in PathVars)
 						{
 							if (String.Compare(PathAdd, PathVar, true) == 0)
 							{
@@ -2007,7 +2006,7 @@ namespace EpicGames.Core
 					}
 
 					// perform adds, but don't add duplicates
-					foreach (String PathAdd in PathAdds)
+					foreach (string PathAdd in PathAdds)
 					{
 						if (!ModifiedPathVars.Contains(PathAdd))
 						{
@@ -2016,7 +2015,7 @@ namespace EpicGames.Core
 						}
 					}
 
-					String ModifiedPath = String.Join(PathDelimiter, ModifiedPathVars);
+					string ModifiedPath = String.Join(PathDelimiter, ModifiedPathVars);
 					Environment.SetEnvironmentVariable("PATH", ModifiedPath);
 
 					Reader.Close();
@@ -2054,7 +2053,7 @@ namespace EpicGames.Core
 
 		protected void InvalidateCurrentlyInstalledAutoSDK()
 		{
-			String PlatformSDKRoot = GetPathToPlatformAutoSDKs();
+			string PlatformSDKRoot = GetPathToPlatformAutoSDKs();
 			if (Directory.Exists(PlatformSDKRoot))
 			{
 				string SDKFilename = Path.Combine(PlatformSDKRoot, CurrentlyInstalledSDKStringManifest);
@@ -2123,7 +2122,7 @@ namespace EpicGames.Core
 		}
 
 		// This tracks if we have already checked the sdk installation.
-		private Int32 SDKCheckStatus = -1;
+		private int SDKCheckStatus = -1;
 
 		// true if we've ever overridden the process's environment with AutoSDK data.  After that, manual installs cannot be considered valid ever again.
 		private bool bLocalProcessSetupAutoSDK = false;
@@ -2135,7 +2134,7 @@ namespace EpicGames.Core
 
 		protected bool HasParentProcessSetupAutoSDK([NotNullWhen(true)] out string? OutAutoSDKSetupValue)
 		{
-			String AutoSDKSetupVarName = GetPlatformAutoSDKSetupEnvVar();
+			string AutoSDKSetupVarName = GetPlatformAutoSDKSetupEnvVar();
 			OutAutoSDKSetupValue = Environment.GetEnvironmentVariable(AutoSDKSetupVarName);
 
 			if (!String.IsNullOrEmpty(OutAutoSDKSetupValue))
diff --git a/Engine/Source/Programs/Shared/EpicGames.Core/ManagedProcess.cs b/Engine/Source/Programs/Shared/EpicGames.Core/ManagedProcess.cs
--- a/Engine/Source/Programs/Shared/EpicGames.Core/ManagedProcess.cs
+++ b/Engine/Source/Programs/Shared/EpicGames.Core/ManagedProcess.cs
@@ -40,8 +40,8 @@ namespace EpicGames.Core
 	{
 #pragma warning disable IDE0049 // Naming Styles
 #pragma warning disable IDE1006 // Naming Styles
-		const UInt32 JOB_OBJECT_LIMIT_KILL_ON_JOB_CLOSE = 0x00002000;
-		const UInt32 JOB_OBJECT_LIMIT_BREAKAWAY_OK = 0x00000800;
+		const uint JOB_OBJECT_LIMIT_KILL_ON_JOB_CLOSE = 0x00002000;
+		const uint JOB_OBJECT_LIMIT_BREAKAWAY_OK = 0x00000800;
 
 		[StructLayout(LayoutKind.Sequential)]
 		struct JOBOBJECT_BASIC_LIMIT_INFORMATION
@@ -340,11 +340,11 @@ namespace EpicGames.Core
 		[DllImport("kernel32.dll", SetLastError = true)]
 		static extern int TerminateProcess(SafeHandleZeroOrMinusOneIsInvalid hProcess, uint uExitCode);
 
-		const UInt32 INFINITE = 0xFFFFFFFF;
-		const UInt32 WAIT_FAILED = 0xFFFFFFFF;
+		const uint INFINITE = 0xFFFFFFFF;
+		const uint WAIT_FAILED = 0xFFFFFFFF;
 
 		[DllImport("kernel32.dll", SetLastError = true)]
-		static extern UInt32 WaitForSingleObject(SafeHandleZeroOrMinusOneIsInvalid hHandle, UInt32 dwMilliseconds);
+		static extern uint WaitForSingleObject(SafeHandleZeroOrMinusOneIsInvalid hHandle, uint dwMilliseconds);
 
 		[DllImport("kernel32.dll", SetLastError = true)]
 		[return: MarshalAs(UnmanagedType.Bool)]
@@ -355,19 +355,19 @@ namespace EpicGames.Core
 			out System.Runtime.InteropServices.ComTypes.FILETIME lpUserTime);
 
 		[DllImport("kernel32.dll")]
-		static extern UInt16 GetActiveProcessorGroupCount();
+		static extern ushort GetActiveProcessorGroupCount();
 
 		[DllImport("kernel32.dll")]
-		static extern UInt32 GetActiveProcessorCount(UInt16 GroupNumber);
+		static extern uint GetActiveProcessorCount(ushort GroupNumber);
 
 		[StructLayout(LayoutKind.Sequential)]
 		class GROUP_AFFINITY
 		{
-			public UInt64 Mask;
-			public UInt16 Group;
-			public UInt16 Reserved0;
-			public UInt16 Reserved1;
-			public UInt16 Reserved2;
+			public ulong Mask;
+			public ushort Group;
+			public ushort Reserved0;
+			public ushort Reserved1;
+			public ushort Reserved2;
 		}
 
 		[DllImport("kernel32.dll", SetLastError = true)]
diff --git a/Engine/Source/Programs/Shared/EpicGames.Core/UHTTypes.cs b/Engine/Source/Programs/Shared/EpicGames.Core/UHTTypes.cs
--- a/Engine/Source/Programs/Shared/EpicGames.Core/UHTTypes.cs
+++ b/Engine/Source/Programs/Shared/EpicGames.Core/UHTTypes.cs
@@ -1,5 +1,6 @@
 ﻿// Copyright Epic Games, Inc. All Rights Reserved.
 
+using System;
 using System.Collections.Generic;
 using System.Text.Json.Serialization;
 
@@ -107,9 +108,9 @@ namespace EpicGames.Core
 			public UHTModuleType ModuleType { get; set; } = UHTModuleType.Program;
 			[JsonConverter(typeof(JsonStringEnumConverter))]
 			public EPackageOverrideType OverrideModuleType { get; set; } = EPackageOverrideType.None;
-			public string BaseDirectory { get; set; } = string.Empty;
-			public string IncludeBase { get; set; } = string.Empty; // The include path which all UHT-generated includes should be relative to
-			public string OutputDirectory { get; set; } = string.Empty;
+			public string BaseDirectory { get; set; } = String.Empty;
+			public string IncludeBase { get; set; } = String.Empty; // The include path which all UHT-generated includes should be relative to
+			public string OutputDirectory { get; set; } = String.Empty;
 			public List<string> ClassesHeaders { get; set; } = [];
 			public List<string> PublicHeaders { get; set; } = [];
 			public List<string> InternalHeaders { get; set; } = [];
@@ -159,7 +160,7 @@ namespace EpicGames.Core
 				string? stringValue;
 				if (TryGetDefine(name, out stringValue))
 				{
-					return int.TryParse(stringValue, out value);
+					return Int32.TryParse(stringValue, out value);
 				}
 				value = 0;
 				return false;
@@ -174,17 +175,17 @@ namespace EpicGames.Core
 		/// <summary>
 		/// The engine path on the local machine
 		/// </summary>
-		public string RootLocalPath { get; set; } = string.Empty;
+		public string RootLocalPath { get; set; } = String.Empty;
 
 		/// <summary>
 		/// Name of the target currently being compiled
 		/// </summary>
-		public string TargetName { get; set; } = string.Empty;
+		public string TargetName { get; set; } = String.Empty;
 
 		/// <summary>
 		/// File to contain additional dependencies that the generated code depends on
 		/// </summary>
-		public string ExternalDependenciesFile { get; set; } = string.Empty;
+		public string ExternalDependenciesFile { get; set; } = String.Empty;
 
 		/// <summary>
 		/// List of modules
diff --git a/Engine/Source/Programs/Shared/EpicGames.Core/Win32ExceptionWithCode.cs b/Engine/Source/Programs/Shared/EpicGames.Core/Win32ExceptionWithCode.cs
--- a/Engine/Source/Programs/Shared/EpicGames.Core/Win32ExceptionWithCode.cs
+++ b/Engine/Source/Programs/Shared/EpicGames.Core/Win32ExceptionWithCode.cs
@@ -41,6 +41,6 @@ namespace EpicGames.Core
 		/// <summary>
 		/// Returns the exception message. Overriden to include the error code in the message.
 		/// </summary>
-		public override string Message => String.Format("{0} (code 0x{1:X8})", base.Message, base.NativeErrorCode);
+		public override string Message => String.Format("{0} (code 0x{1:X8})", base.Message, NativeErrorCode);
 	}
 }
