Horde: Register the standard IHttpClientFactory functionality along with Horde types, since we depend on it.

#rnx

diff --git a/Engine/Source/Programs/Shared/EpicGames.Horde/HordeExtensions.cs b/Engine/Source/Programs/Shared/EpicGames.Horde/HordeExtensions.cs
--- a/Engine/Source/Programs/Shared/EpicGames.Horde/HordeExtensions.cs
+++ b/Engine/Source/Programs/Shared/EpicGames.Horde/HordeExtensions.cs
@@ -31,6 +31,7 @@ namespace EpicGames.Horde
 
 			serviceCollection.AddLogging();
 
+			serviceCollection.AddHttpClient();
 			serviceCollection.AddSingleton<BundleCache>(sp => new BundleCache(sp.GetRequiredService<IOptions<HordeOptions>>().Value.BundleCache));
 			serviceCollection.AddSingleton<StorageBackendCache>(CreateBackendCache);
 			serviceCollection.AddSingleton<HttpStorageBackendFactory>();
