Pin nuget package versions

#rnx

diff --git a/Engine/Source/Programs/AutomationTool.Tests/AutomationTool.Tests.csproj b/Engine/Source/Programs/AutomationTool.Tests/AutomationTool.Tests.csproj
--- a/Engine/Source/Programs/AutomationTool.Tests/AutomationTool.Tests.csproj
+++ b/Engine/Source/Programs/AutomationTool.Tests/AutomationTool.Tests.csproj
@@ -17,6 +17,7 @@
       <PrivateAssets>all</PrivateAssets>
       <IncludeAssets>runtime; build; native; contentfiles; analyzers; buildtransitive</IncludeAssets>
     </PackageReference>
+    <PackageReference Include="System.Text.Json" Version="8.0.5" />
   </ItemGroup>
 
   <ItemGroup>
diff --git a/Engine/Source/Programs/Shared/EpicGames.AspNet/EpicGames.AspNet.csproj b/Engine/Source/Programs/Shared/EpicGames.AspNet/EpicGames.AspNet.csproj
--- a/Engine/Source/Programs/Shared/EpicGames.AspNet/EpicGames.AspNet.csproj
+++ b/Engine/Source/Programs/Shared/EpicGames.AspNet/EpicGames.AspNet.csproj
@@ -44,5 +44,6 @@
   </ItemGroup>
   <ItemGroup>
     <PackageReference Include="Microsoft.CSharp" Version="4.7.0" />
+    <PackageReference Include="System.Text.Json" Version="8.0.5" />
   </ItemGroup>
 </Project>
\ No newline at end of file
diff --git a/Engine/Source/Programs/Shared/EpicGames.BuildGraph.Tests/EpicGames.BuildGraph.Tests.csproj b/Engine/Source/Programs/Shared/EpicGames.BuildGraph.Tests/EpicGames.BuildGraph.Tests.csproj
--- a/Engine/Source/Programs/Shared/EpicGames.BuildGraph.Tests/EpicGames.BuildGraph.Tests.csproj
+++ b/Engine/Source/Programs/Shared/EpicGames.BuildGraph.Tests/EpicGames.BuildGraph.Tests.csproj
@@ -54,6 +54,7 @@
       <PrivateAssets>all</PrivateAssets>
       <IncludeAssets>runtime; build; native; contentfiles; analyzers; buildtransitive</IncludeAssets>
     </PackageReference>
+    <PackageReference Include="System.Text.Json" Version="8.0.5" />
   </ItemGroup>
 
   <ItemGroup>
