Horde: Include HttpRequestException in retry policy again. Don't think we should need to retry these situations, but S3 seems to cause them intermittently.

#rnx

diff --git a/Engine/Source/Programs/Shared/EpicGames.Horde/HordeHttpMessageHandler.cs b/Engine/Source/Programs/Shared/EpicGames.Horde/HordeHttpMessageHandler.cs
--- a/Engine/Source/Programs/Shared/EpicGames.Horde/HordeHttpMessageHandler.cs
+++ b/Engine/Source/Programs/Shared/EpicGames.Horde/HordeHttpMessageHandler.cs
@@ -79,7 +79,8 @@ namespace EpicGames.Horde
 			// Policy for transient errors is the same as HttpPolicyExtensions.HandleTransientHttpError(), but excludes HttpStatusCode.ServiceUnavailable (which is used as a response
 			// when allocating compute resources when none are available). This pathway is handled explicitly on the application side.
 			return Policy<HttpResponseMessage>
-				.HandleResult(x => (x.StatusCode >= HttpStatusCode.InternalServerError && x.StatusCode != HttpStatusCode.ServiceUnavailable) || x.StatusCode == HttpStatusCode.RequestTimeout)
+				.Handle<HttpRequestException>()
+				.OrResult(x => (x.StatusCode >= HttpStatusCode.InternalServerError && x.StatusCode != HttpStatusCode.ServiceUnavailable) || x.StatusCode == HttpStatusCode.RequestTimeout)
 				.WaitAndRetryAsync(retryTimes, OnTimeoutAsync);
 		}
 	}
