Two second connect timeout and no printed message when UBT telemetry fails to connect or send.

diff --git a/Engine/Source/Programs/Shared/EpicGames.Core/Telemetry/DataRouterTelemetryService.cs b/Engine/Source/Programs/Shared/EpicGames.Core/Telemetry/DataRouterTelemetryService.cs
--- a/Engine/Source/Programs/Shared/EpicGames.Core/Telemetry/DataRouterTelemetryService.cs
+++ b/Engine/Source/Programs/Shared/EpicGames.Core/Telemetry/DataRouterTelemetryService.cs
@@ -581,6 +581,9 @@ namespace EpicGames.Core.Telemetry
 			{
 				Logger.LogError("{Exception}", ex.Message);
 			}
+			catch (TaskCanceledException /*ex*/)
+			{
+			}
 			catch (Exception ex)
 			{
 				Logger.LogError("{Exception}", ex.ToString());
diff --git a/Engine/Source/Programs/Shared/EpicGames.Core/Telemetry/DefaultTelemetryHttpClient.cs b/Engine/Source/Programs/Shared/EpicGames.Core/Telemetry/DefaultTelemetryHttpClient.cs
--- a/Engine/Source/Programs/Shared/EpicGames.Core/Telemetry/DefaultTelemetryHttpClient.cs
+++ b/Engine/Source/Programs/Shared/EpicGames.Core/Telemetry/DefaultTelemetryHttpClient.cs
@@ -36,7 +36,10 @@ namespace EpicGames.Core.Telemetry
 		public static HttpClient GetClient()
 		{
 			// per MS documentation, with a lifetime specified we shouldn't have socket exhaustion issues
-			SocketsHttpHandler socketHandler = new SocketsHttpHandler { PooledConnectionLifetime = TimeSpan.FromMinutes(15) };
+			SocketsHttpHandler socketHandler = new SocketsHttpHandler {
+				ConnectTimeout  = TimeSpan.FromSeconds(2),
+				PooledConnectionLifetime = TimeSpan.FromMinutes(15)
+			};
 #pragma warning disable CA2000 // Dispose objects before losing scope
 			PolicyHttpMessageHandler pollyHandler = new PolicyHttpMessageHandler(GetRetryPolicy())
 			{
