Horde: Add more information to exceptions thrown when unable to memory map output files.

#rnx

diff --git a/Engine/Source/Programs/Shared/EpicGames.Core/MemoryMappedView.cs b/Engine/Source/Programs/Shared/EpicGames.Core/MemoryMappedView.cs
--- a/Engine/Source/Programs/Shared/EpicGames.Core/MemoryMappedView.cs
+++ b/Engine/Source/Programs/Shared/EpicGames.Core/MemoryMappedView.cs
@@ -44,10 +44,10 @@ namespace EpicGames.Core
 		const uint FILE_MAP_READ = 0x0004;
 		const uint FILE_MAP_EXECUTE = 0x0020;
 
-		[DllImport("kernel32.dll")]
+		[DllImport("kernel32.dll", SetLastError = true)]
 		static extern unsafe void* MapViewOfFile(SafeMemoryMappedFileHandle handle, uint desiredAccess, uint fileOffsetHigh, uint fileOffsetLow, long numberOfBytes);
 
-		[DllImport("kernel32.dll")]
+		[DllImport("kernel32.dll", SetLastError = true)]
 		static extern unsafe bool UnmapViewOfFile(void* ptr);
 
 		byte* _data;
diff --git a/Engine/Source/Programs/Shared/EpicGames.Horde/Storage/Nodes/DirectoryNode.Extract.cs b/Engine/Source/Programs/Shared/EpicGames.Horde/Storage/Nodes/DirectoryNode.Extract.cs
--- a/Engine/Source/Programs/Shared/EpicGames.Horde/Storage/Nodes/DirectoryNode.Extract.cs
+++ b/Engine/Source/Programs/Shared/EpicGames.Horde/Storage/Nodes/DirectoryNode.Extract.cs
@@ -475,7 +475,7 @@ namespace EpicGames.Horde.Storage.Nodes
 				{
 					// Process as many chunks as we can for this file
 					using MemoryMappedFile memoryMappedFile = MemoryMappedFile.CreateFromFile(stream, null, file.FileEntry.Length, MemoryMappedFileAccess.ReadWrite, HandleInheritability.None, false);
-					using MemoryMappedView memoryMappedView = new MemoryMappedView(memoryMappedFile, 0, file.FileEntry.Length);
+					using MemoryMappedView memoryMappedView = CreateMemoryMappedView(memoryMappedFile, file);
 
 					for (int chunkIdx = 0; chunkIdx < chunks.Count; chunkIdx++)
 					{
@@ -508,6 +508,20 @@ namespace EpicGames.Horde.Storage.Nodes
 			}
 		}
 
+		static MemoryMappedView CreateMemoryMappedView(MemoryMappedFile memoryMappedFile, OutputFile file)
+		{
+			try
+			{
+				MemoryMappedView view = new MemoryMappedView(memoryMappedFile, 0, file.FileEntry.Length);
+				return view;
+			}
+			catch (Exception ex)
+			{
+				file.FileInfo.Refresh();
+				throw new StorageException($"Unable to create mapped view of file (entry length: {file.FileEntry.Length}, file length: {file.FileInfo.Length}): {ex.Message}", ex);
+			}
+		}
+
 #pragma warning disable IDE0051
 		// Update counters for extracting chunks without writing any data. Useful for profiling bottlenecks in other stages of the pipeline.
 		static Task ExtractChunksToNullAsync(OutputFile file, ArraySegment<BatchReadResponseItem<OutputChunk>> chunks, OutputStats stats, ILogger logger, CancellationToken cancellationToken)
