[UBA]
* Hooked up cross architecture logic from c# to c++

diff --git a/Engine/Source/Programs/Shared/EpicGames.UBA/ISessionServer.cs b/Engine/Source/Programs/Shared/EpicGames.UBA/ISessionServer.cs
--- a/Engine/Source/Programs/Shared/EpicGames.UBA/ISessionServer.cs
+++ b/Engine/Source/Programs/Shared/EpicGames.UBA/ISessionServer.cs
@@ -93,8 +93,13 @@ namespace EpicGames.UBA
 	/// <summary>
 	/// Event args for remote process slot available event
 	/// </summary>
-	public class RemoteProcessSlotAvailableEventArgs : EventArgs
+	/// <param name="isCrossArchitecture">Is true if the available remote process slot is on a cross architecture machine</param>
+	public class RemoteProcessSlotAvailableEventArgs(bool isCrossArchitecture) : EventArgs
 	{
+		/// <summary>
+		/// Is true if the available remote process slot is on a cross architecture machine
+		/// </summary>
+		public bool IsCrossArchitecture { get; } = isCrossArchitecture;
 	}
 
 	/// <summary>
@@ -169,8 +174,9 @@ namespace EpicGames.UBA
 		/// <param name="weight">Number of cores this process uses</param>
 		/// <param name="knownInputs">Optionally contains input that we know process will need. Memory block containing zero-terminated strings with an extra termination in the end.</param>
 		/// <param name="knownInputsCount">Number of strings in known inputs</param>
+		/// <param name="canExecuteCrossArchitecture">True means that this process can run on a different architecture and expects a mapping to exist</param>
 		/// <returns>The remote process being run</returns>
-		public abstract IProcess RunProcessRemote(ProcessStartInfo info, IProcess.ExitedEventHandler? exitedEventHandler, double weight = 1.0, byte[]? knownInputs = null, uint knownInputsCount = 0);
+		public abstract IProcess RunProcessRemote(ProcessStartInfo info, IProcess.ExitedEventHandler? exitedEventHandler, double weight = 1.0, byte[]? knownInputs = null, uint knownInputsCount = 0, bool canExecuteCrossArchitecture = false);
 
 		/// <summary>
 		/// Refresh cached information about directories
diff --git a/Engine/Source/Programs/Shared/EpicGames.UBA/Impl/SessionServerImpl.cs b/Engine/Source/Programs/Shared/EpicGames.UBA/Impl/SessionServerImpl.cs
--- a/Engine/Source/Programs/Shared/EpicGames.UBA/Impl/SessionServerImpl.cs
+++ b/Engine/Source/Programs/Shared/EpicGames.UBA/Impl/SessionServerImpl.cs
@@ -63,7 +63,7 @@ namespace EpicGames.UBA.Impl
 
 	internal class SessionServerImpl : ISessionServer
 	{
-		delegate void RemoteProcessSlotAvailableCallback();
+		delegate void RemoteProcessSlotAvailableCallback(bool isCrossArchitecture);
 		delegate void RemoteProcessReturnedCallback(nint handle);
 
 		nint _handle = IntPtr.Zero;
@@ -107,7 +107,7 @@ namespace EpicGames.UBA.Impl
 		static extern nint SessionServer_RunProcess(nint server, nint info, bool async, bool enableDetour);
 
 		[DllImport("UbaHost", CharSet = CharSet.Auto)]
-		static extern nint SessionServer_RunProcessRemote(nint server, nint info, float weight, byte[]? knownInputs, uint knownInputsCount);
+		static extern nint SessionServer_RunProcessRemote(nint server, nint info, float weight, byte[]? knownInputs, uint knownInputsCount, bool allowCrossArchitecture);
 
 		[DllImport("UbaHost", CharSet = CharSet.Auto)]
 		static extern void SessionServer_SetMaxRemoteProcessCount(nint server, uint count);
@@ -118,6 +118,9 @@ namespace EpicGames.UBA.Impl
 		[DllImport("UbaHost", CharSet = CharSet.Auto)]
 		static extern void SessionServer_SetCustomCasKeyFromTrackedInputs(nint server, nint process, string filename, string workingdir);
 
+		[DllImport("UbaHost", CharSet = CharSet.Auto)]
+		static extern void SessionServer_RegisterCrossArchitectureMapping(nint server, string from, string to);
+
 		[DllImport("UbaHost", CharSet = CharSet.Auto)]
 		static extern void SessionServer_PrintSummary(nint server);
 
@@ -152,6 +155,9 @@ namespace EpicGames.UBA.Impl
 			_handle = SessionServer_Create(_info.GetHandle(), environmentMemory.GetBuffer(), (uint)environmentMemory.Position);
 			SessionServer_SetRemoteProcessAvailable(_handle, _remoteProcessSlotAvailableCallbackDelegate);
 			SessionServer_SetRemoteProcessReturned(_handle, _remoteProcessReturnedCallbackDelegate);
+
+			foreach (var kv in Utils.CrossArchitecturePaths)
+				SessionServer_RegisterCrossArchitectureMapping(_handle, kv.Key, kv.Value);
 		}
 
 		#region IDisposable
@@ -192,10 +198,10 @@ namespace EpicGames.UBA.Impl
 			return process;
 		}
 
-		public IProcess RunProcessRemote(ProcessStartInfo info, IProcess.ExitedEventHandler? exitedEventHandler, double weight, byte[]? knownInputs, uint knownInputsCount)
+		public IProcess RunProcessRemote(ProcessStartInfo info, IProcess.ExitedEventHandler? exitedEventHandler, double weight, byte[]? knownInputs, uint knownInputsCount, bool canExecuteCrossArchitecture)
 		{
 			IProcessStartInfo startInfo = IProcessStartInfo.CreateProcessStartInfo(info, exitedEventHandler != null);
-			nint processPtr = SessionServer_RunProcessRemote(_handle, startInfo.GetHandle(), (float)weight, knownInputs, knownInputsCount);
+			nint processPtr = SessionServer_RunProcessRemote(_handle, startInfo.GetHandle(), (float)weight, knownInputs, knownInputsCount, canExecuteCrossArchitecture);
 			IProcess process = IProcess.CreateProcess(processPtr, startInfo, exitedEventHandler, info.UserData);
 			_remoteProcesses.AddOrUpdate(process.Hash, process, (k, v) => process);
 			return process;
@@ -225,9 +231,9 @@ namespace EpicGames.UBA.Impl
 		public void CancelAll() => SessionServer_CancelAll(_handle);
 		#endregion
 
-		void RaiseRemoteProcessSlotAvailable()
+		void RaiseRemoteProcessSlotAvailable(bool isCrossArchitecture)
 		{
-			RemoteProcessSlotAvailable?.Invoke(this, new RemoteProcessSlotAvailableEventArgs());
+			RemoteProcessSlotAvailable?.Invoke(this, new RemoteProcessSlotAvailableEventArgs(isCrossArchitecture));
 		}
 
 		void RaiseRemoteProcessReturned(nint handle)
diff --git a/Engine/Source/Programs/UnrealBuildTool/Executors/UnrealBuildAccelerator/UBAExecutor.cs b/Engine/Source/Programs/UnrealBuildTool/Executors/UnrealBuildAccelerator/UBAExecutor.cs
--- a/Engine/Source/Programs/UnrealBuildTool/Executors/UnrealBuildAccelerator/UBAExecutor.cs
+++ b/Engine/Source/Programs/UnrealBuildTool/Executors/UnrealBuildAccelerator/UBAExecutor.cs
@@ -538,7 +538,8 @@ namespace UnrealBuildTool
 			using ImmediateActionQueue queue = CreateActionQueue(inputActions, actionArtifactCache, UBAConfig.CacheMaxWorkers, logger);
 			int actionLimit = Math.Min(NumParallelProcesses, queue.TotalActions);
 			queue.CreateAutomaticRunner(action => RunActionLocal(queue, action), bUseActionWeights, actionLimit, NumParallelProcesses);
-			ImmediateActionQueueRunner remoteRunner = queue.CreateManualRunner(action => RunActionRemote(queue, action));
+			ImmediateActionQueueRunner remoteRunner = queue.CreateManualRunner(action => RunActionRemote(queue, action, false));
+			ImmediateActionQueueRunner crossArchitectureRemoteRunner = queue.CreateManualRunner(action => RunActionRemote(queue, action, true));
 			queue.CancellationToken.Register(onCancel);
 
 			queue.OnArtifactsRead += (action) => { UpdateCacheProgress(queue); UpdateProgress(queue); };
@@ -552,7 +553,8 @@ namespace UnrealBuildTool
 			// Handle process available from remote
 			session!.RemoteProcessSlotAvailable += (sender, args) =>
 			{
-				if (!queue.TryStartOneAction(remoteRunner) && (_bIsRemoteActionsAllowed && !UBAConfig.bForceBuildAllRemote))
+				ImmediateActionQueueRunner runner = args.IsCrossArchitecture ? crossArchitectureRemoteRunner : remoteRunner;
+				if (!queue.TryStartOneAction(runner) && (_bIsRemoteActionsAllowed && !UBAConfig.bForceBuildAllRemote))
 				{
 					uint count = ActionsLeftThatCanRunRemotely(queue);
 
@@ -864,13 +866,21 @@ namespace UnrealBuildTool
 			};
 		}
 
-		Func<Task>? RunActionRemote(ImmediateActionQueue queue, LinkedAction action)
+		Func<Task>? RunActionRemote(ImmediateActionQueue queue, LinkedAction action, bool isCrossArchitecture)
 		{
 			if (!CanRunRemotely(action))
 			{
 				return null;
 			}
 
+			// TODO: Implement this
+			bool canExecuteCrossArchitecture = true; // action.bCanExecuteCrossArchitecture
+
+			if (isCrossArchitecture && !canExecuteCrossArchitecture)
+			{
+				return null;
+			}
+
 			return () =>
 			{
 				if (_ubaDurationWaitingForRemote == TimeSpan.Zero)
@@ -916,6 +926,9 @@ namespace UnrealBuildTool
 					}
 				}
 
+				// TODO: Implement this
+				bool canExecuteCrossArchitecture = true; // action.bCanExecuteCrossArchitecture
+
 				ProcessStartInfo startInfo = GetActionStartInfo(action, out FileItem? pchItem);
 				_session!.RunProcessRemote(startInfo, (s, e) =>
 				{
@@ -971,7 +984,7 @@ namespace UnrealBuildTool
 					List<string> logLines = e.LogLines;
 					logLines.RemoveAll((line) => line.StartsWith("   Creating library ", StringComparison.OrdinalIgnoreCase) && line.EndsWith(".exp", StringComparison.OrdinalIgnoreCase));
 					ActionFinished(queue, new ExecuteResults(logLines, e.ExitCode, executionTime, processorTime, additionalDescription), action, pchItem, process);
-				}, action.Weight, knownInputs, knownInputsCount);
+				}, action.Weight, knownInputs, knownInputsCount, canExecuteCrossArchitecture);
 				return Task.CompletedTask;
 			};
 		}
