[UBA]
* Fixed bug where bool was used instead of byte when transferring from native to managed code.

diff --git a/Engine/Source/Programs/Shared/EpicGames.UBA/Impl/SessionServerImpl.cs b/Engine/Source/Programs/Shared/EpicGames.UBA/Impl/SessionServerImpl.cs
--- a/Engine/Source/Programs/Shared/EpicGames.UBA/Impl/SessionServerImpl.cs
+++ b/Engine/Source/Programs/Shared/EpicGames.UBA/Impl/SessionServerImpl.cs
@@ -63,7 +63,7 @@ namespace EpicGames.UBA.Impl
 
 	internal class SessionServerImpl : ISessionServer
 	{
-		delegate void RemoteProcessSlotAvailableCallback(bool isCrossArchitecture);
+		delegate void RemoteProcessSlotAvailableCallback(nint userData, byte isCrossArchitecture);
 		delegate void RemoteProcessReturnedCallback(nint handle);
 
 		nint _handle = IntPtr.Zero;
@@ -77,7 +77,7 @@ namespace EpicGames.UBA.Impl
 		static extern nint SessionServer_Create(nint info, byte[] environment, uint environmentSize);
 
 		[DllImport("UbaHost", CharSet = CharSet.Auto)]
-		static extern void SessionServer_SetRemoteProcessAvailable(nint server, RemoteProcessSlotAvailableCallback func);
+		static extern void SessionServer_SetRemoteProcessAvailable(nint server, RemoteProcessSlotAvailableCallback func, nint userData);
 
 		[DllImport("UbaHost", CharSet = CharSet.Auto)]
 		static extern void SessionServer_SetRemoteProcessReturned(nint server, RemoteProcessReturnedCallback func);
@@ -153,7 +153,7 @@ namespace EpicGames.UBA.Impl
 			}
 
 			_handle = SessionServer_Create(_info.GetHandle(), environmentMemory.GetBuffer(), (uint)environmentMemory.Position);
-			SessionServer_SetRemoteProcessAvailable(_handle, _remoteProcessSlotAvailableCallbackDelegate);
+			SessionServer_SetRemoteProcessAvailable(_handle, _remoteProcessSlotAvailableCallbackDelegate, 0);
 			SessionServer_SetRemoteProcessReturned(_handle, _remoteProcessReturnedCallbackDelegate);
 
 			foreach (var kv in Utils.CrossArchitecturePaths)
@@ -231,9 +231,9 @@ namespace EpicGames.UBA.Impl
 		public void CancelAll() => SessionServer_CancelAll(_handle);
 		#endregion
 
-		void RaiseRemoteProcessSlotAvailable(bool isCrossArchitecture)
+		void RaiseRemoteProcessSlotAvailable(nint userData, byte isCrossArchitecture)
 		{
-			RemoteProcessSlotAvailable?.Invoke(this, new RemoteProcessSlotAvailableEventArgs(isCrossArchitecture));
+			RemoteProcessSlotAvailable?.Invoke(this, new RemoteProcessSlotAvailableEventArgs(isCrossArchitecture != 0));
 		}
 
 		void RaiseRemoteProcessReturned(nint handle)
