Slowed the read rate to avoid out of memory issues.

diff --git a/Engine/Source/Programs/Shared/EpicGames.Horde/Storage/Nodes/DirectoryNode.Extract.cs b/Engine/Source/Programs/Shared/EpicGames.Horde/Storage/Nodes/DirectoryNode.Extract.cs
--- a/Engine/Source/Programs/Shared/EpicGames.Horde/Storage/Nodes/DirectoryNode.Extract.cs
+++ b/Engine/Source/Programs/Shared/EpicGames.Horde/Storage/Nodes/DirectoryNode.Extract.cs
@@ -298,7 +298,7 @@ namespace EpicGames.Horde.Storage.Nodes
 		/// <param name="cancellationToken">Cancellation token for the operation</param>
 		public static async Task ExtractAsync(this DirectoryNode directoryNode, DirectoryInfo directoryInfo, IProgress<IExtractStats>? progress, TimeSpan frequency, ILogger logger, CancellationToken cancellationToken)
 		{
-			int numReadTasks = 16;
+			int numReadTasks = 8; //Slowed the read rate to avoid out of memory issues
 			int numDecodeTasks = Math.Min(Environment.ProcessorCount, 16);
 			int numWriteTasks = Math.Min(1 + (int)(directoryNode.Length / (16 * 1024 * 1024)), 16);
 			logger.LogInformation("Using {NumReadTasks} read tasks, {NumDecodeTasks} decode tasks, {NumWriteTasks} write tasks", numReadTasks, numDecodeTasks, numWriteTasks);
