EpicGames.UBA: analysis fixes

#rnx

diff --git a/Engine/Source/Programs/Shared/EpicGames.UBA/Impl/SessionServerImpl.cs b/Engine/Source/Programs/Shared/EpicGames.UBA/Impl/SessionServerImpl.cs
--- a/Engine/Source/Programs/Shared/EpicGames.UBA/Impl/SessionServerImpl.cs
+++ b/Engine/Source/Programs/Shared/EpicGames.UBA/Impl/SessionServerImpl.cs
@@ -156,8 +156,10 @@ namespace EpicGames.UBA.Impl
 			SessionServer_SetRemoteProcessAvailable(_handle, _remoteProcessSlotAvailableCallbackDelegate, 0);
 			SessionServer_SetRemoteProcessReturned(_handle, _remoteProcessReturnedCallbackDelegate);
 
-			foreach (var kv in Utils.CrossArchitecturePaths)
+			foreach (KeyValuePair<string, string> kv in Utils.CrossArchitecturePaths)
+			{
 				SessionServer_RegisterCrossArchitectureMapping(_handle, kv.Key, kv.Value);
+			}
 		}
 
 		#region IDisposable
