EpicGames.Horde: Prevent send and possible exceptions in background send for ComputeSocket

diff --git a/Engine/Source/Programs/Shared/EpicGames.Horde/Compute/ComputeSocket.cs b/Engine/Source/Programs/Shared/EpicGames.Horde/Compute/ComputeSocket.cs
--- a/Engine/Source/Programs/Shared/EpicGames.Horde/Compute/ComputeSocket.cs
+++ b/Engine/Source/Programs/Shared/EpicGames.Horde/Compute/ComputeSocket.cs
@@ -860,7 +860,7 @@ namespace EpicGames.Horde.Compute
 				while (!cancellationToken.IsCancellationRequested)
 				{
 					ReadOnlyMemory<byte> memory = reader.GetReadBuffer();
-					if (memory.Length > 0)
+					if (memory.Length > 0 && !_isClosing)
 					{
 						await SendAsync(channelId, memory, cancellationToken);
 						reader.AdvanceReadPosition(memory.Length);
@@ -879,7 +879,10 @@ namespace EpicGames.Horde.Compute
 			}
 			catch (Exception ex)
 			{
-				_logger.LogInformation(ex, "Exception in background send: {Message}", ex.Message);
+				if (!_isClosing)
+				{
+					_logger.LogInformation(ex, "Exception in background send: {Message}", ex.Message);
+				}
 			}
 		}
 	}
