EpicGames.Core: Add helper checking if current Windows process is running with administrator privileges

diff --git a/Engine/Source/Programs/Shared/EpicGames.Core/UserSecurityUtils.cs b/Engine/Source/Programs/Shared/EpicGames.Core/UserSecurityUtils.cs
new file mode 100644
--- /dev/null
+++ b/Engine/Source/Programs/Shared/EpicGames.Core/UserSecurityUtils.cs
@@ -0,0 +1,36 @@
+// Copyright Epic Games, Inc. All Rights Reserved.
+
+using System;
+using System.Security.Principal;
+
+namespace EpicGames.Core;
+
+/// <summary>
+/// Provides utility methods for checking user security contexts and permissions
+/// </summary>
+public static class UserSecurityUtils
+{
+	/// <summary>
+	/// Checks if the current process is running with administrator privileges on Windows
+	/// </summary>
+	/// <returns>True if process has administrator privileges</returns>
+	public static bool IsAdministrator()
+	{
+		if (!OperatingSystem.IsWindows())
+		{
+			return false;
+		}
+		
+		try
+		{
+			using WindowsIdentity identity = WindowsIdentity.GetCurrent();
+			WindowsPrincipal principal = new (identity);
+			return principal.IsInRole(WindowsBuiltInRole.Administrator);
+		}
+		catch
+		{
+			return false;
+		}
+	}
+}
+
