EpicGames.Core: Add directory helpers for copy and total disk size

diff --git a/Engine/Source/Programs/Shared/EpicGames.Core/DirectoryReference.cs b/Engine/Source/Programs/Shared/EpicGames.Core/DirectoryReference.cs
--- a/Engine/Source/Programs/Shared/EpicGames.Core/DirectoryReference.cs
+++ b/Engine/Source/Programs/Shared/EpicGames.Core/DirectoryReference.cs
@@ -229,7 +229,47 @@ namespace EpicGames.Core
 		/// <param name="destDir">The destination directory</param>
 		public static void Move(DirectoryReference sourceDir, DirectoryReference destDir)
 			=> Directory.Move(sourceDir.FullName, destDir.FullName);
-
+		
+		/// <summary>
+		/// Copies a directory from one location to another
+		/// </summary>
+		/// <param name="sourceDir">The source directory</param>
+		/// <param name="destDir">The destination directory</param>
+		public static void Copy(DirectoryReference sourceDir, DirectoryReference destDir)
+		{
+			CreateDirectory(destDir);
+			
+			foreach (string file in Directory.GetFiles(sourceDir.FullName))
+			{
+				File.Copy(file, Path.Combine(destDir.FullName, Path.GetFileName(file)), true);
+			}
+			
+			foreach (string directory in Directory.GetDirectories(sourceDir.FullName))
+			{
+				Copy(new DirectoryReference(directory), Combine(destDir, Path.GetFileName(directory)));
+			}
+		}
+		
+		/// <summary>
+		/// Calculates the total size in bytes of all files in a directory and its subdirectories
+		/// </summary>
+		/// <param name="dirPath">Path to directory</param>
+		/// <returns>The total size in bytes of all accessible files in the directory and its subdirectories</returns>
+		public static long GetTotalSize(DirectoryReference dirPath)
+		{
+			long size = 0;
+			foreach (FileReference file in EnumerateFiles(dirPath, "*", SearchOption.AllDirectories))
+			{
+				try
+				{
+					size += file.ToFileInfo().Length;
+				}
+				catch (UnauthorizedAccessException) { }
+				catch (FileNotFoundException) { }
+			}
+			return size;
+		}
+		
 		/// <summary>
 		/// Helper function to create a remote directory reference. Unlike normal DirectoryReference objects, these aren't converted to a full path in the local filesystem.
 		/// </summary>
