Fixed typo "nofity" to "notify" in UHT error message.

diff --git a/Engine/Source/Programs/Shared/EpicGames.UHT/Types/UhtProperty.cs b/Engine/Source/Programs/Shared/EpicGames.UHT/Types/UhtProperty.cs
--- a/Engine/Source/Programs/Shared/EpicGames.UHT/Types/UhtProperty.cs
+++ b/Engine/Source/Programs/Shared/EpicGames.UHT/Types/UhtProperty.cs
@@ -1940,12 +1940,12 @@ namespace EpicGames.UHT.Types
 				{
 					if (classObj.ClassType != UhtClassType.Class)
 					{
-						this.LogError("FieldNofity are not valid on UInterface.");
+						this.LogError("FieldNotify are not valid on UInterface.");
 					}
 				}
 				else
 				{
-					this.LogError("FieldNofity property are only valid as UClass member variable.");
+					this.LogError("FieldNotify property are only valid as UClass member variable.");
 				}
 			}
 		}
