EpicGames.UBA: Use concurrent collections in Utils

#rnx

diff --git a/Engine/Source/Programs/Shared/EpicGames.UBA/Utils.cs b/Engine/Source/Programs/Shared/EpicGames.UBA/Utils.cs
--- a/Engine/Source/Programs/Shared/EpicGames.UBA/Utils.cs
+++ b/Engine/Source/Programs/Shared/EpicGames.UBA/Utils.cs
@@ -1,8 +1,10 @@
 ﻿// Copyright Epic Games, Inc. All Rights Reserved.
 
 using System;
+using System.Collections.Concurrent;
 using System.Collections.Generic;
 using System.IO;
+using System.Linq;
 using System.Reflection;
 using System.Runtime.InteropServices;
 
@@ -17,21 +19,21 @@ namespace EpicGames.UBA
 		/// Is UBA available?
 		/// </summary>
 		public static bool IsAvailable() => s_available.Value;
-		static readonly Lazy<bool> s_available = new Lazy<bool>(() => File.Exists(GetLibraryPath()));
+		static readonly Lazy<bool> s_available = new(() => File.Exists(GetLibraryPath()));
 
 		/// <summary>
 		/// Paths that are not allowed to be transferred over the network for UBA remote agents.
 		/// </summary>
 		/// <returns>enumerable of disallowed paths</returns>
-		public static IEnumerable<string> DisallowedPaths() => s_disallowedPaths;
-		static readonly SortedSet<string> s_disallowedPaths = [];
+		public static IEnumerable<string> DisallowedPaths() => s_disallowedPaths.Distinct().Order();
+		static readonly ConcurrentBag<string> s_disallowedPaths = [];
 
 		/// <summary>
 		/// Mapping of binary paths for cross architecture host binaries, to allow for using helpers of a different architecture.
 		/// </summary>
 		/// <returns>dictionary of binary mappings, where the key is the binary for the current host architecture</returns>
 		public static IDictionary<string, string> CrossArchitecturePaths => s_crossArchitecturePaths;
-		static readonly Dictionary<string, string> s_crossArchitecturePaths = [];
+		static readonly ConcurrentDictionary<string, string> s_crossArchitecturePaths = [];
 
 		/// <summary>
 		/// Registers a path that is not allowed to be transferred over the network for UBA remote agents.
@@ -39,7 +41,10 @@ namespace EpicGames.UBA
 		/// <param name="paths">The paths to add to the disallowed list</param>
 		public static void RegisterDisallowedPaths(params string[] paths)
 		{
-			s_disallowedPaths.UnionWith(paths);
+			foreach (string path in paths)
+			{
+				s_disallowedPaths.Add(path);
+			}
 			DisallowedPathRegistered?.Invoke(s_disallowedPaths, new(paths));
 		}
 
@@ -59,14 +64,14 @@ namespace EpicGames.UBA
 		/// </summary>
 		/// <param name="sender">collection that is being changed</param>
 		/// <param name="e">event args containing which paths were added</param>
-		public delegate void DisallowedPathRegisteredEventHandler(SortedSet<string> sender, DisallowedPathRegisteredEventArgs e);
+		public delegate void DisallowedPathRegisteredEventHandler(ConcurrentBag<string> sender, DisallowedPathRegisteredEventArgs e);
 
 		/// <summary>
 		/// Delegate for registering a cross architecture path
 		/// </summary>
 		/// <param name="sender">collection that is being changed</param>
 		/// <param name="e">event args containing which path was added</param>
-		public delegate void CrossArchitecturePathRegisteredEventHandler(Dictionary<string, string> sender, CrossArchitecturePathRegisteredEventArgs e);
+		public delegate void CrossArchitecturePathRegisteredEventHandler(ConcurrentDictionary<string, string> sender, CrossArchitecturePathRegisteredEventArgs e);
 
 		/// <summary>
 		/// Remote disallowed path registered event handler
