Horde: Fix static analysis warning.

#rnx

diff --git a/Engine/Source/Programs/Shared/EpicGames.Horde/Storage/ContentChunker.cs b/Engine/Source/Programs/Shared/EpicGames.Horde/Storage/ContentChunker.cs
--- a/Engine/Source/Programs/Shared/EpicGames.Horde/Storage/ContentChunker.cs
+++ b/Engine/Source/Programs/Shared/EpicGames.Horde/Storage/ContentChunker.cs
@@ -8,7 +8,6 @@ using System.Threading.Channels;
 using System.Threading.Tasks;
 using EpicGames.Core;
 using EpicGames.Horde.Storage.Nodes;
-using ZstdSharp.Unsafe;
 
 namespace EpicGames.Horde.Storage
 {
