Horde: Do not try to parse portable PDB files. These are not supported by symstore.

diff --git a/Engine/Source/Programs/Shared/EpicGames.Horde/Symbols/SymStore.cs b/Engine/Source/Programs/Shared/EpicGames.Horde/Symbols/SymStore.cs
--- a/Engine/Source/Programs/Shared/EpicGames.Horde/Symbols/SymStore.cs
+++ b/Engine/Source/Programs/Shared/EpicGames.Horde/Symbols/SymStore.cs
@@ -112,6 +112,8 @@ namespace EpicGames.Horde.Symbols
 		#endregion
 		#region PDB files
 
+		static ReadOnlyMemory<byte> s_portablePdbSignature = new byte[] { 0x42, 0x53, 0x4a, 0x42 };
+
 		/// <summary>
 		/// Parse the hash value from a PDB file, consisting of concatenated hex values for the PDB GUID and age (the number of times it has been written).
 		/// 
@@ -129,6 +131,13 @@ namespace EpicGames.Horde.Symbols
 			using MemoryMappedFile memoryMappedFile = MemoryMappedFile.CreateFromFile(stream, null, stream.Length, MemoryMappedFileAccess.Read, HandleInheritability.None, false);
 			using MemoryMappedView memoryMappedView = new MemoryMappedView(memoryMappedFile, 0, stream.Length, MemoryMappedFileAccess.Read);
 
+			// Check if this is a portable PDB. Symstore does not support these.
+			Span<byte> header = memoryMappedView.GetMemory(0, 4).Span.Slice(0, 4);
+			if (header.SequenceEqual(s_portablePdbSignature.Span))
+			{
+				return null;
+			}
+
 			// Create the MSF file
 			MsfFile msfFile = new MsfFile(memoryMappedView);
 
