Horde: Fix static analysis warning.

#rnx

diff --git a/Engine/Source/Programs/Shared/EpicGames.Horde/Symbols/SymStore.cs b/Engine/Source/Programs/Shared/EpicGames.Horde/Symbols/SymStore.cs
--- a/Engine/Source/Programs/Shared/EpicGames.Horde/Symbols/SymStore.cs
+++ b/Engine/Source/Programs/Shared/EpicGames.Horde/Symbols/SymStore.cs
@@ -112,7 +112,7 @@ namespace EpicGames.Horde.Symbols
 		#endregion
 		#region PDB files
 
-		static ReadOnlyMemory<byte> s_portablePdbSignature = new byte[] { 0x42, 0x53, 0x4a, 0x42 };
+		static readonly ReadOnlyMemory<byte> s_portablePdbSignature = new byte[] { 0x42, 0x53, 0x4a, 0x42 };
 
 		/// <summary>
 		/// Parse the hash value from a PDB file, consisting of concatenated hex values for the PDB GUID and age (the number of times it has been written).
