Horde: Fix exception when chunking empty files, due to no chunks being returned.

#rnx

diff --git a/Engine/Source/Programs/Shared/EpicGames.Horde/Storage/ContentChunker.cs b/Engine/Source/Programs/Shared/EpicGames.Horde/Storage/ContentChunker.cs
--- a/Engine/Source/Programs/Shared/EpicGames.Horde/Storage/ContentChunker.cs
+++ b/Engine/Source/Programs/Shared/EpicGames.Horde/Storage/ContentChunker.cs
@@ -399,6 +399,7 @@ namespace EpicGames.Horde.Storage
 			int _boundaryIdx;
 			uint _rollingHash;
 			ReadOnlyMemory<byte> _data;
+			bool _firstChunk = true;
 
 			public LatestCdcOutput(ParallelBuzHashChunker pipeline, object? userData, Task<Block> blockTask)
 			{
@@ -421,6 +422,11 @@ namespace EpicGames.Horde.Storage
 				RingBuffer ringBuffer = _pipeline._ringBuffer;
 				LeafChunkedDataNodeOptions options = _pipeline._options;
 
+				// Always need to return at least one chunk, even if the source data is empty. Use this flag to track whether
+				// we've returned something.
+				bool firstChunk = _firstChunk;
+				_firstChunk = false;
+
 				for (; ; )
 				{
 					Block block = await _blockTask.WaitAsync(cancellationToken);
@@ -452,12 +458,13 @@ namespace EpicGames.Horde.Storage
 						_data = blockData.Slice(_blockOffset);
 						_ = BuzHash.FindChunkLength(_data.Span, options.MinSize, options.MaxSize, options.TargetSize, out _rollingHash);
 
-						if (_data.Length == 0)
+						if (_data.Length == 0 && !firstChunk)
 						{
 							ringBuffer.Free(block.BufferOffset, block.Length);
 							_blockTask = null!;
 							return false;
 						}
+
 						return true;
 					}
 
@@ -581,8 +588,16 @@ namespace EpicGames.Horde.Storage
 					TaskCompletionSource<Block>? blockTcs = inputRequest.FirstBlockTcs;
 					for (long sourceOffset = 0; blockTcs != null;)
 					{
+						// Read a block of data
 						int readLength = (int)Math.Min(input.Length - sourceOffset, ReadBlockSize);
 
+						int bufferOffset = await _ringBuffer.AllocAsync(readLength, cancellationToken);
+						Memory<byte> data = _ringBuffer.GetData(bufferOffset, readLength);
+						Task readTask = await input.StartReadAsync(data, cancellationToken);
+
+						sourceOffset += readLength;
+
+						// Post the request to handle the result
 						TaskCompletionSource<Block>? nextBlockTcs;
 						if (sourceOffset >= input.Length)
 						{
@@ -593,14 +608,9 @@ namespace EpicGames.Horde.Storage
 							nextBlockTcs = new TaskCompletionSource<Block>(TaskCreationOptions.RunContinuationsAsynchronously);
 						}
 
-						int bufferOffset = await _ringBuffer.AllocAsync(readLength, cancellationToken);
-						Memory<byte> data = _ringBuffer.GetData(bufferOffset, readLength);
-						Task readTask = await input.StartReadAsync(data, cancellationToken);
-
 						BlockRequest blockRequest = new BlockRequest(bufferOffset, readLength, readTask, blockTcs, nextBlockTcs?.Task);
 						await writer.WriteAsync(blockRequest, cancellationToken);
 
-						sourceOffset += readLength;
 						blockTcs = nextBlockTcs;
 					}
 				}
