Horde: Fix test error due to content chunker not enumerating an empty chunk for an empty input.

#rnx

diff --git a/Engine/Source/Programs/Shared/EpicGames.Horde/Storage/ContentChunker.cs b/Engine/Source/Programs/Shared/EpicGames.Horde/Storage/ContentChunker.cs
--- a/Engine/Source/Programs/Shared/EpicGames.Horde/Storage/ContentChunker.cs
+++ b/Engine/Source/Programs/Shared/EpicGames.Horde/Storage/ContentChunker.cs
@@ -160,6 +160,7 @@ namespace EpicGames.Horde.Storage
 			readonly SerialBuzHashChunker _pipeline;
 			public readonly ChunkerSource Input;
 
+			public bool _firstOutput = true;
 			public uint _rollingHash;
 			public long _inputOffset;
 			public long _outputOffset;
@@ -237,10 +238,13 @@ namespace EpicGames.Horde.Storage
 
 		async ValueTask<bool> MoveNextAsync(ChunkerSource input, Output output, CancellationToken cancellationToken)
 		{
+			bool firstOutput = output._firstOutput;
+			output._firstOutput = false;
+
 			if (output._outputOffset == output.Input.Length)
 			{
 				output._data = ReadOnlyMemory<byte>.Empty;
-				return false;
+				return firstOutput;
 			}
 			else
 			{
