Horde: Add support for parallel uploads of bundles. Bundle locators are now assigned by the client, allowing the next bundle to be assembled before the previous one has finished uploading.

diff --git a/Engine/Source/Programs/Horde/Drivers/JobDriver/Execution/JobExecutor.cs b/Engine/Source/Programs/Horde/Drivers/JobDriver/Execution/JobExecutor.cs
--- a/Engine/Source/Programs/Horde/Drivers/JobDriver/Execution/JobExecutor.cs
+++ b/Engine/Source/Programs/Horde/Drivers/JobDriver/Execution/JobExecutor.cs
@@ -844,7 +844,7 @@ namespace JobDriver.Execution
 				Stopwatch timer = Stopwatch.StartNew();
 
 				IHashedBlobRef<DirectoryNode> rootNodeRef;
-				await using (IBlobWriter blobWriter = storage.CreateBlobWriter(artifact.RefName))
+				await using (IBlobWriter blobWriter = storage.CreateBlobWriter(artifact.RefName, cancellationToken: cancellationToken))
 				{
 					DirectoryNode buildGraphNode = new DirectoryNode();
 					await buildGraphNode.AddFilesAsync(workspaceDir, buildGraphFiles, blobWriter, cancellationToken: cancellationToken);
@@ -1374,7 +1374,7 @@ namespace JobDriver.Execution
 				Stopwatch timer = Stopwatch.StartNew();
 
 				IHashedBlobRef<DirectoryNode> outputNodeRef;
-				await using (IBlobWriter blobWriter = storage.CreateBlobWriter(artifact.RefName))
+				await using (IBlobWriter blobWriter = storage.CreateBlobWriter(artifact.RefName, cancellationToken: cancellationToken))
 				{
 					DirectoryNode outputNode = new DirectoryNode();
 
@@ -1464,7 +1464,7 @@ namespace JobDriver.Execution
 					Stopwatch timer = Stopwatch.StartNew();
 
 					IHashedBlobRef<DirectoryNode> outputNodeRef;
-					await using (IBlobWriter blobWriter = storage.CreateBlobWriter(artifact.RefName))
+					await using (IBlobWriter blobWriter = storage.CreateBlobWriter(artifact.RefName, cancellationToken: cancellationToken))
 					{
 						DirectoryNode outputNode = new DirectoryNode();
 
diff --git a/Engine/Source/Programs/Shared/EpicGames.Horde/Storage/Bundles/BundleCache.cs b/Engine/Source/Programs/Shared/EpicGames.Horde/Storage/Bundles/BundleCache.cs
--- a/Engine/Source/Programs/Shared/EpicGames.Horde/Storage/Bundles/BundleCache.cs
+++ b/Engine/Source/Programs/Shared/EpicGames.Horde/Storage/Bundles/BundleCache.cs
@@ -61,6 +61,9 @@ namespace EpicGames.Horde.Storage.Bundles
 		readonly IScopedAsyncCache<BundlePageCacheKey, IReadOnlyMemoryOwner<byte>> _bundlePageCache;
 		readonly MemoryCache? _headerCache;
 		readonly MemoryCache? _packetCache;
+#pragma warning disable CA2213
+		readonly SemaphoreSlim _writeSemaphore;
+#pragma warning restore CA2213
 
 		internal IScopedAsyncCache<PacketReaderCacheKey, PacketReader> PacketReaderCache => _packetReaderCache;
 		internal IScopedAsyncCache<BundlePageCacheKey, IReadOnlyMemoryOwner<byte>> BundlePageCache => _bundlePageCache;
@@ -95,6 +98,11 @@ namespace EpicGames.Horde.Storage.Bundles
 		/// </summary>
 		public bool HasPacketCache => _packetCache != null;
 
+		/// <summary>
+		/// Semaphore for writing new data
+		/// </summary>
+		public SemaphoreSlim WriteSemaphore => _writeSemaphore;
+
 		/// <summary>
 		/// Constructor
 		/// </summary>
@@ -153,11 +161,15 @@ namespace EpicGames.Horde.Storage.Bundles
 			{
 				_packetCache = new MemoryCache(new MemoryCacheOptions { SizeLimit = options.PacketCacheSize });
 			}
+
+			_writeSemaphore = new SemaphoreSlim(64);
 		}
 
 		/// <inheritdoc/>
 		public ValueTask DisposeAsync()
 		{
+			WriteSemaphore.Dispose();
+
 			_packetReaderCache.Clear();
 			_bundlePageCache.Clear();
 
diff --git a/Engine/Source/Programs/Shared/EpicGames.Horde/Storage/Bundles/BundleHandle.cs b/Engine/Source/Programs/Shared/EpicGames.Horde/Storage/Bundles/BundleHandle.cs
--- a/Engine/Source/Programs/Shared/EpicGames.Horde/Storage/Bundles/BundleHandle.cs
+++ b/Engine/Source/Programs/Shared/EpicGames.Horde/Storage/Bundles/BundleHandle.cs
@@ -12,6 +12,17 @@ namespace EpicGames.Horde.Storage.Bundles
 	/// </summary>
 	public abstract class BundleHandle
 	{
+		/// <summary>
+		/// Locator for this bundle
+		/// </summary>
+		public BlobLocator Locator { get; }
+
+		/// <summary>
+		/// Constructor
+		/// </summary>
+		protected BundleHandle(BlobLocator locator)
+			=> Locator = locator;
+
 		/// <summary>
 		/// Flush contents of the bundle to storage
 		/// </summary>
@@ -30,7 +41,11 @@ namespace EpicGames.Horde.Storage.Bundles
 		/// <summary>
 		/// Attempt to get a locator for this bundle
 		/// </summary>
-		public abstract bool TryGetLocator(out BlobLocator locator);
+		public bool TryGetLocator(out BlobLocator locator)
+		{
+			locator = Locator;
+			return true;
+		}
 	}
 
 	/// <summary>
@@ -39,15 +54,14 @@ namespace EpicGames.Horde.Storage.Bundles
 	class FlushedBundleHandle : BundleHandle
 	{
 		readonly BundleStorageNamespace _storageNamespace;
-		readonly BlobLocator _locator;
 
 		/// <summary>
 		/// Constructor
 		/// </summary>
 		public FlushedBundleHandle(BundleStorageNamespace storageNamespace, BlobLocator locator)
+			: base(locator)
 		{
 			_storageNamespace = storageNamespace;
-			_locator = locator;
 		}
 
 		/// <inheritdoc/>
@@ -55,26 +69,19 @@ namespace EpicGames.Horde.Storage.Bundles
 
 		/// <inheritdoc/>
 		public override Task<Stream> OpenAsync(int offset = 0, int? length = null, CancellationToken cancellationToken = default)
-			=> _storageNamespace.Backend.OpenBlobAsync(_locator, offset, length, cancellationToken);
+			=> _storageNamespace.Backend.OpenBlobAsync(Locator, offset, length, cancellationToken);
 
 		/// <inheritdoc/>
 		public override async ValueTask<IReadOnlyMemoryOwner<byte>> ReadAsync(int offset, int? length, CancellationToken cancellationToken = default)
-			=> await _storageNamespace.Backend.ReadBlobAsync(_locator, offset, length, cancellationToken);
-
-		/// <inheritdoc/>
-		public override bool TryGetLocator(out BlobLocator locator)
-		{
-			locator = _locator;
-			return true;
-		}
+			=> await _storageNamespace.Backend.ReadBlobAsync(Locator, offset, length, cancellationToken);
 
 		/// <inheritdoc/>
-		public override bool Equals(object? obj) => obj is FlushedBundleHandle other && _locator == other._locator;
+		public override bool Equals(object? obj) => obj is FlushedBundleHandle other && Locator == other.Locator;
 
 		/// <inheritdoc/>
-		public override int GetHashCode() => _locator.GetHashCode();
+		public override int GetHashCode() => Locator.GetHashCode();
 
 		/// <inheritdoc/>
-		public override string ToString() => _locator.ToString();
+		public override string ToString() => Locator.ToString();
 	}
 }
diff --git a/Engine/Source/Programs/Shared/EpicGames.Horde/Storage/Bundles/BundleStorageNamespace.cs b/Engine/Source/Programs/Shared/EpicGames.Horde/Storage/Bundles/BundleStorageNamespace.cs
--- a/Engine/Source/Programs/Shared/EpicGames.Horde/Storage/Bundles/BundleStorageNamespace.cs
+++ b/Engine/Source/Programs/Shared/EpicGames.Horde/Storage/Bundles/BundleStorageNamespace.cs
@@ -25,6 +25,7 @@ namespace EpicGames.Horde.Storage.Bundles
 		readonly BundleOptions _options;
 		readonly PacketReaderStats _packetReaderStats = new PacketReaderStats();
 		readonly Bundles.V1.BundleReader _bundleReader;
+		readonly ILogger _logger;
 
 		internal Bundles.V1.BundleReader BundleReader => _bundleReader;
 
@@ -57,6 +58,7 @@ namespace EpicGames.Horde.Storage.Bundles
 			_cache = cache;
 			_options = options ?? BundleOptions.Default;
 			_bundleReader = new Bundles.V1.BundleReader(this, cache, logger);
+			_logger = logger;
 		}
 
 		/// <summary>
@@ -146,7 +148,7 @@ namespace EpicGames.Horde.Storage.Bundles
 			return locators;
 		}
 
-		#region Nodes
+		#region Blobs
 
 		/// <inheritdoc/>
 		public IBlobRef CreateBlobRef(BlobLocator locator)
@@ -175,7 +177,7 @@ namespace EpicGames.Horde.Storage.Bundles
 		}
 
 		/// <inheritdoc/>
-		public IBlobWriter CreateBlobWriter(string? basePath = null, BlobSerializerOptions? serializerOptions = null)
+		public IBlobWriter CreateBlobWriter(string? basePath = null, BlobSerializerOptions? serializerOptions = null, CancellationToken cancellationToken = default)
 		{
 			BundleVersion version = _options.MaxVersion;
 			if (version == BundleVersion.LatestV1)
@@ -184,7 +186,7 @@ namespace EpicGames.Horde.Storage.Bundles
 			}
 			else if (version == BundleVersion.LatestV2)
 			{
-				return new Bundles.V2.BundleWriter(this, basePath, _cache, _options, serializerOptions);
+				return new Bundles.V2.BundleWriter(this, basePath, _cache, _options, serializerOptions, _logger, cancellationToken);
 			}
 			else
 			{
diff --git a/Engine/Source/Programs/Shared/EpicGames.Horde/Storage/Bundles/V2/BundleWriter.cs b/Engine/Source/Programs/Shared/EpicGames.Horde/Storage/Bundles/V2/BundleWriter.cs
--- a/Engine/Source/Programs/Shared/EpicGames.Horde/Storage/Bundles/V2/BundleWriter.cs
+++ b/Engine/Source/Programs/Shared/EpicGames.Horde/Storage/Bundles/V2/BundleWriter.cs
@@ -10,6 +10,7 @@ using System.Runtime.CompilerServices;
 using System.Threading;
 using System.Threading.Tasks;
 using EpicGames.Core;
+using Microsoft.Extensions.Logging;
 
 namespace EpicGames.Horde.Storage.Bundles.V2
 {
@@ -37,7 +38,7 @@ namespace EpicGames.Horde.Storage.Bundles.V2
 
 			public override async ValueTask<BlobData> ReadExportAsync(int exportIdx, CancellationToken cancellationToken = default)
 			{
-				lock (_bundle.LockObject)
+				lock (_bundle._lockObject)
 				{
 					if (_flushedHandle == null)
 					{
@@ -51,14 +52,14 @@ namespace EpicGames.Horde.Storage.Bundles.V2
 				=> _flushedHandle?.TryAppendIdentifier(builder) ?? false;
 		}
 
-		// Fragment of a bundle that needs to be written to storage.
+		// Bundle that has not yet been written to storage.
 		internal sealed class PendingBundleHandle : BundleHandle, IDisposable
 		{
-			readonly object _lockObject = new object();
+			public readonly object _lockObject = new object();
 
+			readonly BundleWriter _outer;
 			readonly BundleStorageNamespace _storageNamespace;
-			readonly string? _basePath;
-			readonly BundleCache _cache;
+			readonly BundleCache _bundleCache;
 			readonly BundleOptions _bundleOptions;
 
 			FlushedBundleHandle? _flushedHandle;
@@ -69,10 +70,8 @@ namespace EpicGames.Horde.Storage.Bundles.V2
 			RefCountedMemoryWriter? _encodedPacketWriter;
 			readonly HashSet<BlobLocator> _bundleImports = new HashSet<BlobLocator>();
 
-			/// <summary>
-			/// Object used for locking access to this bundle's state
-			/// </summary>
-			public object LockObject => _lockObject;
+			public Task? _writeTask;
+			public readonly TaskCompletionSource _startWriteTcs = new TaskCompletionSource();
 
 			public FlushedBundleHandle? FlushedHandle => _flushedHandle;
 
@@ -81,14 +80,15 @@ namespace EpicGames.Horde.Storage.Bundles.V2
 			/// </summary>
 			public int Length => _encodedPacketWriter?.Length ?? throw new InvalidOperationException("Bundle has been flushed");
 
-			public PendingBundleHandle(BundleStorageNamespace storageNamespace, string? basePath, BundleCache cache, BundleOptions bundleOptions)
+			public PendingBundleHandle(BundleWriter outer, BlobLocator locator)
+				: base(locator)
 			{
-				_storageNamespace = storageNamespace;
-				_basePath = basePath;
-				_cache = cache;
-				_bundleOptions = bundleOptions;
+				_outer = outer;
+				_storageNamespace = _outer._storageNamespace;
+				_bundleCache = _outer._bundleCache;
+				_bundleOptions = _outer._bundleOptions;
 
-				_encodedPacketWriter = new RefCountedMemoryWriter(_cache.Allocator, 65536, nameof(PendingBundleHandle));
+				_encodedPacketWriter = new RefCountedMemoryWriter(_bundleCache.Allocator, 65536, nameof(PendingBundleHandle));
 
 				StartPacket();
 			}
@@ -156,7 +156,7 @@ namespace EpicGames.Horde.Storage.Bundles.V2
 				RuntimeAssert(_packetWriter == null);
 
 				_packetHandle = new PendingPacketHandle(this);
-				_packetWriter = new PacketWriter(this, _packetHandle, _cache.Allocator, _lockObject);
+				_packetWriter = new PacketWriter(this, _packetHandle, _bundleCache.Allocator, _lockObject);
 			}
 
 			void FinishPacket()
@@ -175,7 +175,7 @@ namespace EpicGames.Horde.Storage.Bundles.V2
 					// Point the packet handle to the encoded data
 					lock (_lockObject)
 					{
-						FlushedPacketHandle flushedPacketHandle = new FlushedPacketHandle(_storageNamespace, this, packetOffset, packetLength, _cache);
+						FlushedPacketHandle flushedPacketHandle = new FlushedPacketHandle(_storageNamespace, this, packetOffset, packetLength, _bundleCache);
 						_packetHandle.CompletePacket(flushedPacketHandle);
 					}
 				}
@@ -188,10 +188,15 @@ namespace EpicGames.Horde.Storage.Bundles.V2
 				_packetHandle = null;
 			}
 
-			// Write this bundle to storage
 			public override async ValueTask FlushAsync(CancellationToken cancellationToken = default)
 			{
-				// Check we haven't already flushed this bundle
+				Task writeTask = await _outer.StartWriteAsync(this, cancellationToken);
+				await writeTask.WaitAsync(cancellationToken);
+			}
+
+			// Handle writing the bundle to storage
+			public async Task WriteAsync(CancellationToken cancellationToken = default)
+			{
 				if (_encodedPacketWriter == null)
 				{
 					return;
@@ -211,8 +216,8 @@ namespace EpicGames.Horde.Storage.Bundles.V2
 				FlushedBundleHandle flushedHandle;
 				using (ReadOnlySequenceStream stream = new ReadOnlySequenceStream(_encodedPacketWriter.AsSequence()))
 				{
-					BlobLocator locator = await _storageNamespace.Backend.WriteBlobAsync(stream, _bundleImports, _basePath, cancellationToken);
-					flushedHandle = new FlushedBundleHandle(_storageNamespace, locator);
+					await _storageNamespace.Backend.WriteBlobAsync(Locator, stream, _bundleImports, cancellationToken);
+					flushedHandle = new FlushedBundleHandle(_storageNamespace, Locator);
 				}
 
 				// Release all the intermediate data
@@ -267,40 +272,57 @@ namespace EpicGames.Horde.Storage.Bundles.V2
 
 				return await _flushedHandle.ReadAsync(offset, length, cancellationToken);
 			}
-
-			/// <inheritdoc/>
-			public override bool TryGetLocator(out BlobLocator locator)
-			{
-				if (_flushedHandle == null)
-				{
-					locator = default;
-					return false;
-				}
-				return _flushedHandle.TryGetLocator(out locator);
-			}
 		}
 
 		readonly BundleStorageNamespace _storageNamespace;
 		readonly string? _basePath;
 		readonly BundleOptions _bundleOptions;
 		readonly BundleCache _bundleCache;
+		readonly ILogger _logger;
+		readonly CancellationToken _cancellationToken;
+
+		readonly string _locatorPrefix;
+		int _locatorSuffix;
 
 		PendingBundleHandle _currentBundle;
 
+		readonly AsyncEvent _taskFinished = new AsyncEvent();
+		int _numTasks;
+
 		/// <summary>
-		/// 
+		/// Constructor
 		/// </summary>
-		public BundleWriter(BundleStorageNamespace storageNamespace, string? basePath, BundleCache bundleCache, BundleOptions bundleOptions, BlobSerializerOptions? blobOptions)
+		public BundleWriter(BundleStorageNamespace storageNamespace, string? basePath, BundleCache bundleCache, BundleOptions bundleOptions, BlobSerializerOptions? blobOptions, ILogger logger, CancellationToken cancellationToken)
 			: base(blobOptions)
 		{
+			string instanceSuffix = Guid.NewGuid().ToString("n");
+
 			_storageNamespace = storageNamespace;
 			_basePath = basePath;
 			_bundleOptions = bundleOptions;
 			_bundleCache = bundleCache;
+			_logger = logger;
+			_cancellationToken = cancellationToken;
+
+			if (String.IsNullOrEmpty(basePath))
+			{
+				_locatorPrefix = instanceSuffix;
+			}
+			else if (basePath.EndsWith("/", StringComparison.Ordinal))
+			{
+				_locatorPrefix = $"{basePath}{instanceSuffix}";
+			}
+			else
+			{
+				_locatorPrefix = $"{basePath}/{instanceSuffix}";
+			}
 
-			_currentBundle = new PendingBundleHandle(storageNamespace, basePath, bundleCache, _bundleOptions);
+			_currentBundle = new PendingBundleHandle(this, GetNextLocator());
 		}
 
+		BlobLocator GetNextLocator()
+			=> new BlobLocator($"{_locatorPrefix}_{Interlocked.Increment(ref _locatorSuffix)}");
+
 		/// <inheritdoc/>
 		public override async ValueTask DisposeAsync()
 		{
@@ -308,17 +330,78 @@ namespace EpicGames.Horde.Storage.Bundles.V2
 			_currentBundle.Dispose();
 		}
 
+		async Task WriteCurrentBundleAsync(CancellationToken cancellationToken = default)
+		{
+			await StartWriteAsync(_currentBundle, cancellationToken);
+			_currentBundle = new PendingBundleHandle(this, GetNextLocator());
+		}
+
+		async Task<Task> StartWriteAsync(PendingBundleHandle bundle, CancellationToken cancellationToken)
+		{
+			lock (bundle._lockObject)
+			{
+				bundle._writeTask ??= Task.Run(() => HandleBundleWriteAsync(bundle), CancellationToken.None);
+			}
+			await bundle._startWriteTcs.Task.WaitAsync(cancellationToken);
+			return bundle._writeTask;
+		}
+
+		async Task HandleBundleWriteAsync(PendingBundleHandle bundle)
+		{
+			Interlocked.Increment(ref _numTasks);
+			try
+			{
+				await _bundleCache.WriteSemaphore.WaitAsync(_cancellationToken);
+				bundle._startWriteTcs.SetResult();
+
+				try
+				{
+					await bundle.WriteAsync(_cancellationToken);
+				}
+				catch (OperationCanceledException)
+				{
+					_logger.LogDebug("Write of {Locator} was cancelled", bundle.Locator);
+				}
+				catch (Exception ex)
+				{
+					_logger.LogError(ex, "Error writing {Locator}: {Message}", bundle.Locator, ex.Message);
+				}
+				finally
+				{
+					bundle.Dispose();
+					_bundleCache.WriteSemaphore.Release();
+				}
+			}
+			catch (Exception ex)
+			{
+				bundle._startWriteTcs.SetException(ex);
+			}
+			finally
+			{
+				Interlocked.Decrement(ref _numTasks);
+				_taskFinished.Set();
+			}
+		}
+
 		/// <inheritdoc/>
 		public override async Task FlushAsync(CancellationToken cancellationToken = default)
 		{
-			await _currentBundle.FlushAsync(cancellationToken);
-			_currentBundle.Dispose();
-			_currentBundle = new PendingBundleHandle(_storageNamespace, _basePath, _bundleCache, _bundleOptions);
+			await WriteCurrentBundleAsync(cancellationToken);
+
+			for (; ; )
+			{
+				Task taskFinished = _taskFinished.Task;
+				if (Interlocked.CompareExchange(ref _numTasks, 0, 0) == 0)
+				{
+					break;
+				}
+				await taskFinished.WaitAsync(cancellationToken);
+			}
 		}
 
 		/// <inheritdoc/>
 		public override IBlobWriter Fork()
-			=> new BundleWriter(_storageNamespace, _basePath, _bundleCache, _bundleOptions, Options);
+			=> new BundleWriter(_storageNamespace, _basePath, _bundleCache, _bundleOptions, Options, _logger, _cancellationToken);
 
 		/// <inheritdoc/>
 		public override Memory<byte> GetOutputBuffer(int usedSize, int desiredSize)
@@ -330,7 +413,7 @@ namespace EpicGames.Horde.Storage.Bundles.V2
 			ExportHandle exportHandle = _currentBundle.CompleteExport(type, size, references, aliases);
 			if (_currentBundle.Length > _bundleOptions.MaxBlobSize)
 			{
-				await FlushAsync(cancellationToken);
+				await WriteCurrentBundleAsync(cancellationToken);
 			}
 			return exportHandle;
 		}
diff --git a/Engine/Source/Programs/Shared/EpicGames.Horde/Storage/IStorageNamespace.cs b/Engine/Source/Programs/Shared/EpicGames.Horde/Storage/IStorageNamespace.cs
--- a/Engine/Source/Programs/Shared/EpicGames.Horde/Storage/IStorageNamespace.cs
+++ b/Engine/Source/Programs/Shared/EpicGames.Horde/Storage/IStorageNamespace.cs
@@ -37,8 +37,9 @@ namespace EpicGames.Horde.Storage
 		/// </summary>
 		/// <param name="basePath">Base path for any nodes written from the writer.</param>
 		/// <param name="serializerOptions">Options for serializing classes</param>
+		/// <param name="cancellationToken">Cancellation token used for any buffered blob writes. Blob writers will flush on close, unless this cancellation token is signalled.</param>
 		/// <returns>New writer instance. Must be disposed after use.</returns>
-		IBlobWriter CreateBlobWriter(string? basePath = null, BlobSerializerOptions? serializerOptions = null);
+		IBlobWriter CreateBlobWriter(string? basePath = null, BlobSerializerOptions? serializerOptions = null, CancellationToken cancellationToken = default);
 
 		#endregion
 
diff --git a/Engine/Source/Programs/Shared/EpicGames.Horde/Storage/KeyValueStorageNamespace.cs b/Engine/Source/Programs/Shared/EpicGames.Horde/Storage/KeyValueStorageNamespace.cs
--- a/Engine/Source/Programs/Shared/EpicGames.Horde/Storage/KeyValueStorageNamespace.cs
+++ b/Engine/Source/Programs/Shared/EpicGames.Horde/Storage/KeyValueStorageNamespace.cs
@@ -141,7 +141,7 @@ namespace EpicGames.Horde.Storage.Clients
 		}
 
 		/// <inheritdoc/>
-		public IBlobWriter CreateBlobWriter(string? basePath = null, BlobSerializerOptions? options = null) => new Writer(this, basePath, options);
+		public IBlobWriter CreateBlobWriter(string? basePath = null, BlobSerializerOptions? options = null, CancellationToken cancellationToken = default) => new Writer(this, basePath, options);
 
 		/// <inheritdoc/>
 		public async ValueTask<BlobData> ReadBlobAsync(BlobLocator locator, CancellationToken cancellationToken = default)
