Static analysis fixes.

#rnx

diff --git a/Engine/Source/Programs/Horde/Plugins/Build/HordeServer.Build/Streams/StreamMessages.cs b/Engine/Source/Programs/Horde/Plugins/Build/HordeServer.Build/Streams/StreamMessages.cs
--- a/Engine/Source/Programs/Horde/Plugins/Build/HordeServer.Build/Streams/StreamMessages.cs
+++ b/Engine/Source/Programs/Horde/Plugins/Build/HordeServer.Build/Streams/StreamMessages.cs
@@ -273,6 +273,7 @@ namespace HordeServer.Streams
 		/// <param name="template">The actual template</param>
 		/// <param name="stepStates">The template step states</param>
 		/// <param name="schedulerTimeZone">The scheduler time zone</param>
+		/// <param name="canRun">Whether the user can run this template</param>
 		public GetTemplateRefResponse(TemplateId id, ITemplateRef templateRef, ITemplate template, List<GetTemplateStepStateResponse>? stepStates, TimeZoneInfo schedulerTimeZone, bool canRun)
 			: base(template)
 		{
diff --git a/Engine/Source/Programs/Shared/EpicGames.Horde/Storage/Bundles/V2/BundleWriter.cs b/Engine/Source/Programs/Shared/EpicGames.Horde/Storage/Bundles/V2/BundleWriter.cs
--- a/Engine/Source/Programs/Shared/EpicGames.Horde/Storage/Bundles/V2/BundleWriter.cs
+++ b/Engine/Source/Programs/Shared/EpicGames.Horde/Storage/Bundles/V2/BundleWriter.cs
@@ -38,7 +38,7 @@ namespace EpicGames.Horde.Storage.Bundles.V2
 
 			public override async ValueTask<BlobData> ReadExportAsync(int exportIdx, CancellationToken cancellationToken = default)
 			{
-				lock (_bundle._lockObject)
+				lock (_bundle.LockObject)
 				{
 					if (_flushedHandle == null)
 					{
@@ -55,7 +55,7 @@ namespace EpicGames.Horde.Storage.Bundles.V2
 		// Bundle that has not yet been written to storage.
 		internal sealed class PendingBundleHandle : BundleHandle, IDisposable
 		{
-			public readonly object _lockObject = new object();
+			public readonly object LockObject = new object();
 
 			readonly BundleWriter _outer;
 			readonly BundleStorageNamespace _storageNamespace;
@@ -71,7 +71,7 @@ namespace EpicGames.Horde.Storage.Bundles.V2
 			readonly HashSet<BlobLocator> _bundleImports = new HashSet<BlobLocator>();
 
 			public Task? _writeTask;
-			public readonly TaskCompletionSource _startWriteTcs = new TaskCompletionSource();
+			public readonly TaskCompletionSource StartWriteTcs = new TaskCompletionSource();
 
 			public FlushedBundleHandle? FlushedHandle => _flushedHandle;
 
@@ -156,7 +156,7 @@ namespace EpicGames.Horde.Storage.Bundles.V2
 				RuntimeAssert(_packetWriter == null);
 
 				_packetHandle = new PendingPacketHandle(this);
-				_packetWriter = new PacketWriter(this, _packetHandle, _bundleCache.Allocator, _lockObject);
+				_packetWriter = new PacketWriter(this, _packetHandle, _bundleCache.Allocator, LockObject);
 			}
 
 			void FinishPacket()
@@ -173,7 +173,7 @@ namespace EpicGames.Horde.Storage.Bundles.V2
 					int packetLength = _encodedPacketWriter.Length - packetOffset;
 
 					// Point the packet handle to the encoded data
-					lock (_lockObject)
+					lock (LockObject)
 					{
 						FlushedPacketHandle flushedPacketHandle = new FlushedPacketHandle(_storageNamespace, this, packetOffset, packetLength, _bundleCache);
 						_packetHandle.CompletePacket(flushedPacketHandle);
@@ -222,7 +222,7 @@ namespace EpicGames.Horde.Storage.Bundles.V2
 
 				// Release all the intermediate data
 				List<(ExportHandle, AliasInfo)>? pendingExportAliases = _pendingExportAliases;
-				lock (_lockObject)
+				lock (LockObject)
 				{
 					_flushedHandle = flushedHandle;
 					ReleaseResources();
@@ -257,7 +257,7 @@ namespace EpicGames.Horde.Storage.Bundles.V2
 			{
 				if (_flushedHandle == null)
 				{
-					lock (_lockObject)
+					lock (LockObject)
 					{
 						if (_flushedHandle == null)
 						{
@@ -338,11 +338,11 @@ namespace EpicGames.Horde.Storage.Bundles.V2
 
 		async Task<Task> StartWriteAsync(PendingBundleHandle bundle, CancellationToken cancellationToken)
 		{
-			lock (bundle._lockObject)
+			lock (bundle.LockObject)
 			{
 				bundle._writeTask ??= Task.Run(() => HandleBundleWriteAsync(bundle), CancellationToken.None);
 			}
-			await bundle._startWriteTcs.Task.WaitAsync(cancellationToken);
+			await bundle.StartWriteTcs.Task.WaitAsync(cancellationToken);
 			return bundle._writeTask;
 		}
 
@@ -352,7 +352,7 @@ namespace EpicGames.Horde.Storage.Bundles.V2
 			try
 			{
 				await _bundleCache.WriteSemaphore.WaitAsync(_cancellationToken);
-				bundle._startWriteTcs.SetResult();
+				bundle.StartWriteTcs.SetResult();
 
 				try
 				{
@@ -374,7 +374,7 @@ namespace EpicGames.Horde.Storage.Bundles.V2
 			}
 			catch (Exception ex)
 			{
-				bundle._startWriteTcs.SetException(ex);
+				bundle.StartWriteTcs.SetException(ex);
 			}
 			finally
 			{
