Horde: Fix deadlock when flushing aliases that reference the current bundle.

#rnx

diff --git a/Engine/Source/Programs/Shared/EpicGames.Horde/Storage/Bundles/V2/BundleWriter.cs b/Engine/Source/Programs/Shared/EpicGames.Horde/Storage/Bundles/V2/BundleWriter.cs
--- a/Engine/Source/Programs/Shared/EpicGames.Horde/Storage/Bundles/V2/BundleWriter.cs
+++ b/Engine/Source/Programs/Shared/EpicGames.Horde/Storage/Bundles/V2/BundleWriter.cs
@@ -235,7 +235,7 @@ namespace EpicGames.Horde.Storage.Bundles.V2
 				{
 					foreach ((ExportHandle exportHandle, AliasInfo aliasInfo) in pendingExportAliases)
 					{
-						await _storageNamespace.AddAliasAsync(aliasInfo.Name, exportHandle, aliasInfo.Rank, aliasInfo.Data, cancellationToken);
+						await _storageNamespace.Backend.AddAliasAsync(aliasInfo.Name, exportHandle.GetLocator(), aliasInfo.Rank, aliasInfo.Data, cancellationToken);
 					}
 				}
 			}
