Horde: Add a NullBlobWriter implementation of IBlobWriter, which discards all written data. Useful for performance testing.

#rnx

diff --git a/Engine/Source/Programs/Shared/EpicGames.Horde/Storage/IBlobWriter.cs b/Engine/Source/Programs/Shared/EpicGames.Horde/Storage/IBlobWriter.cs
--- a/Engine/Source/Programs/Shared/EpicGames.Horde/Storage/IBlobWriter.cs
+++ b/Engine/Source/Programs/Shared/EpicGames.Horde/Storage/IBlobWriter.cs
@@ -206,6 +206,77 @@ namespace EpicGames.Horde.Storage
 		public abstract ValueTask DisposeAsync();
 	}
 
+	/// <summary>
+	/// Implementation of <see cref="IBlobWriter"/> which discards any written data
+	/// </summary>
+	public sealed class NullBlobWriter : BlobWriter
+	{
+		class BlobRef : IHashedBlobRef
+		{
+			public IoHash Hash { get; }
+
+			public IBlobRef Innermost 
+				=> throw new NotImplementedException();
+
+			public BlobRef(IoHash hash)
+				=> Hash = hash;
+
+			public ValueTask FlushAsync(CancellationToken cancellationToken = default)
+				=> default;
+
+			public ValueTask<BlobData> ReadBlobDataAsync(CancellationToken cancellationToken = default)
+				=> throw new NotSupportedException();
+
+			public bool TryGetLocator([NotNullWhen(true)] out BlobLocator locator)
+			{
+				locator = new BlobLocator(Hash.ToUtf8String());
+				return true;
+			}
+		}
+
+		readonly ChunkedMemoryWriter _memoryWriter;
+		readonly BlobSerializerOptions _options;
+
+		/// <summary>
+		/// Constructor
+		/// </summary>
+		public NullBlobWriter(BlobSerializerOptions options)
+			: base(options)
+		{
+			_memoryWriter = new ChunkedMemoryWriter();
+			_options = options;
+		}
+
+		/// <inheritdoc/>
+		public override ValueTask DisposeAsync()
+		{
+			GC.SuppressFinalize(this);
+			_memoryWriter.Dispose();
+			return new ValueTask();
+		}
+
+		/// <inheritdoc/>
+		public override ValueTask<IHashedBlobRef> WriteBlobAsync(BlobType type, int size, IReadOnlyList<IBlobRef> imports, IReadOnlyList<AliasInfo> aliases, CancellationToken cancellationToken)
+		{
+			Memory<byte> memory = _memoryWriter.GetMemoryAndAdvance(size);
+			IoHash hash = IoHash.Compute(memory.Span);
+			_memoryWriter.Clear();
+			return new ValueTask<IHashedBlobRef>(new BlobRef(hash));
+		}
+
+		/// <inheritdoc/>
+		public override Memory<byte> GetOutputBuffer(int usedSize, int desiredSize)
+			=> _memoryWriter.GetMemory(usedSize, desiredSize);
+
+		/// <inheritdoc/>
+		public override Task FlushAsync(CancellationToken cancellationToken = default)
+			=> Task.CompletedTask;
+
+		/// <inheritdoc/>
+		public override IBlobWriter Fork()
+			=> new NullBlobWriter(_options);
+	}
+
 	/// <summary>
 	/// Implementation of <see cref="IBlobWriter"/> which just buffers data in memory
 	/// </summary>
