Rename default GetSparseClassData accessor to GetMutableSparseClassData to make mutability more obvious and to encourage use of more explicit accessor with EGetSparseClassDataMethod parameter. ( initial submit, old version not deprecated yet )

diff --git a/Engine/Source/Programs/Shared/EpicGames.UHT/Exporters/CodeGen/UhtHeaderCodeGeneratorCppFile.cs b/Engine/Source/Programs/Shared/EpicGames.UHT/Exporters/CodeGen/UhtHeaderCodeGeneratorCppFile.cs
--- a/Engine/Source/Programs/Shared/EpicGames.UHT/Exporters/CodeGen/UhtHeaderCodeGeneratorCppFile.cs
+++ b/Engine/Source/Programs/Shared/EpicGames.UHT/Exporters/CodeGen/UhtHeaderCodeGeneratorCppFile.cs
@@ -1625,6 +1625,11 @@ namespace EpicGames.UHT.Exporters.CodeGen
 				builder.Append("\treturn static_cast<F").Append(sparseDataType).Append("*>(GetClass()->GetOrCreateSparseClassData());\r\n");
 				builder.Append("}\r\n");
 
+				builder.Append('F').Append(sparseDataType).Append("* ").Append(classObj.SourceName).Append("::GetMutable").Append(sparseDataType).Append("() const \r\n");
+				builder.Append("{\r\n");
+				builder.Append("\treturn static_cast<F").Append(sparseDataType).Append("*>(GetClass()->GetOrCreateSparseClassData());\r\n");
+				builder.Append("}\r\n");
+
 				builder.Append("const F").Append(sparseDataType).Append("* ").Append(classObj.SourceName).Append("::Get").Append(sparseDataType).Append("(EGetSparseClassDataMethod GetMethod) const\r\n");
 				builder.Append("{\r\n");
 				builder.Append("\treturn static_cast<const F").Append(sparseDataType).Append("*>(GetClass()->GetSparseClassData(GetMethod));\r\n");
diff --git a/Engine/Source/Programs/Shared/EpicGames.UHT/Exporters/CodeGen/UhtHeaderCodeGeneratorHFile.cs b/Engine/Source/Programs/Shared/EpicGames.UHT/Exporters/CodeGen/UhtHeaderCodeGeneratorHFile.cs
--- a/Engine/Source/Programs/Shared/EpicGames.UHT/Exporters/CodeGen/UhtHeaderCodeGeneratorHFile.cs
+++ b/Engine/Source/Programs/Shared/EpicGames.UHT/Exporters/CodeGen/UhtHeaderCodeGeneratorHFile.cs
@@ -953,6 +953,8 @@ namespace EpicGames.UHT.Exporters.CodeGen
 					string sparseDataType = sparseScriptStruct.EngineName;
 
 					builder.Append(api).Append('F').Append(sparseDataType).Append("* Get").Append(sparseDataType).Append("() const; \\\r\n");
+
+					builder.Append(api).Append('F').Append(sparseDataType).Append("* GetMutable").Append(sparseDataType).Append("() const; \\\r\n");
 					builder.Append(api).Append("const F").Append(sparseDataType).Append("* Get").Append(sparseDataType).Append("(EGetSparseClassDataMethod GetMethod) const; \\\r\n");
 					builder.Append(api).Append("static UScriptStruct* StaticGet").Append(sparseDataType).Append("ScriptStruct(); \\\r\n");
 				}
