EpicGames.Core: Include exception message/stack in fatal errors triggered via CommandHost.

#rnx

diff --git a/Engine/Source/Programs/Shared/EpicGames.Core/CommandHost.cs b/Engine/Source/Programs/Shared/EpicGames.Core/CommandHost.cs
--- a/Engine/Source/Programs/Shared/EpicGames.Core/CommandHost.cs
+++ b/Engine/Source/Programs/Shared/EpicGames.Core/CommandHost.cs
@@ -289,12 +289,12 @@ namespace EpicGames.Core
 			}
 			catch (FatalErrorException ex)
 			{
-				logger.LogCritical(ex, "Fatal error.");
+				logger.LogCritical(ex, "Fatal error: {Message}", ex.Message);
 				return ex.ExitCode;
 			}
 			catch (Exception ex)
 			{
-				logger.LogCritical(ex, "Fatal error.");
+				logger.LogCritical(ex, "Fatal error: {ExceptionInfo}", ex.ToString());
 				return 1;
 			}
 		}
