Horde: Add an IStorageWriter interface, which allows buffering and batching metadata update requests. Use this to upload artifact aliases in a single batch.

diff --git a/Engine/Source/Programs/Horde/Plugins/Build/HordeServer.Build/Artifacts/ArtifactCollection.cs b/Engine/Source/Programs/Horde/Plugins/Build/HordeServer.Build/Artifacts/ArtifactCollection.cs
--- a/Engine/Source/Programs/Horde/Plugins/Build/HordeServer.Build/Artifacts/ArtifactCollection.cs
+++ b/Engine/Source/Programs/Horde/Plugins/Build/HordeServer.Build/Artifacts/ArtifactCollection.cs
@@ -145,8 +145,11 @@ namespace HordeServer.Artifacts
 
 		class ArtifactBuilder : IArtifactBuilder
 		{
+			record class AliasInfo(string Name, IBlobRef BlobRef, int Rank, ReadOnlyMemory<byte> Data);
+
 			readonly IArtifact _artifact;
 			readonly IStorageNamespace _storage;
+			readonly List<AliasInfo> _aliases = new List<AliasInfo>();
 
 			public ArtifactId Id => _artifact.Id;
 			public ArtifactName Name => _artifact.Name;
@@ -166,11 +169,21 @@ namespace HordeServer.Artifacts
 			}
 
 			public Task AddAliasAsync(string name, IBlobRef handle, int rank = 0, ReadOnlyMemory<byte> data = default, CancellationToken cancellationToken = default)
-				=> _storage.AddAliasAsync(name, handle, rank, data, cancellationToken);
+			{
+				_aliases.Add(new AliasInfo(name, handle, rank, data.ToArray()));
+				return Task.CompletedTask;
+			}
 
 			public async Task<IArtifact> CompleteAsync(IHashedBlobRef blobRef, CancellationToken cancellationToken = default)
 			{
-				await _storage.WriteRefAsync(_artifact.RefName, blobRef, cancellationToken: cancellationToken);
+				await using (IStorageWriter writer = _storage.CreateWriter(cancellationToken))
+				{
+					foreach (AliasInfo aliasInfo in _aliases)
+					{
+						await writer.AddAliasAsync(aliasInfo.Name, aliasInfo.BlobRef, aliasInfo.Rank, aliasInfo.Data, cancellationToken);
+					}
+					await writer.WriteRefAsync(_artifact.RefName, blobRef, cancellationToken: cancellationToken);
+				}
 				return _artifact;
 			}
 
diff --git a/Engine/Source/Programs/Horde/Plugins/Ddc/HordeServer.Ddc/Ddc/RefService.cs b/Engine/Source/Programs/Horde/Plugins/Ddc/HordeServer.Ddc/Ddc/RefService.cs
--- a/Engine/Source/Programs/Horde/Plugins/Ddc/HordeServer.Ddc/Ddc/RefService.cs
+++ b/Engine/Source/Programs/Horde/Plugins/Ddc/HordeServer.Ddc/Ddc/RefService.cs
@@ -24,7 +24,15 @@ namespace HordeServer.Ddc
 		public async Task<bool> DeleteAsync(NamespaceId ns, BucketId bucket, RefId key, CancellationToken cancellationToken)
 		{
 			IStorageNamespace storageNamespace = _storageClient.GetNamespace(ns);
-			return await storageNamespace.DeleteRefAsync(GetRefName(bucket, key), cancellationToken);
+
+			RefName refName = GetRefName(bucket, key);
+			if (!await storageNamespace.RefExistsAsync(refName, cancellationToken: cancellationToken))
+			{
+				return false;
+			}
+
+			await storageNamespace.DeleteRefAsync(refName, cancellationToken);
+			return true;
 		}
 
 		public async Task<bool> ExistsAsync(NamespaceId ns, BucketId bucket, RefId key, CancellationToken cancellationToken)
diff --git a/Engine/Source/Programs/Horde/Plugins/Storage/HordeServer.Storage/Storage/StorageController.cs b/Engine/Source/Programs/Horde/Plugins/Storage/HordeServer.Storage/Storage/StorageController.cs
--- a/Engine/Source/Programs/Horde/Plugins/Storage/HordeServer.Storage/Storage/StorageController.cs
+++ b/Engine/Source/Programs/Horde/Plugins/Storage/HordeServer.Storage/Storage/StorageController.cs
@@ -63,14 +63,14 @@ namespace HordeServer.Storage
 		}
 
 		/// <summary>
-		/// Update aliases in the storage service
+		/// Update metadata in the storage service
 		/// </summary>
 		/// <param name="namespaceId">Namespace to write to</param>
 		/// <param name="request">Request for the alias  to write</param>
 		/// <param name="cancellationToken">Cancellation token for the operation</param>
 		[HttpPost]
 		[Route("/api/v1/storage/{namespaceId}")]
-		public async Task<ActionResult> UpdateNamespaceAsync(NamespaceId namespaceId, [FromBody] UpdateNamespaceRequest request, CancellationToken cancellationToken)
+		public async Task<ActionResult> UpdateNamespaceAsync(NamespaceId namespaceId, [FromBody] UpdateMetadataRequest request, CancellationToken cancellationToken)
 		{
 			IStorageBackend? backend = _storageService.TryCreateBackend(namespaceId);
 			if (backend == null)
@@ -97,23 +97,7 @@ namespace HordeServer.Storage
 			}
 
 			// Execute the requests
-			foreach (AddAliasRequest addAlias in request.AddAliases)
-			{
-				await backend.AddAliasAsync(addAlias.Name, addAlias.Target, addAlias.Rank, addAlias.Data ?? Array.Empty<byte>(), cancellationToken);
-			}
-			foreach (RemoveAliasRequest removeAlias in request.RemoveAliases)
-			{
-				await backend.RemoveAliasAsync(removeAlias.Name, removeAlias.Target, cancellationToken);
-			}
-			foreach (AddRefRequest writeRef in request.AddRefs)
-			{
-				await backend.WriteRefAsync(writeRef.RefName, new HashedBlobRefValue(writeRef.Hash, writeRef.Target), writeRef.Options, cancellationToken);
-			}
-			foreach (RemoveRefRequest removeRef in request.RemoveRefs)
-			{
-				await backend.DeleteRefAsync(removeRef.RefName, cancellationToken);
-			}
-
+			await backend.UpdateMetadataAsync(request, cancellationToken);
 			return Ok();
 		}
 
@@ -365,7 +349,8 @@ namespace HordeServer.Storage
 				return Forbid(StorageAclAction.WriteRefs, namespaceId);
 			}
 
-			await backend.DeleteRefAsync(refName, cancellationToken);
+			RemoveRefRequest removeRef = new RemoveRefRequest { RefName = refName };
+			await backend.UpdateMetadataAsync(new UpdateMetadataRequest { RemoveRefs = [removeRef] }, cancellationToken);
 			return Ok();
 		}
 
@@ -397,7 +382,8 @@ namespace HordeServer.Storage
 			}
 #pragma warning restore CS0618 // Type or member is obsolete
 
-			await backend.WriteRefAsync(refName, new HashedBlobRefValue(request.Hash, request.Target), request.Options, cancellationToken);
+			AddRefRequest addRef = new AddRefRequest { RefName = refName, Hash = request.Hash, Target = request.Target, Data = request.Data, Options = request.Options };
+			await backend.UpdateMetadataAsync(new UpdateMetadataRequest { AddRefs = [addRef] }, cancellationToken);
 			return Ok();
 		}
 
diff --git a/Engine/Source/Programs/Horde/Plugins/Storage/HordeServer.Storage/Storage/StorageService.cs b/Engine/Source/Programs/Horde/Plugins/Storage/HordeServer.Storage/Storage/StorageService.cs
--- a/Engine/Source/Programs/Horde/Plugins/Storage/HordeServer.Storage/Storage/StorageService.cs
+++ b/Engine/Source/Programs/Horde/Plugins/Storage/HordeServer.Storage/Storage/StorageService.cs
@@ -194,6 +194,22 @@ namespace HordeServer.Storage
 
 			#endregion
 
+			/// <inheritdoc/>
+			public async Task UpdateMetadataAsync(UpdateMetadataRequest request, CancellationToken cancellationToken = default)
+			{
+				await Parallel.ForEachAsync(request.AddAliases, cancellationToken, async (addAlias, cancellationToken)
+					=> await _outer.AddAliasAsync(NamespaceId, addAlias.Name, addAlias.Target, addAlias.Rank, addAlias.Data, cancellationToken));
+
+				await Parallel.ForEachAsync(request.RemoveAliases, cancellationToken, async (removeAlias, cancellationToken)
+					=> await _outer.RemoveAliasAsync(NamespaceId, removeAlias.Name, removeAlias.Target, cancellationToken));
+
+				await Parallel.ForEachAsync(request.AddRefs, cancellationToken, async (addRef, cancellationToken) 
+					=> await _outer.WriteRefAsync(NamespaceId, addRef.RefName, new HashedBlobRefValue(addRef.Hash, addRef.Target), addRef.Options, cancellationToken));
+
+				await Parallel.ForEachAsync(request.RemoveRefs, cancellationToken, async (removeRef, cancellationToken)
+					=> await _outer.DeleteRefAsync(NamespaceId, removeRef.RefName, cancellationToken));
+			}
+
 			/// <inheritdoc/>
 			public void GetStats(StorageStats stats) => _store.GetStats(stats);
 		}
diff --git a/Engine/Source/Programs/Shared/EpicGames.Horde/Artifacts/HttpArtifactCollection.cs b/Engine/Source/Programs/Shared/EpicGames.Horde/Artifacts/HttpArtifactCollection.cs
--- a/Engine/Source/Programs/Shared/EpicGames.Horde/Artifacts/HttpArtifactCollection.cs
+++ b/Engine/Source/Programs/Shared/EpicGames.Horde/Artifacts/HttpArtifactCollection.cs
@@ -63,9 +63,12 @@ namespace EpicGames.Horde.Artifacts
 
 		class ArtifactBuilder : IArtifactBuilder
 		{
+			record class AliasInfo(string Name, IBlobRef BlobRef, int Rank, ReadOnlyMemory<byte> Data);
+
 			readonly Artifact _artifact;
 			readonly CreateArtifactResponse _response;
 			readonly IStorageNamespace _namespace;
+			readonly List<AliasInfo> _aliases = new List<AliasInfo>();
 
 			public ArtifactId Id => _artifact.Id;
 			public ArtifactName Name => _artifact.Name;
@@ -87,12 +90,22 @@ namespace EpicGames.Horde.Artifacts
 
 			/// <inheritdoc/>
 			public Task AddAliasAsync(string name, IBlobRef handle, int rank = 0, ReadOnlyMemory<byte> data = default, CancellationToken cancellationToken = default)
-				=> _namespace.AddAliasAsync(name, handle, rank, data, cancellationToken);
+			{
+				_aliases.Add(new AliasInfo(name, handle, rank, data.ToArray()));
+				return Task.CompletedTask;
+			}
 
 			/// <inheritdoc/>
 			public async Task<IArtifact> CompleteAsync(IHashedBlobRef blobRef, CancellationToken cancellationToken = default)
 			{
-				await _namespace.WriteRefAsync(_response.RefName, blobRef, cancellationToken: cancellationToken);
+				await using (IStorageWriter writer = _namespace.CreateWriter(cancellationToken))
+				{
+					foreach (AliasInfo aliasInfo in _aliases)
+					{
+						await writer.AddAliasAsync(aliasInfo.Name, aliasInfo.BlobRef, aliasInfo.Rank, aliasInfo.Data, cancellationToken);
+					}
+					await writer.WriteRefAsync(_artifact.RefName, blobRef, cancellationToken: cancellationToken);
+				}
 				return _artifact;
 			}
 
diff --git a/Engine/Source/Programs/Shared/EpicGames.Horde/Compute/AgentStorageClient.cs b/Engine/Source/Programs/Shared/EpicGames.Horde/Compute/AgentStorageClient.cs
--- a/Engine/Source/Programs/Shared/EpicGames.Horde/Compute/AgentStorageClient.cs
+++ b/Engine/Source/Programs/Shared/EpicGames.Horde/Compute/AgentStorageClient.cs
@@ -115,6 +115,27 @@ namespace EpicGames.Horde.Compute
 
 		#endregion
 
+		/// <inheritdoc/>
+		public async Task UpdateMetadataAsync(UpdateMetadataRequest request, CancellationToken cancellationToken = default)
+		{
+			foreach (AddAliasRequest addAlias in request.AddAliases)
+			{
+				await AddAliasAsync(addAlias.Name, addAlias.Target, addAlias.Rank, addAlias.Data, cancellationToken);
+			}
+			foreach (RemoveAliasRequest removeAlias in request.RemoveAliases)
+			{
+				await RemoveAliasAsync(removeAlias.Name, removeAlias.Target, cancellationToken);
+			}
+			foreach (AddRefRequest addRef in request.AddRefs)
+			{
+				await WriteRefAsync(addRef.RefName, new HashedBlobRefValue(addRef.Hash, addRef.Target), addRef.Options, cancellationToken);
+			}
+			foreach (RemoveRefRequest removeRef in request.RemoveRefs)
+			{
+				await DeleteRefAsync(removeRef.RefName, cancellationToken);
+			}
+		}
+
 		/// <inheritdoc/>
 		public void GetStats(StorageStats stats)
 		{
diff --git a/Engine/Source/Programs/Shared/EpicGames.Horde/Storage/Backends/FileStorageBackend.cs b/Engine/Source/Programs/Shared/EpicGames.Horde/Storage/Backends/FileStorageBackend.cs
--- a/Engine/Source/Programs/Shared/EpicGames.Horde/Storage/Backends/FileStorageBackend.cs
+++ b/Engine/Source/Programs/Shared/EpicGames.Horde/Storage/Backends/FileStorageBackend.cs
@@ -190,6 +190,27 @@ namespace EpicGames.Horde.Storage.Backends
 
 		#endregion
 
+		/// <inheritdoc/>
+		public async Task UpdateMetadataAsync(UpdateMetadataRequest request, CancellationToken cancellationToken = default)
+		{
+			foreach (AddAliasRequest addAlias in request.AddAliases)
+			{
+				await AddAliasAsync(addAlias.Name, addAlias.Target, addAlias.Rank, addAlias.Data, cancellationToken);
+			}
+			foreach (RemoveAliasRequest removeAlias in request.RemoveAliases)
+			{
+				await RemoveAliasAsync(removeAlias.Name, removeAlias.Target, cancellationToken);
+			}
+			foreach (AddRefRequest addRef in request.AddRefs)
+			{
+				await WriteRefAsync(addRef.RefName, new HashedBlobRefValue(addRef.Hash, addRef.Target), addRef.Options, cancellationToken);
+			}
+			foreach (RemoveRefRequest removeRef in request.RemoveRefs)
+			{
+				await DeleteRefAsync(removeRef.RefName, cancellationToken);
+			}
+		}
+
 		/// <inheritdoc/>
 		public void GetStats(StorageStats stats)
 			=> _objectStore.GetStats(stats);
diff --git a/Engine/Source/Programs/Shared/EpicGames.Horde/Storage/Backends/HttpStorageBackend.cs b/Engine/Source/Programs/Shared/EpicGames.Horde/Storage/Backends/HttpStorageBackend.cs
--- a/Engine/Source/Programs/Shared/EpicGames.Horde/Storage/Backends/HttpStorageBackend.cs
+++ b/Engine/Source/Programs/Shared/EpicGames.Horde/Storage/Backends/HttpStorageBackend.cs
@@ -256,32 +256,6 @@ namespace EpicGames.Horde.Storage.Backends
 
 		#region Aliases
 
-		/// <inheritdoc/>
-		public async Task AddAliasAsync(string name, BlobLocator target, int rank = 0, ReadOnlyMemory<byte> data = default, CancellationToken cancellationToken = default)
-		{
-			using (HttpClient httpClient = _createClient())
-			{
-				UpdateNamespaceRequest request = new UpdateNamespaceRequest();
-				request.AddAliases.Add(new AddAliasRequest { Name = name, Target = target, Rank = rank, Data = data.ToArray() });
-
-				using HttpResponseMessage response = await httpClient.PostAsJsonAsync($"{_basePath}", request, cancellationToken: cancellationToken);
-				await EnsureSuccessAsync(response, cancellationToken);
-			}
-		}
-
-		/// <inheritdoc/>
-		public async Task RemoveAliasAsync(string name, BlobLocator target, CancellationToken cancellationToken = default)
-		{
-			using (HttpClient httpClient = _createClient())
-			{
-				UpdateNamespaceRequest request = new UpdateNamespaceRequest();
-				request.RemoveAliases.Add(new RemoveAliasRequest { Name = name, Target = target });
-
-				using HttpResponseMessage response = await httpClient.PostAsJsonAsync($"{_basePath}", request, cancellationToken: cancellationToken);
-				await EnsureSuccessAsync(response, cancellationToken);
-			}
-		}
-
 		/// <inheritdoc/>
 		public async Task<BlobAliasLocator[]> FindAliasesAsync(string alias, int? maxResults = null, CancellationToken cancellationToken = default)
 		{
@@ -319,32 +293,6 @@ namespace EpicGames.Horde.Storage.Backends
 
 		#region Refs
 
-		/// <inheritdoc/>
-		public async Task<bool> DeleteRefAsync(RefName name, CancellationToken cancellationToken)
-		{
-			_logger.LogDebug("Deleting ref {RefName}", name);
-			using (HttpClient httpClient = _createClient())
-			{
-				using (HttpRequestMessage request = new HttpRequestMessage(HttpMethod.Delete, $"{_basePath}/refs/{name}"))
-				{
-					using (HttpResponseMessage response = await httpClient.SendAsync(request, cancellationToken))
-					{
-						if (response.IsSuccessStatusCode)
-						{
-							return true;
-						}
-						if (response.StatusCode == HttpStatusCode.NotFound)
-						{
-							return false;
-						}
-
-						await EnsureSuccessAsync(response, cancellationToken);
-						return false;
-					}
-				}
-			}
-		}
-
 		/// <inheritdoc/>
 		public async Task<HashedBlobRefValue?> TryReadRefAsync(RefName name, RefCacheTime cacheTime = default, CancellationToken cancellationToken = default)
 		{
@@ -383,26 +331,27 @@ namespace EpicGames.Horde.Storage.Backends
 			}
 		}
 
+		#endregion
+
 		/// <inheritdoc/>
-		public async Task WriteRefAsync(RefName name, HashedBlobRefValue value, RefOptions? options = null, CancellationToken cancellationToken = default)
+		public async Task UpdateMetadataAsync(UpdateMetadataRequest request, CancellationToken cancellationToken = default)
 		{
-			_logger.LogDebug("Writing ref {RefName} -> {Hash} / {Locator}", name, value.Hash, value.Locator);
-			using (HttpClient httpClient = _createClient())
+			foreach (AddRefRequest addRef in request.AddRefs)
+			{
+				_logger.LogDebug("Writing ref {RefName} -> {Hash} / {Locator}", addRef.RefName, addRef.Hash, addRef.Target);
+			}
+			foreach (RemoveRefRequest removeRef in request.RemoveRefs)
 			{
-				WriteRefRequest request = new WriteRefRequest();
-				request.Hash = value.Hash;
-				request.Target = value.Locator;
-				request.Options = options;
+				_logger.LogDebug("Deleting ref {RefName}", removeRef.RefName);
+			}
 
-				using (HttpResponseMessage response = await httpClient.PutAsync($"{_basePath}/refs/{name}", request, cancellationToken))
-				{
-					await EnsureSuccessAsync(response, cancellationToken);
-				}
+			using (HttpClient httpClient = _createClient())
+			{
+				using HttpResponseMessage response = await httpClient.PostAsJsonAsync($"{_basePath}", request, cancellationToken: cancellationToken);
+				await EnsureSuccessAsync(response, cancellationToken);
 			}
 		}
 
-		#endregion
-
 		/// <inheritdoc/>
 		public void GetStats(StorageStats stats)
 		{
diff --git a/Engine/Source/Programs/Shared/EpicGames.Horde/Storage/Backends/JupiterStorageBackend.cs b/Engine/Source/Programs/Shared/EpicGames.Horde/Storage/Backends/JupiterStorageBackend.cs
--- a/Engine/Source/Programs/Shared/EpicGames.Horde/Storage/Backends/JupiterStorageBackend.cs
+++ b/Engine/Source/Programs/Shared/EpicGames.Horde/Storage/Backends/JupiterStorageBackend.cs
@@ -455,6 +455,27 @@ namespace EpicGames.Horde.Storage.Backends
 
 		#endregion
 
+		/// <inheritdoc/>
+		public async Task UpdateMetadataAsync(UpdateMetadataRequest request, CancellationToken cancellationToken = default)
+		{
+			foreach (AddAliasRequest addAlias in request.AddAliases)
+			{
+				await AddAliasAsync(addAlias.Name, addAlias.Target, addAlias.Rank, addAlias.Data, cancellationToken);
+			}
+			foreach (RemoveAliasRequest removeAlias in request.RemoveAliases)
+			{
+				await RemoveAliasAsync(removeAlias.Name, removeAlias.Target, cancellationToken);
+			}
+			foreach (AddRefRequest addRef in request.AddRefs)
+			{
+				await WriteRefAsync(addRef.RefName, new HashedBlobRefValue(addRef.Hash, addRef.Target), addRef.Options, cancellationToken);
+			}
+			foreach (RemoveRefRequest removeRef in request.RemoveRefs)
+			{
+				await DeleteRefAsync(removeRef.RefName, cancellationToken);
+			}
+		}
+
 		/// <inheritdoc/>
 		public void GetStats(StorageStats stats)
 		{ }
diff --git a/Engine/Source/Programs/Shared/EpicGames.Horde/Storage/Backends/MemoryStorageBackend.cs b/Engine/Source/Programs/Shared/EpicGames.Horde/Storage/Backends/MemoryStorageBackend.cs
--- a/Engine/Source/Programs/Shared/EpicGames.Horde/Storage/Backends/MemoryStorageBackend.cs
+++ b/Engine/Source/Programs/Shared/EpicGames.Horde/Storage/Backends/MemoryStorageBackend.cs
@@ -203,6 +203,27 @@ namespace EpicGames.Horde.Storage.Backends
 
 		#endregion
 
+		/// <inheritdoc/>
+		public async Task UpdateMetadataAsync(UpdateMetadataRequest request, CancellationToken cancellationToken = default)
+		{
+			foreach (AddAliasRequest addAlias in request.AddAliases)
+			{
+				await AddAliasAsync(addAlias.Name, addAlias.Target, addAlias.Rank, addAlias.Data, cancellationToken);
+			}
+			foreach (RemoveAliasRequest removeAlias in request.RemoveAliases)
+			{
+				await RemoveAliasAsync(removeAlias.Name, removeAlias.Target, cancellationToken);
+			}
+			foreach (AddRefRequest addRef in request.AddRefs)
+			{
+				await WriteRefAsync(addRef.RefName, new HashedBlobRefValue(addRef.Hash, addRef.Target), addRef.Options, cancellationToken);
+			}
+			foreach (RemoveRefRequest removeRef in request.RemoveRefs)
+			{
+				await DeleteRefAsync(removeRef.RefName, cancellationToken);
+			}
+		}
+
 		/// <inheritdoc/>
 		public void GetStats(StorageStats stats)
 		{
diff --git a/Engine/Source/Programs/Shared/EpicGames.Horde/Storage/Bundles/BundleStorageNamespace.cs b/Engine/Source/Programs/Shared/EpicGames.Horde/Storage/Bundles/BundleStorageNamespace.cs
--- a/Engine/Source/Programs/Shared/EpicGames.Horde/Storage/Bundles/BundleStorageNamespace.cs
+++ b/Engine/Source/Programs/Shared/EpicGames.Horde/Storage/Bundles/BundleStorageNamespace.cs
@@ -20,6 +20,22 @@ namespace EpicGames.Horde.Storage.Bundles
 	/// </summary>
 	public sealed class BundleStorageNamespace : IStorageNamespace
 	{
+		class StorageWriterImpl : StorageWriter
+		{
+			readonly BundleStorageNamespace _outer;
+			readonly IStorageBackend _storageBackend;
+
+			public StorageWriterImpl(BundleStorageNamespace outer, IStorageBackend storageBackend, CancellationToken cancellationToken)
+				: base(storageBackend, cancellationToken)
+			{
+				_outer = outer;
+				_storageBackend = storageBackend;
+			}
+
+			protected override IBlobWriter CreateBlobWriter(string? basePath = null, BlobSerializerOptions? serializerOptions = null, CancellationToken cancellationToken = default)
+				=> _outer.CreateBlobWriter(basePath, serializerOptions, cancellationToken);
+		}
+
 		readonly IStorageBackend _backend;
 		readonly BundleCache _cache;
 		readonly BundleOptions _options;
@@ -148,6 +164,10 @@ namespace EpicGames.Horde.Storage.Bundles
 			return locators;
 		}
 
+		/// <inheritdoc/>
+		public IStorageWriter CreateWriter(CancellationToken cancellationToken)
+			=> new StorageWriterImpl(this, _backend, cancellationToken);
+
 		#region Blobs
 
 		/// <inheritdoc/>
@@ -198,20 +218,6 @@ namespace EpicGames.Horde.Storage.Bundles
 
 		#region Aliases
 
-		/// <inheritdoc/>
-		public async Task AddAliasAsync(string name, IBlobRef handle, int rank, ReadOnlyMemory<byte> data, CancellationToken cancellationToken)
-		{
-			await handle.FlushAsync(cancellationToken);
-			await _backend.AddAliasAsync(name, handle.GetLocator(), rank, data, cancellationToken);
-		}
-
-		/// <inheritdoc/>
-		public async Task RemoveAliasAsync(string name, IBlobRef handle, CancellationToken cancellationToken)
-		{
-			await handle.FlushAsync(cancellationToken);
-			await _backend.RemoveAliasAsync(name, handle.GetLocator(), cancellationToken);
-		}
-
 		/// <inheritdoc/>
 		public async Task<BlobAlias[]> FindAliasesAsync(string name, int? maxLength = null, CancellationToken cancellationToken = default)
 		{
@@ -230,10 +236,6 @@ namespace EpicGames.Horde.Storage.Bundles
 
 		#region Refs
 
-		/// <inheritdoc/>
-		public Task<bool> DeleteRefAsync(RefName name, CancellationToken cancellationToken = default)
-			=> _backend.DeleteRefAsync(name, cancellationToken);
-
 		/// <inheritdoc/>
 		public async Task<IHashedBlobRef?> TryReadRefAsync(RefName name, RefCacheTime cacheTime = default, CancellationToken cancellationToken = default)
 		{
@@ -245,13 +247,6 @@ namespace EpicGames.Horde.Storage.Bundles
 			return this.CreateBlobRef(target.Hash, target.Locator);
 		}
 
-		/// <inheritdoc/>
-		public async Task WriteRefAsync(RefName name, IHashedBlobRef target, RefOptions? options = null, CancellationToken cancellationToken = default)
-		{
-			await target.FlushAsync(cancellationToken);
-			await _backend.WriteRefAsync(name, new HashedBlobRefValue(target.Hash, target.GetLocator()), options, cancellationToken);
-		}
-
 		#endregion
 
 		/// <inheritdoc/>
diff --git a/Engine/Source/Programs/Shared/EpicGames.Horde/Storage/Bundles/V2/BundleWriter.cs b/Engine/Source/Programs/Shared/EpicGames.Horde/Storage/Bundles/V2/BundleWriter.cs
--- a/Engine/Source/Programs/Shared/EpicGames.Horde/Storage/Bundles/V2/BundleWriter.cs
+++ b/Engine/Source/Programs/Shared/EpicGames.Horde/Storage/Bundles/V2/BundleWriter.cs
@@ -190,8 +190,11 @@ namespace EpicGames.Horde.Storage.Bundles.V2
 
 			public override async ValueTask FlushAsync(CancellationToken cancellationToken = default)
 			{
-				Task writeTask = await _outer.StartWriteAsync(this, cancellationToken);
-				await writeTask.WaitAsync(cancellationToken);
+				if (_flushedHandle == null)
+				{
+					Task writeTask = await _outer.StartWriteAsync(this, cancellationToken);
+					await writeTask.WaitAsync(cancellationToken);
+				}
 			}
 
 			// Handle writing the bundle to storage
@@ -235,7 +238,7 @@ namespace EpicGames.Horde.Storage.Bundles.V2
 				{
 					foreach ((ExportHandle exportHandle, AliasInfo aliasInfo) in pendingExportAliases)
 					{
-						await _storageNamespace.Backend.AddAliasAsync(aliasInfo.Name, exportHandle.GetLocator(), aliasInfo.Rank, aliasInfo.Data, cancellationToken);
+						await _storageNamespace.AddAliasAsync(aliasInfo.Name, exportHandle, aliasInfo.Rank, aliasInfo.Data, cancellationToken);
 					}
 				}
 			}
diff --git a/Engine/Source/Programs/Shared/EpicGames.Horde/Storage/IStorageBackend.cs b/Engine/Source/Programs/Shared/EpicGames.Horde/Storage/IStorageBackend.cs
--- a/Engine/Source/Programs/Shared/EpicGames.Horde/Storage/IStorageBackend.cs
+++ b/Engine/Source/Programs/Shared/EpicGames.Horde/Storage/IStorageBackend.cs
@@ -94,24 +94,6 @@ namespace EpicGames.Horde.Storage
 
 		#region Aliases
 
-		/// <summary>
-		/// Adds an alias to a given blob
-		/// </summary>
-		/// <param name="name">Alias for the blob</param>
-		/// <param name="locator">Locator for the blob</param>
-		/// <param name="rank">Rank for this alias. In situations where an alias has multiple mappings, the alias with the highest rank will be returned by default.</param>
-		/// <param name="data">Additional data to be stored inline with the alias</param>
-		/// <param name="cancellationToken">Cancellation token for the operation</param>
-		Task AddAliasAsync(string name, BlobLocator locator, int rank = 0, ReadOnlyMemory<byte> data = default, CancellationToken cancellationToken = default);
-
-		/// <summary>
-		/// Removes an alias from a blob
-		/// </summary>
-		/// <param name="name">Name of the alias</param>
-		/// <param name="locator">Locator for the blob</param>
-		/// <param name="cancellationToken">Cancellation token for the operation</param>
-		Task RemoveAliasAsync(string name, BlobLocator locator, CancellationToken cancellationToken = default);
-
 		/// <summary>
 		/// Finds blobs with the given alias. Unlike refs, aliases do not serve as GC roots.
 		/// </summary>
@@ -134,24 +116,15 @@ namespace EpicGames.Horde.Storage
 		/// <returns>Blob pointed to by the ref</returns>
 		Task<HashedBlobRefValue?> TryReadRefAsync(RefName name, RefCacheTime cacheTime = default, CancellationToken cancellationToken = default);
 
-		/// <summary>
-		/// Writes a new ref to the store
-		/// </summary>
-		/// <param name="name">Ref to write</param>
-		/// <param name="value">Value for the ref</param>
-		/// <param name="options">Options for the new ref</param>
-		/// <param name="cancellationToken">Cancellation token for the operation</param>
-		/// <returns>Unique identifier for the blob</returns>
-		Task WriteRefAsync(RefName name, HashedBlobRefValue value, RefOptions? options = null, CancellationToken cancellationToken = default);
+		#endregion
 
 		/// <summary>
-		/// Reads data for a ref from the store
+		/// Batch request to update metadata
 		/// </summary>
-		/// <param name="name">The ref identifier</param>
+		/// <param name="request">Options for the update</param>
 		/// <param name="cancellationToken">Cancellation token for the operation</param>
-		Task<bool> DeleteRefAsync(RefName name, CancellationToken cancellationToken = default);
-
-		#endregion
+		/// <returns></returns>
+		Task UpdateMetadataAsync(UpdateMetadataRequest request, CancellationToken cancellationToken = default);
 
 		/// <summary>
 		/// Gets stats for this storage backend
diff --git a/Engine/Source/Programs/Shared/EpicGames.Horde/Storage/IStorageNamespace.cs b/Engine/Source/Programs/Shared/EpicGames.Horde/Storage/IStorageNamespace.cs
--- a/Engine/Source/Programs/Shared/EpicGames.Horde/Storage/IStorageNamespace.cs
+++ b/Engine/Source/Programs/Shared/EpicGames.Horde/Storage/IStorageNamespace.cs
@@ -23,6 +23,13 @@
 	/// </summary>
 	public interface IStorageNamespace
 	{
+		/// <summary>
+		/// Creates a writer for updating the namespace
+		/// </summary>
+		/// <param name="cancellationToken">Cancellation token for async writing. The writer will flush on dispose unless this cancellation token is signalled.</param>
+		/// <returns>New writer instance</returns>
+		IStorageWriter CreateWriter(CancellationToken cancellationToken = default);
+
 		#region Blobs
 
 		/// <summary>
@@ -46,23 +53,5 @@
 		#region Aliases
 
 		/// <summary>
-		/// Adds an alias to a given blob
-		/// </summary>
-		/// <param name="name">Alias for the blob</param>
-		/// <param name="handle">Locator for the blob</param>
-		/// <param name="rank">Rank for this alias. In situations where an alias has multiple mappings, the alias with the highest rank will be returned by default.</param>
-		/// <param name="data">Additional data to be stored inline with the alias</param>
-		/// <param name="cancellationToken">Cancellation token for the operation</param>
-		Task AddAliasAsync(string name, IBlobRef handle, int rank = 0, ReadOnlyMemory<byte> data = default, CancellationToken cancellationToken = default);
-
-		/// <summary>
-		/// Removes an alias from a blob
-		/// </summary>
-		/// <param name="name">Name of the alias</param>
-		/// <param name="handle">Locator for the blob</param>
-		/// <param name="cancellationToken">Cancellation token for the operation</param>
-		Task RemoveAliasAsync(string name, IBlobRef handle, CancellationToken cancellationToken = default);
-
-		/// <summary>
 		/// Finds blobs with the given alias. Unlike refs, aliases do not serve as GC roots.
 		/// </summary>
@@ -85,23 +74,6 @@
 		/// <returns>Blob pointed to by the ref</returns>
 		Task<IHashedBlobRef?> TryReadRefAsync(RefName name, RefCacheTime cacheTime = default, CancellationToken cancellationToken = default);
 
-		/// <summary>
-		/// Writes a new ref to the store
-		/// </summary>
-		/// <param name="name">Ref to write</param>
-		/// <param name="target">Handle to the target blob</param>
-		/// <param name="options">Options for the new ref</param>
-		/// <param name="cancellationToken">Cancellation token for the operation</param>
-		/// <returns>Unique identifier for the blob</returns>
-		Task WriteRefAsync(RefName name, IHashedBlobRef target, RefOptions? options = null, CancellationToken cancellationToken = default);
-
-		/// <summary>
-		/// Reads data for a ref from the store
-		/// </summary>
-		/// <param name="name">The ref identifier</param>
-		/// <param name="cancellationToken">Cancellation token for the operation</param>
-		Task<bool> DeleteRefAsync(RefName name, CancellationToken cancellationToken = default);
-
 		#endregion
 
 		/// <summary>
@@ -301,12 +273,41 @@
 		/// </summary>
 		/// <param name="store">The store instance to read from</param>
 		/// <param name="refName">Ref name to use as a base path</param>
-		public static IBlobWriter CreateBlobWriter(this IStorageNamespace store, RefName refName) => store.CreateBlobWriter(refName.ToString());
+		public static IBlobWriter CreateBlobWriter(this IStorageNamespace store, RefName refName) 
+			=> store.CreateBlobWriter(refName.ToString());
 
 		#endregion
 
 		#region Aliases
 
 		/// <summary>
+		/// Adds an alias to a given blob
+		/// </summary>
+		/// <param name="store">The store instance to write to</param>
+		/// <param name="name">Alias for the blob</param>
+		/// <param name="handle">Locator for the blob</param>
+		/// <param name="rank">Rank for this alias. In situations where an alias has multiple mappings, the alias with the highest rank will be returned by default.</param>
+		/// <param name="data">Additional data to be stored inline with the alias</param>
+		/// <param name="cancellationToken">Cancellation token for the operation</param>
+		public static async Task AddAliasAsync(this IStorageNamespace store, string name, IBlobRef handle, int rank = 0, ReadOnlyMemory<byte> data = default, CancellationToken cancellationToken = default)
+		{
+			await using IStorageWriter writer = store.CreateWriter(cancellationToken);
+			await writer.AddAliasAsync(name, handle, rank, data, cancellationToken);
+		}
+
+		/// <summary>
+		/// Removes an alias from a blob
+		/// </summary>
+		/// <param name="store">The store instance to write to</param>
+		/// <param name="name">Name of the alias</param>
+		/// <param name="handle">Locator for the blob</param>
+		/// <param name="cancellationToken">Cancellation token for the operation</param>
+		public static async Task RemoveAliasAsync(this IStorageNamespace store, string name, IBlobRef handle, CancellationToken cancellationToken = default)
+		{
+			await using IStorageWriter writer = store.CreateWriter(cancellationToken);
+			await writer.RemoveAliasAsync(name, handle, cancellationToken);
+		}
+
+		/// <summary>
 		/// Finds blobs with the given alias. Unlike refs, aliases do not serve as GC roots.
 		/// </summary>
@@ -352,6 +353,33 @@
 			return refTarget ?? throw new RefNameNotFoundException(name);
 		}
 
+		/// <summary>
+		/// Writes a new ref to the store
+		/// </summary>
+		/// <param name="store">The store instance to write to</param>
+		/// <param name="name">Ref to write</param>
+		/// <param name="target">Handle to the target blob</param>
+		/// <param name="options">Options for the new ref</param>
+		/// <param name="cancellationToken">Cancellation token for the operation</param>
+		/// <returns>Unique identifier for the blob</returns>
+		public static async Task WriteRefAsync(this IStorageNamespace store, RefName name, IHashedBlobRef target, RefOptions? options = null, CancellationToken cancellationToken = default)
+		{
+			await using IStorageWriter writer = store.CreateWriter(cancellationToken);
+			await writer.WriteRefAsync(name, target, options, cancellationToken);
+		}
+
+		/// <summary>
+		/// Reads data for a ref from the store
+		/// </summary>
+		/// <param name="store">The store instance to write to</param>
+		/// <param name="name">The ref identifier</param>
+		/// <param name="cancellationToken">Cancellation token for the operation</param>
+		public static async Task DeleteRefAsync(this IStorageNamespace store, RefName name, CancellationToken cancellationToken = default)
+		{
+			await using IStorageWriter writer = store.CreateWriter(cancellationToken);
+			await writer.DeleteRefAsync(name, cancellationToken);
+		}
+
 		#endregion
 
 		/// <summary>
diff --git a/Engine/Source/Programs/Shared/EpicGames.Horde/Storage/IStorageWriter.cs b/Engine/Source/Programs/Shared/EpicGames.Horde/Storage/IStorageWriter.cs
new file mode 100644
--- /dev/null
+++ b/Engine/Source/Programs/Shared/EpicGames.Horde/Storage/IStorageWriter.cs
@@ -0,0 +1,160 @@
+// Copyright Epic Games, Inc. All Rights Reserved.
+
+using System;
+using System.Collections.Generic;
+using System.Linq;
+using System.Text;
+using System.Threading;
+using System.Threading.Tasks;
+
+namespace EpicGames.Horde.Storage
+{
+	/// <summary>
+	/// Interface for a batching storage writer
+	/// </summary>
+	public interface IStorageWriter : IAsyncDisposable
+	{
+		#region Blobs
+
+		/// <summary>
+		/// Creates a new writer for storage blobs
+		/// </summary>
+		/// <param name="basePath">Base path for any nodes written from the writer.</param>
+		/// <param name="serializerOptions">Options for serializing classes</param>
+		/// <returns>New writer instance. Must be disposed after use.</returns>
+		IBlobWriter CreateBlobWriter(string? basePath = null, BlobSerializerOptions? serializerOptions = null);
+
+		#endregion
+
+		#region Aliases
+
+		/// <summary>
+		/// Adds an alias to a given blob
+		/// </summary>
+		/// <param name="name">Alias for the blob</param>
+		/// <param name="handle">Locator for the blob</param>
+		/// <param name="rank">Rank for this alias. In situations where an alias has multiple mappings, the alias with the highest rank will be returned by default.</param>
+		/// <param name="data">Additional data to be stored inline with the alias</param>
+		/// <param name="cancellationToken">Cancellation token for the operation</param>
+		Task AddAliasAsync(string name, IBlobRef handle, int rank = 0, ReadOnlyMemory<byte> data = default, CancellationToken cancellationToken = default);
+
+		/// <summary>
+		/// Removes an alias from a blob
+		/// </summary>
+		/// <param name="name">Name of the alias</param>
+		/// <param name="handle">Locator for the blob</param>
+		/// <param name="cancellationToken">Cancellation token for the operation</param>
+		Task RemoveAliasAsync(string name, IBlobRef handle, CancellationToken cancellationToken = default);
+
+		#endregion
+
+		#region Refs
+
+		/// <summary>
+		/// Writes a new ref to the store
+		/// </summary>
+		/// <param name="name">Ref to write</param>
+		/// <param name="target">Handle to the target blob</param>
+		/// <param name="options">Options for the new ref</param>
+		/// <param name="cancellationToken">Cancellation token for the operation</param>
+		/// <returns>Unique identifier for the blob</returns>
+		Task WriteRefAsync(RefName name, IHashedBlobRef target, RefOptions? options = null, CancellationToken cancellationToken = default);
+
+		/// <summary>
+		/// Reads data for a ref from the store
+		/// </summary>
+		/// <param name="name">The ref identifier</param>
+		/// <param name="cancellationToken">Cancellation token for the operation</param>
+		Task DeleteRefAsync(RefName name, CancellationToken cancellationToken = default);
+
+		#endregion
+
+		/// <summary>
+		/// Flush any buffered writes
+		/// </summary>
+		/// <param name="cancellationToken">Cancellation token for the operation</param>
+		Task FlushAsync(CancellationToken cancellationToken = default);
+	}
+
+	/// <summary>
+	/// Default base implementation of <see cref="IStorageWriter"/> which combines metadata updates into a <see cref="UpdateMetadataRequest"/>
+	/// </summary>
+	public abstract class StorageWriter : IStorageWriter
+	{
+		readonly IStorageBackend _storageBackend;
+		readonly CancellationToken _cancellationToken;
+		UpdateMetadataRequest _request = new UpdateMetadataRequest();
+
+		/// <summary>
+		/// Constructor
+		/// </summary>
+		protected StorageWriter(IStorageBackend storageBackend, CancellationToken cancellationToken = default)
+		{
+			_storageBackend = storageBackend;
+			_cancellationToken = cancellationToken;
+		}
+
+		/// <inheritdoc/>
+		public async ValueTask DisposeAsync()
+			=> await FlushAsync(_cancellationToken);
+
+		/// <inheritdoc/>
+		public IBlobWriter CreateBlobWriter(string? basePath = null, BlobSerializerOptions? serializerOptions = null)
+			=> CreateBlobWriter(basePath, serializerOptions, _cancellationToken);
+
+		/// <inheritdoc cref="IStorageWriter.CreateBlobWriter(string?, BlobSerializerOptions?)"/>
+		protected abstract IBlobWriter CreateBlobWriter(string? basePath = null, BlobSerializerOptions? serializerOptions = null, CancellationToken cancellationToken = default);
+
+		#region Aliases
+
+		/// <inheritdoc/>
+		public async Task AddAliasAsync(string name, IBlobRef handle, int rank = 0, ReadOnlyMemory<byte> data = default, CancellationToken cancellationToken = default)
+		{
+			await handle.FlushAsync(cancellationToken);
+			_request.AddAliases.Add(new AddAliasRequest { Name = name, Rank = rank, Data = data.ToArray(), Target = handle.GetLocator() });
+			await CheckForFlushAsync(cancellationToken);
+		}
+
+		/// <inheritdoc/>
+		public async Task RemoveAliasAsync(string name, IBlobRef handle, CancellationToken cancellationToken = default)
+		{
+			await handle.FlushAsync(cancellationToken);
+			_request.RemoveAliases.Add(new RemoveAliasRequest { Name = name, Target = handle.GetLocator() });
+			await CheckForFlushAsync(cancellationToken);
+		}
+
+		#endregion
+
+		/// <inheritdoc/>
+		public async Task WriteRefAsync(RefName name, IHashedBlobRef target, RefOptions? options = null, CancellationToken cancellationToken = default)
+		{
+			await target.FlushAsync(cancellationToken);
+			options = (options != null) ? new RefOptions(options.Lifetime, options.Extend) : null;
+			_request.AddRefs.Add(new AddRefRequest { RefName = name, Hash = target.Hash, Target = target.GetLocator(), Options = options });
+			await CheckForFlushAsync(cancellationToken);
+		}
+
+		/// <inheritdoc/>
+		public async Task DeleteRefAsync(RefName name, CancellationToken cancellationToken = default)
+		{
+			_request.RemoveRefs.Add(new RemoveRefRequest { RefName = name });
+			await CheckForFlushAsync(cancellationToken);
+		}
+
+		async Task CheckForFlushAsync(CancellationToken cancellationToken = default)
+		{
+			int numItems = _request.AddAliases.Count + _request.RemoveAliases.Count + _request.AddRefs.Count + _request.RemoveRefs.Count;
+			if (numItems > 1000)
+			{
+				await FlushAsync(cancellationToken);
+			}
+		}
+
+		/// <inheritdoc/>
+		public async Task FlushAsync(CancellationToken cancellationToken = default)
+		{
+			await _storageBackend.UpdateMetadataAsync(_request, cancellationToken);
+			_request = new UpdateMetadataRequest();
+		}
+	}
+}
diff --git a/Engine/Source/Programs/Shared/EpicGames.Horde/Storage/KeyValueStorageNamespace.cs b/Engine/Source/Programs/Shared/EpicGames.Horde/Storage/KeyValueStorageNamespace.cs
--- a/Engine/Source/Programs/Shared/EpicGames.Horde/Storage/KeyValueStorageNamespace.cs
+++ b/Engine/Source/Programs/Shared/EpicGames.Horde/Storage/KeyValueStorageNamespace.cs
@@ -53,7 +53,22 @@ namespace EpicGames.Horde.Storage.Clients
 			public override int GetHashCode() => _locator.GetHashCode();
 		}
 
-		class Writer : BlobWriter
+		class StorageWriterImpl : StorageWriter
+		{
+			readonly KeyValueStorageNamespace _keyValueStorageNamespace;
+
+			public StorageWriterImpl(KeyValueStorageNamespace keyValueStorageNamespace, IStorageBackend storageBackend, CancellationToken cancellationToken)
+				: base(storageBackend, cancellationToken)
+			{
+				_keyValueStorageNamespace = keyValueStorageNamespace;
+			}
+
+			/// <inheritdoc/>
+			protected override IBlobWriter CreateBlobWriter(string? basePath = null, BlobSerializerOptions? serializerOptions = null, CancellationToken cancellationToken = default)
+				=> _keyValueStorageNamespace.CreateBlobWriter(basePath, serializerOptions, cancellationToken);
+		}
+
+		class BlobWriterImpl : BlobWriter
 		{
 			readonly KeyValueStorageNamespace _outer;
 			readonly string _basePath;
@@ -63,7 +78,7 @@ namespace EpicGames.Horde.Storage.Clients
 			/// <summary>
 			/// Constructor
 			/// </summary>
-			public Writer(KeyValueStorageNamespace outer, string? basePath, BlobSerializerOptions? options)
+			public BlobWriterImpl(KeyValueStorageNamespace outer, string? basePath, BlobSerializerOptions? options)
 				: base(options)
 			{
 				_outer = outer;
@@ -82,7 +97,7 @@ namespace EpicGames.Horde.Storage.Clients
 			public override Task FlushAsync(CancellationToken cancellationToken = default) => Task.CompletedTask;
 
 			/// <inheritdoc/>
-			public override IBlobWriter Fork() => new Writer(_outer, _basePath, Options);
+			public override IBlobWriter Fork() => new BlobWriterImpl(_outer, _basePath, Options);
 
 			/// <inheritdoc/>
 			public override Memory<byte> GetOutputBuffer(int usedSize, int desiredSize)
@@ -132,16 +147,19 @@ namespace EpicGames.Horde.Storage.Clients
 		/// </summary>
 		public static KeyValueStorageNamespace CreateInMemory() => new KeyValueStorageNamespace(new MemoryStorageBackend());
 
+		/// <inheritdoc/>
+		public IStorageWriter CreateWriter(CancellationToken cancellationToken)
+			=> new StorageWriterImpl(this, _backend, cancellationToken);
+
 		#region Blobs
 
 		/// <inheritdoc/>
 		public IBlobRef CreateBlobRef(BlobLocator locator)
-		{
-			return new Handle(this, locator);
-		}
+			=> new Handle(this, locator);
 
 		/// <inheritdoc/>
-		public IBlobWriter CreateBlobWriter(string? basePath = null, BlobSerializerOptions? options = null, CancellationToken cancellationToken = default) => new Writer(this, basePath, options);
+		public IBlobWriter CreateBlobWriter(string? basePath = null, BlobSerializerOptions? options = null, CancellationToken cancellationToken = default) 
+			=> new BlobWriterImpl(this, basePath, options);
 
 		/// <inheritdoc/>
 		public async ValueTask<BlobData> ReadBlobAsync(BlobLocator locator, CancellationToken cancellationToken = default)
@@ -170,20 +188,6 @@ namespace EpicGames.Horde.Storage.Clients
 
 		#region Aliases
 
-		/// <inheritdoc/>
-		public async Task AddAliasAsync(string name, IBlobRef handle, int rank = 0, ReadOnlyMemory<byte> data = default, CancellationToken cancellationToken = default)
-		{
-			await handle.FlushAsync(cancellationToken);
-			await _backend.AddAliasAsync(name, handle.GetLocator(), rank, data, cancellationToken);
-		}
-
-		/// <inheritdoc/>
-		public async Task RemoveAliasAsync(string name, IBlobRef handle, CancellationToken cancellationToken = default)
-		{
-			await handle.FlushAsync(cancellationToken);
-			await _backend.RemoveAliasAsync(name, handle.GetLocator(), cancellationToken);
-		}
-
 		/// <inheritdoc/>
 		public async Task<BlobAlias[]> FindAliasesAsync(string name, int? maxResults = null, CancellationToken cancellationToken = default)
 		{
@@ -206,17 +210,6 @@ namespace EpicGames.Horde.Storage.Clients
 			return this.CreateBlobRef(value.Hash, value.Locator);
 		}
 
-		/// <inheritdoc/>
-		public async Task WriteRefAsync(RefName name, IHashedBlobRef value, RefOptions? options = null, CancellationToken cancellationToken = default)
-		{
-			await value.FlushAsync(cancellationToken);
-			await _backend.WriteRefAsync(name, new HashedBlobRefValue(value.Hash, value.GetLocator()), options, cancellationToken);
-		}
-
-		/// <inheritdoc/>
-		public Task<bool> DeleteRefAsync(RefName name, CancellationToken cancellationToken = default)
-			=> _backend.DeleteRefAsync(name, cancellationToken);
-
 		#endregion
 
 		/// <inheritdoc/>
diff --git a/Engine/Source/Programs/Shared/EpicGames.Horde/Storage/StorageBackendCache.cs b/Engine/Source/Programs/Shared/EpicGames.Horde/Storage/StorageBackendCache.cs
--- a/Engine/Source/Programs/Shared/EpicGames.Horde/Storage/StorageBackendCache.cs
+++ b/Engine/Source/Programs/Shared/EpicGames.Horde/Storage/StorageBackendCache.cs
@@ -153,12 +153,6 @@ namespace EpicGames.Horde.Storage
 
 			#region Aliases
 
-			public Task AddAliasAsync(string name, BlobLocator locator, int rank = 0, ReadOnlyMemory<byte> data = default, CancellationToken cancellationToken = default)
-				=> _inner.AddAliasAsync(name, locator, rank, data, cancellationToken);
-
-			public Task RemoveAliasAsync(string name, BlobLocator locator, CancellationToken cancellationToken = default)
-				=> _inner.RemoveAliasAsync(name, locator, cancellationToken);
-
 			public Task<BlobAliasLocator[]> FindAliasesAsync(string name, int? maxResults = null, CancellationToken cancellationToken = default)
 				=> _inner.FindAliasesAsync(name, maxResults, cancellationToken);
 
@@ -169,13 +163,10 @@ namespace EpicGames.Horde.Storage
 			public Task<HashedBlobRefValue?> TryReadRefAsync(RefName name, RefCacheTime cacheTime = default, CancellationToken cancellationToken = default)
 				=> _inner.TryReadRefAsync(name, cacheTime, cancellationToken);
 
-			public Task WriteRefAsync(RefName name, HashedBlobRefValue value, RefOptions? options = null, CancellationToken cancellationToken = default)
-				=> _inner.WriteRefAsync(name, value, options, cancellationToken);
-
-			public Task<bool> DeleteRefAsync(RefName name, CancellationToken cancellationToken = default)
-				=> _inner.DeleteRefAsync(name, cancellationToken);
-
 			#endregion
+
+			public Task UpdateMetadataAsync(UpdateMetadataRequest request, CancellationToken cancellationToken = default)
+				=> _inner.UpdateMetadataAsync(request, cancellationToken);
 		}
 
 		object LockObject => _items;
diff --git a/Engine/Source/Programs/Shared/EpicGames.Horde/Storage/StorageMessages.cs b/Engine/Source/Programs/Shared/EpicGames.Horde/Storage/StorageMessages.cs
--- a/Engine/Source/Programs/Shared/EpicGames.Horde/Storage/StorageMessages.cs
+++ b/Engine/Source/Programs/Shared/EpicGames.Horde/Storage/StorageMessages.cs
@@ -80,9 +80,9 @@ namespace EpicGames.Horde.Storage
 	}
 
 	/// <summary>
-	/// Request to batch update items in the database
+	/// Request to batch update metadata in the database
 	/// </summary>
-	public class UpdateNamespaceRequest
+	public class UpdateMetadataRequest
 	{
 		/// <summary>
 		/// List of aliases to add
