Horde: Static analysis fixes.

#rnx

diff --git a/Engine/Source/Programs/Horde/Plugins/Tools/HordeServer.Tools/Commands/Generate/ToolDataCommand.cs b/Engine/Source/Programs/Horde/Plugins/Tools/HordeServer.Tools/Commands/Generate/ToolDataCommand.cs
--- a/Engine/Source/Programs/Horde/Plugins/Tools/HordeServer.Tools/Commands/Generate/ToolDataCommand.cs
+++ b/Engine/Source/Programs/Horde/Plugins/Tools/HordeServer.Tools/Commands/Generate/ToolDataCommand.cs
@@ -108,7 +108,7 @@ namespace HordeServer.Commands.Generate
 
 					logger.LogInformation("");
 				}
-				await client.WriteRefAsync(refName, dirNodeRef);
+				await client.AddRefAsync(refName, dirNodeRef);
 			}
 
 			// Update the server config to include the bundled tool
diff --git a/Engine/Source/Programs/Shared/EpicGames.Horde.Tests/BundleStoreTests.cs b/Engine/Source/Programs/Shared/EpicGames.Horde.Tests/BundleStoreTests.cs
--- a/Engine/Source/Programs/Shared/EpicGames.Horde.Tests/BundleStoreTests.cs
+++ b/Engine/Source/Programs/Shared/EpicGames.Horde.Tests/BundleStoreTests.cs
@@ -170,7 +170,7 @@ namespace EpicGames.Horde.Tests
 				IHashedBlobRef<SimpleNode> rootRef = await writer.WriteBlobAsync(root);
 				await writer.FlushAsync();
 
-				await store.WriteRefAsync(new RefName("test"), rootRef);
+				await store.AddRefAsync(new RefName("test"), rootRef);
 
 				BundleReader reader = new BundleReader(store, BundleCache.None, NullLogger.Instance);
 				await CheckTreeAsync(root);
@@ -228,7 +228,7 @@ namespace EpicGames.Horde.Tests
 			{
 				inputRef = await writer.WriteBlobAsync(new SimpleNode(new ReadOnlySequence<byte>(new byte[] { (byte)123 }), Array.Empty<IHashedBlobRef<SimpleNode>>()));
 			}
-			await store.WriteRefAsync(refName, inputRef);
+			await store.AddRefAsync(refName, inputRef);
 
 			SimpleNode node = await store.ReadRefTargetAsync<SimpleNode>(refName);
 
diff --git a/Engine/Source/Programs/Shared/EpicGames.Horde.Tests/BundleTests.cs b/Engine/Source/Programs/Shared/EpicGames.Horde.Tests/BundleTests.cs
--- a/Engine/Source/Programs/Shared/EpicGames.Horde.Tests/BundleTests.cs
+++ b/Engine/Source/Programs/Shared/EpicGames.Horde.Tests/BundleTests.cs
@@ -195,7 +195,7 @@ namespace EpicGames.Horde.Tests
 			}
 
 			RefName refName = new RefName("testref");
-			await store.WriteRefAsync(refName, rootRef);
+			await store.AddRefAsync(refName, rootRef);
 
 			// Should be stored inline
 			MemoryStorageBackend memoryStore = (MemoryStorageBackend)store.Backend;
@@ -249,7 +249,7 @@ namespace EpicGames.Horde.Tests
 
 				RefName refName = new RefName("ref");
 				IHashedBlobRef<DirectoryNode> rootRef = await writer.WriteBlobAsync(root);
-				await storage.WriteRefAsync(refName, rootRef);
+				await storage.AddRefAsync(refName, rootRef);
 			}
 
 			MemoryStorageBackend memoryStore = (MemoryStorageBackend)storage.Backend;
diff --git a/Engine/Source/Programs/Shared/EpicGames.Horde/Storage/Backends/FileStorageBackend.cs b/Engine/Source/Programs/Shared/EpicGames.Horde/Storage/Backends/FileStorageBackend.cs
--- a/Engine/Source/Programs/Shared/EpicGames.Horde/Storage/Backends/FileStorageBackend.cs
+++ b/Engine/Source/Programs/Shared/EpicGames.Horde/Storage/Backends/FileStorageBackend.cs
@@ -137,15 +137,9 @@ namespace EpicGames.Horde.Storage.Backends
 		#region Refs
 
 		/// <inheritdoc/>
-		public Task<bool> DeleteRefAsync(RefName name, CancellationToken cancellationToken = default)
+		public void DeleteRef(RefName name)
 		{
-			FileInfo file = GetRefFile(name).ToFileInfo();
-			if (file.Exists)
-			{
-				file.Delete();
-				return Task.FromResult(true);
-			}
-			return Task.FromResult(false);
+			FileReference.Delete(GetRefFile(name));
 		}
 
 		/// <inheritdoc/>
@@ -167,7 +161,7 @@ namespace EpicGames.Horde.Storage.Backends
 		}
 
 		/// <inheritdoc/>
-		public async Task WriteRefAsync(RefName name, HashedBlobRefValue value, RefOptions? options = null, CancellationToken cancellationToken = default)
+		public async Task WriteRefAsync(RefName name, HashedBlobRefValue value, CancellationToken cancellationToken = default)
 		{
 			FileReference file = GetRefFile(name);
 			DirectoryReference.CreateDirectory(file.Directory);
@@ -203,11 +197,11 @@ namespace EpicGames.Horde.Storage.Backends
 			}
 			foreach (AddRefRequest addRef in request.AddRefs)
 			{
-				await WriteRefAsync(addRef.RefName, new HashedBlobRefValue(addRef.Hash, addRef.Target), addRef.Options, cancellationToken);
+				await WriteRefAsync(addRef.RefName, new HashedBlobRefValue(addRef.Hash, addRef.Target), cancellationToken);
 			}
 			foreach (RemoveRefRequest removeRef in request.RemoveRefs)
 			{
-				await DeleteRefAsync(removeRef.RefName, cancellationToken);
+				DeleteRef(removeRef.RefName);
 			}
 		}
 
diff --git a/Engine/Source/Programs/Shared/EpicGames.Horde/Storage/Backends/JupiterStorageBackend.cs b/Engine/Source/Programs/Shared/EpicGames.Horde/Storage/Backends/JupiterStorageBackend.cs
--- a/Engine/Source/Programs/Shared/EpicGames.Horde/Storage/Backends/JupiterStorageBackend.cs
+++ b/Engine/Source/Programs/Shared/EpicGames.Horde/Storage/Backends/JupiterStorageBackend.cs
@@ -412,7 +412,7 @@ namespace EpicGames.Horde.Storage.Backends
 		}
 
 		/// <inheritdoc/>
-		public async Task WriteRefAsync(RefName name, HashedBlobRefValue value, RefOptions? options = null, CancellationToken cancellationToken = default)
+		public async Task WriteRefAsync(RefName name, HashedBlobRefValue value, CancellationToken cancellationToken = default)
 		{
 			_logger.LogDebug("Writing ref {RefName} -> {Locator}", name, value.Locator);
 			using (HttpClient httpClient = _createClient())
@@ -468,7 +468,7 @@ namespace EpicGames.Horde.Storage.Backends
 			}
 			foreach (AddRefRequest addRef in request.AddRefs)
 			{
-				await WriteRefAsync(addRef.RefName, new HashedBlobRefValue(addRef.Hash, addRef.Target), addRef.Options, cancellationToken);
+				await WriteRefAsync(addRef.RefName, new HashedBlobRefValue(addRef.Hash, addRef.Target), cancellationToken);
 			}
 			foreach (RemoveRefRequest removeRef in request.RemoveRefs)
 			{
diff --git a/Engine/Source/Programs/Shared/EpicGames.Horde/Storage/Bundles/BundleStorageNamespace.cs b/Engine/Source/Programs/Shared/EpicGames.Horde/Storage/Bundles/BundleStorageNamespace.cs
--- a/Engine/Source/Programs/Shared/EpicGames.Horde/Storage/Bundles/BundleStorageNamespace.cs
+++ b/Engine/Source/Programs/Shared/EpicGames.Horde/Storage/Bundles/BundleStorageNamespace.cs
@@ -23,13 +23,11 @@ namespace EpicGames.Horde.Storage.Bundles
 		class StorageWriterImpl : StorageWriter
 		{
 			readonly BundleStorageNamespace _outer;
-			readonly IStorageBackend _storageBackend;
 
 			public StorageWriterImpl(BundleStorageNamespace outer, IStorageBackend storageBackend, CancellationToken cancellationToken)
 				: base(storageBackend, cancellationToken)
 			{
 				_outer = outer;
-				_storageBackend = storageBackend;
 			}
 
 			protected override IBlobWriter CreateBlobWriter(string? basePath = null, BlobSerializerOptions? serializerOptions = null, CancellationToken cancellationToken = default)
diff --git a/Engine/Source/Programs/Shared/EpicGames.Horde/Storage/IStorageWriter.cs b/Engine/Source/Programs/Shared/EpicGames.Horde/Storage/IStorageWriter.cs
--- a/Engine/Source/Programs/Shared/EpicGames.Horde/Storage/IStorageWriter.cs
+++ b/Engine/Source/Programs/Shared/EpicGames.Horde/Storage/IStorageWriter.cs
@@ -1,9 +1,6 @@
 // Copyright Epic Games, Inc. All Rights Reserved.
 
 using System;
-using System.Collections.Generic;
-using System.Linq;
-using System.Text;
 using System.Threading;
 using System.Threading.Tasks;
 
@@ -96,7 +93,10 @@ namespace EpicGames.Horde.Storage
 
 		/// <inheritdoc/>
 		public async ValueTask DisposeAsync()
-			=> await FlushAsync(_cancellationToken);
+		{
+			await FlushAsync(_cancellationToken);
+			GC.SuppressFinalize(this);
+		}
 
 		/// <inheritdoc/>
 		public IBlobWriter CreateBlobWriter(string? basePath = null, BlobSerializerOptions? serializerOptions = null)
