Horde: Missing file.

#rnx

diff --git a/Engine/Source/Programs/Shared/EpicGames.Horde/Storage/Backends/MemoryStorageBackend.cs b/Engine/Source/Programs/Shared/EpicGames.Horde/Storage/Backends/MemoryStorageBackend.cs
--- a/Engine/Source/Programs/Shared/EpicGames.Horde/Storage/Backends/MemoryStorageBackend.cs
+++ b/Engine/Source/Programs/Shared/EpicGames.Horde/Storage/Backends/MemoryStorageBackend.cs
@@ -99,14 +99,13 @@ namespace EpicGames.Horde.Storage.Backends
 		#region Aliases
 
 		/// <inheritdoc/>
-		public Task AddAliasAsync(string name, BlobLocator locator, int rank = 0, ReadOnlyMemory<byte> data = default, CancellationToken cancellationToken = default)
+		public void AddAlias(string name, BlobLocator locator, int rank = 0, ReadOnlyMemory<byte> data = default)
 		{
 			_aliases.AddOrUpdate(name, _ => new AliasListNode(locator, rank, data, null), (_, entry) => new AliasListNode(locator, rank, data, entry));
-			return Task.CompletedTask;
 		}
 
 		/// <inheritdoc/>
-		public Task RemoveAliasAsync(string name, BlobLocator locator, CancellationToken cancellationToken = default)
+		public void RemoveAlias(string name, BlobLocator locator)
 		{
 			for (; ; )
 			{
@@ -137,7 +136,6 @@ namespace EpicGames.Horde.Storage.Backends
 					}
 				}
 			}
-			return Task.CompletedTask;
 		}
 
 		static AliasListNode? RemoveAliasFromList(AliasListNode? entry, BlobLocator locator)
@@ -177,9 +175,6 @@ namespace EpicGames.Horde.Storage.Backends
 
 		#region Refs
 
-		/// <inheritdoc/>
-		public Task<bool> DeleteRefAsync(RefName name, CancellationToken cancellationToken) => Task.FromResult(_refs.TryRemove(name, out _));
-
 		/// <inheritdoc/>
 		public Task<HashedBlobRefValue?> TryReadRefAsync(RefName name, RefCacheTime cacheTime = default, CancellationToken cancellationToken = default)
 		{
@@ -194,34 +189,28 @@ namespace EpicGames.Horde.Storage.Backends
 			}
 		}
 
-		/// <inheritdoc/>
-		public Task WriteRefAsync(RefName name, HashedBlobRefValue value, RefOptions? options = null, CancellationToken cancellationToken = default)
-		{
-			_refs[name] = value;
-			return Task.CompletedTask;
-		}
-
 		#endregion
 
 		/// <inheritdoc/>
-		public async Task UpdateMetadataAsync(UpdateMetadataRequest request, CancellationToken cancellationToken = default)
+		public Task UpdateMetadataAsync(UpdateMetadataRequest request, CancellationToken cancellationToken = default)
 		{
 			foreach (AddAliasRequest addAlias in request.AddAliases)
 			{
-				await AddAliasAsync(addAlias.Name, addAlias.Target, addAlias.Rank, addAlias.Data, cancellationToken);
+				AddAlias(addAlias.Name, addAlias.Target, addAlias.Rank, addAlias.Data);
 			}
 			foreach (RemoveAliasRequest removeAlias in request.RemoveAliases)
 			{
-				await RemoveAliasAsync(removeAlias.Name, removeAlias.Target, cancellationToken);
+				RemoveAlias(removeAlias.Name, removeAlias.Target);
 			}
 			foreach (AddRefRequest addRef in request.AddRefs)
 			{
-				await WriteRefAsync(addRef.RefName, new HashedBlobRefValue(addRef.Hash, addRef.Target), addRef.Options, cancellationToken);
+				_refs[addRef.RefName] = new HashedBlobRefValue(addRef.Hash, addRef.Target);
 			}
 			foreach (RemoveRefRequest removeRef in request.RemoveRefs)
 			{
-				await DeleteRefAsync(removeRef.RefName, cancellationToken);
+				_refs.TryRemove(removeRef.RefName, out _);
 			}
+			return Task.CompletedTask;
 		}
 
 		/// <inheritdoc/>
