Horde: Fix errors uploading blobs to storage when not using a backend that support redirects.

#rnx

diff --git a/Engine/Source/Programs/Shared/EpicGames.Horde/Storage/Backends/HttpStorageBackend.cs b/Engine/Source/Programs/Shared/EpicGames.Horde/Storage/Backends/HttpStorageBackend.cs
--- a/Engine/Source/Programs/Shared/EpicGames.Horde/Storage/Backends/HttpStorageBackend.cs
+++ b/Engine/Source/Programs/Shared/EpicGames.Horde/Storage/Backends/HttpStorageBackend.cs
@@ -167,7 +167,7 @@ namespace EpicGames.Horde.Storage.Backends
 				}
 			}
 
-			WriteBlobResponse response = await SendWriteRequestAsync(null, streamContent, imports, prefix, cancellationToken);
+			WriteBlobResponse response = await SendWriteRequestAsync(locator, streamContent, imports, prefix, cancellationToken);
 			_supportsUploadRedirects = response.SupportsRedirects ?? false;
 			_logger.LogDebug("Written {Locator} (direct)", response.Blob);
 			return new BlobLocator(response.Blob);
