Fixing spelling errors in UHT

#rnx

diff --git a/Engine/Source/Programs/Shared/EpicGames.UHT/Exporters/CodeGen/UhtHeaderCodeGenerator.cs b/Engine/Source/Programs/Shared/EpicGames.UHT/Exporters/CodeGen/UhtHeaderCodeGenerator.cs
--- a/Engine/Source/Programs/Shared/EpicGames.UHT/Exporters/CodeGen/UhtHeaderCodeGenerator.cs
+++ b/Engine/Source/Programs/Shared/EpicGames.UHT/Exporters/CodeGen/UhtHeaderCodeGenerator.cs
@@ -414,27 +414,27 @@ namespace EpicGames.UHT.Exporters.CodeGen
 
 		protected static StringBuilder AppendParameters(StringBuilder builder, UhtFunction function, UhtPropertyTextType textType, string? extraParameter, bool skipParameterName)
 		{
-			bool needsSeperator = false;
+			bool needsSeparator = false;
 
 			builder.Append('(');
 
 			if (extraParameter != null)
 			{
 				builder.Append(extraParameter);
-				needsSeperator = true;
+				needsSeparator = true;
 			}
 
 			foreach (UhtType parameter in function.ParameterProperties.Span)
 			{
 				if (parameter is UhtProperty property)
 				{
-					if (needsSeperator)
+					if (needsSeparator)
 					{
 						builder.Append(", ");
 					}
 					else
 					{
-						needsSeperator = true;
+						needsSeparator = true;
 					}
 					builder.AppendFullDecl(property, textType, skipParameterName);
 				}
diff --git a/Engine/Source/Programs/Shared/EpicGames.UHT/Tokenizer/UhtToken.cs b/Engine/Source/Programs/Shared/EpicGames.UHT/Tokenizer/UhtToken.cs
--- a/Engine/Source/Programs/Shared/EpicGames.UHT/Tokenizer/UhtToken.cs
+++ b/Engine/Source/Programs/Shared/EpicGames.UHT/Tokenizer/UhtToken.cs
@@ -627,12 +627,12 @@ namespace EpicGames.UHT.Tokenizer
 		public static string Join(char separator, IEnumerable<UhtToken> tokens)
 		{
 			StringBuilder builder = new();
-			bool includeSeperator = false;
+			bool includeSeparator = false;
 			foreach (UhtToken token in tokens)
 			{
-				if (!includeSeperator)
+				if (!includeSeparator)
 				{
-					includeSeperator = true;
+					includeSeparator = true;
 				}
 				else
 				{
@@ -652,12 +652,12 @@ namespace EpicGames.UHT.Tokenizer
 		public static string Join(string separator, IEnumerable<UhtToken> tokens)
 		{
 			StringBuilder builder = new();
-			bool includeSeperator = false;
+			bool includeSeparator = false;
 			foreach (UhtToken token in tokens)
 			{
-				if (!includeSeperator)
+				if (!includeSeparator)
 				{
-					includeSeperator = true;
+					includeSeparator = true;
 				}
 				else
 				{
diff --git a/Engine/Source/Programs/Shared/EpicGames.UHT/Tokenizer/UhtTokenReader.cs b/Engine/Source/Programs/Shared/EpicGames.UHT/Tokenizer/UhtTokenReader.cs
--- a/Engine/Source/Programs/Shared/EpicGames.UHT/Tokenizer/UhtTokenReader.cs
+++ b/Engine/Source/Programs/Shared/EpicGames.UHT/Tokenizer/UhtTokenReader.cs
@@ -266,16 +266,16 @@ namespace EpicGames.UHT.Tokenizer
 		/// Join the tokens in the list
 		/// </summary>
 		/// <param name="builder">Destination builder</param>
-		/// <param name="seperator">Separator between the tokens</param>
+		/// <param name="separator">Separator between the tokens</param>
 		/// <returns></returns>
-		public void Join(StringBuilder builder, string seperator = "")
+		public void Join(StringBuilder builder, string separator = "")
 		{
 			builder.Append(Token.Value.ToString());
 			UhtTokenList list = this;
 			while (list.Next != null)
 			{
 				list = list.Next;
-				builder.Append(seperator);
+				builder.Append(separator);
 				builder.Append(list.Token.Value.ToString());
 			}
 		}
@@ -283,16 +283,16 @@ namespace EpicGames.UHT.Tokenizer
 		/// <summary>
 		/// Join the tokens in the list
 		/// </summary>
-		/// <param name="seperator">Separator between the tokens</param>
+		/// <param name="separator">Separator between the tokens</param>
 		/// <returns></returns>
-		public string Join(string seperator = "")
+		public string Join(string separator = "")
 		{
 			if (Next == null)
 			{
 				return Token.Value.ToString();
 			}
 			StringBuilder builder = new();
-			Join(builder, seperator);
+			Join(builder, separator);
 			return builder.ToString();
 		}
 
