Horde: Add an ExtractOptions object which can be passed to DirectoryNode.ExtractAsync() to configure various parameters without requiring as many overloads.

#rnx

diff --git a/Engine/Source/Programs/AutomationTool/BuildGraph/Tasks/RetrieveArtifactTask.cs b/Engine/Source/Programs/AutomationTool/BuildGraph/Tasks/RetrieveArtifactTask.cs
--- a/Engine/Source/Programs/AutomationTool/BuildGraph/Tasks/RetrieveArtifactTask.cs
+++ b/Engine/Source/Programs/AutomationTool/BuildGraph/Tasks/RetrieveArtifactTask.cs
@@ -187,7 +187,7 @@ namespace AutomationTool.Tasks
 
 			// Upload the files
 			Stopwatch timer = Stopwatch.StartNew();
-			await artifact.Content.ExtractAsync(outputDir.ToDirectoryInfo(), new ExtractStatsLogger(Logger), Logger, CancellationToken.None);
+			await artifact.Content.ExtractAsync(outputDir.ToDirectoryInfo(), new ExtractOptions { Progress = new ExtractStatsLogger(Logger) }, Logger, CancellationToken.None);
 			Logger.LogInformation("Retrieved artifact {ArtifactId} in {Time:n1}s", artifact.Id, timer.Elapsed.TotalSeconds);
 		}
 
diff --git a/Engine/Source/Programs/Horde/Drivers/JobDriver/Execution/JobExecutor.cs b/Engine/Source/Programs/Horde/Drivers/JobDriver/Execution/JobExecutor.cs
--- a/Engine/Source/Programs/Horde/Drivers/JobDriver/Execution/JobExecutor.cs
+++ b/Engine/Source/Programs/Horde/Drivers/JobDriver/Execution/JobExecutor.cs
@@ -1102,7 +1102,7 @@ namespace JobDriver.Execution
 				DirectoryNode? buildGraphDir = await node.TryOpenDirectoryAsync(BuildGraphTempStorageDir, cancellationToken: cancellationToken);
 				if (buildGraphDir != null)
 				{
-					await buildGraphDir.ExtractAsync(new DirectoryInfo(workspaceDir.FullName), logger, cancellationToken);
+					await buildGraphDir.ExtractAsync(new DirectoryInfo(workspaceDir.FullName), new ExtractOptions(), logger, cancellationToken);
 					logger.LogInformation("Copying preprocessed script from {BuildGraphFolderName} into {OutputDir}", BuildGraphTempStorageDir, workspaceDir);
 				}
 				else
diff --git a/Engine/Source/Programs/Horde/Drivers/JobDriver/Utility/TempStorage.cs b/Engine/Source/Programs/Horde/Drivers/JobDriver/Utility/TempStorage.cs
--- a/Engine/Source/Programs/Horde/Drivers/JobDriver/Utility/TempStorage.cs
+++ b/Engine/Source/Programs/Horde/Drivers/JobDriver/Utility/TempStorage.cs
@@ -783,7 +783,7 @@ namespace JobDriver.Utility
 				Stopwatch timer = Stopwatch.StartNew();
 
 				// Add all the files and flush the ref
-				await rootDirEntry.Handle.ExtractAsync(rootDir.ToDirectoryInfo(), new ExtractStatsLogger(logger), logger, cancellationToken);
+				await rootDirEntry.Handle.ExtractAsync(rootDir.ToDirectoryInfo(), new ExtractOptions { Progress = new ExtractStatsLogger(logger) }, logger, cancellationToken);
 
 				StorageStats deltaStats = StorageStats.GetDelta(initialStats, storageNamespace.GetStats());
 				logger.LogInformation("{Stats}", $"Elapsed: {(int)timer.Elapsed.TotalSeconds}s, {String.Join(", ", deltaStats.Values.Select(x => $"{x.Item1}: {x.Item2:n0}"))}");
diff --git a/Engine/Source/Programs/Horde/HordeCmd/Commands/Archives/ArchiveExtract.cs b/Engine/Source/Programs/Horde/HordeCmd/Commands/Archives/ArchiveExtract.cs
--- a/Engine/Source/Programs/Horde/HordeCmd/Commands/Archives/ArchiveExtract.cs
+++ b/Engine/Source/Programs/Horde/HordeCmd/Commands/Archives/ArchiveExtract.cs
@@ -38,6 +38,10 @@ namespace Horde.Commands.Archives
 		[Description("If set, deletes the contents of the output directory before extraction.")]
 		public bool CleanOutput { get; set; }
 
+		[CommandLine("-VerifyOutput")]
+		[Description("Hashes output data to ensure it is as expected")]
+		public bool VerifyOutput { get; set; }
+
 		public ArchiveExtract(HttpStorageClient storageClient, BundleCache bundleCache)
 			: base(storageClient, bundleCache)
 		{
@@ -51,24 +55,28 @@ namespace Horde.Commands.Archives
 				FileUtils.ForceDeleteDirectoryContents(OutputDir);
 			}
 
+			ExtractOptions options = new ExtractOptions();
+			options.Progress = new ExtractStatsLogger(logger);
+			options.VerifyOutput = VerifyOutput;
+
 			if (File != null)
 			{
 				using MemoryMappedFileCache memoryMappedFileCache = new MemoryMappedFileCache();
 				IStorageNamespace store = BundleStorageNamespace.CreateFromDirectory(File.Directory, BundleCache, memoryMappedFileCache, logger);
 				IBlobRef<DirectoryNode> handle = store.CreateBlobRef<DirectoryNode>(await FileStorageBackend.ReadRefAsync(File));
-				await ExecuteInternalAsync(store, handle, logger);
+				await ExecuteInternalAsync(store, handle, options, logger);
 			}
 			else if (Ref != null)
 			{
 				IStorageNamespace store = GetStorageNamespace();
 				IBlobRef<DirectoryNode> handle = await store.ReadRefAsync<DirectoryNode>(new RefName(Ref));
-				await ExecuteInternalAsync(store, handle, logger);
+				await ExecuteInternalAsync(store, handle, options, logger);
 			}
 			else if (Node != null)
 			{
 				IStorageNamespace store = GetStorageNamespace();
 				IBlobRef<DirectoryNode> handle = store.CreateBlobRef<DirectoryNode>(new BlobLocator(Node));
-				await ExecuteInternalAsync(store, handle, logger);
+				await ExecuteInternalAsync(store, handle, options, logger);
 			}
 			else
 			{
@@ -78,11 +86,11 @@ namespace Horde.Commands.Archives
 			return 0;
 		}
 
-		protected async Task ExecuteInternalAsync(IStorageNamespace store, IBlobRef<DirectoryNode> handle, ILogger logger)
+		protected async Task ExecuteInternalAsync(IStorageNamespace store, IBlobRef<DirectoryNode> handle, ExtractOptions options, ILogger logger)
 		{
 			Stopwatch timer = Stopwatch.StartNew();
 
-			await handle.ExtractAsync(OutputDir.ToDirectoryInfo(), new ExtractStatsLogger(logger), logger, CancellationToken.None);
+			await handle.ExtractAsync(OutputDir.ToDirectoryInfo(), options, logger, CancellationToken.None);
 
 			logger.LogInformation("Elapsed: {Time}s", timer.Elapsed.TotalSeconds);
 
diff --git a/Engine/Source/Programs/Horde/HordeCmd/Commands/Artifacts/ArtifactDownload.cs b/Engine/Source/Programs/Horde/HordeCmd/Commands/Artifacts/ArtifactDownload.cs
--- a/Engine/Source/Programs/Horde/HordeCmd/Commands/Artifacts/ArtifactDownload.cs
+++ b/Engine/Source/Programs/Horde/HordeCmd/Commands/Artifacts/ArtifactDownload.cs
@@ -34,6 +34,10 @@ namespace Horde.Commands.Artifacts
 		[Description("If set, deletes the contents of the output directory before extraction.")]
 		public bool CleanOutput { get; set; }
 
+		[CommandLine("-VerifyOutput")]
+		[Description("Hashes output data to ensure it is as expected")]
+		public bool VerifyOutput { get; set; }
+
 		readonly IHordeClient _hordeClient;
 
 		public ArtifactDownload(IHordeClient hordeClient)
@@ -79,7 +83,7 @@ namespace Horde.Commands.Artifacts
 			else
 			{
 				logger.LogError("Either -Id=... or -File=... must be specified.");
-				return 1;
+				return 1; 
 			}
 
 			if (CleanOutput)
@@ -88,8 +92,12 @@ namespace Horde.Commands.Artifacts
 				FileUtils.ForceDeleteDirectoryContents(OutputDir);
 			}
 
+			ExtractOptions options = new ExtractOptions();
+			options.Progress = new ExtractStatsLogger(logger);
+			options.VerifyOutput = VerifyOutput;
+
 			Stopwatch timer = Stopwatch.StartNew();
-			await handle.ExtractAsync(OutputDir.ToDirectoryInfo(), new ExtractStatsLogger(logger), logger, CancellationToken.None);
+			await handle.ExtractAsync(OutputDir.ToDirectoryInfo(), options, logger, CancellationToken.None);
 			logger.LogInformation("Elapsed: {Time}s", timer.Elapsed.TotalSeconds);
 
 			if (Stats)
diff --git a/Engine/Source/Programs/Shared/EpicGames.Horde.Tests/ChunkingTests.cs b/Engine/Source/Programs/Shared/EpicGames.Horde.Tests/ChunkingTests.cs
--- a/Engine/Source/Programs/Shared/EpicGames.Horde.Tests/ChunkingTests.cs
+++ b/Engine/Source/Programs/Shared/EpicGames.Horde.Tests/ChunkingTests.cs
@@ -195,7 +195,7 @@ namespace EpicGames.Horde.Tests
 			DirectoryInfo tempDir = new DirectoryInfo(Path.Combine(Path.GetTempPath(), Guid.NewGuid().ToString()));
 			try
 			{
-				await handle.ExtractAsync(tempDir, null, NullLogger.Instance, CancellationToken.None);
+				await handle.ExtractAsync(tempDir, new ExtractOptions(), NullLogger.Instance, CancellationToken.None);
 
 				byte[] outputData = await File.ReadAllBytesAsync(Path.Combine(tempDir.FullName, "test"));
 				Assert.IsTrue(outputData.SequenceEqual(data));
diff --git a/Engine/Source/Programs/Shared/EpicGames.Horde/Compute/AgentMessageHandler.cs b/Engine/Source/Programs/Shared/EpicGames.Horde/Compute/AgentMessageHandler.cs
--- a/Engine/Source/Programs/Shared/EpicGames.Horde/Compute/AgentMessageHandler.cs
+++ b/Engine/Source/Programs/Shared/EpicGames.Horde/Compute/AgentMessageHandler.cs
@@ -197,7 +197,7 @@ namespace EpicGames.Horde.Compute
 				throw new InvalidOperationException("Cannot write files outside sandbox");
 			}
 
-			await directoryRef.ExtractAsync(outputDir.ToDirectoryInfo(), _logger, cancellationToken);
+			await directoryRef.ExtractAsync(outputDir.ToDirectoryInfo(), new ExtractOptions(), _logger, cancellationToken);
 			await VerifyFilesAsync(outputDir, directoryRef, cancellationToken);
 
 			using (IAgentMessageBuilder message = await channel.CreateMessageAsync(AgentMessageType.WriteFilesResponse, cancellationToken))
diff --git a/Engine/Source/Programs/Shared/EpicGames.Horde/Storage/BatchReader.cs b/Engine/Source/Programs/Shared/EpicGames.Horde/Storage/BatchReader.cs
--- a/Engine/Source/Programs/Shared/EpicGames.Horde/Storage/BatchReader.cs
+++ b/Engine/Source/Programs/Shared/EpicGames.Horde/Storage/BatchReader.cs
@@ -3,6 +3,7 @@
 using System;
 using System.Collections.Generic;
 using System.Diagnostics.CodeAnalysis;
+using System.IO;
 using System.Linq;
 using System.Threading;
 using System.Threading.Channels;
@@ -19,7 +20,7 @@ namespace EpicGames.Horde.Storage
 	/// <typeparam name="TContext">User defined context type for the request</typeparam>
 	/// <param name="BlobRef">Reference to the blob data</param>
 	/// <param name="Context">User defined context</param>
-	public record class BatchReadRequest<TContext>(IBlobRef BlobRef, TContext Context);
+	public record class BatchReadRequest<TContext>(IHashedBlobRef BlobRef, TContext Context);
 
 	/// <summary>
 	/// Individual response for a blob that was read
@@ -76,7 +77,7 @@ namespace EpicGames.Horde.Storage
 			public int MaxOffset => Handle.PacketOffset + Handle.PacketLength;
 		}
 
-		record class ExportRequest(ExportHandle Handle, TContext Context);
+		record class ExportRequest(IoHash Hash, ExportHandle Handle, TContext Context);
 
 		record class PacketResponse(IReadOnlyMemoryOwner<byte> Data, List<ExportRequest> Exports) : IDisposable
 		{
@@ -90,6 +91,8 @@ namespace EpicGames.Horde.Storage
 		// Maximum gap between reads that should be coalesced
 		const long CoalesceReadsBelowSize = 2 * 1024 * 1024;
 
+		readonly bool _verifyHashes;
+
 		readonly Channel<BatchReadRequest<TContext>> _blobRequests;
 		readonly Channel<BatchReadRequest<TContext>> _otherRequests;
 		readonly Channel<BundleRequest> _bundleRequests;
@@ -120,8 +123,10 @@ namespace EpicGames.Horde.Storage
 		/// <summary>
 		/// Constructor
 		/// </summary>
-		public BatchReader()
+		public BatchReader(bool verifyHashes)
 		{
+			_verifyHashes = verifyHashes;
+
 			_blobRequests = Channel.CreateUnbounded<BatchReadRequest<TContext>>();
 			_otherRequests = Channel.CreateUnbounded<BatchReadRequest<TContext>>();
 			_bundleRequests = Channel.CreateUnbounded<BundleRequest>();
@@ -184,8 +189,9 @@ namespace EpicGames.Horde.Storage
 				try
 				{
 					blobData = await request.BlobRef.ReadBlobDataAsync(cancellationToken);
+					VerifyHash(request.BlobRef, blobData.Data.Span, request.BlobRef.Hash);
 #pragma warning disable CA2000
-					BatchReadResponse<TContext> response = new BatchReadResponse<TContext>([new BatchReadResponseItem<TContext>(blobData, request.Context)]);
+					BatchReadResponse < TContext> response = new BatchReadResponse<TContext>([new BatchReadResponseItem<TContext>(blobData, request.Context)]);
 					await _blobResponses.Writer.WriteAsync(response, cancellationToken);
 #pragma warning restore CA2000
 				}
@@ -198,7 +204,7 @@ namespace EpicGames.Horde.Storage
 
 		#region Grouping requests
 
-		record class PacketExport(FlushedPacketHandle PacketHandle, ExportHandle ExportHandle, TContext Context);
+		record class PacketExport(FlushedPacketHandle PacketHandle, IoHash ExportHash, ExportHandle ExportHandle, TContext Context);
 
 		async Task CreateBundleRequestsAsync(CancellationToken cancellationToken)
 		{
@@ -276,7 +282,7 @@ namespace EpicGames.Horde.Storage
 			List<PacketRequest> packetRequests = new List<PacketRequest>();
 			foreach (IGrouping<int, PacketExport> group in exports.GroupBy(x => x.PacketHandle.PacketOffset))
 			{
-				List<ExportRequest> exportRequests = group.Select(x => new ExportRequest(x.ExportHandle, x.Context)).ToList();
+				List<ExportRequest> exportRequests = group.Select(x => new ExportRequest(x.ExportHash, x.ExportHandle, x.Context)).ToList();
 				PacketRequest packetRequest = new PacketRequest(group.First().PacketHandle, exportRequests);
 				packetRequests.Add(packetRequest);
 			}
@@ -303,7 +309,7 @@ namespace EpicGames.Horde.Storage
 		{
 			if (chunk.BlobRef.Innermost is ExportHandle exportHandle && exportHandle.Packet is FlushedPacketHandle packetHandle)
 			{
-				export = new PacketExport(packetHandle, exportHandle, chunk.Context);
+				export = new PacketExport(packetHandle, chunk.BlobRef.Hash, exportHandle, chunk.Context);
 				return true;
 			}
 			else
@@ -339,6 +345,18 @@ namespace EpicGames.Horde.Storage
 			}
 		}
 
+		void VerifyHash(IBlobRef blobRef, ReadOnlySpan<byte> data, IoHash hash)
+		{
+			if (_verifyHashes)
+			{
+				IoHash actualHash = IoHash.Compute(data);
+				if (actualHash != hash)
+				{
+					throw new InvalidDataException($"Invalid hash while reading {blobRef.GetLocator()}: Expected {hash}, got {actualHash}.");
+				}
+			}
+		}
+
 		#endregion
 		#region Packet reads
 
@@ -365,6 +383,7 @@ namespace EpicGames.Horde.Storage
 						{
 							ExportRequest exportRequest = packetRequest.Exports[idx];
 							BlobData blobData = packetReader.ReadExport(exportRequest.Handle.ExportIdx);
+							VerifyHash(exportRequest.Handle, blobData.Data.Span, exportRequest.Hash);
 							items[idx] = new BatchReadResponseItem<TContext>(blobData, exportRequest.Context);
 						}
 
diff --git a/Engine/Source/Programs/Shared/EpicGames.Horde/Storage/Bundles/V2/BundleWriter.cs b/Engine/Source/Programs/Shared/EpicGames.Horde/Storage/Bundles/V2/BundleWriter.cs
--- a/Engine/Source/Programs/Shared/EpicGames.Horde/Storage/Bundles/V2/BundleWriter.cs
+++ b/Engine/Source/Programs/Shared/EpicGames.Horde/Storage/Bundles/V2/BundleWriter.cs
@@ -125,7 +125,7 @@ namespace EpicGames.Horde.Storage.Bundles.V2
 				return _packetWriter.GetExport(exportIdx);
 			}
 
-			public ExportHandle CompleteExport(BlobType type, int size, IReadOnlyList<IBlobRef> imports, IReadOnlyList<AliasInfo> aliases)
+			public HashedExportHandle CompleteExport(BlobType type, int size, IReadOnlyList<IBlobRef> imports, IReadOnlyList<AliasInfo> aliases)
 			{
 				RuntimeAssert(_packetWriter != null);
 				RuntimeAssert(_packetHandle != null);
@@ -133,12 +133,12 @@ namespace EpicGames.Horde.Storage.Bundles.V2
 				IoHash hash = IoHash.Compute(_packetWriter.GetOutputBuffer(size, size).Span);
 
 				int exportIdx = _packetWriter.CompleteExport(size, type, imports);
-				ExportHandle exportHandle = new ExportHandle(hash, _packetHandle, exportIdx);
+				HashedExportHandle exportHandle = new HashedExportHandle(hash, _packetHandle, exportIdx);
 
 				if (aliases.Count > 0)
 				{
 					_pendingExportAliases ??= new List<(ExportHandle, AliasInfo)>();
-					_pendingExportAliases.AddRange(aliases.Select(x => (exportHandle, x)));
+					_pendingExportAliases.AddRange(aliases.Select(x => ((ExportHandle)exportHandle, x)));
 				}
 
 				if (_packetWriter.Length > Math.Min(_bundleOptions.MinCompressionPacketSize, _bundleOptions.MaxBlobSize))
@@ -413,7 +413,7 @@ namespace EpicGames.Horde.Storage.Bundles.V2
 		/// <inheritdoc/>
 		public override async ValueTask<IHashedBlobRef> WriteBlobAsync(BlobType type, int size, IReadOnlyList<IBlobRef> references, IReadOnlyList<AliasInfo> aliases, CancellationToken cancellationToken = default)
 		{
-			ExportHandle exportHandle = _currentBundle.CompleteExport(type, size, references, aliases);
+			HashedExportHandle exportHandle = _currentBundle.CompleteExport(type, size, references, aliases);
 			if (_currentBundle.Length > _bundleOptions.MaxBlobSize)
 			{
 				await WriteCurrentBundleAsync(cancellationToken);
diff --git a/Engine/Source/Programs/Shared/EpicGames.Horde/Storage/Bundles/V2/ExportHandle.cs b/Engine/Source/Programs/Shared/EpicGames.Horde/Storage/Bundles/V2/ExportHandle.cs
--- a/Engine/Source/Programs/Shared/EpicGames.Horde/Storage/Bundles/V2/ExportHandle.cs
+++ b/Engine/Source/Programs/Shared/EpicGames.Horde/Storage/Bundles/V2/ExportHandle.cs
@@ -12,7 +12,7 @@ namespace EpicGames.Horde.Storage.Bundles.V2
 	/// <summary>
 	/// Handle class for blobs exported from a bundle
 	/// </summary>
-	public class ExportHandle : IHashedBlobRef
+	public class ExportHandle : IBlobRef
 	{
 		static readonly Utf8String s_fragmentPrefix = new Utf8String("exp=");
 
@@ -26,22 +26,17 @@ namespace EpicGames.Horde.Storage.Bundles.V2
 		/// </summary>
 		public int ExportIdx => _exportIdx;
 
-		readonly IoHash _hash;
 		readonly PacketHandle _packet;
 		readonly int _exportIdx;
 
 		/// <inheritdoc/>
 		public IBlobRef Innermost => this;
 
-		/// <inheritdoc/>
-		public IoHash Hash => _hash;
-
 		/// <summary>
 		/// Constructor
 		/// </summary>
-		public ExportHandle(IoHash hash, PacketHandle packet, int exportIdx)
+		public ExportHandle(PacketHandle packet, int exportIdx)
 		{
-			_hash = hash;
 			_packet = packet;
 			_exportIdx = exportIdx;
 		}
@@ -144,4 +139,24 @@ namespace EpicGames.Horde.Storage.Bundles.V2
 		public override int GetHashCode()
 			=> HashCode.Combine(Packet, ExportIdx);
 	}
+
+	/// <summary>
+	/// Implementation of <see cref="ExportHandle"/> which also contains a hash
+	/// </summary>
+	public class HashedExportHandle : ExportHandle, IHashedBlobRef
+	{
+		readonly IoHash _hash;
+
+		/// <inheritdoc/>
+		public IoHash Hash => _hash;
+
+		/// <summary>
+		/// Constructor
+		/// </summary>
+		public HashedExportHandle(IoHash hash, PacketHandle packet, int exportIdx)
+			: base(packet, exportIdx)
+		{
+			_hash = hash;
+		}
+	}
 }
diff --git a/Engine/Source/Programs/Shared/EpicGames.Horde/Storage/Nodes/DirectoryNode.Extract.cs b/Engine/Source/Programs/Shared/EpicGames.Horde/Storage/Nodes/DirectoryNode.Extract.cs
--- a/Engine/Source/Programs/Shared/EpicGames.Horde/Storage/Nodes/DirectoryNode.Extract.cs
+++ b/Engine/Source/Programs/Shared/EpicGames.Horde/Storage/Nodes/DirectoryNode.Extract.cs
@@ -104,20 +104,46 @@ namespace EpicGames.Horde.Storage.Nodes
 	}
 
 	/// <summary>
-	/// Extension methods for extracting data from directory nodes
+	/// Options for extracting data
 	/// </summary>
-	public static class DirectoryNodeExtract
+	public class ExtractOptions
 	{
 		/// <summary>
-		/// Utility function to allow extracting a packed directory to disk
+		/// Number of async tasks to spawn for reading
 		/// </summary>
-		/// <param name="directoryRef">Directory to update</param>
-		/// <param name="directoryInfo"></param>
-		/// <param name="logger"></param>
-		/// <param name="cancellationToken"></param>
-		public static Task ExtractAsync(this IBlobRef<DirectoryNode> directoryRef, DirectoryInfo directoryInfo, ILogger logger, CancellationToken cancellationToken)
-			=> ExtractAsync(directoryRef, directoryInfo, null, logger, cancellationToken);
+		public int? NumReadTasks { get; set; }
+
+		/// <summary>
+		/// Number of async tasks to spawn for decoding data
+		/// </summary>
+		public int? NumDecodeTasks { get; set; }
+
+		/// <summary>
+		/// Number of async tasks to spawn for writing output
+		/// </summary>
+		public int? NumWriteTasks { get; set; }
+
+		/// <summary>
+		/// Output for progress updates
+		/// </summary>
+		public IProgress<IExtractStats>? Progress { get; set; }
 
+		/// <summary>
+		/// Frequency that the progress object is updated
+		/// </summary>
+		public TimeSpan ProgressUpdateFrequency { get; set; } = TimeSpan.FromSeconds(5.0);
+
+		/// <summary>
+		/// Whether to hash downloaded data to ensure it's correct
+		/// </summary>
+		public bool VerifyOutput { get; set; }
+	}
+
+	/// <summary>
+	/// Extension methods for extracting data from directory nodes
+	/// </summary>
+	public static class DirectoryNodeExtract
+	{
 		class OutputStats
 		{
 			public int _writtenFiles;
@@ -221,7 +247,7 @@ namespace EpicGames.Horde.Storage.Nodes
 				_chunkWriter = chunkWriter;
 			}
 
-			public async Task WriteAsync(long offset, long length, IBlobRef handle, CancellationToken cancellationToken)
+			public async Task WriteAsync(long offset, long length, IHashedBlobRef handle, CancellationToken cancellationToken)
 			{
 				if (_bufferedChunk != null)
 				{
@@ -249,6 +275,17 @@ namespace EpicGames.Horde.Storage.Nodes
 		}
 #pragma warning restore IDE0060
 
+		/// <summary>
+		/// Utility function to allow extracting a packed directory to disk
+		/// </summary>
+		/// <param name="directoryRef">Directory to update</param>
+		/// <param name="directoryInfo"></param>
+		/// <param name="logger"></param>
+		/// <param name="cancellationToken"></param>
+		[Obsolete("Use an ExtractAsync overload that takes an ExtractOptions object")]
+		public static Task ExtractAsync(this IBlobRef<DirectoryNode> directoryRef, DirectoryInfo directoryInfo, ILogger logger, CancellationToken cancellationToken)
+			=> ExtractAsync(directoryRef, directoryInfo, new ExtractOptions(), logger, cancellationToken);
+
 		/// <summary>
 		/// Utility function to allow extracting a packed directory to disk
 		/// </summary>
@@ -257,9 +294,10 @@ namespace EpicGames.Horde.Storage.Nodes
 		/// <param name="progress">Sink for progress updates</param>
 		/// <param name="logger">Logger for output</param>
 		/// <param name="cancellationToken">Cancellation token for the operation</param>
+		[Obsolete("Use an ExtractAsync overload that takes an ExtractOptions object")]
 		public static Task ExtractAsync(this IBlobRef<DirectoryNode> directoryRef, DirectoryInfo directoryInfo, IProgress<IExtractStats>? progress, ILogger logger, CancellationToken cancellationToken)
 		{
-			return ExtractAsync(directoryRef, directoryInfo, progress, TimeSpan.FromSeconds(5.0), logger, cancellationToken);
+			return ExtractAsync(directoryRef, directoryInfo, new ExtractOptions { Progress = progress }, logger, cancellationToken);
 		}
 
 		/// <summary>
@@ -271,10 +309,11 @@ namespace EpicGames.Horde.Storage.Nodes
 		/// <param name="frequency">Frequency for progress updates</param>
 		/// <param name="logger">Logger for output</param>
 		/// <param name="cancellationToken">Cancellation token for the operation</param>
+		[Obsolete("Use an ExtractAsync overload that takes an ExtractOptions object")]
 		public static async Task ExtractAsync(this IBlobRef<DirectoryNode> directoryRef, DirectoryInfo directoryInfo, IProgress<IExtractStats>? progress, TimeSpan frequency, ILogger logger, CancellationToken cancellationToken)
 		{
 			DirectoryNode directoryNode = await directoryRef.ReadBlobAsync(cancellationToken);
-			await ExtractAsync(directoryNode, directoryInfo, progress, frequency, logger, cancellationToken);
+			await ExtractAsync(directoryNode, directoryInfo, new ExtractOptions { Progress = progress, ProgressUpdateFrequency = frequency }, logger, cancellationToken);
 		}
 
 		/// <summary>
@@ -284,38 +323,54 @@ namespace EpicGames.Horde.Storage.Nodes
 		/// <param name="directoryInfo">Direcotry to write to</param>
 		/// <param name="logger">Logger for output</param>
 		/// <param name="cancellationToken">Cancellation token for the operation</param>
+		[Obsolete("Use an ExtractAsync overload that takes an ExtractOptions object")]
 		public static Task ExtractAsync(this DirectoryNode directoryNode, DirectoryInfo directoryInfo, ILogger logger, CancellationToken cancellationToken)
-			=> ExtractAsync(directoryNode, directoryInfo, null, TimeSpan.FromDays(1.0), logger, cancellationToken);
+			=> ExtractAsync(directoryNode, directoryInfo, new ExtractOptions(), logger, cancellationToken);
+
+		/// <summary>
+		/// Utility function to allow extracting a packed directory to disk
+		/// </summary>
+		/// <param name="directoryRef">Directory to extract</param>
+		/// <param name="directoryInfo">Direcotry to write to</param>
+		/// <param name="options">Options for the download</param>
+		/// <param name="logger">Logger for output</param>
+		/// <param name="cancellationToken">Cancellation token for the operation</param>
+		public static async Task ExtractAsync(this IBlobRef<DirectoryNode> directoryRef, DirectoryInfo directoryInfo, ExtractOptions? options, ILogger logger, CancellationToken cancellationToken)
+		{
+			DirectoryNode directoryNode = await directoryRef.ReadBlobAsync(cancellationToken);
+			await ExtractAsync(directoryNode, directoryInfo, options, logger, cancellationToken);
+		}
 
 		/// <summary>
 		/// Utility function to allow extracting a packed directory to disk
 		/// </summary>
 		/// <param name="directoryNode">Directory to extract</param>
 		/// <param name="directoryInfo">Direcotry to write to</param>
-		/// <param name="progress">Sink for progress updates</param>
-		/// <param name="frequency">Frequency for progress updates</param>
+		/// <param name="options">Options for the download</param>
 		/// <param name="logger">Logger for output</param>
 		/// <param name="cancellationToken">Cancellation token for the operation</param>
-		public static async Task ExtractAsync(this DirectoryNode directoryNode, DirectoryInfo directoryInfo, IProgress<IExtractStats>? progress, TimeSpan frequency, ILogger logger, CancellationToken cancellationToken)
+		public static async Task ExtractAsync(this DirectoryNode directoryNode, DirectoryInfo directoryInfo, ExtractOptions? options, ILogger logger, CancellationToken cancellationToken)
 		{
-			int numReadTasks = 8; //Slowed the read rate to avoid out of memory issues
-			int numDecodeTasks = Math.Min(Environment.ProcessorCount, 16);
-			int numWriteTasks = Math.Min(1 + (int)(directoryNode.Length / (16 * 1024 * 1024)), 16);
-			logger.LogInformation("Using {NumReadTasks} read tasks, {NumDecodeTasks} decode tasks, {NumWriteTasks} write tasks", numReadTasks, numDecodeTasks, numWriteTasks);
+			options ??= new ExtractOptions();
+
+			int numReadTasks = options.NumReadTasks ?? 8; //Slowed the read rate to avoid out of memory issues
+			int numDecodeTasks = options.NumDecodeTasks ?? Math.Min(Environment.ProcessorCount, 16);
+			int numWriteTasks = options.NumWriteTasks ?? Math.Min(1 + (int)(directoryNode.Length / (16 * 1024 * 1024)), 16);
+			logger.LogDebug("Using {NumReadTasks} read tasks, {NumDecodeTasks} decode tasks, {NumWriteTasks} write tasks", numReadTasks, numDecodeTasks, numWriteTasks);
 
 			await using (AsyncPipeline pipeline = new AsyncPipeline(cancellationToken))
 			{
-				using BatchReader<OutputChunk> batchReader = new BatchReader<OutputChunk>();
+				using BatchReader<OutputChunk> batchReader = new BatchReader<OutputChunk>(options.VerifyOutput);
 
 				_ = pipeline.AddTask(ctx => FindOutputChunksRootAsync(directoryInfo, directoryNode, batchReader.RequestWriter, logger, ctx));
 				pipeline.AddBatchReaderTasks(numReadTasks, numDecodeTasks, batchReader);
 
 				OutputStats outputStats = new OutputStats();
-				Task[] writeTasks = pipeline.AddTasks(numWriteTasks, ctx => WriteAsync(batchReader.ResponseReader, outputStats, logger, ctx));
+				Task[] writeTasks = pipeline.AddTasks(numWriteTasks, ctx => WriteAsync(batchReader.ResponseReader, outputStats, options, logger, ctx));
 
-				if (progress != null)
+				if (options.Progress != null)
 				{
-					_ = pipeline.AddTask(ctx => UpdateStatsAsync(batchReader, outputStats, Task.WhenAll(writeTasks), progress, frequency, ctx));
+					_ = pipeline.AddTask(ctx => UpdateStatsAsync(batchReader, outputStats, Task.WhenAll(writeTasks), options.Progress, options.ProgressUpdateFrequency, ctx));
 				}
 
 				await pipeline.WaitForCompletionAsync();
@@ -402,7 +457,7 @@ namespace EpicGames.Horde.Storage.Nodes
 
 		#region Write to disk
 
-		static async Task WriteAsync(ChannelReader<BatchReadResponse<OutputChunk>> batchReader, OutputStats outputStats, ILogger logger, CancellationToken cancellationToken)
+		static async Task WriteAsync(ChannelReader<BatchReadResponse<OutputChunk>> batchReader, OutputStats outputStats, ExtractOptions options, ILogger logger, CancellationToken cancellationToken)
 		{
 			const int WriteBatchSize = 64;
 			while (await batchReader.WaitToReadAsync(cancellationToken))
@@ -415,7 +470,7 @@ namespace EpicGames.Horde.Storage.Nodes
 						List<Task> tasks = new List<Task>();
 						foreach (IReadOnlyList<BatchReadResponseItem<OutputChunk>> group in batch.Items.Batch(WriteBatchSize))
 						{
-							tasks.Add(WriteChunksAsync(group.ToArray(), outputStats, logger, cancellationToken));
+							tasks.Add(WriteChunksAsync(group.ToArray(), outputStats, options, logger, cancellationToken));
 						}
 						await Task.WhenAll(tasks);
 					}
@@ -427,7 +482,7 @@ namespace EpicGames.Horde.Storage.Nodes
 			}
 		}
 
-		static async Task WriteChunksAsync(ArraySegment<BatchReadResponseItem<OutputChunk>> chunks, OutputStats outputStats, ILogger logger, CancellationToken cancellationToken)
+		static async Task WriteChunksAsync(ArraySegment<BatchReadResponseItem<OutputChunk>> chunks, OutputStats outputStats, ExtractOptions options, ILogger logger, CancellationToken cancellationToken)
 		{
 			for (int chunkIdx = 0; chunkIdx < chunks.Count;)
 			{
@@ -441,7 +496,7 @@ namespace EpicGames.Horde.Storage.Nodes
 
 				try
 				{
-					await ExtractChunksToFileAsync(file, chunks.Slice(chunkIdx, maxChunkIdx - chunkIdx), outputStats, logger, cancellationToken);
+					await ExtractChunksToFileAsync(file, chunks.Slice(chunkIdx, maxChunkIdx - chunkIdx), outputStats, options, logger, cancellationToken);
 					//await ExtractChunksToNullAsync(file, chunks.Slice(chunkIdx, maxChunkIdx - chunkIdx), outputStats, logger, cancellationToken);
 				}
 				catch (OperationCanceledException)
@@ -457,7 +512,7 @@ namespace EpicGames.Horde.Storage.Nodes
 			}
 		}
 
-		static async Task ExtractChunksToFileAsync(OutputFile file, ArraySegment<BatchReadResponseItem<OutputChunk>> chunks, OutputStats stats, ILogger logger, CancellationToken cancellationToken)
+		static async Task ExtractChunksToFileAsync(OutputFile file, ArraySegment<BatchReadResponseItem<OutputChunk>> chunks, OutputStats stats, ExtractOptions options, ILogger logger, CancellationToken cancellationToken)
 		{
 			// Open the file for the current chunk
 			int remainingChunks = 0;
@@ -496,6 +551,20 @@ namespace EpicGames.Horde.Storage.Nodes
 			// Set correct permissions on the output file
 			if (remainingChunks == 0)
 			{
+				if (options.VerifyOutput)
+				{
+					IoHash computedHash;
+					using (Stream stream = file.OpenStream())
+					{
+						computedHash = await IoHash.ComputeAsync(stream, cancellationToken);
+					}
+					if (computedHash != file.FileEntry.StreamHash)
+					{
+						throw new InvalidDataException($"File {file.Path} does not have correct hash after extract; expected {file.FileEntry.StreamHash}, got {computedHash}.");
+					}
+					logger.LogDebug("Verified hash of {Path}: {Hash}", file.Path, computedHash);
+				}
+
 				file.FileInfo.Refresh();
 				FileEntry.SetPermissions(file.FileInfo!, file.FileEntry.Flags);
 
@@ -560,35 +629,46 @@ namespace EpicGames.Horde.Storage.Nodes
 		static async Task UpdateStatsAsync(BatchReader<OutputChunk> batchReader, OutputStats outputStats, Task writeTask, IProgress<IExtractStats> progress, TimeSpan frequency, CancellationToken cancellationToken)
 		{
 			Stopwatch timer = Stopwatch.StartNew();
+			double lastOutputTime = 0.0;
 
-			long lastExtractSize = 0;
-			long lastDownloadSize = 0;
+			// Store previous download progress so we can sample the download rate over time
+			double rateAverageOverTime = 5.0;
+			Queue<(double, ExtractStats)> samples = new Queue<(double, ExtractStats)>();
+			samples.Enqueue((0.0, new ExtractStats()));
 
 			Task? completeTask = null;
 			while (completeTask != writeTask)
 			{
-				completeTask = await Task.WhenAny(writeTask, Task.Delay(frequency, cancellationToken));
+				completeTask = await Task.WhenAny(writeTask, Task.Delay(frequency / 4.0, cancellationToken));
+
+				double timeSeconds = timer.Elapsed.TotalSeconds;
 
 				ExtractStats stats = new ExtractStats();
 				stats.NumFiles = Interlocked.CompareExchange(ref outputStats._writtenFiles, 0, 0);
 				stats.ExtractSize = Interlocked.CompareExchange(ref outputStats._writtenBytes, 0, 0);
+				samples.Enqueue((timeSeconds, stats));
 
 				BatchReaderStats batchStats = batchReader.GetStats();
 				stats.DownloadSize = batchStats.BytesRead;
 
-				double elapsedSeconds = timer.Elapsed.TotalSeconds;
-				if (elapsedSeconds > 1.0)
+				(double TimeSeconds, ExtractStats Stats) firstSample;
+				while (samples.TryPeek(out firstSample) && samples.Count > 2 && timeSeconds > firstSample.TimeSeconds + rateAverageOverTime)
 				{
-					stats.ExtractRate = (stats.ExtractSize - lastExtractSize) / elapsedSeconds;
-					stats.DownloadRate = (stats.DownloadSize - lastDownloadSize) / elapsedSeconds;
-
-					lastExtractSize = stats.ExtractSize;
-					lastDownloadSize = stats.DownloadSize;
+					samples.Dequeue();
+				}
 
-					timer.Restart();
+				double elapsedSeconds = timeSeconds - firstSample.TimeSeconds;
+				if (elapsedSeconds > 1.0)
+				{
+					stats.ExtractRate = (stats.ExtractSize - firstSample.Stats.ExtractSize) / elapsedSeconds;
+					stats.DownloadRate = (stats.DownloadSize - firstSample.Stats.DownloadSize) / elapsedSeconds;
 				}
 
-				progress.Report(stats);
+				if (timeSeconds > lastOutputTime + frequency.TotalSeconds)
+				{
+					progress.Report(stats);
+					lastOutputTime = timeSeconds;
+				}
 			}
 		}
 
diff --git a/Engine/Source/Programs/UnrealGameSync/UnrealGameSync/ArtifactDownload.cs b/Engine/Source/Programs/UnrealGameSync/UnrealGameSync/ArtifactDownload.cs
--- a/Engine/Source/Programs/UnrealGameSync/UnrealGameSync/ArtifactDownload.cs
+++ b/Engine/Source/Programs/UnrealGameSync/UnrealGameSync/ArtifactDownload.cs
@@ -68,7 +68,11 @@ namespace UnrealGameSync
 				progress.Report("Starting...");
 				ILogger logger = serviceProvider.GetRequiredService<ILogger<ArtifactDownload>>();
 
-				await blobRef.ExtractAsync(outputDir.ToDirectoryInfo(), new CopyProgressAdapter(progress), TimeSpan.FromSeconds(0.2), serviceProvider.GetRequiredService<ILogger<ArtifactDownload>>(), cancellationToken);
+				ExtractOptions options = new ExtractOptions();
+				options.Progress = new CopyProgressAdapter(progress);
+				options.ProgressUpdateFrequency = TimeSpan.FromSeconds(0.2);
+
+				await blobRef.ExtractAsync(outputDir.ToDirectoryInfo(), options, serviceProvider.GetRequiredService<ILogger<ArtifactDownload>>(), cancellationToken);
 
 				logger.LogInformation("Download complete");
 			}
diff --git a/Engine/Source/Programs/UnrealGameSync/UnrealGameSync/Launcher.cs b/Engine/Source/Programs/UnrealGameSync/UnrealGameSync/Launcher.cs
--- a/Engine/Source/Programs/UnrealGameSync/UnrealGameSync/Launcher.cs
+++ b/Engine/Source/Programs/UnrealGameSync/UnrealGameSync/Launcher.cs
@@ -190,7 +190,7 @@ namespace UnrealGameSync
 						IBlobRef<DirectoryNode> handle = storageNamespace.CreateBlobRef<DirectoryNode>(deployment.Locator);
 
 						DirectoryNode rootNode = await handle.ReadBlobAsync(cancellationToken: cancellationToken);
-						await rootNode.ExtractAsync(new DirectoryInfo(applicationFolder), new ExtractStatsLogger(0, rootNode.Length, logger), TimeSpan.FromSeconds(5.0), logger, cancellationToken);
+						await rootNode.ExtractAsync(new DirectoryInfo(applicationFolder), new ExtractOptions { Progress = new ExtractStatsLogger(0, rootNode.Length, logger) }, logger, cancellationToken);
 
 						// Update the version
 						if (!TryWriteAllText(syncVersionFile, requiredSyncText))
diff --git a/Engine/Source/Programs/UnrealToolbox/Plugins/Artifacts/ArtifactsPlugin.cs b/Engine/Source/Programs/UnrealToolbox/Plugins/Artifacts/ArtifactsPlugin.cs
--- a/Engine/Source/Programs/UnrealToolbox/Plugins/Artifacts/ArtifactsPlugin.cs
+++ b/Engine/Source/Programs/UnrealToolbox/Plugins/Artifacts/ArtifactsPlugin.cs
@@ -102,7 +102,7 @@ namespace UnrealToolbox.Plugins.Artifacts
 					progress.SetTotalSize(node.Length);
 
 					progress.Status = "Starting...";
-					await blobRef.ExtractAsync(options.OutputDir.ToDirectoryInfo(), progress, TimeSpan.FromSeconds(0.2), _logger, linkedCancellationToken);
+					await blobRef.ExtractAsync(options.OutputDir.ToDirectoryInfo(), new ExtractOptions { Progress = progress, ProgressUpdateFrequency = TimeSpan.FromSeconds(0.2) }, _logger, linkedCancellationToken);
 
 					progress.Status = "Download complete.";
 				}
diff --git a/Engine/Source/Programs/UnrealToolbox/SelfUpdateService.cs b/Engine/Source/Programs/UnrealToolbox/SelfUpdateService.cs
--- a/Engine/Source/Programs/UnrealToolbox/SelfUpdateService.cs
+++ b/Engine/Source/Programs/UnrealToolbox/SelfUpdateService.cs
@@ -315,7 +315,7 @@ namespace UnrealToolbox
 
 				FileUtils.ForceDeleteDirectoryContents(updateDir);
 
-				await deployment.Content.ExtractAsync(updateDir.ToDirectoryInfo(), _logger, cancellationToken);
+				await deployment.Content.ExtractAsync(updateDir.ToDirectoryInfo(), new ExtractOptions(), _logger, cancellationToken);
 
 				selfUpdate.UpdateVersion = updateVersion;
 				_logger.LogInformation("Extracted {UpdateVersion} to {UpdateDir}", updateVersion, updateDir);
diff --git a/Engine/Source/Programs/UnrealToolbox/ToolCatalog.cs b/Engine/Source/Programs/UnrealToolbox/ToolCatalog.cs
--- a/Engine/Source/Programs/UnrealToolbox/ToolCatalog.cs
+++ b/Engine/Source/Programs/UnrealToolbox/ToolCatalog.cs
@@ -481,7 +481,7 @@ namespace UnrealToolbox
 
 					DirectoryReference toolDir = GetToolDir(item.Id, deploymentInfo.Id);
 					DirectoryReference.CreateDirectory(toolDir);
-					await deployment.Content.ExtractAsync(toolDir.ToDirectoryInfo(), _logger, cancellationToken);
+					await deployment.Content.ExtractAsync(toolDir.ToDirectoryInfo(), new ExtractOptions(), _logger, cancellationToken);
 
 					await InstallAsync(item, toolDir, deploymentInfo, cancellationToken);
 				}
