Horde: Move job template interfaces and messages into EpicGames.Horde.

#rnx

diff --git a/Engine/Source/Programs/Horde/Plugins/Build/HordeServer.Build.Tests/BuildTestSetup.cs b/Engine/Source/Programs/Horde/Plugins/Build/HordeServer.Build.Tests/BuildTestSetup.cs
--- a/Engine/Source/Programs/Horde/Plugins/Build/HordeServer.Build.Tests/BuildTestSetup.cs
+++ b/Engine/Source/Programs/Horde/Plugins/Build/HordeServer.Build.Tests/BuildTestSetup.cs
@@ -7,6 +7,7 @@ using Amazon.EC2;
 using EpicGames.Horde.Agents;
 using EpicGames.Horde.Agents.Pools;
 using EpicGames.Horde.Artifacts;
+using EpicGames.Horde.Jobs.Templates;
 using EpicGames.Horde.Tools;
 using HordeCommon.Rpc.Messages;
 using HordeServer.Acls;
@@ -81,7 +82,7 @@ namespace HordeServer.Tests
 		public AgentRelayService AgentRelayService => ServiceProvider.GetRequiredService<AgentRelayService>();
 		public ICommitService CommitService => ServiceProvider.GetRequiredService<ICommitService>();
 		public GlobalsService GlobalsService => ServiceProvider.GetRequiredService<GlobalsService>();
-		public ITemplateCollection TemplateCollection => ServiceProvider.GetRequiredService<ITemplateCollection>();
+		public ITemplateCollectionInternal TemplateCollection => ServiceProvider.GetRequiredService<ITemplateCollectionInternal>();
 		internal PerforceServiceStub PerforceService => (PerforceServiceStub)ServiceProvider.GetRequiredService<IPerforceService>();
 		public ISubscriptionCollection SubscriptionCollection => ServiceProvider.GetRequiredService<ISubscriptionCollection>();
 		public INotificationService NotificationService => ServiceProvider.GetRequiredService<INotificationService>();
diff --git a/Engine/Source/Programs/Horde/Plugins/Build/HordeServer.Build.Tests/Devices/DeviceServiceTests.cs b/Engine/Source/Programs/Horde/Plugins/Build/HordeServer.Build.Tests/Devices/DeviceServiceTests.cs
--- a/Engine/Source/Programs/Horde/Plugins/Build/HordeServer.Build.Tests/Devices/DeviceServiceTests.cs
+++ b/Engine/Source/Programs/Horde/Plugins/Build/HordeServer.Build.Tests/Devices/DeviceServiceTests.cs
@@ -11,7 +11,6 @@ using EpicGames.Horde.Streams;
 using HordeServer.Devices;
 using HordeServer.Jobs;
 using HordeServer.Jobs.Graphs;
-using HordeServer.Jobs.Templates;
 using HordeServer.Users;
 using HordeServer.Utilities;
 using Microsoft.AspNetCore.Http;
@@ -751,7 +750,7 @@ namespace HordeServer.Tests.Devices
 			List<IDevicePool> pools = DeviceService!.GetPools();
 			List<IDevicePlatform> platforms = DeviceService!.GetPlatforms();
 			List<IDevice> devices = (await DeviceService!.GetDevicesAsync()).OrderBy(d => d.PoolId.ToString()).ThenBy(d => d.PlatformId.ToString()).ToList();
-			List<IDeviceTelemetry>  telemetry = new List<IDeviceTelemetry>();
+			List<IDeviceTelemetry> telemetry = new List<IDeviceTelemetry>();
 
 			DateTime referencePoint = DateTime.Now - TimeSpan.FromDays(1);
 			AddLoadSpikes(telemetry, devices, pools[0].Id, platforms[0].Id, ref referencePoint, 3, TimeSpan.FromMinutes(3));
diff --git a/Engine/Source/Programs/Horde/Plugins/Build/HordeServer.Build.Tests/Fixture.cs b/Engine/Source/Programs/Horde/Plugins/Build/HordeServer.Build.Tests/Fixture.cs
--- a/Engine/Source/Programs/Horde/Plugins/Build/HordeServer.Build.Tests/Fixture.cs
+++ b/Engine/Source/Programs/Horde/Plugins/Build/HordeServer.Build.Tests/Fixture.cs
@@ -35,7 +35,7 @@ namespace HordeServer.Tests
 		public string Agent1Name { get; private set; } = null!;
 		public const string PoolName = "TestingPool";
 
-		public static async Task<Fixture> CreateAsync(ConfigService configService, IGraphCollection graphCollection, ITemplateCollection templateCollection, JobService jobService, AgentService agentService, IPluginCollection pluginCollection, ServerSettings serverSettings)
+		public static async Task<Fixture> CreateAsync(ConfigService configService, IGraphCollection graphCollection, ITemplateCollectionInternal templateCollection, JobService jobService, AgentService agentService, IPluginCollection pluginCollection, ServerSettings serverSettings)
 		{
 			Fixture fixture = new Fixture();
 			await fixture.PopulateAsync(configService, graphCollection, templateCollection, jobService, agentService, pluginCollection, serverSettings);
@@ -46,7 +46,7 @@ namespace HordeServer.Tests
 			return fixture;
 		}
 
-		private async Task PopulateAsync(ConfigService configService, IGraphCollection graphCollection, ITemplateCollection templateCollection, JobService jobService, AgentService agentService, IPluginCollection pluginCollection, ServerSettings serverSettings)
+		private async Task PopulateAsync(ConfigService configService, IGraphCollection graphCollection, ITemplateCollectionInternal templateCollection, JobService jobService, AgentService agentService, IPluginCollection pluginCollection, ServerSettings serverSettings)
 		{
 			GraphStub fg = new (ContentHash.Empty, 1122);
 			Template = await templateCollection.GetOrAddAsync(new TemplateConfig { Name = "Test template" });
diff --git a/Engine/Source/Programs/Horde/Plugins/Build/HordeServer.Build.Tests/Jobs/Bisect/BisectTests.cs b/Engine/Source/Programs/Horde/Plugins/Build/HordeServer.Build.Tests/Jobs/Bisect/BisectTests.cs
--- a/Engine/Source/Programs/Horde/Plugins/Build/HordeServer.Build.Tests/Jobs/Bisect/BisectTests.cs
+++ b/Engine/Source/Programs/Horde/Plugins/Build/HordeServer.Build.Tests/Jobs/Bisect/BisectTests.cs
@@ -10,7 +10,6 @@ using EpicGames.Horde.Streams;
 using HordeServer.Jobs;
 using HordeServer.Jobs.Bisect;
 using HordeServer.Jobs.Graphs;
-using HordeServer.Jobs.Templates;
 using HordeServer.Logs;
 using HordeServer.Projects;
 using HordeServer.Streams;
diff --git a/Engine/Source/Programs/Horde/Plugins/Build/HordeServer.Build.Tests/Jobs/JobCollectionTests.cs b/Engine/Source/Programs/Horde/Plugins/Build/HordeServer.Build.Tests/Jobs/JobCollectionTests.cs
--- a/Engine/Source/Programs/Horde/Plugins/Build/HordeServer.Build.Tests/Jobs/JobCollectionTests.cs
+++ b/Engine/Source/Programs/Horde/Plugins/Build/HordeServer.Build.Tests/Jobs/JobCollectionTests.cs
@@ -17,7 +17,6 @@ using EpicGames.Horde.Users;
 using HordeServer.Agents.Sessions;
 using HordeServer.Jobs;
 using HordeServer.Jobs.Graphs;
-using HordeServer.Jobs.Templates;
 using HordeServer.Logs;
 using HordeServer.Projects;
 using HordeServer.Streams;
diff --git a/Engine/Source/Programs/Horde/Plugins/Build/HordeServer.Build.Tests/Jobs/JobExpirationTests.cs b/Engine/Source/Programs/Horde/Plugins/Build/HordeServer.Build.Tests/Jobs/JobExpirationTests.cs
--- a/Engine/Source/Programs/Horde/Plugins/Build/HordeServer.Build.Tests/Jobs/JobExpirationTests.cs
+++ b/Engine/Source/Programs/Horde/Plugins/Build/HordeServer.Build.Tests/Jobs/JobExpirationTests.cs
@@ -6,7 +6,6 @@ using EpicGames.Horde.Projects;
 using EpicGames.Horde.Streams;
 using HordeServer.Jobs;
 using HordeServer.Jobs.Graphs;
-using HordeServer.Jobs.Templates;
 using HordeServer.Projects;
 using HordeServer.Server;
 using HordeServer.Streams;
diff --git a/Engine/Source/Programs/Horde/Plugins/Build/HordeServer.Build/BuildPlugin.cs b/Engine/Source/Programs/Horde/Plugins/Build/HordeServer.Build/BuildPlugin.cs
--- a/Engine/Source/Programs/Horde/Plugins/Build/HordeServer.Build/BuildPlugin.cs
+++ b/Engine/Source/Programs/Horde/Plugins/Build/HordeServer.Build/BuildPlugin.cs
@@ -3,6 +3,7 @@
 using System.Diagnostics.CodeAnalysis;
 using EpicGames.Core;
 using EpicGames.Horde.Artifacts;
+using EpicGames.Horde.Jobs.Templates;
 using HordeServer.Acls;
 using HordeServer.Agents;
 using HordeServer.Agents.Fleet;
@@ -84,7 +85,9 @@ namespace HordeServer
 			services.AddSingleton<IUgsMetadataCollection, UgsMetadataCollection>();
 			services.AddSingleton<ISubscriptionCollection, SubscriptionCollection>();
 			services.AddSingleton<IStreamCollection, StreamCollection>();
-			services.AddSingleton<ITemplateCollection, TemplateCollection>();
+			services.AddSingleton<TemplateCollection>();
+			services.AddSingleton<ITemplateCollection>(sp => sp.GetRequiredService<TemplateCollection>());
+			services.AddSingleton<ITemplateCollectionInternal>(sp => sp.GetRequiredService<TemplateCollection>());
 			services.AddSingleton<ITestDataCollection, TestDataCollection>();
 			services.AddSingleton<IUtilizationDataCollection, UtilizationDataCollection>();
 			services.AddSingleton<ITemplateCollection, TemplateCollection>();
diff --git a/Engine/Source/Programs/Horde/Plugins/Build/HordeServer.Build/Commits/ICommitCollection.cs b/Engine/Source/Programs/Horde/Plugins/Build/HordeServer.Build/Commits/ICommitCollection.cs
--- a/Engine/Source/Programs/Horde/Plugins/Build/HordeServer.Build/Commits/ICommitCollection.cs
+++ b/Engine/Source/Programs/Horde/Plugins/Build/HordeServer.Build/Commits/ICommitCollection.cs
@@ -1,8 +1,8 @@
 ﻿// Copyright Epic Games, Inc. All Rights Reserved.
 
 using EpicGames.Horde.Commits;
+using EpicGames.Horde.Jobs.Templates;
 using EpicGames.Perforce;
-using HordeServer.Jobs.Templates;
 
 namespace HordeServer.Commits
 {
diff --git a/Engine/Source/Programs/Horde/Plugins/Build/HordeServer.Build/Jobs/Bisect/BisectService.cs b/Engine/Source/Programs/Horde/Plugins/Build/HordeServer.Build/Jobs/Bisect/BisectService.cs
--- a/Engine/Source/Programs/Horde/Plugins/Build/HordeServer.Build/Jobs/Bisect/BisectService.cs
+++ b/Engine/Source/Programs/Horde/Plugins/Build/HordeServer.Build/Jobs/Bisect/BisectService.cs
@@ -2,6 +2,7 @@
 
 using EpicGames.Core;
 using EpicGames.Horde.Jobs;
+using EpicGames.Horde.Jobs.Templates;
 using HordeServer.Commits;
 using HordeServer.Jobs.Graphs;
 using HordeServer.Jobs.Templates;
@@ -18,14 +19,14 @@ namespace HordeServer.Jobs.Bisect
 		readonly IBisectTaskCollection _bisectTaskCollection;
 		readonly IJobCollection _jobCollection;
 		readonly IJobStepRefCollection _jobStepRefCollection;
-		readonly ITemplateCollection _templateCollection;
+		readonly ITemplateCollectionInternal _templateCollection;
 		readonly IGraphCollection _graphCollection;
 		readonly ICommitService _commitService;
 		readonly IOptionsMonitor<BuildConfig> _buildConfig;
 		readonly ILogger _logger;
 		readonly ITicker _ticker;
 
-		public BisectService(IBisectTaskCollection bisectTaskCollection, IJobCollection jobCollection, IJobStepRefCollection jobStepRefCollection, ITemplateCollection templateCollection, IGraphCollection graphCollection, ICommitService commitService, IOptionsMonitor<BuildConfig> buildConfig, IClock clock, ILogger<BisectService> logger)
+		public BisectService(IBisectTaskCollection bisectTaskCollection, IJobCollection jobCollection, IJobStepRefCollection jobStepRefCollection, ITemplateCollectionInternal templateCollection, IGraphCollection graphCollection, ICommitService commitService, IOptionsMonitor<BuildConfig> buildConfig, IClock clock, ILogger<BisectService> logger)
 		{
 			_bisectTaskCollection = bisectTaskCollection;
 			_jobCollection = jobCollection;
diff --git a/Engine/Source/Programs/Horde/Plugins/Build/HordeServer.Build/Jobs/Graphs/GraphCollection.cs b/Engine/Source/Programs/Horde/Plugins/Build/HordeServer.Build/Jobs/Graphs/GraphCollection.cs
--- a/Engine/Source/Programs/Horde/Plugins/Build/HordeServer.Build/Jobs/Graphs/GraphCollection.cs
+++ b/Engine/Source/Programs/Horde/Plugins/Build/HordeServer.Build/Jobs/Graphs/GraphCollection.cs
@@ -4,8 +4,8 @@ using System.Diagnostics;
 using EpicGames.Core;
 using EpicGames.Horde.Artifacts;
 using EpicGames.Horde.Jobs;
+using EpicGames.Horde.Jobs.Templates;
 using HordeCommon;
-using HordeServer.Jobs.Templates;
 using HordeServer.Server;
 using Microsoft.Extensions.Caching.Memory;
 using MongoDB.Bson;
diff --git a/Engine/Source/Programs/Horde/Plugins/Build/HordeServer.Build/Jobs/Graphs/IGraphCollection.cs b/Engine/Source/Programs/Horde/Plugins/Build/HordeServer.Build/Jobs/Graphs/IGraphCollection.cs
--- a/Engine/Source/Programs/Horde/Plugins/Build/HordeServer.Build/Jobs/Graphs/IGraphCollection.cs
+++ b/Engine/Source/Programs/Horde/Plugins/Build/HordeServer.Build/Jobs/Graphs/IGraphCollection.cs
@@ -1,7 +1,7 @@
 // Copyright Epic Games, Inc. All Rights Reserved.
 
 using EpicGames.Core;
-using HordeServer.Jobs.Templates;
+using EpicGames.Horde.Jobs.Templates;
 
 namespace HordeServer.Jobs.Graphs
 {
diff --git a/Engine/Source/Programs/Horde/Plugins/Build/HordeServer.Build/Jobs/JobRpcService.cs b/Engine/Source/Programs/Horde/Plugins/Build/HordeServer.Build/Jobs/JobRpcService.cs
--- a/Engine/Source/Programs/Horde/Plugins/Build/HordeServer.Build/Jobs/JobRpcService.cs
+++ b/Engine/Source/Programs/Horde/Plugins/Build/HordeServer.Build/Jobs/JobRpcService.cs
@@ -11,6 +11,7 @@ using EpicGames.Horde.Agents.Sessions;
 using EpicGames.Horde.Artifacts;
 using EpicGames.Horde.Commits;
 using EpicGames.Horde.Jobs;
+using EpicGames.Horde.Jobs.Templates;
 using EpicGames.Horde.Logs;
 using Google.Protobuf.WellKnownTypes;
 using Grpc.Core;
@@ -20,7 +21,6 @@ using HordeServer.Agents;
 using HordeServer.Agents.Pools;
 using HordeServer.Artifacts;
 using HordeServer.Jobs.Graphs;
-using HordeServer.Jobs.Templates;
 using HordeServer.Jobs.TestData;
 using HordeServer.Logs;
 using HordeServer.Perforce;
@@ -122,8 +122,8 @@ namespace HordeServer.Jobs
 							{
 								AgentWorkspaceInfo? prevWorkspace = currentWorkspaces.FirstOrDefault(x => x.Identifier.Equals(identifier, StringComparison.OrdinalIgnoreCase));
 								AgentWorkspaceInfo? nextWorkspace = conformWorkspaces.FirstOrDefault(x => x.Identifier.Equals(identifier, StringComparison.OrdinalIgnoreCase));
-								string prevWorkspaceText = (prevWorkspace == null)? "(does not exist)" : JsonSerializer.Serialize(prevWorkspace);
-								string nextWorkspaceText = (nextWorkspace == null)? "(does not exist)" : JsonSerializer.Serialize(nextWorkspace);
+								string prevWorkspaceText = (prevWorkspace == null) ? "(does not exist)" : JsonSerializer.Serialize(prevWorkspace);
+								string nextWorkspaceText = (nextWorkspace == null) ? "(does not exist)" : JsonSerializer.Serialize(nextWorkspace);
 
 								if (!String.Equals(prevWorkspaceText, nextWorkspaceText, StringComparison.Ordinal))
 								{
diff --git a/Engine/Source/Programs/Horde/Plugins/Build/HordeServer.Build/Jobs/JobService.cs b/Engine/Source/Programs/Horde/Plugins/Build/HordeServer.Build/Jobs/JobService.cs
--- a/Engine/Source/Programs/Horde/Plugins/Build/HordeServer.Build/Jobs/JobService.cs
+++ b/Engine/Source/Programs/Horde/Plugins/Build/HordeServer.Build/Jobs/JobService.cs
@@ -79,7 +79,7 @@ namespace HordeServer.Jobs
 		readonly INotificationTriggerCollection _triggerCollection;
 		readonly JobTaskSource _jobTaskSource;
 		readonly IStreamCollection _streamCollection;
-		readonly ITemplateCollection _templateCollection;
+		readonly ITemplateCollectionInternal _templateCollection;
 		readonly IPerforceService _perforceService;
 		readonly IOptionsMonitor<BuildConfig> _buildConfig;
 		readonly Tracer _tracer;
@@ -115,7 +115,7 @@ namespace HordeServer.Jobs
 		/// <summary>
 		/// Constructor
 		/// </summary>
-		public JobService(IJobCollection jobs, IGraphCollection graphs, IAgentCollection agents, IJobStepRefCollection jobStepRefs, IBisectTaskCollection bisectTasks, IJobTimingCollection jobTimings, IUserCollection userCollection, INotificationTriggerCollection triggerCollection, JobTaskSource jobTaskSource, IStreamCollection streamCollection, ITemplateCollection templateCollection, IPerforceService perforceService, Tracer tracer, IOptionsMonitor<BuildConfig> buildConfig, ILogger<JobService> logger)
+		public JobService(IJobCollection jobs, IGraphCollection graphs, IAgentCollection agents, IJobStepRefCollection jobStepRefs, IBisectTaskCollection bisectTasks, IJobTimingCollection jobTimings, IUserCollection userCollection, INotificationTriggerCollection triggerCollection, JobTaskSource jobTaskSource, IStreamCollection streamCollection, ITemplateCollectionInternal templateCollection, IPerforceService perforceService, Tracer tracer, IOptionsMonitor<BuildConfig> buildConfig, ILogger<JobService> logger)
 		{
 			_jobs = jobs;
 			_graphs = graphs;
diff --git a/Engine/Source/Programs/Horde/Plugins/Build/HordeServer.Build/Jobs/JobsController.cs b/Engine/Source/Programs/Horde/Plugins/Build/HordeServer.Build/Jobs/JobsController.cs
--- a/Engine/Source/Programs/Horde/Plugins/Build/HordeServer.Build/Jobs/JobsController.cs
+++ b/Engine/Source/Programs/Horde/Plugins/Build/HordeServer.Build/Jobs/JobsController.cs
@@ -41,7 +41,7 @@ namespace HordeServer.Jobs
 		private readonly ICommitService _commitService;
 		private readonly IPerforceService _perforce;
 		private readonly JobService _jobService;
-		private readonly ITemplateCollection _templateCollection;
+		private readonly ITemplateCollectionInternal _templateCollection;
 		private readonly IArtifactCollection _artifactCollection;
 		private readonly IUserCollection _userCollection;
 		private readonly INotificationService _notificationService;
@@ -52,7 +52,7 @@ namespace HordeServer.Jobs
 		/// <summary>
 		/// Constructor
 		/// </summary>
-		public JobsController(IGraphCollection graphs, ICommitService commitService, IPerforceService perforce, JobService jobService, ITemplateCollection templateCollection, IArtifactCollection artifactCollection, IUserCollection userCollection, INotificationService notificationService, AgentService agentService, IOptionsSnapshot<BuildConfig> buildConfig, ILogger<JobsController> logger)
+		public JobsController(IGraphCollection graphs, ICommitService commitService, IPerforceService perforce, JobService jobService, ITemplateCollectionInternal templateCollection, IArtifactCollection artifactCollection, IUserCollection userCollection, INotificationService notificationService, AgentService agentService, IOptionsSnapshot<BuildConfig> buildConfig, ILogger<JobsController> logger)
 		{
 			_graphs = graphs;
 			_commitService = commitService;
diff --git a/Engine/Source/Programs/Horde/Plugins/Build/HordeServer.Build/Jobs/Schedules/ScheduleService.cs b/Engine/Source/Programs/Horde/Plugins/Build/HordeServer.Build/Jobs/Schedules/ScheduleService.cs
--- a/Engine/Source/Programs/Horde/Plugins/Build/HordeServer.Build/Jobs/Schedules/ScheduleService.cs
+++ b/Engine/Source/Programs/Horde/Plugins/Build/HordeServer.Build/Jobs/Schedules/ScheduleService.cs
@@ -70,7 +70,7 @@ namespace HordeServer.Jobs.Schedules
 		readonly IDowntimeService _downtimeService;
 		readonly JobService _jobService;
 		readonly IStreamCollection _streamCollection;
-		readonly ITemplateCollection _templateCollection;
+		readonly ITemplateCollectionInternal _templateCollection;
 		readonly IClock _clock;
 		readonly Tracer _tracer;
 		readonly ILogger _logger;
@@ -84,7 +84,7 @@ namespace HordeServer.Jobs.Schedules
 		/// <summary>
 		/// Constructor
 		/// </summary>
-		public ScheduleService(IAuditLog<ScheduleId> auditLog, IRedisService redis, IGraphCollection graphs, ICommitService commitService, IJobCollection jobCollection, JobService jobService, IDowntimeService downtimeService, IStreamCollection streamCollection, ITemplateCollection templateCollection, IMongoService mongoService, IClock clock, IOptionsMonitor<BuildConfig> buildConfig, Tracer tracer, ILogger<ScheduleService> logger)
+		public ScheduleService(IAuditLog<ScheduleId> auditLog, IRedisService redis, IGraphCollection graphs, ICommitService commitService, IJobCollection jobCollection, JobService jobService, IDowntimeService downtimeService, IStreamCollection streamCollection, ITemplateCollectionInternal templateCollection, IMongoService mongoService, IClock clock, IOptionsMonitor<BuildConfig> buildConfig, Tracer tracer, ILogger<ScheduleService> logger)
 		{
 			_auditLog = auditLog;
 			_graphs = graphs;
diff --git a/Engine/Source/Programs/Horde/Plugins/Build/HordeServer.Build/Jobs/Templates/ITemplate.cs b/Engine/Source/Programs/Horde/Plugins/Build/HordeServer.Build/Jobs/Templates/ITemplate.cs
deleted file mode 100644
diff --git a/Engine/Source/Programs/Horde/Plugins/Build/HordeServer.Build/Jobs/Templates/ITemplateCollection.cs b/Engine/Source/Programs/Horde/Plugins/Build/HordeServer.Build/Jobs/Templates/ITemplateCollection.cs
deleted file mode 100644
diff --git a/Engine/Source/Programs/Horde/Plugins/Build/HordeServer.Build/Jobs/Templates/ITemplateCollectionInternal.cs b/Engine/Source/Programs/Horde/Plugins/Build/HordeServer.Build/Jobs/Templates/ITemplateCollectionInternal.cs
new file mode 100644
--- /dev/null
+++ b/Engine/Source/Programs/Horde/Plugins/Build/HordeServer.Build/Jobs/Templates/ITemplateCollectionInternal.cs
@@ -0,0 +1,20 @@
+// Copyright Epic Games, Inc. All Rights Reserved.
+
+using EpicGames.Horde.Jobs.Templates;
+
+namespace HordeServer.Jobs.Templates
+{
+	/// <summary>
+	/// Internal implementation of <see cref="ITemplateCollection"/> which
+	/// allows registering new template instances.
+	/// </summary>
+	public interface ITemplateCollectionInternal : ITemplateCollection
+	{
+		/// <summary>
+		/// Gets a template instance from its configuration
+		/// </summary>
+		/// <param name="templateConfig">The template configuration</param>
+		/// <returns>Template instance</returns>
+		Task<ITemplate> GetOrAddAsync(TemplateConfig templateConfig);
+	}
+}
diff --git a/Engine/Source/Programs/Horde/Plugins/Build/HordeServer.Build/Jobs/Templates/TemplateCollection.cs b/Engine/Source/Programs/Horde/Plugins/Build/HordeServer.Build/Jobs/Templates/TemplateCollection.cs
--- a/Engine/Source/Programs/Horde/Plugins/Build/HordeServer.Build/Jobs/Templates/TemplateCollection.cs
+++ b/Engine/Source/Programs/Horde/Plugins/Build/HordeServer.Build/Jobs/Templates/TemplateCollection.cs
@@ -14,7 +14,7 @@ namespace HordeServer.Jobs.Templates
 	/// <summary>
 	/// Collection of template documents
 	/// </summary>
-	public sealed class TemplateCollection : ITemplateCollection, IDisposable
+	public sealed class TemplateCollection : ITemplateCollectionInternal, IDisposable
 	{
 		/// <summary>
 		/// Document describing a job template. These objects are considered immutable once created and uniquely referenced by hash, in order to de-duplicate across all job runs.
diff --git a/Engine/Source/Programs/Horde/Plugins/Build/HordeServer.Build/Jobs/Templates/TemplateMessages.cs b/Engine/Source/Programs/Horde/Plugins/Build/HordeServer.Build/Jobs/Templates/TemplateMessages.cs
deleted file mode 100644
diff --git a/Engine/Source/Programs/Horde/Plugins/Build/HordeServer.Build/Streams/StreamMessages.cs b/Engine/Source/Programs/Horde/Plugins/Build/HordeServer.Build/Streams/StreamMessages.cs
--- a/Engine/Source/Programs/Horde/Plugins/Build/HordeServer.Build/Streams/StreamMessages.cs
+++ b/Engine/Source/Programs/Horde/Plugins/Build/HordeServer.Build/Streams/StreamMessages.cs
@@ -10,7 +10,6 @@ using EpicGames.Horde.Streams;
 using EpicGames.Horde.Users;
 using HordeServer.Issues;
 using HordeServer.Jobs.Schedules;
-using HordeServer.Jobs.Templates;
 using HordeServer.Projects;
 
 #pragma warning disable CA2227 // Change x to be read-only by removing the property setter
diff --git a/Engine/Source/Programs/Horde/Plugins/Build/HordeServer.Build/Streams/StreamsController.cs b/Engine/Source/Programs/Horde/Plugins/Build/HordeServer.Build/Streams/StreamsController.cs
--- a/Engine/Source/Programs/Horde/Plugins/Build/HordeServer.Build/Streams/StreamsController.cs
+++ b/Engine/Source/Programs/Horde/Plugins/Build/HordeServer.Build/Streams/StreamsController.cs
@@ -33,7 +33,7 @@ namespace HordeServer.Streams
 	{
 		private readonly IStreamCollection _streamCollection;
 		private readonly ICommitService _commitService;
-		private readonly ITemplateCollection _templateCollection;
+		private readonly ITemplateCollectionInternal _templateCollection;
 		private readonly IJobStepRefCollection _jobStepRefCollection;
 		private readonly IUserCollection _userCollection;
 		private readonly ScheduleService _scheduleService;
@@ -44,7 +44,7 @@ namespace HordeServer.Streams
 		/// <summary>
 		/// Constructor
 		/// </summary>
-		public StreamsController(IStreamCollection streamCollection, ICommitService commitService, ITemplateCollection templateCollection, IJobStepRefCollection jobStepRefCollection, IUserCollection userCollection, ScheduleService scheduleService, Tracer tracer, IClock clock, IOptionsSnapshot<BuildConfig> buildConfig)
+		public StreamsController(IStreamCollection streamCollection, ICommitService commitService, ITemplateCollectionInternal templateCollection, IJobStepRefCollection jobStepRefCollection, IUserCollection userCollection, ScheduleService scheduleService, Tracer tracer, IClock clock, IOptionsSnapshot<BuildConfig> buildConfig)
 		{
 			_streamCollection = streamCollection;
 			_commitService = commitService;
diff --git a/Engine/Source/Programs/Shared/EpicGames.Horde/Jobs/Templates/ITemplate.cs b/Engine/Source/Programs/Shared/EpicGames.Horde/Jobs/Templates/ITemplate.cs
new file mode 100644
--- /dev/null
+++ b/Engine/Source/Programs/Shared/EpicGames.Horde/Jobs/Templates/ITemplate.cs
@@ -0,0 +1,328 @@
+// Copyright Epic Games, Inc. All Rights Reserved.
+
+using System.Collections.Generic;
+using EpicGames.Core;
+
+#pragma warning disable CA1716 // Rename virtual/interface member x so that it no longer conflicts with the reserved language keyword 'Default'. Using a reserved keyword as the name of a virtual/interface member makes it harder for consumers in other languages to override/implement the member. 
+#pragma warning disable CA2227 // Change x to be read-only by removing the property setter
+
+namespace EpicGames.Horde.Jobs.Templates
+{
+	/// <summary>
+	/// Document describing a job template. These objects are considered immutable once created and uniquely referenced by hash, in order to de-duplicate across all job runs.
+	/// </summary>
+	public interface ITemplate
+	{
+		/// <summary>
+		/// Hash of this template
+		/// </summary>
+		public ContentHash Hash { get; }
+
+		/// <summary>
+		/// Name of the template.
+		/// </summary>
+		public string Name { get; }
+
+		/// <summary>
+		/// Description for the template
+		/// </summary>
+		public string? Description { get; }
+
+		/// <summary>
+		/// Priority of this job
+		/// </summary>
+		public Priority? Priority { get; }
+
+		/// <summary>
+		/// Whether to allow preflights for this job type
+		/// </summary>
+		public bool AllowPreflights { get; }
+
+		/// <summary>
+		/// Whether to always issues for jobs using this template
+		/// </summary>
+		public bool UpdateIssues { get; }
+
+		/// <summary>
+		/// Whether to promote issues by default for jobs using this template
+		/// </summary>
+		public bool PromoteIssuesByDefault { get; }
+
+		/// <summary>
+		/// Agent type to use for parsing the job state
+		/// </summary>
+		public string? InitialAgentType { get; }
+
+		/// <summary>
+		/// Path to a file within the stream to submit to generate a new changelist for jobs
+		/// </summary>
+		public string? SubmitNewChange { get; }
+
+		/// <summary>
+		/// Description for new changelists
+		/// </summary>
+		public string? SubmitDescription { get; }
+
+		/// <summary>
+		/// Optional predefined user-defined properties for this job
+		/// </summary>
+		public IReadOnlyList<string> Arguments { get; }
+
+		/// <summary>
+		/// Parameters for this template
+		/// </summary>
+		public IReadOnlyList<ITemplateParameter> Parameters { get; }
+	}
+
+	/// <summary>
+	/// Base class for parameters used to configure templates via the new build dialog
+	/// </summary>
+	public interface ITemplateParameter
+	{
+		/// <summary>
+		/// Gets the arguments for a job given a set of parameters
+		/// </summary>
+		/// <param name="parameters">Map of parameter id to value</param>
+		/// <param name="scheduledBuild">Whether this is a scheduled build</param>
+		/// <param name="arguments">Receives command line arguments for the job</param>
+		void GetArguments(IReadOnlyDictionary<ParameterId, string> parameters, bool scheduledBuild, List<string> arguments);
+
+		/// <summary>
+		/// Gets the default arguments for this parameter and its children
+		/// </summary>
+		/// <param name="parameters">List of default parameters</param>
+		/// <param name="scheduledBuild">Whether the arguments are being queried for a scheduled build</param>
+		void GetDefaultParameters(Dictionary<ParameterId, string> parameters, bool scheduledBuild);
+	}
+
+	/// <summary>
+	/// Allows the user to toggle an option on or off
+	/// </summary>
+	public interface ITemplateBoolParameter : ITemplateParameter
+	{
+		/// <summary>
+		/// Identifier for this parameter
+		/// </summary>
+		ParameterId Id { get; }
+
+		/// <summary>
+		/// Label to display next to this parameter.
+		/// </summary>
+		string Label { get; }
+
+		/// <summary>
+		/// Argument to add if this parameter is enabled
+		/// </summary>
+		string? ArgumentIfEnabled { get; }
+
+		/// <summary>
+		/// Arguments to add if this parameter is enabled
+		/// </summary>
+		IReadOnlyList<string>? ArgumentsIfEnabled { get; }
+
+		/// <summary>
+		/// Argument to add if this parameter is disabled
+		/// </summary>
+		string? ArgumentIfDisabled { get; }
+
+		/// <summary>
+		/// Arguments to add if this parameter is disabled
+		/// </summary>
+		IReadOnlyList<string>? ArgumentsIfDisabled { get; }
+
+		/// <summary>
+		/// Whether this option should be enabled by default
+		/// </summary>
+		bool Default { get; }
+
+		/// <summary>
+		/// Whether this option should be enabled by default
+		/// </summary>
+		bool? ScheduleOverride { get; }
+
+		/// <summary>
+		/// Tool tip text to display
+		/// </summary>
+		string? ToolTip { get; }
+	}
+
+	/// <summary>
+	/// Free-form text entry parameter
+	/// </summary>
+	public interface ITemplateTextParameter : ITemplateParameter
+	{
+		/// <summary>
+		/// Identifier for this parameter
+		/// </summary>
+		ParameterId Id { get; }
+
+		/// <summary>
+		/// Label to display next to this parameter. Should default to the parameter name.
+		/// </summary>
+		string Label { get; }
+
+		/// <summary>
+		/// Argument to add (will have the value of this field appended)
+		/// </summary>
+		string Argument { get; }
+
+		/// <summary>
+		/// Default value for this argument
+		/// </summary>
+		string Default { get; }
+
+		/// <summary>
+		/// Override for this argument in scheduled builds. 
+		/// </summary>
+		string? ScheduleOverride { get; }
+
+		/// <summary>
+		/// Hint text to display when the field is empty
+		/// </summary>
+		string? Hint { get; }
+
+		/// <summary>
+		/// Regex used to validate values entered into this text field.
+		/// </summary>
+		string? Validation { get; }
+
+		/// <summary>
+		/// Message displayed to explain valid values if validation fails.
+		/// </summary>
+		string? ValidationError { get; }
+
+		/// <summary>
+		/// Tool tip text to display
+		/// </summary>
+		string? ToolTip { get; }
+	}
+
+	/// <summary>
+	/// Style of list parameter
+	/// </summary>
+	public enum TemplateListParameterStyle
+	{
+		/// <summary>
+		/// Regular drop-down list. One item is always selected.
+		/// </summary>
+		List,
+
+		/// <summary>
+		/// Drop-down list with checkboxes
+		/// </summary>
+		MultiList,
+
+		/// <summary>
+		/// Tag picker from list of options
+		/// </summary>
+		TagPicker,
+	}
+
+	/// <summary>
+	/// Allows the user to select a value from a constrained list of choices
+	/// </summary>
+	public interface ITemplateListParameter : ITemplateParameter
+	{
+		/// <summary>
+		/// Label to display next to this parameter.
+		/// </summary>
+		string Label { get; }
+
+		/// <summary>
+		/// Style of picker parameter to use
+		/// </summary>
+		TemplateListParameterStyle Style { get; }
+
+		/// <summary>
+		/// List of values to display in the list
+		/// </summary>
+		IReadOnlyList<ITemplateListParameterItem> Items { get; }
+
+		/// <summary>
+		/// Tool tip text to display
+		/// </summary>
+		string? ToolTip { get; }
+	}
+
+	/// <summary>
+	/// Possible option for a list parameter
+	/// </summary>
+	public interface ITemplateListParameterItem
+	{
+		/// <summary>
+		/// Identifier for this parameter
+		/// </summary>
+		ParameterId Id { get; }
+
+		/// <summary>
+		/// Group to display this entry in
+		/// </summary>
+		string? Group { get; }
+
+		/// <summary>
+		/// Text to display for this option.
+		/// </summary>
+		string Text { get; }
+
+		/// <summary>
+		/// Argument to add if this parameter is enabled.
+		/// </summary>
+		string? ArgumentIfEnabled { get; }
+
+		/// <summary>
+		/// Arguments to add if this parameter is enabled.
+		/// </summary>
+		IReadOnlyList<string>? ArgumentsIfEnabled { get; }
+
+		/// <summary>
+		/// Argument to add if this parameter is disabled.
+		/// </summary>
+		string? ArgumentIfDisabled { get; }
+
+		/// <summary>
+		/// Arguments to add if this parameter is disabled.
+		/// </summary>
+		IReadOnlyList<string>? ArgumentsIfDisabled { get; }
+
+		/// <summary>
+		/// Whether this item is selected by default
+		/// </summary>
+		bool Default { get; }
+
+		/// <summary>
+		/// Whether this item is selected by default
+		/// </summary>
+		bool? ScheduleOverride { get; }
+	}
+
+	/// <summary>
+	/// Extension methods for templates
+	/// </summary>
+	public static class TemplateExtensions
+	{
+		/// <summary>
+		/// Gets the full argument list for a template
+		/// </summary>
+		public static void GetArgumentsForParameters(this ITemplate template, IReadOnlyDictionary<ParameterId, string> parameters, List<string> arguments)
+		{
+			arguments.AddRange(template.Arguments);
+
+			foreach (ITemplateParameter parameter in template.Parameters)
+			{
+				parameter.GetArguments(parameters, false, arguments);
+			}
+		}
+
+		/// <summary>
+		/// Gets the arguments for default options in this template. Does not include the standard template arguments.
+		/// </summary>
+		/// <returns>List of default arguments</returns>
+		public static void GetDefaultParameters(this ITemplate template, Dictionary<ParameterId, string> parameters, bool scheduledBuild)
+		{
+			foreach (ITemplateParameter parameter in template.Parameters)
+			{
+				parameter.GetDefaultParameters(parameters, scheduledBuild);
+			}
+		}
+	}
+}
diff --git a/Engine/Source/Programs/Shared/EpicGames.Horde/Jobs/Templates/ITemplateCollection.cs b/Engine/Source/Programs/Shared/EpicGames.Horde/Jobs/Templates/ITemplateCollection.cs
new file mode 100644
--- /dev/null
+++ b/Engine/Source/Programs/Shared/EpicGames.Horde/Jobs/Templates/ITemplateCollection.cs
@@ -0,0 +1,20 @@
+// Copyright Epic Games, Inc. All Rights Reserved.
+
+using System.Threading.Tasks;
+using EpicGames.Core;
+
+namespace EpicGames.Horde.Jobs.Templates
+{
+	/// <summary>
+	/// Interface for a collection of template documents
+	/// </summary>
+	public interface ITemplateCollection
+	{
+		/// <summary>
+		/// Gets a template by ID
+		/// </summary>
+		/// <param name="templateId">Unique id of the template</param>
+		/// <returns>The template document</returns>
+		Task<ITemplate?> GetAsync(ContentHash templateId);
+	}
+}
diff --git a/Engine/Source/Programs/Shared/EpicGames.Horde/Jobs/Templates/TemplateMessages.cs b/Engine/Source/Programs/Shared/EpicGames.Horde/Jobs/Templates/TemplateMessages.cs
new file mode 100644
--- /dev/null
+++ b/Engine/Source/Programs/Shared/EpicGames.Horde/Jobs/Templates/TemplateMessages.cs
@@ -0,0 +1,371 @@
+// Copyright Epic Games, Inc. All Rights Reserved.
+
+using System;
+using System.Collections.Generic;
+using System.Linq;
+using EpicGames.Core;
+
+#pragma warning disable CA2227 // Change x to be read-only by removing the property setter
+
+namespace EpicGames.Horde.Jobs.Templates
+{
+	/// <summary>
+	/// Response describing a template
+	/// </summary>
+	public class GetTemplateResponseBase
+	{
+		/// <summary>
+		/// Name of the template
+		/// </summary>
+		public string Name { get; set; }
+
+		/// <summary>
+		/// Description for the template
+		/// </summary>
+		public string? Description { get; set; }
+
+		/// <summary>
+		/// Default priority for this job
+		/// </summary>
+		public Priority? Priority { get; set; }
+
+		/// <summary>
+		/// Whether to allow preflights of this template
+		/// </summary>
+		public bool AllowPreflights { get; set; }
+
+		/// <summary>
+		/// Whether to always update issues on jobs using this template
+		/// </summary>
+		public bool UpdateIssues { get; set; }
+
+		/// <summary>
+		/// The initial agent type to parse the BuildGraph script on
+		/// </summary>
+		public string? InitialAgentType { get; set; }
+
+		/// <summary>
+		/// Path to a file within the stream to submit to generate a new changelist for jobs
+		/// </summary>
+		public string? SubmitNewChange { get; }
+
+		/// <summary>
+		/// Parameters for the job.
+		/// </summary>
+		public List<string> Arguments { get; set; }
+
+		/// <summary>
+		/// List of parameters for this template
+		/// </summary>
+		public List<GetTemplateParameterResponse> Parameters { get; set; }
+
+		/// <summary>
+		/// Parameterless constructor for serialization
+		/// </summary>
+		protected GetTemplateResponseBase()
+		{
+			Name = null!;
+			AllowPreflights = true;
+			Arguments = new List<string>();
+			Parameters = new List<GetTemplateParameterResponse>();
+		}
+
+		/// <summary>
+		/// Constructor
+		/// </summary>
+		/// <param name="template">The template to construct from</param>
+		public GetTemplateResponseBase(ITemplate template)
+		{
+			Name = template.Name;
+			Description = template.Description;
+			Priority = template.Priority;
+			AllowPreflights = template.AllowPreflights;
+			UpdateIssues = template.UpdateIssues;
+			InitialAgentType = template.InitialAgentType;
+			SubmitNewChange = template.SubmitNewChange;
+			Arguments = new List<string>(template.Arguments);
+			Parameters = template.Parameters.ConvertAll(x => CreateParameterResponse(x));
+		}
+
+		static GetTemplateParameterResponse CreateParameterResponse(ITemplateParameter parameter)
+		{
+			return parameter switch
+			{
+				ITemplateBoolParameter boolParameter
+					=> new GetTemplateBoolParameterResponse(boolParameter),
+
+				ITemplateTextParameter textParameter
+					=> new GetTemplateTextParameterResponse(textParameter),
+
+				ITemplateListParameter listParameter
+					=> new GetTemplateListParameterResponse(listParameter),
+
+				_ => throw new NotImplementedException()
+			};
+		}
+	}
+
+	/// <summary>
+	/// Response describing a template
+	/// </summary>
+	public class GetTemplateResponse : GetTemplateResponseBase
+	{
+		/// <summary>
+		/// Unique id of the template
+		/// </summary>
+		public string Id { get; set; }
+
+		/// <summary>
+		/// Parameterless constructor for serialization
+		/// </summary>
+		protected GetTemplateResponse()
+			: base()
+		{
+			Id = null!;
+		}
+
+		/// <summary>
+		/// Constructor
+		/// </summary>
+		/// <param name="template">The template to construct from</param>
+		public GetTemplateResponse(ITemplate template)
+			: base(template)
+		{
+			Id = template.Hash.ToString();
+		}
+	}
+
+	/// <summary>
+	/// Base class for template parameters
+	/// </summary>
+	[JsonKnownTypes(typeof(GetTemplateBoolParameterResponse), typeof(GetTemplateTextParameterResponse), typeof(GetTemplateListParameterResponse))]
+	public abstract class GetTemplateParameterResponse : ITemplateParameter
+	{
+		void ITemplateParameter.GetArguments(IReadOnlyDictionary<ParameterId, string> parameters, bool scheduledBuild, List<string> arguments)
+			=> throw new NotSupportedException();
+
+		void ITemplateParameter.GetDefaultParameters(Dictionary<ParameterId, string> parameters, bool scheduledBuild)
+			=> throw new NotSupportedException();
+	}
+
+	/// <summary>
+	/// Allows the user to toggle an option on or off
+	/// </summary>
+	[JsonDiscriminator("Bool")]
+	public class GetTemplateBoolParameterResponse : GetTemplateParameterResponse, ITemplateBoolParameter
+	{
+		/// <inheritdoc/>
+		public ParameterId Id { get; set; }
+
+		/// <inheritdoc/>
+		public string Label { get; set; }
+
+		/// <inheritdoc/>
+		public string? ArgumentIfEnabled { get; set; }
+
+		/// <inheritdoc/>
+		public List<string>? ArgumentsIfEnabled { get; set; }
+
+		/// <inheritdoc/>
+		public string? ArgumentIfDisabled { get; set; }
+
+		/// <inheritdoc/>
+		public List<string>? ArgumentsIfDisabled { get; set; }
+
+		/// <inheritdoc/>
+		public bool Default { get; set; }
+
+		/// <inheritdoc/>
+		public bool? ScheduleOverride { get; set; }
+
+		/// <inheritdoc/>
+		public string? ToolTip { get; set; }
+
+		IReadOnlyList<string>? ITemplateBoolParameter.ArgumentsIfEnabled => ArgumentsIfEnabled;
+		IReadOnlyList<string>? ITemplateBoolParameter.ArgumentsIfDisabled => ArgumentsIfDisabled;
+
+		/// <summary>
+		/// Default constructor
+		/// </summary>
+		public GetTemplateBoolParameterResponse()
+		{
+			Label = String.Empty;
+		}
+
+		/// <summary>
+		/// Constructor
+		/// </summary>
+		public GetTemplateBoolParameterResponse(ITemplateBoolParameter parameter)
+		{
+			Id = parameter.Id;
+			Label = parameter.Label;
+			ArgumentIfEnabled = parameter.ArgumentIfEnabled;
+			ArgumentsIfEnabled = parameter.ArgumentsIfEnabled?.ToList();
+			ArgumentIfDisabled = parameter.ArgumentIfDisabled;
+			ArgumentsIfDisabled = parameter.ArgumentsIfDisabled?.ToList();
+			Default = parameter.Default;
+			ScheduleOverride = parameter.ScheduleOverride;
+			ToolTip = parameter.ToolTip;
+		}
+	}
+
+	/// <summary>
+	/// Free-form text entry parameter
+	/// </summary>
+	[JsonDiscriminator("Text")]
+	public class GetTemplateTextParameterResponse : GetTemplateParameterResponse, ITemplateTextParameter
+	{
+		/// <inheritdoc/>
+		public ParameterId Id { get; set; }
+
+		/// <inheritdoc/>
+		public string Label { get; set; }
+
+		/// <inheritdoc/>
+		public string Argument { get; set; }
+
+		/// <inheritdoc/>
+		public string Default { get; set; }
+
+		/// <inheritdoc/>
+		public string? ScheduleOverride { get; set; }
+
+		/// <inheritdoc/>
+		public string? Hint { get; set; }
+
+		/// <inheritdoc/>
+		public string? Validation { get; set; }
+
+		/// <inheritdoc/>
+		public string? ValidationError { get; set; }
+
+		/// <inheritdoc/>
+		public string? ToolTip { get; set; }
+
+		/// <summary>
+		/// Default constructor
+		/// </summary>
+		public GetTemplateTextParameterResponse()
+		{
+			Label = String.Empty;
+			Argument = String.Empty;
+			Default = String.Empty;
+		}
+
+		/// <summary>
+		/// Constructor
+		/// </summary>
+		public GetTemplateTextParameterResponse(ITemplateTextParameter parameter)
+		{
+			Id = parameter.Id;
+			Label = parameter.Label;
+			Argument = parameter.Argument;
+			Default = parameter.Default;
+			ScheduleOverride = parameter.ScheduleOverride;
+			Hint = parameter.Hint;
+			Validation = parameter.Validation;
+			ValidationError = parameter.ValidationError;
+			ToolTip = parameter.ToolTip;
+		}
+	}
+
+	/// <summary>
+	/// Allows the user to select a value from a constrained list of choices
+	/// </summary>
+	[JsonDiscriminator("List")]
+	public class GetTemplateListParameterResponse : GetTemplateParameterResponse, ITemplateListParameter
+	{
+		/// <inheritdoc/>
+		public string Label { get; }
+
+		/// <inheritdoc/>
+		public TemplateListParameterStyle Style { get; }
+
+		/// <inheritdoc cref="ITemplateListParameter.Items"/>
+		public List<GetTemplateListParameterItemResponse> Items { get; set; }
+
+		/// <inheritdoc/>
+		public string? ToolTip { get; }
+
+		IReadOnlyList<ITemplateListParameterItem> ITemplateListParameter.Items => Items;
+
+		/// <summary>
+		/// Default constructor
+		/// </summary>
+		public GetTemplateListParameterResponse()
+		{
+			Label = String.Empty;
+			Items = new List<GetTemplateListParameterItemResponse>();
+		}
+
+		/// <summary>
+		/// Constructor
+		/// </summary>
+		public GetTemplateListParameterResponse(ITemplateListParameter parameter)
+		{
+			Label = parameter.Label;
+			Style = parameter.Style;
+			Items = parameter.Items.ConvertAll(x => new GetTemplateListParameterItemResponse(x));
+		}
+	}
+
+	/// <summary>
+	/// Possible option for a list parameter
+	/// </summary>
+	public class GetTemplateListParameterItemResponse : ITemplateListParameterItem
+	{
+		/// <inheritdoc/>
+		public ParameterId Id { get; set; }
+
+		/// <inheritdoc/>
+		public string? Group { get; set; }
+
+		/// <inheritdoc/>
+		public string Text { get; set; }
+
+		/// <inheritdoc/>
+		public string? ArgumentIfEnabled { get; set; }
+
+		/// <inheritdoc/>
+		public List<string>? ArgumentsIfEnabled { get; set; }
+
+		/// <inheritdoc/>
+		public string? ArgumentIfDisabled { get; set; }
+
+		/// <inheritdoc/>
+		public List<string>? ArgumentsIfDisabled { get; set; }
+
+		/// <inheritdoc/>
+		public bool Default { get; set; }
+
+		/// <inheritdoc/>
+		public bool? ScheduleOverride { get; set; }
+
+		IReadOnlyList<string>? ITemplateListParameterItem.ArgumentsIfEnabled => ArgumentsIfEnabled;
+		IReadOnlyList<string>? ITemplateListParameterItem.ArgumentsIfDisabled => ArgumentsIfDisabled;
+
+		/// <summary>
+		/// Default constructor
+		/// </summary>
+		public GetTemplateListParameterItemResponse()
+		{
+			Text = String.Empty;
+		}
+
+		/// <summary>
+		/// Constructor
+		/// </summary>
+		public GetTemplateListParameterItemResponse(ITemplateListParameterItem item)
+		{
+			Id = item.Id;
+			Group = item.Group;
+			Text = item.Text;
+			ArgumentIfEnabled = item.ArgumentIfEnabled;
+			ArgumentsIfEnabled = item.ArgumentsIfEnabled?.ToList();
+			ArgumentIfDisabled = item.ArgumentIfDisabled;
+			ArgumentsIfDisabled = item.ArgumentsIfDisabled?.ToList();
+			Default = item.Default;
+			ScheduleOverride = item.ScheduleOverride;
+		}
+	}
+}
