Horde: Add an interface for job options.

#rnx

diff --git a/Engine/Source/Programs/Shared/EpicGames.Horde/Jobs/JobMessages.cs b/Engine/Source/Programs/Shared/EpicGames.Horde/Jobs/JobMessages.cs
--- a/Engine/Source/Programs/Shared/EpicGames.Horde/Jobs/JobMessages.cs
+++ b/Engine/Source/Programs/Shared/EpicGames.Horde/Jobs/JobMessages.cs
@@ -39,45 +39,74 @@ namespace EpicGames.Horde.Jobs
 	}
 
 	/// <summary>
-	/// Options for executing a job
+	/// Read-only interface for job options
 	/// </summary>
-	public class JobOptions
+	public interface IReadOnlyJobOptions
 	{
 		/// <summary>
 		/// Name of the executor to use
 		/// </summary>
-		public string? Executor { get; set; }
+		string? Executor { get; }
 
 		/// <summary>
 		/// Whether to execute using Wine emulation on Linux
 		/// </summary>
-		public bool? UseWine { get; set; }
+		bool? UseWine { get; }
 
 		/// <summary>
 		/// Executes the job lease in a separate process
 		/// </summary>
-		public bool? RunInSeparateProcess { get; set; }
+		bool? RunInSeparateProcess { get; }
 
 		/// <summary>
 		/// What workspace materializer to use in WorkspaceExecutor. Will override any value from workspace config.
 		/// </summary>
-		public string? WorkspaceMaterializer { get; set; }
+		string? WorkspaceMaterializer { get; }
 
 		/// <summary>
 		/// Options for executing a job inside a container
 		/// </summary>
-		public JobContainerOptions Container { get; set; } = new JobContainerOptions();
+		IReadOnlyJobContainerOptions Container { get; }
 
 		/// <summary>
 		/// Number of days after which to expire jobs
 		/// </summary>
-		public int? ExpireAfterDays { get; set; }
+		int? ExpireAfterDays { get; }
 
 		/// <summary>
 		/// Name of the driver to use
 		/// </summary>
+		string? Driver { get; }
+	}
+
+	/// <summary>
+	/// Options for executing a job
+	/// </summary>
+	public class JobOptions : IReadOnlyJobOptions
+	{
+		/// <inheritdoc/>
+		public string? Executor { get; set; }
+
+		/// <inheritdoc/>
+		public bool? UseWine { get; set; }
+
+		/// <inheritdoc/>
+		public bool? RunInSeparateProcess { get; set; }
+
+		/// <inheritdoc/>
+		public string? WorkspaceMaterializer { get; set; }
+
+		/// <inheritdoc/>
+		public JobContainerOptions Container { get; set; } = new JobContainerOptions();
+
+		/// <inheritdoc/>
+		public int? ExpireAfterDays { get; set; }
+
+		/// <inheritdoc/>
 		public string? Driver { get; set; }
 
+		IReadOnlyJobContainerOptions IReadOnlyJobOptions.Container => Container;
+
 		/// <summary>
 		/// Merge defaults from another options object
 		/// </summary>
@@ -93,10 +122,36 @@ namespace EpicGames.Horde.Jobs
 		}
 	}
 
+	/// <summary>
+	/// Options for a job container
+	/// </summary>
+	public interface IReadOnlyJobContainerOptions
+	{
+		/// <summary>
+		/// Whether to execute job inside a container
+		/// </summary>
+		public bool? Enabled { get; }
+
+		/// <summary>
+		/// Image URL to container, such as "quay.io/podman/hello"
+		/// </summary>
+		public string? ImageUrl { get; }
+
+		/// <summary>
+		/// Container engine executable (docker or with full path like /usr/bin/podman)
+		/// </summary>
+		public string? ContainerEngineExecutable { get; }
+
+		/// <summary>
+		/// Additional arguments to pass to container engine
+		/// </summary>
+		public string? ExtraArguments { get; }
+	}
+
 	/// <summary>
 	/// Options for executing a job inside a container
 	/// </summary>
-	public class JobContainerOptions
+	public class JobContainerOptions : IReadOnlyJobContainerOptions
 	{
 		/// <summary>
 		/// Whether to execute job inside a container
