Horde: Add explicit interface and message type for ChangeQuery type.

#rnx

diff --git a/Engine/Source/Programs/Horde/Plugins/Build/HordeServer.Build/Jobs/JobService.cs b/Engine/Source/Programs/Horde/Plugins/Build/HordeServer.Build/Jobs/JobService.cs
--- a/Engine/Source/Programs/Horde/Plugins/Build/HordeServer.Build/Jobs/JobService.cs
+++ b/Engine/Source/Programs/Horde/Plugins/Build/HordeServer.Build/Jobs/JobService.cs
@@ -380,9 +380,9 @@ namespace HordeServer.Jobs
 		/// Evaluate a change query to determine which CL to run a job at
 		/// </summary>
 		/// <returns></returns>
-		public async Task<CommitIdWithOrder?> EvaluateChangeQueriesAsync(StreamId streamId, List<ChangeQueryConfig> queries, List<CommitTag>? commitTags, ICommitCollection commits, CancellationToken cancellationToken)
+		public async Task<CommitIdWithOrder?> EvaluateChangeQueriesAsync(StreamId streamId, IEnumerable<IChangeQuery> queries, List<CommitTag>? commitTags, ICommitCollection commits, CancellationToken cancellationToken)
 		{
-			foreach (ChangeQueryConfig query in queries)
+			foreach (IChangeQuery query in queries)
 			{
 				CommitIdWithOrder? commitId = await EvaluateChangeQueryAsync(streamId, query, commitTags, commits, cancellationToken);
 				if (commitId != null)
@@ -397,7 +397,7 @@ namespace HordeServer.Jobs
 		/// Evaluate a change query to determine which CL to run a job at
 		/// </summary>
 		/// <returns></returns>
-		public async ValueTask<CommitIdWithOrder?> EvaluateChangeQueryAsync(StreamId streamId, ChangeQueryConfig query, List<CommitTag>? commitTags, ICommitCollection commits, CancellationToken cancellationToken)
+		public async ValueTask<CommitIdWithOrder?> EvaluateChangeQueryAsync(StreamId streamId, IChangeQuery query, List<CommitTag>? commitTags, ICommitCollection commits, CancellationToken cancellationToken)
 		{
 			if (query.Condition == null || query.Condition.Evaluate(propertyName => GetTagPropertyValues(propertyName, commitTags)))
 			{
@@ -415,7 +415,7 @@ namespace HordeServer.Jobs
 				// Find the job with the given result
 				if (query.TemplateId != null)
 				{
-					List<JobStepOutcome> outcomes = query.Outcomes ?? new List<JobStepOutcome> { JobStepOutcome.Success };
+					IReadOnlyList<JobStepOutcome> outcomes = query.Outcomes ?? new List<JobStepOutcome> { JobStepOutcome.Success };
 
 					IReadOnlyList<IJob> jobs = await FindJobsAsync(new FindJobOptions(StreamId: streamId, Templates: new[] { query.TemplateId.Value }, Target: query.Target, State: new[] { JobStepState.Completed }, Outcome: outcomes.ToArray(), ExcludeUserJobs: true, ExcludeCancelled: true), count: 1, cancellationToken: cancellationToken);
 					if (jobs.Count > 0)
diff --git a/Engine/Source/Programs/Horde/Plugins/Build/HordeServer.Build/Jobs/Templates/TemplateConfig.cs b/Engine/Source/Programs/Horde/Plugins/Build/HordeServer.Build/Jobs/Templates/TemplateConfig.cs
--- a/Engine/Source/Programs/Horde/Plugins/Build/HordeServer.Build/Jobs/Templates/TemplateConfig.cs
+++ b/Engine/Source/Programs/Horde/Plugins/Build/HordeServer.Build/Jobs/Templates/TemplateConfig.cs
@@ -4,6 +4,8 @@ using System.ComponentModel.DataAnnotations;
 using System.Text.Json.Serialization;
 using EpicGames.Core;
 using EpicGames.Horde;
+using EpicGames.Horde.Commits;
+using EpicGames.Horde.Common;
 using EpicGames.Horde.Jobs;
 using EpicGames.Horde.Jobs.Templates;
 using HordeServer.Configuration;
@@ -94,6 +96,31 @@ namespace HordeServer.Jobs.Templates
 		internal ContentHash? CachedHash { get; set; }
 	}
 
+	/// <summary>
+	/// Query selecting the base changelist to use
+	/// </summary>
+	public class ChangeQueryConfig : IChangeQuery
+	{
+		/// <inheritdoc/>
+		public string? Name { get; set; }
+
+		/// <inheritdoc/>
+		public Condition? Condition { get; set; }
+
+		/// <inheritdoc/>
+		public TemplateId? TemplateId { get; set; }
+
+		/// <inheritdoc/>
+		public string? Target { get; set; }
+
+		/// <inheritdoc/>
+		public List<JobStepOutcome>? Outcomes { get; set; }
+		IReadOnlyList<JobStepOutcome>? IChangeQuery.Outcomes => Outcomes;
+
+		/// <inheritdoc/>
+		public CommitTag? CommitTag { get; set; }
+	}
+
 	/// <summary>
 	/// Base class for template parameters
 	/// </summary>
diff --git a/Engine/Source/Programs/Horde/Plugins/Build/HordeServer.Build/Streams/StreamMessages.cs b/Engine/Source/Programs/Horde/Plugins/Build/HordeServer.Build/Streams/StreamMessages.cs
--- a/Engine/Source/Programs/Horde/Plugins/Build/HordeServer.Build/Streams/StreamMessages.cs
+++ b/Engine/Source/Programs/Horde/Plugins/Build/HordeServer.Build/Streams/StreamMessages.cs
@@ -347,7 +347,7 @@ namespace HordeServer.Streams
 		/// <summary>
 		/// List of queries for the default changelist
 		/// </summary>
-		public List<ChangeQueryConfig>? DefaultChange { get; set; }
+		public List<ChangeQueryMessage>? DefaultChange { get; set; }
 
 		/// <summary>
 		/// Whether the user is allowed to create jobs from this template
@@ -375,7 +375,7 @@ namespace HordeServer.Streams
 			Schedule = (templateRef.Schedule != null) ? new GetScheduleResponse(templateRef.Schedule, schedulerTimeZone) : null;
 			ChainedJobs = (templateRef.Config.ChainedJobs != null && templateRef.Config.ChainedJobs.Count > 0) ? templateRef.Config.ChainedJobs : null;
 			StepStates = stepStates;
-			DefaultChange = templateRef.Config.DefaultChange;
+			DefaultChange = templateRef.Config.DefaultChange?.ConvertAll(x => new ChangeQueryMessage(x));
 			CanRun = canRun;
 		}
 	}
diff --git a/Engine/Source/Programs/Shared/EpicGames.Horde/Jobs/JobMessages.cs b/Engine/Source/Programs/Shared/EpicGames.Horde/Jobs/JobMessages.cs
--- a/Engine/Source/Programs/Shared/EpicGames.Horde/Jobs/JobMessages.cs
+++ b/Engine/Source/Programs/Shared/EpicGames.Horde/Jobs/JobMessages.cs
@@ -3,6 +3,7 @@
 using System;
 using System.Collections.Generic;
 using System.ComponentModel.DataAnnotations;
+using System.Linq;
 using System.Text.Json.Serialization;
 using EpicGames.Horde.Artifacts;
 using EpicGames.Horde.Commits;
@@ -188,37 +189,45 @@ namespace EpicGames.Horde.Jobs
 	/// <summary>
 	/// Query selecting the base changelist to use
 	/// </summary>
-	public class ChangeQueryConfig
+	public class ChangeQueryMessage : IChangeQuery
 	{
-		/// <summary>
-		/// Name of this query, for display on the dashboard.
-		/// </summary>
+		/// <inheritdoc/>
 		public string? Name { get; set; }
 
-		/// <summary>
-		/// Condition to evaluate before deciding to use this query. May query tags in a preflight.
-		/// </summary>
+		/// <inheritdoc/>
 		public Condition? Condition { get; set; }
 
-		/// <summary>
-		/// The template id to query
-		/// </summary>
+		/// <inheritdoc/>
 		public TemplateId? TemplateId { get; set; }
 
-		/// <summary>
-		/// The target to query
-		/// </summary>
+		/// <inheritdoc/>
 		public string? Target { get; set; }
 
+		/// <inheritdoc cref="IChangeQuery.Outcomes"/>
+		public List<JobStepOutcome>? Outcomes { get; set; }
+		IReadOnlyList<JobStepOutcome>? IChangeQuery.Outcomes => Outcomes;
+
+		/// <inheritdoc/>
+		public CommitTag? CommitTag { get; set; }
+
 		/// <summary>
-		/// Whether to match a job that produced warnings
+		/// Constructor
 		/// </summary>
-		public List<JobStepOutcome>? Outcomes { get; set; }
+		public ChangeQueryMessage()
+		{ }
 
 		/// <summary>
-		/// Finds the last commit with this tag
+		/// Constructor
 		/// </summary>
-		public CommitTag? CommitTag { get; set; }
+		public ChangeQueryMessage(IChangeQuery other)
+		{
+			Name = other.Name;
+			Condition = other.Condition;
+			TemplateId = other.TemplateId;
+			Target = other.Target;
+			Outcomes = other.Outcomes?.ToList();
+			CommitTag = other.CommitTag;
+		}
 	}
 
 	/// <summary>
@@ -268,16 +277,16 @@ namespace EpicGames.Horde.Jobs
 		/// Parameters to use when selecting the change to execute at.
 		/// </summary>
 		[Obsolete("Use ChangeQueries instead")]
-		public ChangeQueryConfig? ChangeQuery
+		public ChangeQueryMessage? ChangeQuery
 		{
 			get => (ChangeQueries != null && ChangeQueries.Count > 0) ? ChangeQueries[0] : null;
-			set => ChangeQueries = (value == null) ? null : new List<ChangeQueryConfig> { value };
+			set => ChangeQueries = (value == null) ? null : new List<ChangeQueryMessage> { value };
 		}
 
 		/// <summary>
 		/// List of change queries to evaluate
 		/// </summary>
-		public List<ChangeQueryConfig>? ChangeQueries { get; set; }
+		public List<ChangeQueryMessage>? ChangeQueries { get; set; }
 
 		/// <summary>
 		/// The preflight changelist number
diff --git a/Engine/Source/Programs/Shared/EpicGames.Horde/Jobs/Templates/ITemplate.cs b/Engine/Source/Programs/Shared/EpicGames.Horde/Jobs/Templates/ITemplate.cs
--- a/Engine/Source/Programs/Shared/EpicGames.Horde/Jobs/Templates/ITemplate.cs
+++ b/Engine/Source/Programs/Shared/EpicGames.Horde/Jobs/Templates/ITemplate.cs
@@ -2,6 +2,8 @@
 
 using System.Collections.Generic;
 using EpicGames.Core;
+using EpicGames.Horde.Commits;
+using EpicGames.Horde.Common;
 
 #pragma warning disable CA1716 // Rename virtual/interface member x so that it no longer conflicts with the reserved language keyword 'Default'. Using a reserved keyword as the name of a virtual/interface member makes it harder for consumers in other languages to override/implement the member. 
 #pragma warning disable CA2227 // Change x to be read-only by removing the property setter
@@ -325,4 +327,40 @@ namespace EpicGames.Horde.Jobs.Templates
 			}
 		}
 	}
+
+	/// <summary>
+	/// Query selecting the base changelist to use
+	/// </summary>
+	public interface IChangeQuery
+	{
+		/// <summary>
+		/// Name of this query, for display on the dashboard.
+		/// </summary>
+		string? Name { get; }
+
+		/// <summary>
+		/// Condition to evaluate before deciding to use this query. May query tags in a preflight.
+		/// </summary>
+		Condition? Condition { get; }
+
+		/// <summary>
+		/// The template id to query
+		/// </summary>
+		TemplateId? TemplateId { get; }
+
+		/// <summary>
+		/// The target to query
+		/// </summary>
+		string? Target { get; }
+
+		/// <summary>
+		/// Whether to match a job that produced warnings
+		/// </summary>
+		IReadOnlyList<JobStepOutcome>? Outcomes { get; }
+
+		/// <summary>
+		/// Finds the last commit with this tag
+		/// </summary>
+		CommitTag? CommitTag { get; }
+	}
 }
