Horde: Move commit interfaces into EpicGames.Horde.

#rnx

diff --git a/Engine/Source/Programs/Horde/Plugins/Build/HordeServer.Build.Tests/Jobs/JobServiceTests.cs b/Engine/Source/Programs/Horde/Plugins/Build/HordeServer.Build.Tests/Jobs/JobServiceTests.cs
--- a/Engine/Source/Programs/Horde/Plugins/Build/HordeServer.Build.Tests/Jobs/JobServiceTests.cs
+++ b/Engine/Source/Programs/Horde/Plugins/Build/HordeServer.Build.Tests/Jobs/JobServiceTests.cs
@@ -11,7 +11,6 @@ using EpicGames.Horde.Jobs.Templates;
 using EpicGames.Horde.Streams;
 using HordeCommon.Rpc.Messages;
 using HordeServer.Agents;
-using HordeServer.Commits;
 using HordeServer.Jobs;
 using HordeServer.Jobs.Templates;
 using HordeServer.Logs;
diff --git a/Engine/Source/Programs/Horde/Plugins/Build/HordeServer.Build.Tests/Stubs/Services/PerforceServiceStub.cs b/Engine/Source/Programs/Horde/Plugins/Build/HordeServer.Build.Tests/Stubs/Services/PerforceServiceStub.cs
--- a/Engine/Source/Programs/Horde/Plugins/Build/HordeServer.Build.Tests/Stubs/Services/PerforceServiceStub.cs
+++ b/Engine/Source/Programs/Horde/Plugins/Build/HordeServer.Build.Tests/Stubs/Services/PerforceServiceStub.cs
@@ -5,7 +5,6 @@ using EpicGames.Core;
 using EpicGames.Horde.Commits;
 using EpicGames.Horde.Streams;
 using EpicGames.Horde.Users;
-using HordeServer.Commits;
 using HordeServer.Perforce;
 using HordeServer.Streams;
 using HordeServer.Users;
@@ -34,10 +33,10 @@ namespace HordeServer.Tests.Stubs.Services
 		public class Commit : ICommit
 		{
 			public StreamId StreamId { get; }
-			
+
 			public int Number { get; }
 			CommitIdWithOrder ICommit.Id => CommitIdWithOrder.FromPerforceChange(Number);
-	
+
 			public int OriginalChange { get; }
 			CommitIdWithOrder ICommit.OriginalCommitId => CommitIdWithOrder.FromPerforceChange(OriginalChange);
 
diff --git a/Engine/Source/Programs/Horde/Plugins/Build/HordeServer.Build/Commits/ICommit.cs b/Engine/Source/Programs/Horde/Plugins/Build/HordeServer.Build/Commits/ICommit.cs
deleted file mode 100644
diff --git a/Engine/Source/Programs/Horde/Plugins/Build/HordeServer.Build/Commits/ICommitCollection.cs b/Engine/Source/Programs/Horde/Plugins/Build/HordeServer.Build/Commits/ICommitCollection.cs
deleted file mode 100644
diff --git a/Engine/Source/Programs/Horde/Plugins/Build/HordeServer.Build/Issues/IssuesController.cs b/Engine/Source/Programs/Horde/Plugins/Build/HordeServer.Build/Issues/IssuesController.cs
--- a/Engine/Source/Programs/Horde/Plugins/Build/HordeServer.Build/Issues/IssuesController.cs
+++ b/Engine/Source/Programs/Horde/Plugins/Build/HordeServer.Build/Issues/IssuesController.cs
@@ -10,7 +10,6 @@ using EpicGames.Horde.Logs;
 using EpicGames.Horde.Streams;
 using EpicGames.Horde.Users;
 using HordeServer.Auditing;
-using HordeServer.Commits;
 using HordeServer.Issues.External;
 using HordeServer.Jobs;
 using HordeServer.Logs;
diff --git a/Engine/Source/Programs/Horde/Plugins/Build/HordeServer.Build/Jobs/Bisect/BisectService.cs b/Engine/Source/Programs/Horde/Plugins/Build/HordeServer.Build/Jobs/Bisect/BisectService.cs
--- a/Engine/Source/Programs/Horde/Plugins/Build/HordeServer.Build/Jobs/Bisect/BisectService.cs
+++ b/Engine/Source/Programs/Horde/Plugins/Build/HordeServer.Build/Jobs/Bisect/BisectService.cs
@@ -1,6 +1,7 @@
 // Copyright Epic Games, Inc. All Rights Reserved.
 
 using EpicGames.Core;
+using EpicGames.Horde.Commits;
 using EpicGames.Horde.Jobs;
 using EpicGames.Horde.Jobs.Graphs;
 using EpicGames.Horde.Jobs.Templates;
diff --git a/Engine/Source/Programs/Horde/Plugins/Build/HordeServer.Build/Jobs/JobService.cs b/Engine/Source/Programs/Horde/Plugins/Build/HordeServer.Build/Jobs/JobService.cs
--- a/Engine/Source/Programs/Horde/Plugins/Build/HordeServer.Build/Jobs/JobService.cs
+++ b/Engine/Source/Programs/Horde/Plugins/Build/HordeServer.Build/Jobs/JobService.cs
@@ -15,7 +15,6 @@ using EpicGames.Horde.Logs;
 using EpicGames.Horde.Streams;
 using EpicGames.Horde.Users;
 using HordeServer.Agents;
-using HordeServer.Commits;
 using HordeServer.Jobs.Bisect;
 using HordeServer.Jobs.Graphs;
 using HordeServer.Jobs.Templates;
diff --git a/Engine/Source/Programs/Horde/Plugins/Build/HordeServer.Build/Perforce/IPerforceService.cs b/Engine/Source/Programs/Horde/Plugins/Build/HordeServer.Build/Perforce/IPerforceService.cs
--- a/Engine/Source/Programs/Horde/Plugins/Build/HordeServer.Build/Perforce/IPerforceService.cs
+++ b/Engine/Source/Programs/Horde/Plugins/Build/HordeServer.Build/Perforce/IPerforceService.cs
@@ -2,7 +2,6 @@
 
 using EpicGames.Horde.Commits;
 using EpicGames.Perforce;
-using HordeServer.Commits;
 using HordeServer.Streams;
 using HordeServer.Users;
 
diff --git a/Engine/Source/Programs/Horde/Plugins/Build/HordeServer.Build/Perforce/PerforceService.cs b/Engine/Source/Programs/Horde/Plugins/Build/HordeServer.Build/Perforce/PerforceService.cs
--- a/Engine/Source/Programs/Horde/Plugins/Build/HordeServer.Build/Perforce/PerforceService.cs
+++ b/Engine/Source/Programs/Horde/Plugins/Build/HordeServer.Build/Perforce/PerforceService.cs
@@ -12,7 +12,6 @@ using EpicGames.Horde.Commits;
 using EpicGames.Horde.Streams;
 using EpicGames.Horde.Users;
 using EpicGames.Perforce;
-using HordeServer.Commits;
 using HordeServer.Streams;
 using HordeServer.Users;
 using HordeServer.Utilities;
diff --git a/Engine/Source/Programs/Horde/Plugins/Build/HordeServer.Build/Perforce/PerforceServiceCache.cs b/Engine/Source/Programs/Horde/Plugins/Build/HordeServer.Build/Perforce/PerforceServiceCache.cs
--- a/Engine/Source/Programs/Horde/Plugins/Build/HordeServer.Build/Perforce/PerforceServiceCache.cs
+++ b/Engine/Source/Programs/Horde/Plugins/Build/HordeServer.Build/Perforce/PerforceServiceCache.cs
@@ -8,7 +8,6 @@ using EpicGames.Horde.Streams;
 using EpicGames.Horde.Users;
 using EpicGames.Perforce;
 using EpicGames.Redis;
-using HordeServer.Commits;
 using HordeServer.Server;
 using HordeServer.Streams;
 using HordeServer.Users;
@@ -790,7 +789,7 @@ namespace HordeServer.Perforce
 
 						Task task = updateEvent.Task;
 
-						await foreach(ICommit commit in FindCachedCommitsInAscendingOrderAsync(minChange, tags, cancellationToken))
+						await foreach (ICommit commit in FindCachedCommitsInAscendingOrderAsync(minChange, tags, cancellationToken))
 						{
 							yield return commit;
 							minChange = commit.Id.GetPerforceChange();
diff --git a/Engine/Source/Programs/Horde/Plugins/Build/HordeServer.Build/Replicators/PerforceReplicator.cs b/Engine/Source/Programs/Horde/Plugins/Build/HordeServer.Build/Replicators/PerforceReplicator.cs
--- a/Engine/Source/Programs/Horde/Plugins/Build/HordeServer.Build/Replicators/PerforceReplicator.cs
+++ b/Engine/Source/Programs/Horde/Plugins/Build/HordeServer.Build/Replicators/PerforceReplicator.cs
@@ -11,7 +11,6 @@ using EpicGames.Horde.Storage;
 using EpicGames.Horde.Storage.Nodes;
 using EpicGames.Horde.Streams;
 using EpicGames.Perforce;
-using HordeServer.Commits;
 using HordeServer.Perforce;
 using HordeServer.Storage;
 using HordeServer.Streams;
diff --git a/Engine/Source/Programs/Horde/Plugins/Build/HordeServer.Build/Streams/IStream.cs b/Engine/Source/Programs/Horde/Plugins/Build/HordeServer.Build/Streams/IStream.cs
--- a/Engine/Source/Programs/Horde/Plugins/Build/HordeServer.Build/Streams/IStream.cs
+++ b/Engine/Source/Programs/Horde/Plugins/Build/HordeServer.Build/Streams/IStream.cs
@@ -9,7 +9,6 @@ using EpicGames.Horde.Projects;
 using EpicGames.Horde.Streams;
 using EpicGames.Horde.Users;
 using Horde.Common.Rpc;
-using HordeServer.Commits;
 using HordeServer.Jobs.Schedules;
 
 namespace HordeServer.Streams
diff --git a/Engine/Source/Programs/Shared/EpicGames.Horde/Commits/ICommit.cs b/Engine/Source/Programs/Shared/EpicGames.Horde/Commits/ICommit.cs
new file mode 100644
--- /dev/null
+++ b/Engine/Source/Programs/Shared/EpicGames.Horde/Commits/ICommit.cs
@@ -0,0 +1,98 @@
+﻿// Copyright Epic Games, Inc. All Rights Reserved.
+
+using System;
+using System.Collections.Generic;
+using System.Threading;
+using System.Threading.Tasks;
+using EpicGames.Core;
+using EpicGames.Horde.Streams;
+using EpicGames.Horde.Users;
+
+namespace EpicGames.Horde.Commits
+{
+	/// <summary>
+	/// Stores metadata about a commit
+	/// </summary>
+	public interface ICommit
+	{
+		/// <summary>
+		/// Id for this commit
+		/// </summary>
+		CommitIdWithOrder Id { get; }
+
+		/// <summary>
+		/// Stream containing the commit
+		/// </summary>
+		StreamId StreamId { get; }
+
+		/// <summary>
+		/// The change that this commit originates from
+		/// </summary>
+		CommitIdWithOrder OriginalCommitId { get; }
+
+		/// <summary>
+		/// The author user id
+		/// </summary>
+		UserId AuthorId { get; }
+
+		/// <summary>
+		/// The owner of this change, if different from the author (due to Robomerge)
+		/// </summary>
+		UserId OwnerId { get; }
+
+		/// <summary>
+		/// Changelist description
+		/// </summary>
+		string Description { get; }
+
+		/// <summary>
+		/// Base path for all files in the change
+		/// </summary>
+		string BasePath { get; }
+
+		/// <summary>
+		/// Date/time that change was committed
+		/// </summary>
+		DateTime DateUtc { get; }
+
+		/// <summary>
+		/// Gets the list of tags for the commit
+		/// </summary>
+		/// <param name="cancellationToken">Cancellation token for the operation</param>
+		/// <returns>True if the commit has the given tag</returns>
+		ValueTask<IReadOnlyList<CommitTag>> GetTagsAsync(CancellationToken cancellationToken);
+
+		/// <summary>
+		/// Determine if this commit matches the given filter. Prefer using commit tags rather than this method; the results can be cached.
+		/// </summary>
+		/// <param name="filter">Filter to test</param>
+		/// <param name="cancellationToken">Cancellation token for the operation</param>
+		/// <returns></returns>
+		ValueTask<bool> MatchesFilterAsync(FileFilter filter, CancellationToken cancellationToken);
+
+		/// <summary>
+		/// Gets the files for this change, relative to the root of the stream
+		/// </summary>
+		/// <param name="minFiles">Minimum number of files to return. The response will include at least this number of files, unless the commit has fewer files.</param>
+		/// <param name="maxFiles">Maximum number of files to return. Querying large number of files may cause performance issues with merge commits.</param>
+		/// <param name="cancellationToken">Cancellation token for the operation</param>
+		/// <returns>List of files modified by this commit</returns>
+		ValueTask<IReadOnlyList<string>> GetFilesAsync(int? minFiles, int? maxFiles, CancellationToken cancellationToken);
+	}
+
+	/// <summary>
+	/// Extension methods for operating on commits
+	/// </summary>
+	public static class CommitExtensions
+	{
+		/// <summary>
+		/// Gets the files for this change, relative to the root of the stream
+		/// </summary>
+		/// <param name="commit">Commit to operate on</param>
+		/// <param name="files">Number of files to return</param>
+		/// <param name="cancellationToken">Cancellation token for the operation</param>
+		/// <returns>List of files modified by this commit</returns>
+		public static ValueTask<IReadOnlyList<string>> GetFilesAsync(this ICommit commit, int files, CancellationToken cancellationToken)
+			=> commit.GetFilesAsync(files, files, cancellationToken);
+	}
+}
diff --git a/Engine/Source/Programs/Shared/EpicGames.Horde/Commits/ICommitCollection.cs b/Engine/Source/Programs/Shared/EpicGames.Horde/Commits/ICommitCollection.cs
new file mode 100644
--- /dev/null
+++ b/Engine/Source/Programs/Shared/EpicGames.Horde/Commits/ICommitCollection.cs
@@ -0,0 +1,136 @@
+﻿// Copyright Epic Games, Inc. All Rights Reserved.
+
+using System;
+using System.Collections.Generic;
+using System.Linq;
+using System.Threading;
+using System.Threading.Tasks;
+using EpicGames.Horde.Jobs.Templates;
+
+namespace EpicGames.Horde.Commits
+{
+	/// <summary>
+	/// VCS abstraction. Provides information about commits to a particular stream.
+	/// </summary>
+	public interface ICommitCollection
+	{
+		/// <summary>
+		/// Creates a new change
+		/// </summary>
+		/// <param name="path">Path to modify in the change</param>
+		/// <param name="description">Description of the change</param>
+		/// <param name="cancellationToken">Cancellation token for the operation</param>
+		/// <returns>New commit information</returns>
+		Task<CommitIdWithOrder> CreateNewAsync(string path, string description, CancellationToken cancellationToken = default);
+
+		/// <summary>
+		/// Gets a commit by id
+		/// </summary>
+		/// <param name="commitId">Commit to query</param>
+		/// <param name="cancellationToken">Cancellation token for the operation</param>
+		/// <returns>Commit details</returns>
+		Task<ICommit> GetAsync(CommitId commitId, CancellationToken cancellationToken = default);
+
+		/// <summary>
+		/// Gets an ordered commit id
+		/// </summary>
+		/// <param name="commitId">The commit to query</param>
+		/// <param name="cancellationToken">Cancellation token for the operation</param>
+		/// <returns>Numbered commit id</returns>
+		ValueTask<CommitIdWithOrder> GetOrderedAsync(CommitId commitId, CancellationToken cancellationToken = default);
+
+		/// <summary>
+		/// Finds changes submitted to a stream, in reverse order.
+		/// </summary>
+		/// <param name="minCommitId">The minimum changelist number</param>
+		/// <param name="includeMinCommit">Whether to include the minimum changelist in the range of enumerated responses</param>
+		/// <param name="maxCommitId">The maximum changelist number</param>
+		/// <param name="includeMaxCommit">Whether to include the maximum changelist in the range of enumerated responses</param>
+		/// <param name="maxResults">Maximum number of results to return</param>
+		/// <param name="tags">Tags for the commits to return</param>
+		/// <param name="cancellationToken">Cancellation token for the operation</param>
+		/// <returns>Changelist information</returns>
+		IAsyncEnumerable<ICommit> FindAsync(CommitId? minCommitId = null, bool includeMinCommit = true, CommitId? maxCommitId = null, bool includeMaxCommit = true, int? maxResults = null, IReadOnlyList<CommitTag>? tags = null, CancellationToken cancellationToken = default);
+
+		/// <summary>
+		/// Subscribes to changes from this commit source
+		/// </summary>
+		/// <param name="minCommitId">Minimum changelist number (exclusive)</param>
+		/// <param name="tags">Tags for the commit to return</param>
+		/// <param name="cancellationToken">Cancellation token for the operation</param>
+		/// <returns>New change information</returns>
+		IAsyncEnumerable<ICommit> SubscribeAsync(CommitId minCommitId, IReadOnlyList<CommitTag>? tags = null, CancellationToken cancellationToken = default);
+	}
+
+	/// <summary>
+	/// Extension methods for <see cref="ICommitCollection"/>
+	/// </summary>
+	public static class CommitCollectionExtensions
+	{
+		/// <summary>
+		/// Creates a new change for a template
+		/// </summary>
+		/// <param name="commitCollection">The Perforce service instance</param>
+		/// <param name="template">The template being built</param>
+		/// <param name="cancellationToken">Cancellation token for the operation</param>
+		/// <returns>New changelist number</returns>
+		public static Task<CommitIdWithOrder> CreateNewAsync(this ICommitCollection commitCollection, ITemplate template, CancellationToken cancellationToken)
+		{
+			string description = (template.SubmitDescription ?? "[Horde] New change for $(TemplateName)").Replace("$(TemplateName)", template.Name, StringComparison.OrdinalIgnoreCase);
+			return commitCollection.CreateNewAsync(template.SubmitNewChange!, description, cancellationToken);
+		}
+
+		/// <summary>
+		/// Finds changes submitted to a stream, in reverse order.
+		/// </summary>
+		/// <param name="commitCollection">Collection to operate on</param>
+		/// <param name="minCommitId">The minimum changelist number</param>
+		/// <param name="maxCommitId">The maximum changelist number</param>
+		/// <param name="maxResults">Maximum number of results to return</param>
+		/// <param name="tags">Tags for the commits to return</param>
+		/// <param name="cancellationToken">Cancellation token for the operation</param>
+		/// <returns>Changelist information</returns>
+		public static IAsyncEnumerable<ICommit> FindAsync(this ICommitCollection commitCollection, CommitId? minCommitId, CommitId? maxCommitId, int? maxResults = null, IReadOnlyList<CommitTag>? tags = null, CancellationToken cancellationToken = default)
+			=> commitCollection.FindAsync(minCommitId: minCommitId, maxCommitId: maxCommitId, maxResults: maxResults, tags: tags, cancellationToken: cancellationToken);
+
+		/// <summary>
+		/// Gets the last code code equal or before the given change number
+		/// </summary>
+		/// <param name="commitCollection">The commit source to query</param>
+		/// <param name="maxCommitId">Maximum code change to query</param>
+		/// <param name="cancellationToken">Cancellation token for the operation</param>
+		/// <returns>The last code change</returns>
+		public static async ValueTask<ICommit?> GetLastCodeChangeAsync(this ICommitCollection commitCollection, CommitId? maxCommitId, CancellationToken cancellationToken = default)
+		{
+			return await commitCollection.FindAsync(minCommitId: null, maxCommitId: maxCommitId, maxResults: 1, tags: new[] { CommitTag.Code }, cancellationToken: cancellationToken).FirstOrDefaultAsync(cancellationToken);
+		}
+
+		/// <summary>
+		/// Finds the latest commit from a source
+		/// </summary>
+		/// <param name="commitCollection">The commit source to query</param>
+		/// <param name="cancellationToken">Cancellation token for the operation</param>
+		/// <returns>The latest commit</returns>
+		public static async Task<ICommit> GetLatestAsync(this ICommitCollection commitCollection, CancellationToken cancellationToken = default)
+		{
+			ICommit? commit = await commitCollection.FindAsync(null, null, maxResults: 1, cancellationToken: cancellationToken).FirstOrDefaultAsync(cancellationToken);
+			if (commit == null)
+			{
+				throw new Exception("No changes found for stream.");
+			}
+			return commit;
+		}
+
+		/// <summary>
+		/// Finds the latest commit from a source
+		/// </summary>
+		/// <param name="commitCollection">The commit source to query</param>
+		/// <param name="cancellationToken">Cancellation token for the operation</param>
+		/// <returns>The latest commit</returns>
+		public static async Task<CommitIdWithOrder> GetLastCommitIdAsync(this ICommitCollection commitCollection, CancellationToken cancellationToken = default)
+		{
+			ICommit commit = await GetLatestAsync(commitCollection, cancellationToken);
+			return commit.Id;
+		}
+	}
+}
diff --git a/Engine/Source/Programs/Shared/EpicGames.Horde/EpicGames.Horde.csproj b/Engine/Source/Programs/Shared/EpicGames.Horde/EpicGames.Horde.csproj
--- a/Engine/Source/Programs/Shared/EpicGames.Horde/EpicGames.Horde.csproj
+++ b/Engine/Source/Programs/Shared/EpicGames.Horde/EpicGames.Horde.csproj
@@ -57,6 +57,7 @@
       <PrivateAssets>all</PrivateAssets>
       <IncludeAssets>runtime; build; native; contentfiles; analyzers; buildtransitive</IncludeAssets>
     </PackageReference>
+    <PackageReference Include="System.Linq.Async" Version="6.0.1" />
     <PackageReference Include="ZstdSharp.Port" Version="0.8.1" />
   </ItemGroup>
   <ItemGroup>
