Horde: Separate EpicGames.Horde.Tests into separate namespaces. Also move serializers for time intervals into EpicGames.Horde.

#rnx

diff --git a/Engine/Source/Programs/Horde/HordeServer.Shared/Utilities/IntervalJsonConverter.cs b/Engine/Source/Programs/Horde/HordeServer.Shared/Utilities/IntervalJsonConverter.cs
deleted file mode 100644
diff --git a/Engine/Source/Programs/Horde/HordeServer.Shared/Utilities/TimeOfDayJsonConverter.cs b/Engine/Source/Programs/Horde/HordeServer.Shared/Utilities/TimeOfDayJsonConverter.cs
deleted file mode 100644
diff --git a/Engine/Source/Programs/Horde/HordeServer.Tests/Utilities/IntervalJsonConverterTests.cs b/Engine/Source/Programs/Horde/HordeServer.Tests/Utilities/IntervalJsonConverterTests.cs
deleted file mode 100644
diff --git a/Engine/Source/Programs/Horde/HordeServer.Tests/Utilities/TimeOfDayJsonConverterTests.cs b/Engine/Source/Programs/Horde/HordeServer.Tests/Utilities/TimeOfDayJsonConverterTests.cs
deleted file mode 100644
diff --git a/Engine/Source/Programs/Horde/Plugins/Build/HordeServer.Build/Jobs/Schedules/ScheduleConfig.cs b/Engine/Source/Programs/Horde/Plugins/Build/HordeServer.Build/Jobs/Schedules/ScheduleConfig.cs
--- a/Engine/Source/Programs/Horde/Plugins/Build/HordeServer.Build/Jobs/Schedules/ScheduleConfig.cs
+++ b/Engine/Source/Programs/Horde/Plugins/Build/HordeServer.Build/Jobs/Schedules/ScheduleConfig.cs
@@ -7,6 +7,7 @@ using System.Text.Json.Serialization;
 using EpicGames.Core;
 using EpicGames.Horde.Commits;
 using EpicGames.Horde.Jobs.Templates;
+using EpicGames.Horde.Utilities;
 using HordeServer.Acls;
 using HordeServer.Streams;
 using HordeServer.Utilities;
diff --git a/Engine/Source/Programs/Shared/EpicGames.Horde.Tests/BackendTests.cs b/Engine/Source/Programs/Shared/EpicGames.Horde.Tests/BackendTests.cs
deleted file mode 100644
diff --git a/Engine/Source/Programs/Shared/EpicGames.Horde.Tests/BlobDataTests.cs b/Engine/Source/Programs/Shared/EpicGames.Horde.Tests/BlobDataTests.cs
deleted file mode 100644
diff --git a/Engine/Source/Programs/Shared/EpicGames.Horde.Tests/BlobPipelineTests.cs b/Engine/Source/Programs/Shared/EpicGames.Horde.Tests/BlobPipelineTests.cs
deleted file mode 100644
diff --git a/Engine/Source/Programs/Shared/EpicGames.Horde.Tests/BufferTests.cs b/Engine/Source/Programs/Shared/EpicGames.Horde.Tests/BufferTests.cs
deleted file mode 100644
diff --git a/Engine/Source/Programs/Shared/EpicGames.Horde.Tests/BundleStoreTests.cs b/Engine/Source/Programs/Shared/EpicGames.Horde.Tests/BundleStoreTests.cs
deleted file mode 100644
diff --git a/Engine/Source/Programs/Shared/EpicGames.Horde.Tests/BundleTests.cs b/Engine/Source/Programs/Shared/EpicGames.Horde.Tests/BundleTests.cs
deleted file mode 100644
diff --git a/Engine/Source/Programs/Shared/EpicGames.Horde.Tests/ChunkingTests.cs b/Engine/Source/Programs/Shared/EpicGames.Horde.Tests/ChunkingTests.cs
deleted file mode 100644
diff --git a/Engine/Source/Programs/Shared/EpicGames.Horde.Tests/CompressionTests.cs b/Engine/Source/Programs/Shared/EpicGames.Horde.Tests/CompressionTests.cs
deleted file mode 100644
diff --git a/Engine/Source/Programs/Shared/EpicGames.Horde.Tests/Compute/BufferTests.cs b/Engine/Source/Programs/Shared/EpicGames.Horde.Tests/Compute/BufferTests.cs
new file mode 100644
--- /dev/null
+++ b/Engine/Source/Programs/Shared/EpicGames.Horde.Tests/Compute/BufferTests.cs
@@ -0,0 +1,199 @@
+// Copyright Epic Games, Inc. All Rights Reserved.
+
+using System;
+using System.IO.Pipelines;
+using System.Linq;
+using System.Runtime.InteropServices;
+using System.Security.Cryptography;
+using System.Threading;
+using System.Threading.Tasks;
+using EpicGames.Horde.Compute;
+using EpicGames.Horde.Compute.Buffers;
+using EpicGames.Horde.Compute.Transports;
+using Microsoft.Extensions.Logging.Abstractions;
+using Microsoft.VisualStudio.TestTools.UnitTesting;
+
+namespace EpicGames.Horde.Tests.Compute
+{
+	[TestClass]
+	public class BufferTests
+	{
+		const int ChannelId = 0;
+
+		[TestMethod]
+		public async Task TestSimpleBufferAsync()
+		{
+			using PooledBuffer buffer = new PooledBuffer(2, 1024);
+			using ComputeBufferWriter bufferWriter = buffer.CreateWriter();
+			bufferWriter.AdvanceWritePosition(10);
+
+			using ComputeBufferReader bufferReader = buffer.CreateReader();
+			await bufferReader.WaitToReadAsync(9);
+			bufferReader.AdvanceReadPosition(9);
+			await bufferReader.WaitToReadAsync(1);
+		}
+
+		[TestMethod]
+		public async Task TestOverflowAsync()
+		{
+			using PooledBuffer buffer = new PooledBuffer(2, 20);
+			using ComputeBufferReader bufferReader = buffer.CreateReader();
+			using ComputeBufferWriter bufferWriter = buffer.CreateWriter();
+
+			// Fill up the first chunk
+			Assert.AreEqual(20, bufferWriter.GetWriteBuffer().Length);
+			bufferWriter.AdvanceWritePosition(10);
+			Assert.AreEqual(10, bufferWriter.GetWriteBuffer().Length);
+			bufferWriter.AdvanceWritePosition(10);
+			Assert.AreEqual(0, bufferWriter.GetWriteBuffer().Length);
+
+			Task waitToWriteTask = bufferWriter.WaitToWriteAsync(1).AsTask();
+			Assert.IsTrue(waitToWriteTask.IsCompleted);
+
+			// Fill up the second chunk
+			Assert.AreEqual(20, bufferWriter.GetWriteBuffer().Length);
+			bufferWriter.AdvanceWritePosition(10);
+			Assert.AreEqual(10, bufferWriter.GetWriteBuffer().Length);
+			bufferWriter.AdvanceWritePosition(10);
+			Assert.AreEqual(0, bufferWriter.GetWriteBuffer().Length);
+
+			waitToWriteTask = bufferWriter.WaitToWriteAsync(1).AsTask();
+			Assert.IsFalse(waitToWriteTask.IsCompleted);
+
+			// Wait for data to be read
+			Assert.AreEqual(20, bufferReader.GetReadBuffer().Length);
+			bufferReader.AdvanceReadPosition(10);
+			Assert.IsFalse(waitToWriteTask.IsCompleted);
+
+			Assert.AreEqual(10, bufferReader.GetReadBuffer().Length);
+			bufferReader.AdvanceReadPosition(10);
+			Assert.AreEqual(0, bufferReader.GetReadBuffer().Length);
+
+			Task waitToReadTask = bufferReader.WaitToReadAsync(1).AsTask();
+			Assert.IsTrue(waitToReadTask.IsCompleted);
+
+			Assert.AreEqual(20, bufferReader.GetReadBuffer().Length);
+			await waitToWriteTask;
+
+			// Make sure both reader and writer have something to work with
+			Assert.AreEqual(20, bufferReader.GetReadBuffer().Length);
+			Assert.AreEqual(20, bufferWriter.GetWriteBuffer().Length);
+		}
+
+		[TestMethod]
+		public async Task TestPooledBufferAsync()
+		{
+			await TestProducerConsumerAsync(length => new PooledBuffer(length), CancellationToken.None);
+		}
+
+		[TestMethod]
+		public async Task TestSharedMemoryBufferAsync()
+		{
+			if (RuntimeInformation.IsOSPlatform(OSPlatform.Windows))
+			{
+				await TestProducerConsumerAsync(length => SharedMemoryBuffer.CreateNew(null, length), CancellationToken.None);
+			}
+		}
+
+		static async Task TestProducerConsumerAsync(Func<int, ComputeBuffer> createBuffer, CancellationToken cancellationToken)
+		{
+			_ = cancellationToken;
+
+			const int Length = 8000;
+
+			Pipe sourceToTargetPipe = new Pipe();
+			Pipe targetToSourcePipe = new Pipe();
+
+			await using PipeTransport producerTransport = new(targetToSourcePipe.Reader, sourceToTargetPipe.Writer);
+			await using PipeTransport consumerTransport = new(sourceToTargetPipe.Reader, targetToSourcePipe.Writer);
+			await using RemoteComputeSocket producerSocket = new(producerTransport, ComputeProtocol.Latest, NullLogger.Instance);
+			await using RemoteComputeSocket consumerSocket = new(consumerTransport, ComputeProtocol.Latest, NullLogger.Instance);
+
+			using ComputeBuffer consumerBuffer = createBuffer(Length);
+			consumerSocket.AttachRecvBuffer(ChannelId, consumerBuffer);
+
+			using ComputeBuffer producerBuffer = createBuffer(Length);
+			producerSocket.AttachSendBuffer(ChannelId, producerBuffer);
+
+			byte[] input = RandomNumberGenerator.GetBytes(Length);
+			Task producerTask = RunProducerAsync(producerBuffer, input);
+
+			using ComputeBufferReader consumerBufferReader = consumerBuffer.CreateReader();
+
+			byte[] output = new byte[Length];
+			await RunConsumerAsync(consumerBufferReader, output);
+
+			await producerTask;
+			Assert.IsTrue(input.SequenceEqual(output));
+		}
+
+		static async Task RunProducerAsync(ComputeBuffer buffer, ReadOnlyMemory<byte> input)
+		{
+			using ComputeBufferWriter writer = buffer.CreateWriter();
+
+			int offset = 0;
+			while (offset < input.Length)
+			{
+				int length = Math.Min(input.Length - offset, 100);
+				await writer.WriteAsync(input.Slice(offset, length));
+				await Task.Delay(10);
+				offset += length;
+			}
+
+			writer.MarkComplete();
+		}
+
+		static async Task RunConsumerAsync(ComputeBufferReader reader, Memory<byte> output)
+		{
+			int offset = 0;
+			while (!reader.IsComplete)
+			{
+				ReadOnlyMemory<byte> memory = reader.GetReadBuffer();
+				if (memory.Length == 0)
+				{
+					await reader.WaitToReadAsync(1, CancellationToken.None);
+					continue;
+				}
+
+				int length = Math.Min(memory.Length, 7);
+				memory.Slice(0, length).CopyTo(output.Slice(offset));
+				reader.AdvanceReadPosition(length);
+				offset += length;
+			}
+		}
+
+		[TestMethod]
+		public async Task TestSendBufferCompleteAsync()
+		{
+			Pipe recvPipe = new Pipe();
+			Pipe sendPipe = new Pipe();
+			await using PipeTransport localTransport = new(sendPipe.Reader, recvPipe.Writer);
+			await using PipeTransport remoteTransport = new(recvPipe.Reader, sendPipe.Writer);
+			await using RemoteComputeSocket localSocket = new(localTransport, ComputeProtocol.Latest, NullLogger.Instance);
+			await using RemoteComputeSocket remoteSocket = new(remoteTransport, ComputeProtocol.Latest, NullLogger.Instance);
+
+			using (PooledBuffer remoteBuffer = new PooledBuffer(1024))
+			{
+				remoteSocket.AttachRecvBuffer(1, remoteBuffer);
+
+				using ComputeBufferReader reader = remoteBuffer.CreateReader();
+
+				// Disposing of the buffer should mark the channel as complete
+				using (PooledBuffer localBuffer = new PooledBuffer(1024))
+				{
+					localSocket.AttachSendBuffer(1, localBuffer);
+
+					using ComputeBufferWriter localBufferWriter = localBuffer.CreateWriter();
+					await localBufferWriter.WriteAsync(new byte[] { 1, 2, 3 });
+				}
+
+				Assert.IsTrue(await reader.WaitToReadAsync(3));
+				Assert.IsTrue(reader.GetReadBuffer().Slice(0, 3).Span.SequenceEqual(new byte[] { 1, 2, 3 }));
+				reader.AdvanceReadPosition(3);
+
+				Assert.IsFalse(await reader.WaitToReadAsync(1));
+				Assert.IsTrue(reader.IsComplete);
+			}
+		}
+	}
+}
diff --git a/Engine/Source/Programs/Shared/EpicGames.Horde.Tests/Compute/ComputeSocketTests.cs b/Engine/Source/Programs/Shared/EpicGames.Horde.Tests/Compute/ComputeSocketTests.cs
new file mode 100644
--- /dev/null
+++ b/Engine/Source/Programs/Shared/EpicGames.Horde.Tests/Compute/ComputeSocketTests.cs
@@ -0,0 +1,247 @@
+// Copyright Epic Games, Inc. All Rights Reserved.
+
+using System;
+using System.Collections.Generic;
+using System.IO.Pipelines;
+using System.Linq;
+using System.Net;
+using System.Net.Sockets;
+using System.Text;
+using System.Threading;
+using System.Threading.Tasks;
+using EpicGames.Core;
+using EpicGames.Horde.Compute;
+using EpicGames.Horde.Compute.Transports;
+using EpicGames.Horde.Storage;
+using EpicGames.Horde.Storage.Bundles;
+using EpicGames.Horde.Storage.Nodes;
+using Microsoft.Extensions.Logging;
+using Microsoft.Extensions.Logging.Abstractions;
+using Microsoft.VisualStudio.TestTools.UnitTesting;
+
+namespace EpicGames.Horde.Tests.Compute
+{
+	[TestClass]
+	public class ComputeSocketTests
+	{
+		class TestComputeSocket : ComputeSocket, IDisposable
+		{
+			public Dictionary<int, ComputeBufferWriter> RecvBufferWriters { get; } = new Dictionary<int, ComputeBufferWriter>();
+			public Dictionary<int, ComputeBufferReader> SendBufferReaders { get; } = new Dictionary<int, ComputeBufferReader>();
+
+			public override ComputeProtocol Protocol => ComputeProtocol.Latest;
+			public override ILogger Logger => NullLogger.Instance;
+
+			public void Dispose()
+			{
+				foreach (ComputeBufferWriter writer in RecvBufferWriters.Values)
+				{
+					writer.Dispose();
+				}
+				foreach (ComputeBufferReader reader in SendBufferReaders.Values)
+				{
+					reader.Dispose();
+				}
+			}
+
+			public override void AttachRecvBuffer(int channelId, ComputeBuffer recvBuffer)
+			{
+				RecvBufferWriters.Add(channelId, recvBuffer.CreateWriter());
+			}
+
+			public override void AttachSendBuffer(int channelId, ComputeBuffer sendBuffer)
+			{
+				SendBufferReaders.Add(channelId, sendBuffer.CreateReader());
+			}
+		}
+
+		class TestLogger : ILogger
+		{
+			public IDisposable? BeginScope<TState>(TState state) where TState : notnull => null!;
+
+			public bool IsEnabled(LogLevel logLevel) => true;
+
+			public void Log<TState>(LogLevel logLevel, EventId eventId, TState state, Exception? exception, Func<TState, Exception?, string> formatter)
+			{
+				Console.WriteLine($"{logLevel}: {formatter(state, exception)}");
+				Assert.IsFalse(logLevel == LogLevel.Error);
+				Assert.IsFalse(logLevel == LogLevel.Warning);
+			}
+		}
+
+		[TestMethod]
+		public async Task TestAgentMessageLoopPipeAsync()
+		{
+			Pipe recvPipe = new Pipe();
+			Pipe sendPipe = new Pipe();
+			await using PipeTransport localTransport = new(sendPipe.Reader, recvPipe.Writer);
+			await using PipeTransport agentTransport = new(recvPipe.Reader, sendPipe.Writer);
+			await using RemoteComputeSocket localSocket = new(localTransport, ComputeProtocol.Latest, new TestLogger());
+			await using RemoteComputeSocket agentSocket = new(agentTransport, ComputeProtocol.Latest, new TestLogger());
+
+			await RunAgentTestsAsync(localSocket, agentSocket);
+		}
+
+		[TestMethod]
+		public async Task TestAgentMessageLoopTcpAsync()
+		{
+			using CancellationTokenSource cts = new(5000);
+			(Socket clientSocket, Socket serverSocket) = await CreateSocketsAsync(cts.Token);
+
+			await using TcpTransport clientTransport = new(clientSocket);
+			await using TcpTransport serverTransport = new(serverSocket);
+			await using RemoteComputeSocket localSocket = new(clientTransport, ComputeProtocol.Latest, new TestLogger());
+			await using RemoteComputeSocket agentSocket = new(serverTransport, ComputeProtocol.Latest, new TestLogger());
+
+			await RunAgentTestsAsync(localSocket, agentSocket, cts.Token);
+		}
+
+		[TestMethod]
+		[DataRow(Encryption.Ssl)]
+		[DataRow(Encryption.SslEcdsaP256)]
+		public async Task TestAgentMessageLoopTcpSslAsync(Encryption encryption)
+		{
+			using CancellationTokenSource cts = new(5000);
+			(Socket clientSocket, Socket serverSocket) = await CreateSocketsAsync(cts.Token);
+
+			byte[] certData = TcpSslTransport.GenerateCert(encryption);
+			await using TcpSslTransport clientTransport = new(clientSocket, certData, false);
+			await using TcpSslTransport serverTransport = new(serverSocket, certData, true);
+
+			Task t1 = clientTransport.AuthenticateAsync(cts.Token);
+			Task t2 = serverTransport.AuthenticateAsync(cts.Token);
+			await t2;
+			await t1;
+
+			await using RemoteComputeSocket localSocket = new(clientTransport, ComputeProtocol.Latest, new TestLogger());
+			await using RemoteComputeSocket agentSocket = new(serverTransport, ComputeProtocol.Latest, new TestLogger());
+			await RunAgentTestsAsync(localSocket, agentSocket, cts.Token);
+		}
+
+		private static async Task<(Socket client, Socket server)> CreateSocketsAsync(CancellationToken cancellationToken)
+		{
+			int port = GetAvailablePort();
+			using TcpListener listener = new (IPAddress.Loopback, port);
+			listener.Start();
+			Socket clientSocket = new(SocketType.Stream, ProtocolType.Tcp);
+			Task clientConnectTask = clientSocket.ConnectAsync(IPAddress.Loopback, port, cancellationToken).AsTask();
+			Socket serverSocket = await listener.AcceptSocketAsync(cancellationToken);
+			await clientConnectTask;
+			return (clientSocket, serverSocket);
+		}
+
+		static async Task RunAgentTestsAsync(RemoteComputeSocket localSocket, RemoteComputeSocket agentSocket, CancellationToken cancellationToken = default)
+		{
+			DirectoryReference tempDir = new DirectoryReference("test-temp");
+			await using (BackgroundTask agentTask = BackgroundTask.StartNew(ctx => RunAgentAsync(agentSocket, tempDir, ctx)))
+			{
+				const int PrimaryChannelId = 0;
+				using (AgentMessageChannel channel = localSocket.CreateAgentMessageChannel(PrimaryChannelId, 4 * 1024 * 1024))
+				{
+					await channel.WaitForAttachAsync(cancellationToken);
+
+					await channel.PingAsync(cancellationToken);
+					using (AgentMessage message = await channel.ReceiveAsync(cancellationToken))
+					{
+						Assert.AreEqual(AgentMessageType.Ping, message.Type);
+						Assert.IsTrue(message.Data.Span.SequenceEqual(ReadOnlySpan<byte>.Empty));
+					}
+
+					await channel.SendXorRequestAsync(new byte[] { 1, 2, 3 }, 44, cancellationToken);
+					using (AgentMessage message = await channel.ReceiveAsync(cancellationToken))
+					{
+						Assert.AreEqual(AgentMessageType.XorResponse, message.Type);
+						Assert.IsTrue(message.Data.Span.SequenceEqual(new byte[] { 1 ^ 44, 2 ^ 44, 3 ^ 44 }));
+					}
+
+					const int SecondaryChannelId = 1;
+					using (AgentMessageChannel channel2 = localSocket.CreateAgentMessageChannel(SecondaryChannelId, 4 * 1024 * 1024))
+					{
+						await channel.ForkAsync(SecondaryChannelId, 4 * 1024 * 1024, cancellationToken);
+
+						await channel2.WaitForAttachAsync(cancellationToken);
+
+						await channel2.SendXorRequestAsync(new byte[] { 1, 2, 3 }, 44, cancellationToken);
+						using (AgentMessage message = await channel2.ReceiveAsync(cancellationToken))
+						{
+							Assert.AreEqual(AgentMessageType.XorResponse, message.Type);
+							Assert.IsTrue(message.Data.Span.SequenceEqual(new byte[] { 1 ^ 44, 2 ^ 44, 3 ^ 44 }));
+						}
+
+						await channel2.CloseAsync(cancellationToken);
+					}
+
+					BundleStorageNamespace storage = BundleStorageNamespace.CreateInMemory(NullLogger.Instance);
+					await using (IBlobWriter blobWriter = storage.CreateBlobWriter(cancellationToken: cancellationToken))
+					{
+						FileReference file = FileReference.Combine(tempDir, "subdir/hello.txt");
+						if (FileReference.Exists(file))
+						{
+							FileReference.Delete(file);
+						}
+						Assert.IsFalse(FileReference.Exists(file));
+
+						byte[] data = Encoding.UTF8.GetBytes("Hello world");
+
+						using ChunkedDataWriter writer = new ChunkedDataWriter(blobWriter, new ChunkingOptions());
+						ChunkedData chunkedData = await writer.CreateAsync(data, cancellationToken);
+
+						DirectoryNode directory = new DirectoryNode();
+						directory.AddFile("hello.txt", FileEntryFlags.None, data.Length, chunkedData);
+
+						IHashedBlobRef<DirectoryNode> directoryRef = await blobWriter.WriteBlobAsync(directory, cancellationToken: cancellationToken);
+
+						DirectoryNode root = new DirectoryNode();
+						root.AddDirectory(new DirectoryEntry("subdir", directory.Length, directoryRef));
+
+						IHashedBlobRef<DirectoryNode> handle = await blobWriter.WriteBlobAsync(root, cancellationToken: cancellationToken);
+						await blobWriter.FlushAsync(cancellationToken);
+
+						await channel.UploadFilesAsync("", handle.GetLocator(), storage.Backend, cancellationToken);
+
+						Assert.IsTrue(FileReference.Exists(file));
+						byte[] readData = await FileReference.ReadAllBytesAsync(file, cancellationToken);
+						Assert.IsTrue(readData.SequenceEqual(data));
+
+						await channel.DeleteFilesAsync(new[] { "subdir/hello.txt" }, cancellationToken);
+
+						await channel.PingAsync(cancellationToken);
+						using (AgentMessage message = await channel.ReceiveAsync(cancellationToken))
+						{
+							Assert.AreEqual(AgentMessageType.Ping, message.Type);
+							Assert.IsTrue(message.Data.Span.SequenceEqual(ReadOnlySpan<byte>.Empty));
+						}
+
+						Assert.IsFalse(FileReference.Exists(file));
+					}
+				}
+			}
+
+			await localSocket.CloseAsync(CancellationToken.None);
+			await agentSocket.CloseAsync(CancellationToken.None);
+		}
+
+		static int GetAvailablePort()
+		{
+			using TcpListener listener = new(IPAddress.Loopback, 0);
+			listener.Start();
+			int port = ((IPEndPoint)listener.LocalEndpoint).Port;
+			listener.Stop();
+			return port;
+		}
+
+		static async Task RunAgentAsync(ComputeSocket socket, DirectoryReference tempDir, CancellationToken cancellationToken)
+		{
+			try
+			{
+				AgentMessageHandler handler = new AgentMessageHandler(tempDir, null, true, null, null, NullLogger.Instance);
+				await handler.RunAsync(socket, cancellationToken);
+			}
+			catch (Exception e)
+			{
+				Console.WriteLine("Exception when running agent:\n" + e);
+				throw;
+			}
+		}
+	}
+}
diff --git a/Engine/Source/Programs/Shared/EpicGames.Horde.Tests/Compute/ExternalIpResolverTests.cs b/Engine/Source/Programs/Shared/EpicGames.Horde.Tests/Compute/ExternalIpResolverTests.cs
new file mode 100644
--- /dev/null
+++ b/Engine/Source/Programs/Shared/EpicGames.Horde.Tests/Compute/ExternalIpResolverTests.cs
@@ -0,0 +1,73 @@
+﻿// Copyright Epic Games, Inc. All Rights Reserved.
+
+using System;
+using System.Net;
+using System.Net.Http;
+using System.Threading;
+using System.Threading.Tasks;
+using EpicGames.Horde.Compute;
+using Microsoft.VisualStudio.TestTools.UnitTesting;
+
+namespace EpicGames.Horde.Tests.Compute;
+
+[TestClass]
+public class ExternalIpResolverTests
+{
+	private class StubMessageHandler : HttpMessageHandler
+	{
+		private readonly HttpStatusCode _statusCode;
+		public string Content { get; set; }
+
+		public StubMessageHandler(HttpStatusCode statusCode, string content)
+		{
+			_statusCode = statusCode;
+			Content = content;
+		}
+
+		protected override Task<HttpResponseMessage> SendAsync(HttpRequestMessage request, CancellationToken cancellationToken)
+		{
+			return Task.FromResult(new HttpResponseMessage(_statusCode) { Content = new StringContent(Content) });
+		}
+	}
+
+	[TestMethod]
+	public async Task BasicAsync()
+	{
+		Assert.AreEqual(IPAddress.Parse("192.168.2.3"), await GetIpAsync("192.168.2.3"));
+		await Assert.ThrowsExceptionAsync<ExternalIpResolverException>(() => GetIpAsync("bad-ip"));
+		await Assert.ThrowsExceptionAsync<ExternalIpResolverException>(() => GetIpAsync("192.168.2.3", HttpStatusCode.NotFound));
+	}
+
+	[TestMethod]
+	public async Task CacheIpAddressAsync()
+	{
+		using StubMessageHandler stub = new(HttpStatusCode.OK, "100.200.1.1");
+		using HttpClient httpClient = new(stub);
+		ExternalIpResolver resolver = new(httpClient);
+		IPAddress ip = await resolver.GetExternalIpAddressAsync(CancellationToken.None);
+		Assert.AreEqual(IPAddress.Parse("100.200.1.1"), ip);
+
+		// Querying for IP again should never reach the IP resolver's HTTP server, so an invalid IP is returned to ensure this
+		stub.Content = "invalid-ip";
+		ip = await resolver.GetExternalIpAddressAsync(CancellationToken.None);
+		Assert.AreEqual(IPAddress.Parse("100.200.1.1"), ip);
+	}
+
+	[TestMethod]
+	[Ignore]
+	public async Task IntegrationAsync()
+	{
+		using HttpClient httpClient = new();
+		ExternalIpResolver resolver = new(httpClient);
+		IPAddress externalIp = await resolver.GetExternalIpAddressAsync();
+		Console.WriteLine("External IP: " + externalIp);
+	}
+
+	private static async Task<IPAddress> GetIpAsync(string content, HttpStatusCode statusCode = HttpStatusCode.OK, CancellationToken cancellationToken = default)
+	{
+		using StubMessageHandler handler = new(statusCode, content);
+		using HttpClient httpClient = new(handler);
+		ExternalIpResolver resolver = new(httpClient);
+		return await resolver.GetExternalIpAddressAsync(cancellationToken);
+	}
+}
\ No newline at end of file
diff --git a/Engine/Source/Programs/Shared/EpicGames.Horde.Tests/Compute/RequestIdAllocatorTests.cs b/Engine/Source/Programs/Shared/EpicGames.Horde.Tests/Compute/RequestIdAllocatorTests.cs
new file mode 100644
--- /dev/null
+++ b/Engine/Source/Programs/Shared/EpicGames.Horde.Tests/Compute/RequestIdAllocatorTests.cs
@@ -0,0 +1,36 @@
+// Copyright Epic Games, Inc. All Rights Reserved.
+
+using System.Collections.Generic;
+using EpicGames.Horde.Compute;
+using Microsoft.VisualStudio.TestTools.UnitTesting;
+
+namespace EpicGames.Horde.Tests.Compute;
+
+[TestClass]
+public class RequestIdAllocatorTests
+{
+	[TestMethod]
+	public void IdsAreReused()
+	{
+		RequestIdAllocator ria = new();
+		List<string> before = new() { ria.AllocateId(), ria.AllocateId(), ria.AllocateId() };
+		ria.StartBatch();
+		List<string> after = new() { ria.AllocateId(), ria.AllocateId(), ria.AllocateId() };
+
+		CollectionAssert.AreEquivalent(before, after);
+	}
+
+	[TestMethod]
+	public void AcceptedIdsAreNotReused()
+	{
+		RequestIdAllocator ria = new();
+		List<string> before = new() { ria.AllocateId(), ria.AllocateId(), ria.AllocateId() };
+		ria.MarkAccepted(before[1]);
+
+		ria.StartBatch();
+		List<string> after = new() { ria.AllocateId(), ria.AllocateId(), ria.AllocateId() };
+		Assert.IsTrue(after.Remove(before[0]));
+		Assert.IsFalse(after.Remove(before[1]));
+		Assert.IsTrue(after.Remove(before[2]));
+	}
+}
diff --git a/Engine/Source/Programs/Shared/EpicGames.Horde.Tests/Compute/TransportTests.cs b/Engine/Source/Programs/Shared/EpicGames.Horde.Tests/Compute/TransportTests.cs
new file mode 100644
--- /dev/null
+++ b/Engine/Source/Programs/Shared/EpicGames.Horde.Tests/Compute/TransportTests.cs
@@ -0,0 +1,105 @@
+// Copyright Epic Games, Inc. All Rights Reserved.
+
+using System;
+using System.IO;
+using System.Linq;
+using System.Security.Cryptography;
+using System.Threading;
+using System.Threading.Tasks;
+using EpicGames.Horde.Compute;
+using EpicGames.Horde.Compute.Transports;
+using Microsoft.VisualStudio.TestTools.UnitTesting;
+
+namespace EpicGames.Horde.Tests.Compute
+{
+	[TestClass]
+	public class TransportTests
+	{
+		[TestMethod]
+		public async Task TestStreamTransportAsync()
+		{
+			byte[] input = new byte[256 * 1024];
+			new Random(0).NextBytes(input);
+
+			byte[] payload;
+			{
+				using MemoryStream memoryStream = new MemoryStream();
+				await using StreamTransport streamTransport = new StreamTransport(memoryStream);
+				await streamTransport.SendAsync(input, CancellationToken.None);
+				payload = memoryStream.ToArray();
+			}
+
+			byte[] output = new byte[input.Length];
+			{
+				using MemoryStream memoryStream = new MemoryStream(payload);
+				await using StreamTransport streamTransport = new StreamTransport(memoryStream);
+				await CopyToAsync(streamTransport, output);
+			}
+
+			for (int idx = 0; idx < input.Length; idx++)
+			{
+				Assert.IsTrue(input[idx] == output[idx]);
+			}
+			Assert.IsTrue(input.SequenceEqual(output));
+
+		}
+
+		[TestMethod]
+		public async Task TestAesTransportAsync()
+		{
+			byte[] key = new byte[AesTransport.KeyLength];
+			new Random(1).NextBytes(key);
+
+			byte[] nonce = new byte[AesTransport.NonceLength];
+			new Random(2).NextBytes(nonce);
+
+			byte[] input = RandomNumberGenerator.GetBytes(256 * 1024);
+			new Random(0).NextBytes(input);
+
+			byte[] encrypted;
+
+			{
+				using MemoryStream memoryStream = new MemoryStream();
+				await using (StreamTransport streamTransport = new StreamTransport(memoryStream))
+				{
+					await using (AesTransport aesTransport = new AesTransport(streamTransport, key, nonce))
+					{
+						await aesTransport.SendAsync(input, CancellationToken.None);
+					}
+				}
+				encrypted = memoryStream.ToArray();
+			}
+
+			byte[] output = new byte[input.Length];
+			{
+				using MemoryStream memoryStream = new MemoryStream(encrypted);
+				await using (StreamTransport streamTransport = new StreamTransport(memoryStream))
+				{
+					await using (AesTransport aesTransport = new AesTransport(streamTransport, key, nonce))
+					{
+						for (int offset = 0; offset < output.Length;)
+						{
+							offset += await aesTransport.RecvAsync(output.AsMemory(offset), CancellationToken.None);
+						}
+					}
+				}
+			}
+
+			for (int idx = 0; idx < input.Length; idx++)
+			{
+				Assert.IsTrue(input[idx] == output[idx]);
+			}
+			Assert.IsTrue(input.SequenceEqual(output));
+		}
+
+		public static async Task CopyToAsync(ComputeTransport transport, Memory<byte> buffer)
+		{
+			for (int offset = 0; offset < buffer.Length;)
+			{
+				int size = await transport.RecvAsync(buffer.Slice(offset), CancellationToken.None);
+				Assert.IsTrue(size > 0);
+				offset += size;
+			}
+		}
+	}
+}
diff --git a/Engine/Source/Programs/Shared/EpicGames.Horde.Tests/ComputeSocketTests.cs b/Engine/Source/Programs/Shared/EpicGames.Horde.Tests/ComputeSocketTests.cs
deleted file mode 100644
diff --git a/Engine/Source/Programs/Shared/EpicGames.Horde.Tests/ContentChunkerTests.cs b/Engine/Source/Programs/Shared/EpicGames.Horde.Tests/ContentChunkerTests.cs
deleted file mode 100644
diff --git a/Engine/Source/Programs/Shared/EpicGames.Horde.Tests/ExternalIpResolverTests.cs b/Engine/Source/Programs/Shared/EpicGames.Horde.Tests/ExternalIpResolverTests.cs
deleted file mode 100644
diff --git a/Engine/Source/Programs/Shared/EpicGames.Horde.Tests/FileMetadataDbTests.cs b/Engine/Source/Programs/Shared/EpicGames.Horde.Tests/FileMetadataDbTests.cs
deleted file mode 100644
diff --git a/Engine/Source/Programs/Shared/EpicGames.Horde.Tests/LogTests.cs b/Engine/Source/Programs/Shared/EpicGames.Horde.Tests/LogTests.cs
deleted file mode 100644
diff --git a/Engine/Source/Programs/Shared/EpicGames.Horde.Tests/Logs/LogTests.cs b/Engine/Source/Programs/Shared/EpicGames.Horde.Tests/Logs/LogTests.cs
new file mode 100644
--- /dev/null
+++ b/Engine/Source/Programs/Shared/EpicGames.Horde.Tests/Logs/LogTests.cs
@@ -0,0 +1,212 @@
+// Copyright Epic Games, Inc. All Rights Reserved.
+
+using System;
+using System.Collections.Generic;
+using System.Linq;
+using System.Text;
+using System.Threading;
+using System.Threading.Tasks;
+using EpicGames.Horde.Logs;
+using EpicGames.Horde.Storage;
+using EpicGames.Horde.Storage.Clients;
+using EpicGames.Horde.Tests.Properties;
+using Microsoft.Extensions.Logging.Abstractions;
+using Microsoft.VisualStudio.TestTools.UnitTesting;
+
+namespace EpicGames.Horde.Tests
+{
+	[TestClass]
+	public sealed class LogTests
+	{
+		private readonly byte[] _data = Resources.TextFile;
+
+		[TestMethod]
+		public void NgramTests()
+		{
+			NgramSetBuilder builder = new NgramSetBuilder();
+
+			ulong[] values = { 1, 2, 3, 5, 8, 13, 21, 34, 55, 89, 144, 233, 377, 610, 987, 1597, 2584, 4181, 6765, 10946, 0xfedcba9876543210UL };
+			foreach (ulong value in values)
+			{
+				builder.Add(value);
+			}
+
+			NgramSet trie = builder.ToNgramSet();
+			Assert.IsTrue(Enumerable.SequenceEqual(trie.EnumerateRange(0, UInt64.MaxValue), values));
+			Assert.IsTrue(Enumerable.SequenceEqual(trie.EnumerateRange(0, 90), values.Where(x => x <= 90)));
+			Assert.IsTrue(Enumerable.SequenceEqual(trie.EnumerateRange(2, 89), values.Where(x => x >= 2 && x <= 89)));
+		}
+
+		[TestMethod]
+		public async Task IndexTestsAsync()
+		{
+			KeyValueStorageNamespace store = KeyValueStorageNamespace.CreateInMemory();
+
+			// Write the test data to the log file in blocks
+			LogBuilder builder = new LogBuilder(LogFormat.Text, 1, 1, NullLogger.Instance);
+
+			int readOffset = 0;
+			int readLineIndex = 0;
+			while (readOffset < _data.Length)
+			{
+				int length = 0;
+				int lineCount = 0;
+
+				for (int idx = readOffset; idx < Math.Min(_data.Length, readOffset + 1883); idx++)
+				{
+					if (_data[idx] == '\n')
+					{
+						length = (idx + 1) - readOffset;
+						lineCount++;
+						break;
+					}
+				}
+
+				builder.WriteData(_data.AsMemory(readOffset, length));
+
+				readOffset += length;
+				readLineIndex += lineCount;
+			}
+
+			// Flush it to storage, and read the finished log node
+			IHashedBlobRef<LogNode> logRef;
+			await using (IBlobWriter writer = store.CreateBlobWriter())
+			{
+				logRef = await builder.FlushAsync(writer, true, CancellationToken.None);
+			}
+			LogNode log = await logRef.ReadBlobAsync();
+
+			// Read the data back out and check it's the same
+			byte[] readData = new byte[_data.Length];
+
+			int writeOffset = 0;
+			await foreach (ReadOnlyMemory<byte> line in log.ReadLogLinesAsync(0))
+			{
+				line.CopyTo(readData.AsMemory(writeOffset));
+				writeOffset += line.Length;
+			}
+
+			int equalSize = 0;
+			while (equalSize < _data.Length && _data[equalSize] == readData[equalSize])
+			{
+				equalSize++;
+			}
+			Assert.AreEqual(equalSize, readOffset);
+
+			// Test some searches
+			LogIndexNode index = await log.IndexRef.ReadBlobAsync();
+			await SearchLogDataTestAsync(index);
+		}
+
+		[TestMethod]
+		public async Task PartialTokenTestsAsync()
+		{
+			IStorageNamespace store = KeyValueStorageNamespace.CreateInMemory();
+
+			// Generate the test data
+			string[] lines =
+			{
+				"abcdefghi\n",
+				"jklmno123\n",
+				"pqrst99uv\n",
+				"wx\n"
+			};
+
+			LogBuilder builder = new LogBuilder(LogFormat.Text, 1, 1, NullLogger.Instance);
+			for (int lineIdx = 0; lineIdx < lines.Length; lineIdx++)
+			{
+				builder.WriteData(Encoding.UTF8.GetBytes(lines[lineIdx]));
+			}
+
+			IHashedBlobRef<LogNode> rootNodeRef;
+			await using (IBlobWriter writer = store.CreateBlobWriter())
+			{
+				rootNodeRef = await builder.FlushAsync(writer, true, CancellationToken.None);
+			}
+
+			// Read it back in and test the index
+			LogNode rootNode = await rootNodeRef.ReadBlobAsync();
+			LogIndexNode index = await rootNode.IndexRef.ReadBlobAsync();
+
+			for (int lineIdx = 0; lineIdx < lines.Length; lineIdx++)
+			{
+				for (int strLen = 1; strLen < 7; strLen++)
+				{
+					for (int strOfs = 0; strOfs + strLen < lines[lineIdx].Length - 1; strOfs++)
+					{
+						string str = lines[lineIdx].Substring(strOfs, strLen);
+
+						SearchStats stats = new SearchStats();
+						List<int> results = await index.SearchAsync(0, new SearchTerm(str), stats).ToListAsync();
+						Assert.AreEqual(1, results.Count);
+						Assert.AreEqual(lineIdx, results[0]);
+
+						Assert.AreEqual(1, stats.NumScannedBlocks);
+						Assert.AreEqual(3, stats.NumSkippedBlocks);
+						Assert.AreEqual(0, stats.NumFalsePositiveBlocks);
+					}
+				}
+			}
+		}
+		/*
+		[TestMethod]
+		public async Task AppendIndexTests()
+		{
+			JobId jobId = JobId.GenerateNewId();
+			ILogFile logFile = await _logFileService.CreateLogFileAsync(jobId, null, LogType.Text);
+
+			logFile = await WriteLogDataAsync(logFile, 0, 0, Encoding.UTF8.GetBytes("abc\n"), true);
+			logFile = await WriteLogDataAsync(logFile, 4, 1, Encoding.UTF8.GetBytes("def\n"), true);
+			logFile = await WriteLogDataAsync(logFile, 8, 2, Encoding.UTF8.GetBytes("ghi\n"), false);
+
+			await ((LogFileService)_logFileService).FlushPendingWritesAsync();
+
+			{
+				LogSearchStats stats = new LogSearchStats();
+				List<int> results = await _logFileService.SearchLogDataAsync(logFile, "abc", 0, 5, stats, CancellationToken.None);
+				Assert.AreEqual(1, results.Count);
+				Assert.AreEqual(0, results[0]);
+
+				Assert.AreEqual(2, stats.NumScannedBlocks); // abc + ghi (no index yet because it hasn't been flushed)
+				Assert.AreEqual(1, stats.NumSkippedBlocks); // def
+				Assert.AreEqual(0, stats.NumFalsePositiveBlocks);
+			}
+			{
+				LogSearchStats stats = new LogSearchStats();
+				List<int> results = await _logFileService.SearchLogDataAsync(logFile, "def", 0, 5, stats, CancellationToken.None);
+				Assert.AreEqual(1, results.Count);
+				Assert.AreEqual(1, results[0]);
+
+				Assert.AreEqual(2, stats.NumScannedBlocks); // def + ghi (no index yet because it hasn't been flushed)
+				Assert.AreEqual(1, stats.NumSkippedBlocks); // abc
+				Assert.AreEqual(0, stats.NumFalsePositiveBlocks);
+			}
+			{
+				LogSearchStats stats = new LogSearchStats();
+				List<int> results = await _logFileService.SearchLogDataAsync(logFile, "ghi", 0, 5, stats, CancellationToken.None);
+				Assert.AreEqual(1, results.Count);
+				Assert.AreEqual(2, results[0]);
+
+				Assert.AreEqual(1, stats.NumScannedBlocks); // ghi
+				Assert.AreEqual(2, stats.NumSkippedBlocks); // abc + def
+				Assert.AreEqual(0, stats.NumFalsePositiveBlocks);
+			}
+		}
+*/
+		static async Task SearchLogDataTestAsync(LogIndexNode index)
+		{
+			await SearchLogDataTestAsync(index, "HISPANIOLA", 0, 4, new[] { 1503, 1520, 1525, 1595 });
+			await SearchLogDataTestAsync(index, "Hispaniola", 0, 4, new[] { 1503, 1520, 1525, 1595 });
+			await SearchLogDataTestAsync(index, "HizpaniolZ", 0, 4, Array.Empty<int>());
+			await SearchLogDataTestAsync(index, "Pieces of eight!", 0, 100, new[] { 2227, 2228, 5840, 5841, 7520 });
+			await SearchLogDataTestAsync(index, "NEWSLETTER", 0, 100, new[] { 7886 });
+		}
+
+		static async Task SearchLogDataTestAsync(LogIndexNode index, string text, int firstLine, int count, int[] expectedLines)
+		{
+			SearchStats stats = new SearchStats();
+			List<int> lines = await index.SearchAsync(firstLine, new SearchTerm(text), stats, CancellationToken.None).Take(count).ToListAsync();
+			Assert.IsTrue(lines.SequenceEqual(expectedLines));
+		}
+	}
+}
diff --git a/Engine/Source/Programs/Shared/EpicGames.Horde.Tests/Logs/ServerLogPacketBuilderTests.cs b/Engine/Source/Programs/Shared/EpicGames.Horde.Tests/Logs/ServerLogPacketBuilderTests.cs
new file mode 100644
--- /dev/null
+++ b/Engine/Source/Programs/Shared/EpicGames.Horde.Tests/Logs/ServerLogPacketBuilderTests.cs
@@ -0,0 +1,424 @@
+// Copyright Epic Games, Inc. All Rights Reserved.
+
+using System;
+using System.Collections.Generic;
+using System.Linq;
+using System.Text;
+using System.Threading;
+using System.Threading.Channels;
+using System.Threading.Tasks;
+using EpicGames.Core;
+using EpicGames.Horde.Artifacts;
+using EpicGames.Horde.Compute;
+using EpicGames.Horde.Logs;
+using EpicGames.Horde.Projects;
+using EpicGames.Horde.Secrets;
+using EpicGames.Horde.Storage;
+using EpicGames.Horde.Storage.Bundles;
+using EpicGames.Horde.Tools;
+using Grpc.Core;
+using Grpc.Net.Client;
+using Horde.Common.Rpc;
+using Microsoft.Extensions.Logging;
+using Microsoft.Extensions.Logging.Abstractions;
+using Microsoft.VisualStudio.TestTools.UnitTesting;
+
+namespace EpicGames.Horde.Tests
+{
+	[TestClass]
+	public class ServerLogPacketBuilderTests
+	{
+		class JsonLoggerImpl : ILogger
+		{
+			public List<LogEvent> _lines = new List<LogEvent>();
+
+			public IDisposable? BeginScope<TState>(TState state) where TState : notnull => null!;
+
+			public bool IsEnabled(LogLevel logLevel) => true;
+
+			public void Log<TState>(LogLevel logLevel, EventId eventId, TState state, Exception? exception, Func<TState, Exception?, string> formatter)
+			{
+				JsonLogEvent logEvent = JsonLogEvent.FromLoggerState(logLevel, eventId, state, exception, formatter);
+				_lines.Add(LogEvent.Read(logEvent.Data.Span));
+			}
+		}
+
+		[TestMethod]
+		public void JsonLogEventTest()
+		{
+			JsonLoggerImpl impl = new JsonLoggerImpl();
+			impl.LogInformation("Hello {Text}", "world");
+			impl.LogInformation("Hello {Text}", "world");
+
+			Assert.AreEqual(impl._lines.Count, 2);
+			Assert.AreEqual(impl._lines[0].ToString(), "Hello world");
+			Assert.AreEqual(impl._lines[0].GetProperty("Text")?.ToString(), "world");
+			Assert.AreEqual(impl._lines[1].ToString(), "Hello world");
+			Assert.AreEqual(impl._lines[1].GetProperty("Text")?.ToString(), "world");
+		}
+
+		static string[] SplitMultiLineMessage(string message)
+		{
+			JsonLogEvent jsonLogEvent;
+			Assert.IsTrue(JsonLogEvent.TryParse(Encoding.UTF8.GetBytes(message), out jsonLogEvent));
+
+			ServerLogPacketBuilder writer = new ServerLogPacketBuilder();
+			int count = writer.SanitizeAndWriteEvent(jsonLogEvent);
+
+			string[] result = Encoding.UTF8.GetString(writer.CreatePacket().Item1.Span).Split('\n', StringSplitOptions.RemoveEmptyEntries);
+			Assert.AreEqual(count, result.Length);
+			return result;
+		}
+
+		[TestMethod]
+		public void MultiLineFormatTest()
+		{
+			const string Msg = @"{""level"":""Information"",""message"":""ignored"",""format"":""This\nis a\nmulti-{Line}-end\nlog message {Var1} {Var2}"",""properties"":{""Line"":""line\nsplit\nin\nfour"",""Var1"":123,""Var2"":{""$type"":""SourceFile"",""$text"":""D:\\build\\\u002B\u002BUE5\\Sync\\GenerateProjectFiles.bat"",""relativePath"":""GenerateProjectFiles.bat"",""depotPath"":""//UE5/Main/GenerateProjectFiles.bat@20392842""}}}";
+
+			string[] result = SplitMultiLineMessage(Msg);
+			Assert.AreEqual(7, result.Length);
+			Assert.AreEqual(@"{""level"":""Information"",""message"":""This"",""format"":""This"",""properties"":{""Var1"":123,""Var2"":{""$type"":""SourceFile"",""$text"":""D:\\build\\\u002B\u002BUE5\\Sync\\GenerateProjectFiles.bat"",""relativePath"":""GenerateProjectFiles.bat"",""depotPath"":""//UE5/Main/GenerateProjectFiles.bat@20392842""}},""line"":0,""lineCount"":7}", result[0]);
+			Assert.AreEqual(@"{""level"":""Information"",""message"":""is a"",""format"":""is a"",""properties"":{""Var1"":123,""Var2"":{""$type"":""SourceFile"",""$text"":""D:\\build\\\u002B\u002BUE5\\Sync\\GenerateProjectFiles.bat"",""relativePath"":""GenerateProjectFiles.bat"",""depotPath"":""//UE5/Main/GenerateProjectFiles.bat@20392842""}},""line"":1,""lineCount"":7}", result[1]);
+			Assert.AreEqual(@"{""level"":""Information"",""message"":""multi-line"",""format"":""multi-{Line$0}"",""properties"":{""Var1"":123,""Var2"":{""$type"":""SourceFile"",""$text"":""D:\\build\\\u002B\u002BUE5\\Sync\\GenerateProjectFiles.bat"",""relativePath"":""GenerateProjectFiles.bat"",""depotPath"":""//UE5/Main/GenerateProjectFiles.bat@20392842""},""Line$0"":""line""},""line"":2,""lineCount"":7}", result[2]);
+			Assert.AreEqual(@"{""level"":""Information"",""message"":""split"",""format"":""{Line$1}"",""properties"":{""Var1"":123,""Var2"":{""$type"":""SourceFile"",""$text"":""D:\\build\\\u002B\u002BUE5\\Sync\\GenerateProjectFiles.bat"",""relativePath"":""GenerateProjectFiles.bat"",""depotPath"":""//UE5/Main/GenerateProjectFiles.bat@20392842""},""Line$1"":""split""},""line"":3,""lineCount"":7}", result[3]);
+			Assert.AreEqual(@"{""level"":""Information"",""message"":""in"",""format"":""{Line$2}"",""properties"":{""Var1"":123,""Var2"":{""$type"":""SourceFile"",""$text"":""D:\\build\\\u002B\u002BUE5\\Sync\\GenerateProjectFiles.bat"",""relativePath"":""GenerateProjectFiles.bat"",""depotPath"":""//UE5/Main/GenerateProjectFiles.bat@20392842""},""Line$2"":""in""},""line"":4,""lineCount"":7}", result[4]);
+			Assert.AreEqual(@"{""level"":""Information"",""message"":""four-end"",""format"":""{Line$3}-end"",""properties"":{""Var1"":123,""Var2"":{""$type"":""SourceFile"",""$text"":""D:\\build\\\u002B\u002BUE5\\Sync\\GenerateProjectFiles.bat"",""relativePath"":""GenerateProjectFiles.bat"",""depotPath"":""//UE5/Main/GenerateProjectFiles.bat@20392842""},""Line$3"":""four""},""line"":5,""lineCount"":7}", result[5]);
+			Assert.AreEqual(@"{""level"":""Information"",""message"":""log message 123 D:\\build\\\u002B\u002BUE5\\Sync\\GenerateProjectFiles.bat"",""format"":""log message {Var1} {Var2}"",""properties"":{""Var1"":123,""Var2"":{""$type"":""SourceFile"",""$text"":""D:\\build\\\u002B\u002BUE5\\Sync\\GenerateProjectFiles.bat"",""relativePath"":""GenerateProjectFiles.bat"",""depotPath"":""//UE5/Main/GenerateProjectFiles.bat@20392842""}},""line"":6,""lineCount"":7}", result[6]);
+		}
+
+		[TestMethod]
+		public void InvalidCharactersText()
+		{
+			byte[] data = Encoding.UTF8.GetBytes(@"{""level"":""Information"",""message"":""test"",""format"":""GOOD\nBAD""}");
+			int idx = data.AsSpan().IndexOf(Encoding.UTF8.GetBytes("BAD"));
+			data[idx] = 0xef;
+			data[idx + 1] = 0xbb;
+
+			ServerLogPacketBuilder writer = new ServerLogPacketBuilder();
+			int count = writer.SanitizeAndWriteEvent(JsonLogEvent.Parse(data));
+			Assert.AreEqual(1, count);
+
+			string[] result = Encoding.UTF8.GetString(writer.CreatePacket().Item1.Span).Split('\n', StringSplitOptions.RemoveEmptyEntries);
+			Assert.AreEqual(1, result.Length);
+
+			LogEvent logEvent = LogEvent.Read(Encoding.UTF8.GetBytes(result[0]));
+
+			const string ExpectedError = "Invalid json log event: {\"level\":\"Information\",\"message\":\"test\",\"format\":\"GOOD\\n\\xef\\xbbD\"}";
+			Assert.AreEqual(ExpectedError, logEvent.ToString());
+		}
+
+		[TestMethod]
+		public void MultiLineMessageTest()
+		{
+			const string Msg = @"{""level"":""Information"",""message"":""This is \na multi-line string""}";
+
+			string[] result = SplitMultiLineMessage(Msg);
+			Assert.AreEqual(2, result.Length);
+			Assert.AreEqual(result[0], @"{""level"":""Information"",""message"":""This is "",""line"":0,""lineCount"":2}");
+			Assert.AreEqual(result[1], @"{""level"":""Information"",""message"":""a multi-line string"",""line"":1,""lineCount"":2}");
+		}
+
+		[TestMethod]
+		public void MultiLineDecoyTest()
+		{
+			const string Msg = @"{""level"":""Information"",""message"":""This is \\\\not a multi-line format string""}";
+
+			string[] result = SplitMultiLineMessage(Msg);
+			Assert.AreEqual(1, result.Length);
+			Assert.AreEqual(result[0], Msg);
+		}
+
+		class GrpcHelpers
+		{
+			class StreamReader<T> : IAsyncStreamReader<T> where T : class
+			{
+				readonly ChannelReader<T> _reader;
+				T? _current;
+
+				public T Current => _current ?? throw new InvalidOperationException();
+
+				public StreamReader(ChannelReader<T> reader)
+					=> _reader = reader;
+
+				public async Task<bool> MoveNext(CancellationToken cancellationToken)
+				{
+					while (await _reader.WaitToReadAsync(cancellationToken))
+					{
+						if (_reader.TryRead(out _current))
+						{
+							return true;
+						}
+					}
+					return false;
+				}
+			}
+
+			class StreamWriter<T> : IClientStreamWriter<T>
+			{
+				readonly ChannelWriter<T> _writer;
+
+				public StreamWriter(ChannelWriter<T> writer)
+					=> _writer = writer;
+
+				public WriteOptions? WriteOptions { get => throw new NotImplementedException(); set => throw new NotImplementedException(); }
+
+				public Task CompleteAsync()
+				{
+					_writer.Complete();
+					return Task.CompletedTask;
+				}
+
+				public async Task WriteAsync(T message)
+					=> await _writer.WriteAsync(message);
+			}
+
+			public static AsyncDuplexStreamingCall<TRequest, TResponse> CreateDuplexCall<TRequest, TResponse>(Func<ChannelReader<TRequest>, ChannelWriter<TResponse>, Task> func)
+				where TRequest : class
+				where TResponse : class
+			{
+				Channel<TRequest> requests = Channel.CreateUnbounded<TRequest>();
+				Channel<TResponse> responses = Channel.CreateUnbounded<TResponse>();
+
+				BackgroundTask runner = BackgroundTask.StartNew(async ctx =>
+				{
+					await func(requests.Reader, responses.Writer);
+					responses.Writer.Complete();
+				});
+
+				return new AsyncDuplexStreamingCall<TRequest, TResponse>(new StreamWriter<TRequest>(requests.Writer), new StreamReader<TResponse>(responses.Reader), Task.FromResult<Metadata>(null!), () => Status.DefaultSuccess, () => null!, () => runner.DisposeAsync().AsTask().Wait());
+			}
+		}
+
+		class FakeLogRpcClient : LogRpc.LogRpcClient
+		{
+			public Dictionary<LogId, BlobLocator> Logs { get; } = new Dictionary<LogId, BlobLocator>();
+
+			public override AsyncUnaryCall<RpcUpdateLogResponse> UpdateLogAsync(RpcUpdateLogRequest request, CallOptions options)
+			{
+				Logs[LogId.Parse(request.LogId)] = new BlobLocator(request.TargetLocator);
+				return new AsyncUnaryCall<RpcUpdateLogResponse>(Task.FromResult(new RpcUpdateLogResponse()), Task.FromResult(new Metadata()), () => Status.DefaultSuccess, () => new Metadata(), () => { });
+			}
+
+			public override AsyncDuplexStreamingCall<RpcUpdateLogTailRequest, RpcUpdateLogTailResponse> UpdateLogTail(CallOptions options)
+			{
+				return GrpcHelpers.CreateDuplexCall<RpcUpdateLogTailRequest, RpcUpdateLogTailResponse>(async (reader, writer) =>
+				{
+					while (await reader.WaitToReadAsync())
+					{
+						reader.TryRead(out RpcUpdateLogTailRequest? request);
+					}
+				});
+			}
+		}
+
+		class FakeHordeClient : IHordeClient, IAsyncDisposable
+		{
+			public FakeLogRpcClient LogRpc { get; } = new FakeLogRpcClient();
+			public Dictionary<string, BundleStorageNamespace> StorageNamespaces { get; } = new Dictionary<string, BundleStorageNamespace>();
+
+			public Uri ServerUrl => throw new NotImplementedException();
+
+			public IArtifactCollection Artifacts => throw new NotImplementedException();
+			public IComputeClient Compute => throw new NotImplementedException();
+			public IProjectCollection Projects => throw new NotImplementedException();
+			public ISecretCollection Secrets => throw new NotImplementedException();
+			public IToolCollection Tools => throw new NotImplementedException();
+
+			public event Action? OnAccessTokenStateChanged
+			{
+				add { }
+				remove { }
+			}
+
+			public Task<bool> LoginAsync(bool allowLogin, CancellationToken cancellationToken)
+				=> throw new NotImplementedException();
+
+			public HordeHttpClient CreateHttpClient()
+				=> throw new NotImplementedException();
+
+			public IStorageNamespace GetStorageNamespace(string relativePath, string? accessToken = null)
+			{
+				BundleStorageNamespace? storageNamespace;
+				if (!StorageNamespaces.TryGetValue(relativePath, out storageNamespace))
+				{
+					storageNamespace = BundleStorageNamespace.CreateInMemory(NullLogger.Instance);
+					StorageNamespaces.Add(relativePath, storageNamespace);
+				}
+				return storageNamespace;
+			}
+
+			public ValueTask DisposeAsync()
+			{
+				StorageNamespaces.Clear();
+				return default;
+			}
+
+			public Task<string?> GetAccessTokenAsync(bool interactive, CancellationToken cancellationToken = default)
+				=> throw new NotImplementedException();
+
+			public Task<GrpcChannel> CreateGrpcChannelAsync(CancellationToken cancellationToken = default)
+				=> throw new NotImplementedException();
+
+			public Task<TClient> CreateGrpcClientAsync<TClient>(CancellationToken cancellationToken = default) where TClient : ClientBase<TClient>
+			{
+				if (typeof(TClient) == typeof(LogRpc.LogRpcClient))
+				{
+					return Task.FromResult<TClient>((TClient)(object)LogRpc);
+				}
+				throw new NotImplementedException();
+			}
+
+			public bool HasValidAccessToken()
+				=> throw new NotImplementedException();
+
+			public IServerLogger CreateServerLogger(LogId logId, LogLevel minimumLevel = LogLevel.Information)
+				=> throw new NotImplementedException();
+		}
+
+		[TestMethod]
+		public async Task StorageLoggerTestAsync()
+		{
+			await using BundleCache cache = new BundleCache();
+			await using FakeHordeClient hordeClient = new FakeHordeClient();
+
+			LogId logId = default;
+
+			const int Count = 20000;
+			await using (ServerLogger logger = new ServerLogger(hordeClient, logId, LogLevel.Information, NullLogger.Instance))
+			{
+				for (int idx = 0; idx < Count; idx++)
+				{
+					logger.LogInformation("Testing {Number}", idx);
+				}
+			}
+
+			// Read the log
+			IStorageNamespace storageNamespace = hordeClient.GetStorageNamespace(logId);
+			LogNode file = await storageNamespace.CreateBlobRef(hordeClient.LogRpc.Logs[logId]).ReadBlobAsync<LogNode>();
+
+			// Check the index text
+			List<Utf8String> extractedIndexText = new List<Utf8String>();
+
+			LogIndexNode index = await file.IndexRef.ReadBlobAsync();
+			foreach (LogChunkRef block in index.PlainTextChunkRefs)
+			{
+				LogChunkNode text = await block.Target.ReadBlobAsync();
+				extractedIndexText.AddRange(text.Lines);
+			}
+
+			Assert.AreEqual(Count, extractedIndexText.Count);
+			for (int idx = 0; idx < Count; idx++)
+			{
+				Assert.AreEqual(extractedIndexText[idx], new Utf8String($"Testing {idx}"));
+			}
+
+			// Check the body text
+			List<string> extractedBodyText = new List<string>();
+
+			foreach (LogChunkRef blockRef in file.TextChunkRefs)
+			{
+				LogChunkNode blockText = await blockRef.Target.ReadBlobAsync();
+				foreach (Utf8String line in blockText.Lines)
+				{
+					LogEvent logEvent = LogEvent.Read(line.Span);
+					extractedBodyText.Add(logEvent.ToString());
+				}
+			}
+
+			Assert.AreEqual(Count, extractedBodyText.Count);
+			for (int idx = 0; idx < Count; idx++)
+			{
+				Assert.AreEqual(extractedBodyText[idx], $"Testing {idx}");
+			}
+		}
+
+		[TestMethod]
+		public void LongMessageTest()
+		{
+			DateTime time = new DateTime(2023, 1, 1, 0, 0, 0, 0, DateTimeKind.Utc);
+
+			string longLineA = String.Join("\n", Enumerable.Range(0, 5).Select(x => $"A{x}"));
+			string longLineB = String.Join("\n", Enumerable.Range(0, 5).Select(x => $"B{x}"));
+			string longLineC = String.Join("\n", Enumerable.Range(0, 5).Select(x => $"C{x}"));
+
+			Dictionary<string, object?> properties = new Dictionary<string, object?>
+			{
+				["LongLineA"] = longLineA,
+				["LongLineB"] = longLineB,
+				["LongLineC"] = longLineC,
+			};
+
+			LogEvent baseEvent = new LogEvent(time, LogLevel.Information, default, $"start {longLineA} x {longLineB} y\nz {longLineC} end", "start {LongLineA} x {LongLineB} y\nz {LongLineC} end", properties, null);
+
+			ServerLogPacketBuilder writer = new ServerLogPacketBuilder();
+			writer.SanitizeAndWriteEvent(new JsonLogEvent(baseEvent));
+
+			string[] output = Encoding.UTF8.GetString(writer.CreatePacket().Item1.Span).Split('\n');
+			string[] expected =
+			{
+				@"{""time"":""2023-01-01T00:00:00"",""level"":""Information"",""message"":""start A0"",""format"":""start {LongLineA$0}"",""properties"":{""LongLineA$0"":""A0""},""line"":0,""lineCount"":14}",
+				@"{""time"":""2023-01-01T00:00:00"",""level"":""Information"",""message"":""A1"",""format"":""{LongLineA$1}"",""properties"":{""LongLineA$1"":""A1""},""line"":1,""lineCount"":14}",
+				@"{""time"":""2023-01-01T00:00:00"",""level"":""Information"",""message"":""A2"",""format"":""{LongLineA$2}"",""properties"":{""LongLineA$2"":""A2""},""line"":2,""lineCount"":14}",
+				@"{""time"":""2023-01-01T00:00:00"",""level"":""Information"",""message"":""A3"",""format"":""{LongLineA$3}"",""properties"":{""LongLineA$3"":""A3""},""line"":3,""lineCount"":14}",
+				@"{""time"":""2023-01-01T00:00:00"",""level"":""Information"",""message"":""A4 x B0"",""format"":""{LongLineA$4} x {LongLineB$0}"",""properties"":{""LongLineA$4"":""A4"",""LongLineB$0"":""B0""},""line"":4,""lineCount"":14}",
+				@"{""time"":""2023-01-01T00:00:00"",""level"":""Information"",""message"":""B1"",""format"":""{LongLineB$1}"",""properties"":{""LongLineB$1"":""B1""},""line"":5,""lineCount"":14}",
+				@"{""time"":""2023-01-01T00:00:00"",""level"":""Information"",""message"":""B2"",""format"":""{LongLineB$2}"",""properties"":{""LongLineB$2"":""B2""},""line"":6,""lineCount"":14}",
+				@"{""time"":""2023-01-01T00:00:00"",""level"":""Information"",""message"":""B3"",""format"":""{LongLineB$3}"",""properties"":{""LongLineB$3"":""B3""},""line"":7,""lineCount"":14}",
+				@"{""time"":""2023-01-01T00:00:00"",""level"":""Information"",""message"":""B4 y"",""format"":""{LongLineB$4} y"",""properties"":{""LongLineB$4"":""B4""},""line"":8,""lineCount"":14}",
+				@"{""time"":""2023-01-01T00:00:00"",""level"":""Information"",""message"":""z C0"",""format"":""z {LongLineC$0}"",""properties"":{""LongLineC$0"":""C0""},""line"":9,""lineCount"":14}",
+				@"{""time"":""2023-01-01T00:00:00"",""level"":""Information"",""message"":""C1"",""format"":""{LongLineC$1}"",""properties"":{""LongLineC$1"":""C1""},""line"":10,""lineCount"":14}",
+				@"{""time"":""2023-01-01T00:00:00"",""level"":""Information"",""message"":""C2"",""format"":""{LongLineC$2}"",""properties"":{""LongLineC$2"":""C2""},""line"":11,""lineCount"":14}",
+				@"{""time"":""2023-01-01T00:00:00"",""level"":""Information"",""message"":""C3"",""format"":""{LongLineC$3}"",""properties"":{""LongLineC$3"":""C3""},""line"":12,""lineCount"":14}",
+				@"{""time"":""2023-01-01T00:00:00"",""level"":""Information"",""message"":""C4 end"",""format"":""{LongLineC$4} end"",""properties"":{""LongLineC$4"":""C4""},""line"":13,""lineCount"":14}",
+				""
+			};
+
+			for (int idx = 0; idx < output.Length; idx++)
+			{
+				Assert.AreEqual(expected[idx], output[idx]);
+			}
+		}
+
+		[TestMethod]
+		public void MaxPacketLengthTest()
+		{
+			DateTime time = new DateTime(2023, 1, 1, 0, 0, 0, 0, DateTimeKind.Utc);
+
+			LogEvent baseEvent = new LogEvent(time, LogLevel.Information, default, "abcdefghijklmnopqrstuvwxyzABCDEFGHIJKLMNOPQRSTUVWXYZ", null, null, null);
+
+			ServerLogPacketBuilder writer = new ServerLogPacketBuilder(maxPacketLength: 85);
+			writer.SanitizeAndWriteEvent(new JsonLogEvent(baseEvent));
+			writer.SanitizeAndWriteEvent(new JsonLogEvent(baseEvent));
+
+			for (int idx = 0; idx < 2; idx++)
+			{
+				string output = Encoding.UTF8.GetString(writer.CreatePacket().Item1.Span);
+				string expected = @"{""time"":""2023-01-01T00:00:00"",""level"":""Information"",""message"":""abcdefghijklmnopqrstuvwxyzABCDEFGHIJKLMNOPQRSTUVWXYZ""}" + "\n";
+				Assert.AreEqual(expected, output);
+			}
+		}
+
+		[TestMethod]
+		public void MaxLineLengthTest()
+		{
+			DateTime time = new DateTime(2023, 1, 1, 0, 0, 0, 0, DateTimeKind.Utc);
+
+			LogEvent baseEvent = new LogEvent(time, LogLevel.Information, default, "abcdefghijklmnopqrstuvwxyzABCDEFGHIJKLMNOPQRSTUVWXYZ", null, null, null);
+
+			ServerLogPacketBuilder writer = new ServerLogPacketBuilder(85);
+			writer.SanitizeAndWriteEvent(new JsonLogEvent(baseEvent));
+
+			string[] output = Encoding.UTF8.GetString(writer.CreatePacket().Item1.Span).Split('\n');
+			string[] expected =
+			{
+				@"{""time"":""2023-01-01T00:00:00"",""level"":""Information"",""message"":""abcdefghijklmnopqrstuvwxyzABCDEFGHI [...]""}",
+				""
+			};
+
+			for (int idx = 0; idx < output.Length; idx++)
+			{
+				Assert.AreEqual(expected[idx], output[idx]);
+			}
+		}
+	}
+}
diff --git a/Engine/Source/Programs/Shared/EpicGames.Horde.Tests/RequestIdAllocatorTests.cs b/Engine/Source/Programs/Shared/EpicGames.Horde.Tests/RequestIdAllocatorTests.cs
deleted file mode 100644
diff --git a/Engine/Source/Programs/Shared/EpicGames.Horde.Tests/ServerLogPacketBuilderTests.cs b/Engine/Source/Programs/Shared/EpicGames.Horde.Tests/ServerLogPacketBuilderTests.cs
deleted file mode 100644
diff --git a/Engine/Source/Programs/Shared/EpicGames.Horde.Tests/Storage/BackendTests.cs b/Engine/Source/Programs/Shared/EpicGames.Horde.Tests/Storage/BackendTests.cs
new file mode 100644
--- /dev/null
+++ b/Engine/Source/Programs/Shared/EpicGames.Horde.Tests/Storage/BackendTests.cs
@@ -0,0 +1,149 @@
+// Copyright Epic Games, Inc. All Rights Reserved.
+
+using System;
+using System.Collections.Generic;
+using System.Linq;
+using System.Text;
+using System.Threading.Tasks;
+using EpicGames.Core;
+using EpicGames.Horde.Storage;
+using EpicGames.Horde.Storage.Backends;
+using Microsoft.Extensions.Logging.Abstractions;
+using Microsoft.VisualStudio.TestTools.UnitTesting;
+
+namespace EpicGames.Horde.Tests.Storage
+{
+	[TestClass]
+	public class BackendTests
+	{
+		class TempDir : IDisposable
+		{
+			readonly DirectoryReference _cacheDir;
+
+			public DirectoryReference Location => _cacheDir;
+
+			public TempDir(string name)
+			{
+				_cacheDir = new DirectoryReference(name);
+				DirectoryReference.CreateDirectory(_cacheDir);
+				FileUtils.ForceDeleteDirectoryContents(_cacheDir);
+			}
+
+			public void Dispose()
+			{
+				FileUtils.ForceDeleteDirectoryContents(_cacheDir);
+				DirectoryReference.Delete(_cacheDir);
+			}
+		}
+
+		[TestMethod]
+		public async Task TestMemoryBackendAsync()
+		{
+			MemoryStorageBackend backend = new MemoryStorageBackend();
+			await TestBackendAsync(backend);
+		}
+
+		[TestMethod]
+		public async Task TestFileBackendAsync()
+		{
+			using (TempDir tempDir = new TempDir("Cache"))
+			{
+				using MemoryMappedFileCache memoryMappedFileCache = new MemoryMappedFileCache();
+				FileStorageBackend backend = new FileStorageBackend(tempDir.Location, memoryMappedFileCache, NullLogger.Instance);
+				await TestBackendAsync(backend);
+			}
+		}
+
+		[TestMethod]
+		public async Task TestFileBackendFileMappingAsync()
+		{
+			using (TempDir tempDir = new TempDir("Cache"))
+			{
+				using MemoryMappedFileCache memoryMappedFileCache = new MemoryMappedFileCache();
+				FileStorageBackend backend = new FileStorageBackend(tempDir.Location, memoryMappedFileCache, NullLogger.Instance);
+				BlobLocator path = await backend.WriteBytesAsync(Encoding.UTF8.GetBytes("hello world"), Array.Empty<BlobLocator>());
+
+				using (IReadOnlyMemoryOwner<byte> handle = await backend.ReadBlobAsync(path, 0, null))
+				{
+					Assert.AreEqual("hello world", Encoding.UTF8.GetString(handle.Memory.Span));
+				}
+				using (IReadOnlyMemoryOwner<byte> handle = await backend.ReadBlobAsync(path, 0, 5))
+				{
+					Assert.AreEqual("hello", Encoding.UTF8.GetString(handle.Memory.Span));
+				}
+				using (IReadOnlyMemoryOwner<byte> handle = await backend.ReadBlobAsync(path, 4, 3))
+				{
+					Assert.AreEqual("o w", Encoding.UTF8.GetString(handle.Memory.Span));
+				}
+			}
+		}
+
+		[TestMethod]
+		public async Task TestCacheBackendAsync()
+		{
+			using (TempDir tempDir = new TempDir("Cache"))
+			{
+				using StorageBackendCache cache = new StorageBackendCache(tempDir.Location, 12, NullLogger.Instance);
+
+				MemoryStorageBackend memoryBackend = new MemoryStorageBackend();
+				IStorageBackend cacheBackend = cache.CreateWrapper("", memoryBackend);
+
+				await TestBackendAsync(cacheBackend);
+
+				Assert.AreEqual(1, cache.Items.Count());
+				byte[] value = await cacheBackend.ReadBytesAsync(cache.GetLocators().First());
+				Assert.IsTrue(value.SequenceEqual(Encoding.UTF8.GetBytes("item 2")));
+
+				byte[] data3 = Encoding.UTF8.GetBytes("3");
+				BlobLocator path3 = await cacheBackend.WriteBytesAsync(data3, Array.Empty<BlobLocator>());
+				await cacheBackend.ReadBytesAsync(path3);
+
+				byte[] data4 = Encoding.UTF8.GetBytes("4");
+				BlobLocator path4 = await cacheBackend.WriteBytesAsync(data4, Array.Empty<BlobLocator>());
+				await cacheBackend.ReadBytesAsync(path4);
+
+				byte[] data5 = Encoding.UTF8.GetBytes("5");
+				BlobLocator path5 = await cacheBackend.WriteBytesAsync(data5, Array.Empty<BlobLocator>());
+				await cacheBackend.ReadBytesAsync(path5);
+
+				HashSet<BlobLocator> paths = new HashSet<BlobLocator>(cache.GetLocators());
+				Assert.AreEqual(4, paths.Count);
+				Assert.IsTrue(paths.Contains(path3));
+				Assert.IsTrue(paths.Contains(path4));
+				Assert.IsTrue(paths.Contains(path5));
+
+				byte[] data6 = Encoding.UTF8.GetBytes("12345678901");
+				BlobLocator path6 = await cacheBackend.WriteBytesAsync(data6, Array.Empty<BlobLocator>());
+				await cacheBackend.ReadBytesAsync(path6);
+
+				paths = new HashSet<BlobLocator>(cache.GetLocators());
+				Assert.AreEqual(2, paths.Count);
+				Assert.IsTrue(paths.Contains(path5));
+				Assert.IsTrue(paths.Contains(path6));
+
+				await cacheBackend.ReadBytesAsync(path3);
+				await cacheBackend.ReadBytesAsync(path4);
+
+				paths = new HashSet<BlobLocator>(cache.GetLocators());
+				Assert.AreEqual(2, paths.Count);
+				Assert.IsTrue(paths.Contains(path3));
+				Assert.IsTrue(paths.Contains(path4));
+			}
+		}
+
+		static async Task TestBackendAsync(IStorageBackend backend)
+		{
+			byte[] data1 = Encoding.UTF8.GetBytes("hello world");
+			byte[] data2 = Encoding.UTF8.GetBytes("item 2");
+
+			BlobLocator path1 = await backend.WriteBytesAsync(data1, Array.Empty<BlobLocator>());
+			byte[] outputData1 = await backend.ReadBytesAsync(path1);
+
+			BlobLocator path2 = await backend.WriteBytesAsync(data2, Array.Empty<BlobLocator>());
+			byte[] outputData2 = await backend.ReadBytesAsync(path2);
+
+			Assert.IsTrue(data1.SequenceEqual(outputData1));
+			Assert.IsTrue(data2.SequenceEqual(outputData2));
+		}
+	}
+}
diff --git a/Engine/Source/Programs/Shared/EpicGames.Horde.Tests/Storage/BlobDataTests.cs b/Engine/Source/Programs/Shared/EpicGames.Horde.Tests/Storage/BlobDataTests.cs
new file mode 100644
--- /dev/null
+++ b/Engine/Source/Programs/Shared/EpicGames.Horde.Tests/Storage/BlobDataTests.cs
@@ -0,0 +1,61 @@
+﻿// Copyright Epic Games, Inc. All Rights Reserved.
+
+using System;
+using System.Collections.Generic;
+using System.Threading;
+using System.Threading.Tasks;
+using EpicGames.Core;
+using EpicGames.Horde.Storage;
+using Microsoft.VisualStudio.TestTools.UnitTesting;
+
+namespace EpicGames.Horde.Tests.Storage
+{
+	[TestClass]
+	public class BlobDataTests
+	{
+		class DummyHandle : IBlobRef
+		{
+			readonly BlobLocator _locator;
+
+			public IBlobRef Innermost => this;
+
+			public DummyHandle(string locator) => _locator = new BlobLocator(new Utf8String(locator));
+			public ValueTask FlushAsync(CancellationToken cancellationToken = default) => default;
+			public ValueTask<BlobData> ReadBlobDataAsync(CancellationToken cancellationToken = default) => throw new NotImplementedException();
+
+			public bool TryGetLocator(out BlobLocator locator)
+			{
+				locator = _locator;
+				return true;
+			}
+
+			public override bool Equals(object? obj) => obj is DummyHandle other && _locator == other._locator;
+			public override int GetHashCode() => _locator.GetHashCode();
+		}
+
+		[TestMethod]
+		public void HeaderSerialization()
+		{
+			List<IBlobRef> refs = new List<IBlobRef>();
+			refs.Add(new DummyHandle("hello"));
+			refs.Add(new DummyHandle("world"));
+
+			using BlobData blobData = new BlobData(new BlobType(Guid.NewGuid(), 4), new byte[] { 1, 2, 3 }, refs);
+
+			byte[] data = EncodedBlobData.Create(blobData);
+			EncodedBlobData packet = new EncodedBlobData(data);
+			Assert.AreEqual(49, data.Length);
+
+			Assert.AreEqual(blobData.Type.Guid, packet.Type.Guid);
+			Assert.AreEqual(blobData.Type.Version, packet.Type.Version);
+			Assert.AreEqual(blobData.Imports.Count, packet.Imports.Count);
+
+			for (int idx = 0; idx < packet.Imports.Count; idx++)
+			{
+				string locator = blobData.Imports[idx].GetLocator().ToString();
+				string encoded = packet.Imports[idx].ToString();
+				Assert.AreEqual(locator, encoded);
+			}
+		}
+	}
+}
diff --git a/Engine/Source/Programs/Shared/EpicGames.Horde.Tests/Storage/BlobPipelineTests.cs b/Engine/Source/Programs/Shared/EpicGames.Horde.Tests/Storage/BlobPipelineTests.cs
new file mode 100644
--- /dev/null
+++ b/Engine/Source/Programs/Shared/EpicGames.Horde.Tests/Storage/BlobPipelineTests.cs
@@ -0,0 +1,48 @@
+// Copyright Epic Games, Inc. All Rights Reserved.
+
+using System;
+using System.Collections.Generic;
+using System.Linq;
+using System.Threading.Tasks;
+using EpicGames.Core;
+using EpicGames.Horde.Storage;
+using EpicGames.Horde.Storage.Bundles;
+using Microsoft.Extensions.Logging.Abstractions;
+using Microsoft.VisualStudio.TestTools.UnitTesting;
+
+namespace EpicGames.Horde.Tests.Storage
+{
+	[TestClass]
+	public class BlobPipelineTests
+	{
+		static readonly BlobType s_blobType = new BlobType(Guid.Empty, 1);
+
+		[TestMethod]
+		public async Task TestAsync()
+		{
+			BundleStorageNamespace client = BundleStorageNamespace.CreateInMemory(NullLogger.Instance);
+
+			List<IHashedBlobRef> handles = new List<IHashedBlobRef>();
+			await using (IBlobWriter writer = client.CreateBlobWriter())
+			{
+				for (int idx = 0; idx < 4000; idx++)
+				{
+					writer.WriteInt32(idx);
+					handles.Add(await writer.CompleteAsync(s_blobType));
+				}
+			}
+
+			await using (BlobPipeline<IoHash> pipeline = new BlobPipeline<IoHash>())
+			{
+				foreach (IHashedBlobRef blobRef in handles)
+				{
+					pipeline.Add(new BlobRequest<IoHash>(blobRef, blobRef.Hash));
+				}
+				pipeline.FinishAdding();
+
+				List<BlobResponse<IoHash>> responses = await pipeline.ReadAllAsync().ToListAsync();
+				Assert.AreEqual(handles.Count, responses.Count);
+			}
+		}
+	}
+}
diff --git a/Engine/Source/Programs/Shared/EpicGames.Horde.Tests/Storage/BundleStoreTests.cs b/Engine/Source/Programs/Shared/EpicGames.Horde.Tests/Storage/BundleStoreTests.cs
new file mode 100644
--- /dev/null
+++ b/Engine/Source/Programs/Shared/EpicGames.Horde.Tests/Storage/BundleStoreTests.cs
@@ -0,0 +1,522 @@
+// Copyright Epic Games, Inc. All Rights Reserved.
+
+using System;
+using System.Buffers;
+using System.Collections.Generic;
+using System.IO;
+using System.Linq;
+using System.Text;
+using System.Threading;
+using System.Threading.Tasks;
+using EpicGames.Core;
+using EpicGames.Horde.Storage;
+using EpicGames.Horde.Storage.Backends;
+using EpicGames.Horde.Storage.Bundles;
+using EpicGames.Horde.Storage.Bundles.V1;
+using EpicGames.Horde.Storage.Clients;
+using EpicGames.Horde.Storage.Nodes;
+using EpicGames.Serialization;
+using Microsoft.Extensions.Caching.Memory;
+using Microsoft.Extensions.Logging.Abstractions;
+using Microsoft.VisualStudio.TestTools.UnitTesting;
+
+namespace EpicGames.Horde.Tests.Storage
+{
+	[TestClass]
+	public class BundleStoreTests
+	{
+		[TestMethod]
+		public async Task CreateBundlesManuallyV1Async()
+		{
+			byte[] a = CreateBundleManually();
+			byte[] b = await CreateBundleNormalAsync();
+			Assert.IsTrue(a.SequenceEqual(b));
+		}
+
+		[BlobConverter(typeof(TextNodeConverter))]
+		class TextNode
+		{
+			public string Text { get; }
+
+			public TextNode(string text) => Text = text;
+		}
+
+		class TextNodeConverter : BlobConverter<TextNode>
+		{
+			public static BlobType BlobType { get; } = new BlobType("{F63606D4-4061-5DBB-446F-55A69E22654F}", 1);
+
+			public override TextNode Read(IBlobReader reader, BlobSerializerOptions options)
+			{
+				string text = reader.ReadString();
+				return new TextNode(text);
+			}
+
+			public override BlobType Write(IBlobWriter writer, TextNode value, BlobSerializerOptions options)
+			{
+				writer.WriteString(value.Text);
+				return BlobType;
+			}
+		}
+
+		static async Task<byte[]> CreateBundleNormalAsync()
+		{
+			MemoryStorageBackend memoryStore = new MemoryStorageBackend();
+			BundleStorageNamespace store = new BundleStorageNamespace(memoryStore, BundleCache.None, new BundleOptions { MaxVersion = BundleVersion.ImportHashes, CompressionFormat = BundleCompressionFormat.None }, NullLogger.Instance);
+			await using IBlobWriter writer = store.CreateBlobWriter();
+
+			TextNode node = new TextNode("Hello world");
+			IHashedBlobRef<TextNode> handle = await writer.WriteBlobAsync(node);
+			await writer.FlushAsync();
+
+			using IReadOnlyMemoryOwner<byte> owner = await memoryStore.ReadBlobAsync(handle.GetLocator().BaseLocator, 0, null);
+			return owner.Memory.ToArray();
+		}
+
+		static byte[] CreateBundleManually()
+		{
+			ArrayMemoryWriter payloadWriter = new ArrayMemoryWriter(200);
+			payloadWriter.WriteString("Hello world");
+			byte[] payload = payloadWriter.WrittenMemory.ToArray();
+
+			List<BlobType> types = new List<BlobType>();
+			types.Add(new BlobType("{F63606D4-4061-5DBB-446F-55A69E22654F}", 1));
+
+			List<BundleExport> exports = new List<BundleExport>();
+			exports.Add(new BundleExport(0, 0, 0, payload.Length, Array.Empty<BundleExportRef>()));
+
+			List<BundlePacket> packets = new List<BundlePacket>();
+			packets.Add(new BundlePacket(BundleCompressionFormat.None, 0, payload.Length, payload.Length));
+
+			BundleHeader header = new BundleHeader(types.ToArray(), Array.Empty<BlobLocator>(), exports.ToArray(), packets.ToArray());
+
+			ReadOnlySequenceBuilder<byte> builder = new ReadOnlySequenceBuilder<byte>();
+			header.AppendTo(builder);
+			builder.Append(payload);
+			return builder.Construct().ToArray();
+		}
+
+		[TestMethod]
+		public async Task TestTreeAsync()
+		{
+			MemoryStorageBackend blobStore = new MemoryStorageBackend();
+
+			BundleOptions bundleOptions = new BundleOptions { MaxBlobSize = 1024 * 1024 };
+			BundleStorageNamespace bundleStore = new BundleStorageNamespace(blobStore, BundleCache.None, bundleOptions, NullLogger.Instance);
+
+			await TestTreeAsync(bundleStore);
+
+			Assert.AreEqual(1, blobStore.Blobs.Count);
+			Assert.AreEqual(1, blobStore.Refs.Count);
+		}
+
+		[TestMethod]
+		public async Task TestTreeSeparateBlobsAsync()
+		{
+			MemoryStorageBackend blobStore = new MemoryStorageBackend();
+
+			BundleOptions bundleOptions = new BundleOptions { MaxBlobSize = 1 };
+			BundleStorageNamespace bundleStore = new BundleStorageNamespace(blobStore, BundleCache.None, bundleOptions, NullLogger.Instance);
+
+			await TestTreeAsync(bundleStore);
+
+			Assert.AreEqual(5, blobStore.Blobs.Count);
+			Assert.AreEqual(1, blobStore.Refs.Count);
+		}
+
+		[BlobConverter(typeof(SimpleNodeConverter))]
+		class SimpleNode
+		{
+			public static BlobType BlobType { get; } = new BlobType("{F63606D4-4061-5DBB-446F-55A69F22654F}", 1);
+
+			public ReadOnlySequence<byte> Data { get; }
+			public IReadOnlyList<IHashedBlobRef<SimpleNode>> Refs { get; }
+
+			public SimpleNode(ReadOnlySequence<byte> data, IReadOnlyList<IHashedBlobRef<SimpleNode>> refs)
+			{
+				Data = data;
+				Refs = refs;
+			}
+		}
+
+		class SimpleNodeConverter : BlobConverter<SimpleNode>
+		{
+			public override SimpleNode Read(IBlobReader reader, BlobSerializerOptions options)
+			{
+				ReadOnlyMemory<byte> data = reader.ReadVariableLengthBytes();
+				IReadOnlyList<IHashedBlobRef<SimpleNode>> refs = reader.ReadVariableLengthArray(() => reader.ReadBlobRef<SimpleNode>());
+				return new SimpleNode(new ReadOnlySequence<byte>(data), refs);
+			}
+
+			public override BlobType Write(IBlobWriter writer, SimpleNode value, BlobSerializerOptions options)
+			{
+				writer.WriteVariableLengthBytes(value.Data);
+				writer.WriteVariableLengthArray(value.Refs, x => writer.WriteBlobRef(x));
+				return SimpleNode.BlobType;
+			}
+		}
+
+		static async Task TestTreeAsync(BundleStorageNamespace store)
+		{
+			// Generate a tree
+			{
+				await using IBlobWriter writer = store.CreateBlobWriter("test");
+
+				SimpleNode node1 = new SimpleNode(new ReadOnlySequence<byte>(new byte[] { 1 }), Array.Empty<IHashedBlobRef<SimpleNode>>());
+				SimpleNode node2 = new SimpleNode(new ReadOnlySequence<byte>(new byte[] { 2 }), new[] { await writer.WriteBlobAsync(node1) });
+				SimpleNode node3 = new SimpleNode(new ReadOnlySequence<byte>(new byte[] { 3 }), new[] { await writer.WriteBlobAsync(node2) });
+				SimpleNode node4 = new SimpleNode(new ReadOnlySequence<byte>(new byte[] { 4 }), Array.Empty<IHashedBlobRef<SimpleNode>>());
+
+				SimpleNode root = new SimpleNode(new ReadOnlySequence<byte>(new byte[] { 5 }), new[] { await writer.WriteBlobAsync(node4), await writer.WriteBlobAsync(node3) });
+				IHashedBlobRef<SimpleNode> rootRef = await writer.WriteBlobAsync(root);
+				await writer.FlushAsync();
+
+				await store.AddRefAsync(new RefName("test"), rootRef);
+
+				BundleReader reader = new BundleReader(store, BundleCache.None, NullLogger.Instance);
+				await CheckTreeAsync(root);
+			}
+
+			// Check we can read it back in
+			{
+				SimpleNode root = await store.ReadRefTargetAsync<SimpleNode>(new RefName("test"));
+				await CheckTreeAsync(root);
+			}
+		}
+
+		static async Task CheckTreeAsync(SimpleNode root)
+		{
+			SimpleNode node5 = root;
+			byte[] data5 = node5.Data.ToArray();
+			Assert.IsTrue(data5.SequenceEqual(new byte[] { 5 }));
+			IReadOnlyList<IHashedBlobRef<SimpleNode>> refs5 = node5.Refs;
+			Assert.AreEqual(2, refs5.Count);
+
+			SimpleNode node4 = await refs5[0].ReadBlobAsync();
+			byte[] data4 = node4.Data.ToArray();
+			Assert.IsTrue(data4.SequenceEqual(new byte[] { 4 }));
+			IReadOnlyList<IHashedBlobRef<SimpleNode>> refs4 = node4.Refs;
+			Assert.AreEqual(0, refs4.Count);
+
+			SimpleNode node3 = await refs5[1].ReadBlobAsync();
+			byte[] data3 = node3.Data.ToArray();
+			Assert.IsTrue(data3.SequenceEqual(new byte[] { 3 }));
+			IReadOnlyList<IHashedBlobRef<SimpleNode>> refs3 = node3.Refs;
+			Assert.AreEqual(1, refs3.Count);
+
+			SimpleNode node2 = await refs3[0].ReadBlobAsync();
+			byte[] data2 = node2.Data.ToArray();
+			Assert.IsTrue(data2.SequenceEqual(new byte[] { 2 }));
+			IReadOnlyList<IHashedBlobRef<SimpleNode>> refs2 = node2.Refs;
+			Assert.AreEqual(1, refs2.Count);
+
+			SimpleNode node1 = await refs2[0].ReadBlobAsync();
+			byte[] data1 = node1.Data.ToArray();
+			Assert.IsTrue(data1.SequenceEqual(new byte[] { 1 }));
+			IReadOnlyList<IHashedBlobRef<SimpleNode>> refs1 = node1.Refs;
+			Assert.AreEqual(0, refs1.Count);
+		}
+
+		[TestMethod]
+		public async Task SimpleNodeAsync()
+		{
+			KeyValueStorageNamespace store = KeyValueStorageNamespace.CreateInMemory();
+
+			RefName refName = new RefName("test");
+
+			IHashedBlobRef<SimpleNode> inputRef;
+			await using (IBlobWriter writer = store.CreateBlobWriter(refName))
+			{
+				inputRef = await writer.WriteBlobAsync(new SimpleNode(new ReadOnlySequence<byte>(new byte[] { (byte)123 }), Array.Empty<IHashedBlobRef<SimpleNode>>()));
+			}
+			await store.AddRefAsync(refName, inputRef);
+
+			SimpleNode node = await store.ReadRefTargetAsync<SimpleNode>(refName);
+
+			Assert.AreEqual(123, node.Data.FirstSpan[0]);
+		}
+
+		[TestMethod]
+		public async Task DirectoryNodesAsync()
+		{
+			KeyValueStorageNamespace store = KeyValueStorageNamespace.CreateInMemory();
+
+			// Generate a tree
+			{
+				await using (IBlobWriter writer = store.CreateBlobWriter(new RefName("test")))
+				{
+					DirectoryNode world = new DirectoryNode();
+					IHashedBlobRef<DirectoryNode> worldRef = await writer.WriteBlobAsync(world);
+
+					DirectoryNode hello = new DirectoryNode();
+					hello.AddDirectory(new DirectoryEntry("world", 0, worldRef));
+					IHashedBlobRef<DirectoryNode> helloRef = await writer.WriteBlobAsync(hello);
+
+					DirectoryNode root = new DirectoryNode(DirectoryFlags.None);
+					root.AddDirectory(new DirectoryEntry("hello", 0, helloRef));
+					IHashedBlobRef<DirectoryNode> rootRef = await writer.WriteBlobAsync(root);
+
+					await writer.FlushAsync();
+
+					await store.AddRefAsync(new RefName("test"), rootRef);
+				}
+			}
+
+			// Check we can read it back in
+			{
+				DirectoryNode root = await store.ReadRefTargetAsync<DirectoryNode>(new RefName("test"));
+				await CheckDirectoryTreeAsync(root);
+			}
+		}
+
+		static async Task CheckDirectoryTreeAsync(DirectoryNode root)
+		{
+			Assert.AreEqual(1, root.Directories.Count);
+			Assert.AreEqual("hello", root.Directories.First().Name);
+
+			DirectoryNode hello = await root.Directories.First().Handle.ReadBlobAsync();
+			Assert.AreEqual(1, hello.Directories.Count);
+			Assert.AreEqual("world", hello.Directories.First().Name);
+
+			DirectoryNode world = await hello.Directories.First().Handle.ReadBlobAsync();
+			Assert.AreEqual(0, world.Directories.Count);
+		}
+
+		[TestMethod]
+		public async Task FileNodesAsync()
+		{
+			using IMemoryCache cache = new MemoryCache(new MemoryCacheOptions());
+			BundleStorageNamespace store = BundleStorageNamespace.CreateInMemory(NullLogger.Instance);
+			BundleReader reader = new BundleReader(store, BundleCache.None, NullLogger.Instance);
+
+			// Generate a tree
+			{
+				await using IBlobWriter writer = store.CreateBlobWriter();
+
+				using ChunkedDataWriter fileWriter = new ChunkedDataWriter(writer, new ChunkingOptions());
+				ChunkedData fileHandle = await fileWriter.CreateAsync(Encoding.UTF8.GetBytes("world"), CancellationToken.None);
+
+				List<FileUpdate> fileUpdates = new List<FileUpdate>();
+				fileUpdates.Add(new FileUpdate("hello/world", FileEntryFlags.None, fileWriter.Length, fileHandle));
+
+				DirectoryNode root = new DirectoryNode();
+				await root.UpdateAsync(fileUpdates, writer);
+
+				IHashedBlobRef<DirectoryNode> rootRef = await writer.WriteBlobAsync(root);
+				await store.AddRefAsync(new RefName("test"), rootRef);
+
+				await CheckFileTreeAsync(root);
+			}
+
+			// Check we can read it back in
+			{
+				DirectoryNode root = await store.ReadRefTargetAsync<DirectoryNode>(new RefName("test"));
+				await CheckFileTreeAsync(root);
+			}
+		}
+
+		static async Task CheckFileTreeAsync(DirectoryNode root)
+		{
+			Assert.AreEqual(1, root.Directories.Count);
+			Assert.AreEqual("hello", root.Directories.First().Name);
+
+			DirectoryNode hello = await root.Directories.First().Handle.ReadBlobAsync();
+			Assert.AreEqual(0, hello.Directories.Count);
+			Assert.AreEqual(1, hello.Files.Count);
+			Assert.AreEqual("world", hello.Files.First().Name);
+
+			ChunkedDataNode world = await hello.Files.First().Target.ReadBlobAsync();
+
+			byte[] worldData = await GetFileDataAsync(world);
+			Assert.IsTrue(worldData.SequenceEqual(Encoding.UTF8.GetBytes("world")));
+		}
+
+		[TestMethod]
+		public async Task StreamTestAsync()
+		{
+			BundleOptions bundleOptions = new BundleOptions { MaxBlobSize = 1024 };
+			BundleStorageNamespace store = BundleStorageNamespace.CreateInMemory(bundleOptions, NullLogger.Instance);
+
+			const int Length = 4096;
+
+			byte[] chunk = new byte[Length];
+			new Random(0).NextBytes(chunk);
+
+			// Generate a tree
+			ChunkedDataNodeRef nodeRef;
+			{
+				await using IBlobWriter writer = store.CreateBlobWriter();
+
+				ChunkingOptions options = new ChunkingOptions();
+				options.LeafOptions = new LeafChunkedDataNodeOptions(128, 256, 64 * 1024);
+
+				using ChunkedDataWriter fileWriter = new ChunkedDataWriter(writer, options);
+				for (int idx = 0; idx < chunk.Length / 16; idx++)
+				{
+					await fileWriter.AppendAsync(chunk.AsMemory(idx * 16, 16), CancellationToken.None);
+				}
+
+				nodeRef = (await fileWriter.FlushAsync(CancellationToken.None)).Root;
+			}
+
+			// Check we can read it back in
+			{
+				ChunkedDataNode newRoot = await nodeRef.ReadBlobAsync();
+
+				using MemoryStream stream = new MemoryStream();
+				await newRoot.CopyToStreamAsync(stream);
+
+				byte[] output = stream.ToArray();
+				Assert.IsTrue(chunk.SequenceEqual(output));
+			}
+		}
+
+		[TestMethod]
+		public async Task LargeFileTestAsync()
+		{
+			BundleOptions bundleOptions = new BundleOptions { MaxBlobSize = 1024 };
+			BundleStorageNamespace store = BundleStorageNamespace.CreateInMemory(bundleOptions, NullLogger.Instance);
+
+			const int Length = 1024;
+			const int Copies = 4096;
+
+			byte[] chunk = new byte[Length];
+			new Random(0).NextBytes(chunk);
+
+			byte[] data = new byte[chunk.Length * Copies];
+			for (int idx = 0; idx < Copies; idx++)
+			{
+				chunk.CopyTo(data.AsSpan(idx * chunk.Length));
+			}
+
+			// Generate a tree
+			DirectoryNode root;
+			{
+				await using DedupeBlobWriter writer = new DedupeBlobWriter(store.CreateBlobWriter());
+
+				ChunkingOptions options = new ChunkingOptions();
+				options.LeafOptions = new LeafChunkedDataNodeOptions(128, 256, 64 * 1024);
+
+				using ChunkedDataWriter fileWriter = new ChunkedDataWriter(writer, options);
+				ChunkedData chunkedData = await fileWriter.CreateAsync(data, CancellationToken.None);
+
+				root = new DirectoryNode(DirectoryFlags.None);
+				root.AddFile("test", FileEntryFlags.None, fileWriter.Length, chunkedData);
+
+				IHashedBlobRef<DirectoryNode> rootRef = await writer.WriteBlobAsync(root);
+				await store.AddRefAsync(new RefName("test"), rootRef);
+
+				await CheckLargeFileTreeAsync(root, data);
+			}
+
+			// Check we can read it back in
+			{
+				DirectoryNode newRoot = await store.ReadRefTargetAsync<DirectoryNode>(new RefName("test"));
+				await CompareTreesAsync(root, newRoot);
+				await CheckLargeFileTreeAsync(root, data);
+
+				FileEntry file = root.GetFileEntry("test");
+
+				Dictionary<BlobLocator, long> locatorToSize = new Dictionary<BlobLocator, long>();
+				await GetUniqueBlobsAsync(file.Target.Handle, locatorToSize);
+
+				long uniqueSize = locatorToSize.Sum(x => x.Value);
+				Assert.IsTrue(uniqueSize < data.Length / 3); // random fraction meaning "lots of dedupe happened"
+			}
+		}
+
+		static async Task GetUniqueBlobsAsync(IBlobRef handle, Dictionary<BlobLocator, long> locatorToSize)
+		{
+			using BlobData data = await handle.ReadBlobDataAsync();
+			locatorToSize[handle.GetLocator()] = data.Data.Length;
+
+			foreach (IBlobRef import in data.Imports)
+			{
+				await GetUniqueBlobsAsync(import, locatorToSize);
+			}
+		}
+
+		static async Task CompareTreesAsync(DirectoryNode oldNode, DirectoryNode newNode)
+		{
+			Assert.AreEqual(oldNode.Length, newNode.Length);
+			Assert.AreEqual(oldNode.Files.Count, newNode.Files.Count);
+			Assert.AreEqual(oldNode.Directories.Count, newNode.Directories.Count);
+
+			foreach ((FileEntry oldFileEntry, FileEntry newFileEntry) in oldNode.Files.Zip(newNode.Files))
+			{
+				ChunkedDataNode oldFile = await oldFileEntry.Target.ReadBlobAsync();
+				ChunkedDataNode newFile = await newFileEntry.Target.ReadBlobAsync();
+				await CompareTreesAsync(oldFile, newFile);
+			}
+		}
+
+		static async Task CompareTreesAsync(ChunkedDataNode oldNode, ChunkedDataNode newNode)
+		{
+			if (oldNode is InteriorChunkedDataNode oldInteriorNode)
+			{
+				InteriorChunkedDataNode newInteriorNode = (InteriorChunkedDataNode)newNode;
+				Assert.AreEqual(oldInteriorNode.Children.Count, newInteriorNode.Children.Count);
+
+				int index = 0;
+				foreach ((ChunkedDataNodeRef oldFileRef, ChunkedDataNodeRef newFileRef) in oldInteriorNode.Children.Zip(newInteriorNode.Children))
+				{
+					ChunkedDataNode oldFile = await oldFileRef.ReadBlobAsync();
+					ChunkedDataNode newFile = await newFileRef.ReadBlobAsync();
+					await CompareTreesAsync(oldFile, newFile);
+					index++;
+				}
+			}
+			else if (oldNode is LeafChunkedDataNode oldLeafNode)
+			{
+				LeafChunkedDataNode newLeafNode = (LeafChunkedDataNode)newNode;
+				Assert.IsTrue(oldLeafNode.Data.Span.SequenceEqual(newLeafNode.Data.Span));
+			}
+			else
+			{
+				throw new NotImplementedException();
+			}
+		}
+
+		static async Task CheckLargeFileTreeAsync(DirectoryNode root, byte[] data)
+		{
+			Assert.AreEqual(0, root.Directories.Count);
+			Assert.AreEqual(1, root.Files.Count);
+
+			ChunkedDataNode world = await root.Files.First().Target.ReadBlobAsync();
+
+			int length = await CheckFileDataAsync(world, data);
+			Assert.AreEqual(data.Length, length);
+		}
+
+		static async Task<int> CheckFileDataAsync(ChunkedDataNode fileNode, ReadOnlyMemory<byte> data)
+		{
+			int offset = 0;
+			if (fileNode is LeafChunkedDataNode leafNode)
+			{
+				Assert.IsTrue(leafNode.Data.Span.SequenceEqual(data.Span.Slice(offset, leafNode.Data.Length)));
+				offset += leafNode.Data.Length;
+			}
+			else if (fileNode is InteriorChunkedDataNode interiorFileNode)
+			{
+				foreach (ChunkedDataNodeRef childRef in interiorFileNode.Children)
+				{
+					ChunkedDataNode child = await childRef.ReadBlobAsync();
+					offset += await CheckFileDataAsync(child, data.Slice(offset));
+				}
+			}
+			else
+			{
+				throw new NotImplementedException();
+			}
+			return offset;
+		}
+
+		static async Task<byte[]> GetFileDataAsync(ChunkedDataNode fileNode)
+		{
+			using (MemoryStream stream = new MemoryStream())
+			{
+				await fileNode.CopyToStreamAsync(stream);
+				return stream.ToArray();
+			}
+		}
+	}
+}
\ No newline at end of file
diff --git a/Engine/Source/Programs/Shared/EpicGames.Horde.Tests/Storage/BundleTests.cs b/Engine/Source/Programs/Shared/EpicGames.Horde.Tests/Storage/BundleTests.cs
new file mode 100644
--- /dev/null
+++ b/Engine/Source/Programs/Shared/EpicGames.Horde.Tests/Storage/BundleTests.cs
@@ -0,0 +1,305 @@
+// Copyright Epic Games, Inc. All Rights Reserved.
+
+using System;
+using System.Collections.Generic;
+using System.IO;
+using System.Linq;
+using System.Threading;
+using System.Threading.Tasks;
+using EpicGames.Core;
+using EpicGames.Horde.Storage;
+using EpicGames.Horde.Storage.Backends;
+using EpicGames.Horde.Storage.Bundles;
+using EpicGames.Horde.Storage.Bundles.V1;
+using EpicGames.Horde.Storage.Nodes;
+using Microsoft.Extensions.Logging.Abstractions;
+using Microsoft.VisualStudio.TestTools.UnitTesting;
+
+namespace EpicGames.Horde.Tests.Storage
+{
+	[TestClass]
+	public sealed class BundleTests
+	{
+		[TestMethod]
+		public void BuzHashTests()
+		{
+			byte[] data = new byte[4096];
+			new Random(0).NextBytes(data);
+
+			const int WindowSize = 128;
+
+			uint rollingHash = 0;
+			for (int maxIdx = 0; maxIdx < data.Length + WindowSize; maxIdx++)
+			{
+				int minIdx = maxIdx - WindowSize;
+
+				if (maxIdx < data.Length)
+				{
+					rollingHash = BuzHash.Add(rollingHash, data[maxIdx]);
+				}
+
+				int length = Math.Min(maxIdx + 1, data.Length) - Math.Max(minIdx, 0);
+				uint cleanHash = BuzHash.Add(0, data.AsSpan(Math.Max(minIdx, 0), length));
+				Assert.AreEqual(rollingHash, cleanHash);
+
+				if (minIdx >= 0)
+				{
+					rollingHash = BuzHash.Sub(rollingHash, data[minIdx], length);
+				}
+			}
+		}
+
+		[TestMethod]
+		public async Task BasicChunkingTestsAsync()
+		{
+			BundleStorageNamespace storage = BundleStorageNamespace.CreateInMemory(NullLogger.Instance);
+
+			RefName refName = new RefName("test");
+			await using IBlobWriter writer = storage.CreateBlobWriter(refName);
+
+			ChunkingOptions options = new ChunkingOptions();
+			options.LeafOptions = new LeafChunkedDataNodeOptions(8, 8, 8);
+
+			using ChunkedDataWriter fileNodeWriter = new ChunkedDataWriter(writer, options);
+
+			ChunkedDataNode node;
+			ChunkedDataNodeRef nodeRef;
+			byte[] data = CreateBuffer(1024);
+
+			nodeRef = (await fileNodeWriter.CreateAsync(data.AsMemory(0, 7), CancellationToken.None)).Root;
+			node = await nodeRef.ReadBlobAsync();
+			Assert.IsTrue(node is LeafChunkedDataNode);
+			Assert.AreEqual(7, ((LeafChunkedDataNode)node).Data.Length);
+			await TestBufferlessReadsAsync(nodeRef, data.AsMemory(0, 7));
+
+			nodeRef = (await fileNodeWriter.CreateAsync(data.AsMemory(0, 8), CancellationToken.None)).Root;
+			node = await nodeRef.ReadBlobAsync();
+			Assert.IsTrue(node is LeafChunkedDataNode);
+			Assert.AreEqual(8, ((LeafChunkedDataNode)node).Data.Length);
+			await TestBufferlessReadsAsync(nodeRef, data.AsMemory(0, 8));
+
+			nodeRef = (await fileNodeWriter.CreateAsync(data.AsMemory(0, 9), CancellationToken.None)).Root;
+			node = await nodeRef.ReadBlobAsync();
+			Assert.IsTrue(node is InteriorChunkedDataNode);
+			Assert.AreEqual(2, ((InteriorChunkedDataNode)node).Children.Count);
+			await TestBufferlessReadsAsync(nodeRef, data.AsMemory(0, 9));
+
+			ChunkedDataNode? childNode1 = await ((InteriorChunkedDataNode)node).Children[0].ReadBlobAsync();
+			Assert.IsNotNull(childNode1);
+			Assert.IsTrue(childNode1 is LeafChunkedDataNode);
+			Assert.AreEqual(8, ((LeafChunkedDataNode)childNode1!).Data.Length);
+
+			ChunkedDataNode? childNode2 = await ((InteriorChunkedDataNode)node).Children[1].ReadBlobAsync();
+			Assert.IsNotNull(childNode2);
+			Assert.IsTrue(childNode2 is LeafChunkedDataNode);
+			Assert.AreEqual(1, ((LeafChunkedDataNode)childNode2!).Data.Length);
+
+			nodeRef = (await fileNodeWriter.CreateAsync(data, CancellationToken.None)).Root;
+			node = await nodeRef.ReadBlobAsync();
+			Assert.IsTrue(node is InteriorChunkedDataNode);
+			await TestBufferlessReadsAsync(nodeRef, data);
+		}
+
+		private static byte[] CreateBuffer(int length)
+		{
+			byte[] output = GC.AllocateUninitializedArray<byte>(length);
+			for (int i = 0; i < length; i++)
+			{
+				output[i] = (byte)i;
+			}
+			return output;
+		}
+
+		private static async Task TestBufferlessReadsAsync(ChunkedDataNodeRef nodeRef, ReadOnlyMemory<byte> expected)
+		{
+			using MemoryStream memoryStream = new MemoryStream();
+			await ChunkedDataNode.CopyToStreamAsync(nodeRef.Handle, memoryStream, default);
+			ReadOnlyMemory<byte> read = memoryStream.ToArray().AsMemory();
+			Assert.IsTrue(read.Span.SequenceEqual(expected.Span));
+		}
+
+		[TestMethod]
+		public void SerializationTests()
+		{
+			BundleHeader oldHeader;
+			{
+				List<BlobType> types = new List<BlobType>();
+				types.Add(new BlobType(Guid.NewGuid(), 0));
+
+				List<BlobLocator> imports = new List<BlobLocator>();
+				imports.Add(new BlobLocator("import1"));
+				imports.Add(new BlobLocator("import2"));
+
+				List<BundleExport> exports = new List<BundleExport>();
+				exports.Add(new BundleExport(0, 0, 0, 2, new BundleExportRef[] { new BundleExportRef(0, 5), new BundleExportRef(0, 6) }));
+				exports.Add(new BundleExport(0, 1, 0, 3, new BundleExportRef[] { new BundleExportRef(-1, 0) }));
+
+				List<BundlePacket> packets = new List<BundlePacket>();
+				packets.Add(new BundlePacket(BundleCompressionFormat.LZ4, 0, 20, 40));
+				packets.Add(new BundlePacket(BundleCompressionFormat.LZ4, 20, 10, 20));
+
+				oldHeader = new BundleHeader(types.ToArray(), imports.ToArray(), exports.ToArray(), packets.ToArray());
+			}
+
+			byte[] serializedData = oldHeader.ToByteArray();
+
+			BundleHeader newHeader = BundleHeader.Read(serializedData);
+
+			Assert.AreEqual(oldHeader.Imports.Count, newHeader.Imports.Count);
+			for (int idx = 0; idx < oldHeader.Imports.Count; idx++)
+			{
+				Assert.AreEqual(oldHeader.Imports[idx], newHeader.Imports[idx]);
+			}
+
+			Assert.AreEqual(oldHeader.Exports.Count, newHeader.Exports.Count);
+			for (int idx = 0; idx < oldHeader.Exports.Count; idx++)
+			{
+				BundleExport oldExport = oldHeader.Exports[idx];
+				BundleExport newExport = newHeader.Exports[idx];
+
+				Assert.AreEqual(oldExport.Hash, newExport.Hash);
+				Assert.AreEqual(oldExport.Length, newExport.Length);
+				Assert.IsTrue(oldExport.References.SequenceEqual(newExport.References));
+			}
+
+			Assert.AreEqual(oldHeader.Packets.Count, newHeader.Packets.Count);
+			for (int idx = 0; idx < oldHeader.Packets.Count; idx++)
+			{
+				BundlePacket oldPacket = oldHeader.Packets[idx];
+				BundlePacket newPacket = newHeader.Packets[idx];
+				Assert.AreEqual(oldPacket.DecodedLength, newPacket.DecodedLength);
+				Assert.AreEqual(oldPacket.EncodedLength, newPacket.EncodedLength);
+			}
+		}
+
+		[TestMethod]
+		public async Task BasicTestDirectoryAsync()
+		{
+			BundleStorageNamespace store = BundleStorageNamespace.CreateInMemory(NullLogger.Instance);
+
+			IHashedBlobRef<DirectoryNode> rootRef;
+			await using (IBlobWriter writer = store.CreateBlobWriter())
+			{
+				DirectoryNode world = new DirectoryNode();
+				IHashedBlobRef<DirectoryNode> worldRef = await writer.WriteBlobAsync(world);
+
+				DirectoryNode hello = new DirectoryNode();
+				hello.AddDirectory(new DirectoryEntry("world", 0, worldRef));
+				IHashedBlobRef<DirectoryNode> helloRef = await writer.WriteBlobAsync(hello);
+
+				DirectoryNode root = new DirectoryNode();
+				root.AddDirectory(new DirectoryEntry("hello", 0, helloRef));
+				rootRef = await writer.WriteBlobAsync(root);
+
+				await writer.FlushAsync();
+			}
+
+			RefName refName = new RefName("testref");
+			await store.AddRefAsync(refName, rootRef);
+
+			// Should be stored inline
+			MemoryStorageBackend memoryStore = (MemoryStorageBackend)store.Backend;
+			Assert.AreEqual(1, memoryStore.Refs.Count);
+			Assert.AreEqual(1, memoryStore.Blobs.Count);
+
+			// Check the ref
+			//			IBlobHandle refTarget =  await store.ReadRefTargetAsync(refName);
+			//			IBlobHandle bundleTarget = store.CreateBlobHandle(refTarget.GetLocator().BaseLocator);
+			//			using BlobData bundleData = await bundleTarget.ReadBlobDataAsync();
+
+			// This is specific to V1 data
+			//			BundleHeader bundleHeader = BundleHeader.Read(bundleData.Data);
+			//			Assert.AreEqual(0, bundleHeader.Imports.Count);
+			//			Assert.AreEqual(3, bundleHeader.Exports.Count);
+
+			// Create a new bundle and read it back in again
+			DirectoryNode newRoot = await store.ReadRefTargetAsync<DirectoryNode>(refName);
+
+			Assert.AreEqual(0, newRoot.Files.Count);
+			Assert.AreEqual(1, newRoot.Directories.Count);
+			DirectoryNode? outputNode = await newRoot.TryOpenDirectoryAsync("hello");
+			Assert.IsNotNull(outputNode);
+
+			Assert.AreEqual(0, outputNode!.Files.Count);
+			Assert.AreEqual(1, outputNode!.Directories.Count);
+			DirectoryNode? outputNode2 = await outputNode.TryOpenDirectoryAsync("world");
+			Assert.IsNotNull(outputNode2);
+
+			Assert.AreEqual(0, outputNode2!.Files.Count);
+			Assert.AreEqual(0, outputNode2!.Directories.Count);
+		}
+
+		[TestMethod]
+		public async Task DedupTestsAsync()
+		{
+			BundleOptions bundleOptions = new BundleOptions();
+			bundleOptions.MaxBlobSize = 1;
+
+			BundleStorageNamespace storage = BundleStorageNamespace.CreateInMemory(bundleOptions, NullLogger.Instance);
+
+			await using (IBlobWriter writer = new DedupeBlobWriter(storage.CreateBlobWriter()))
+			{
+				DirectoryNode root = new DirectoryNode();
+				for (int idx = 1; idx <= 3; idx++)
+				{
+					DirectoryNode node = new DirectoryNode();
+					IHashedBlobRef<DirectoryNode> nodeRef = await writer.WriteBlobAsync(node);
+					root.AddDirectory(new DirectoryEntry($"node{idx}", 0, nodeRef));
+				}
+
+				RefName refName = new RefName("ref");
+				IHashedBlobRef<DirectoryNode> rootRef = await writer.WriteBlobAsync(root);
+				await storage.AddRefAsync(refName, rootRef);
+			}
+
+			MemoryStorageBackend memoryStore = (MemoryStorageBackend)storage.Backend;
+			Assert.AreEqual(1, memoryStore.Refs.Count);
+			Assert.AreEqual(2, memoryStore.Blobs.Count);
+		}
+
+		[TestMethod]
+		public async Task ReloadTestsAsync()
+		{
+			BundleOptions bundleOptions = new BundleOptions();
+			bundleOptions.MaxBlobSize = 1;
+
+			BundleStorageNamespace storage = BundleStorageNamespace.CreateInMemory(bundleOptions, NullLogger.Instance);
+
+			RefName refName = new RefName("ref");
+
+			{
+				await using (IBlobWriter writer = storage.CreateBlobWriter())
+				{
+					IHashedBlobRef<DirectoryNode> rootRef = await writer.WriteBlobAsync(new DirectoryNode());
+					for (int idx = 4; idx >= 1; idx--)
+					{
+						DirectoryNode next = new DirectoryNode();
+						next.AddDirectory(new DirectoryEntry($"node{idx}", 0, rootRef));
+						rootRef = await writer.WriteBlobAsync(next);
+					}
+					await storage.AddRefAsync(refName, rootRef);
+				}
+
+				MemoryStorageBackend memoryStore = (MemoryStorageBackend)storage.Backend;
+				Assert.AreEqual(1, memoryStore.Refs.Count);
+				Assert.AreEqual(5, memoryStore.Blobs.Count);
+			}
+
+			{
+				DirectoryNode root = await storage.ReadRefTargetAsync<DirectoryNode>(refName);
+
+				DirectoryNode? newNode1 = await root.TryOpenDirectoryAsync("node1");
+				Assert.IsNotNull(newNode1);
+
+				DirectoryNode? newNode2 = await newNode1!.TryOpenDirectoryAsync("node2");
+				Assert.IsNotNull(newNode2);
+
+				DirectoryNode? newNode3 = await newNode2!.TryOpenDirectoryAsync("node3");
+				Assert.IsNotNull(newNode3);
+
+				DirectoryNode? newNode4 = await newNode3!.TryOpenDirectoryAsync("node4");
+				Assert.IsNotNull(newNode4);
+			}
+		}
+	}
+}
diff --git a/Engine/Source/Programs/Shared/EpicGames.Horde.Tests/Storage/ChunkingTests.cs b/Engine/Source/Programs/Shared/EpicGames.Horde.Tests/Storage/ChunkingTests.cs
new file mode 100644
--- /dev/null
+++ b/Engine/Source/Programs/Shared/EpicGames.Horde.Tests/Storage/ChunkingTests.cs
@@ -0,0 +1,383 @@
+// Copyright Epic Games, Inc. All Rights Reserved.
+
+using System;
+using System.Collections.Generic;
+using System.IO;
+using System.Linq;
+using System.Text;
+using System.Threading;
+using System.Threading.Tasks;
+using EpicGames.Core;
+using EpicGames.Horde.Storage;
+using EpicGames.Horde.Storage.Clients;
+using EpicGames.Horde.Storage.Nodes;
+using Microsoft.Extensions.Caching.Memory;
+using Microsoft.Extensions.Logging.Abstractions;
+using Microsoft.VisualStudio.TestTools.UnitTesting;
+
+namespace EpicGames.Horde.Tests.Storage
+{
+	[TestClass]
+	public class ChunkingTests
+	{
+		[TestMethod]
+		public void BuzHashTests()
+		{
+			byte[] data = new byte[4096];
+			new Random(0).NextBytes(data);
+
+			const int WindowSize = 128;
+
+			uint rollingHash = 0;
+			for (int maxIdx = 0; maxIdx < data.Length + WindowSize; maxIdx++)
+			{
+				int minIdx = maxIdx - WindowSize;
+
+				if (maxIdx < data.Length)
+				{
+					rollingHash = BuzHash.Add(rollingHash, data[maxIdx]);
+				}
+
+				int length = Math.Min(maxIdx + 1, data.Length) - Math.Max(minIdx, 0);
+				uint cleanHash = BuzHash.Add(0, data.AsSpan(Math.Max(minIdx, 0), length));
+				Assert.AreEqual(rollingHash, cleanHash);
+
+				if (minIdx >= 0)
+				{
+					rollingHash = BuzHash.Sub(rollingHash, data[minIdx], length);
+				}
+			}
+		}
+
+		[TestMethod]
+		public async Task EmptyNodeTestAsync()
+		{
+			KeyValueStorageNamespace store = KeyValueStorageNamespace.CreateInMemory();
+
+			const string RefName = "hello";
+			await using (IBlobWriter writer = store.CreateBlobWriter(RefName))
+			{
+				ChunkingOptions options = new ChunkingOptions();
+				options.LeafOptions = new LeafChunkedDataNodeOptions(64, 64, 64);
+				options.InteriorOptions = new InteriorChunkedDataNodeOptions(4, 4, 4);
+
+				using MemoryStream emptyStream = new MemoryStream();
+				LeafChunkedData leafChunkedData = await LeafChunkedDataNode.CreateFromStreamAsync(writer, emptyStream, new LeafChunkedDataNodeOptions(64, 64, 64), CancellationToken.None);
+				ChunkedData chunkedData = await InteriorChunkedDataNode.CreateTreeAsync(leafChunkedData, new InteriorChunkedDataNodeOptions(4, 4, 4), writer, CancellationToken.None);
+
+				DirectoryNode directory = new DirectoryNode();
+				directory.AddFile("test.foo", FileEntryFlags.None, 0, chunkedData);
+
+				IHashedBlobRef handle = await writer.WriteBlobAsync(directory);
+				await store.AddRefAsync(RefName, handle);
+			}
+		}
+
+		[TestMethod]
+		public async Task FixedSizeChunkingTestsAsync()
+		{
+			ChunkingOptions options = new ChunkingOptions();
+			options.LeafOptions = new LeafChunkedDataNodeOptions(64, 64, 64);
+			options.InteriorOptions = new InteriorChunkedDataNodeOptions(4, 4, 4);
+
+			await TestChunkingAsync(options);
+		}
+
+		[TestMethod]
+		public async Task VariableSizeChunkingTestsAsync()
+		{
+			ChunkingOptions options = new ChunkingOptions();
+			options.LeafOptions = new LeafChunkedDataNodeOptions(32, 64, 96);
+			options.InteriorOptions = new InteriorChunkedDataNodeOptions(1, 4, 12);
+
+			await TestChunkingAsync(options);
+		}
+
+		static async Task TestChunkingAsync(ChunkingOptions options)
+		{
+			using MemoryCache cache = new MemoryCache(new MemoryCacheOptions());
+
+			KeyValueStorageNamespace store = KeyValueStorageNamespace.CreateInMemory();
+
+			await using IBlobWriter writer = store.CreateBlobWriter();
+
+			byte[] data = new byte[4096];
+			new Random(0).NextBytes(data);
+
+			for (int idx = 0; idx < data.Length; idx++)
+			{
+				data[idx] = (byte)idx;
+			}
+
+			ChunkedDataNodeRef handle;
+
+			const int NumIterations = 100;
+			{
+				using ChunkedDataWriter fileWriter = new ChunkedDataWriter(writer, options);
+
+				for (int idx = 0; idx < NumIterations; idx++)
+				{
+					await fileWriter.AppendAsync(data, CancellationToken.None);
+				}
+
+				handle = (await fileWriter.FlushAsync(CancellationToken.None)).Root;
+			}
+
+			ChunkedDataNode root = await handle.ReadBlobAsync();
+
+			byte[] result;
+			using (MemoryStream stream = new MemoryStream())
+			{
+				await root.CopyToStreamAsync(stream);
+				result = stream.ToArray();
+			}
+
+			Assert.AreEqual(NumIterations * data.Length, result.Length);
+
+			for (int idx = 0; idx < NumIterations; idx++)
+			{
+				ReadOnlyMemory<byte> spanData = result.AsMemory(idx * data.Length, data.Length);
+				Assert.IsTrue(spanData.Span.SequenceEqual(data));
+			}
+
+			await CheckSizesAsync(root, options, true);
+		}
+
+		static async Task CheckSizesAsync(ChunkedDataNode node, ChunkingOptions options, bool rightmost)
+		{
+			if (node is LeafChunkedDataNode leafNode)
+			{
+				Assert.IsTrue(rightmost || leafNode.Data.Length >= options.LeafOptions.MinSize);
+				Assert.IsTrue(leafNode.Data.Length <= options.LeafOptions.MaxSize);
+			}
+			else
+			{
+				InteriorChunkedDataNode interiorNode = (InteriorChunkedDataNode)node;
+
+				Assert.IsTrue(rightmost || interiorNode.Children.Count >= options.InteriorOptions.MinChildCount);
+				Assert.IsTrue(interiorNode.Children.Count <= options.InteriorOptions.MaxChildCount);
+
+				int childCount = interiorNode.Children.Count;
+				for (int idx = 0; idx < childCount; idx++)
+				{
+					ChunkedDataNode childNode = await interiorNode.Children[idx].ReadBlobAsync();
+					await CheckSizesAsync(childNode, options, idx == childCount - 1);
+				}
+			}
+		}
+
+		[TestMethod]
+		public async Task ChunkingCompatV2Async()
+		{
+			ChunkingOptions chunkingOptions = new ChunkingOptions();
+			chunkingOptions.LeafOptions = new LeafChunkedDataNodeOptions(2, 2, 2);
+			chunkingOptions.InteriorOptions = new InteriorChunkedDataNodeOptions(2, 2, 2);
+
+			BlobSerializerOptions serializerOptions = new BlobSerializerOptions();
+			serializerOptions.Converters.Add(new InteriorChunkedDataNodeConverter(HordeApiVersion.Initial)); // Does not include length fields in interior nodes
+
+			KeyValueStorageNamespace store = KeyValueStorageNamespace.CreateInMemory();
+
+			byte[] data = Encoding.UTF8.GetBytes("hello world");
+
+			IHashedBlobRef<DirectoryNode> handle;
+			await using (IBlobWriter writer = store.CreateBlobWriter(options: serializerOptions))
+			{
+				using ChunkedDataWriter chunkedWriter = new ChunkedDataWriter(writer, chunkingOptions);
+				await chunkedWriter.AppendAsync(data, CancellationToken.None);
+				ChunkedData chunkedData = await chunkedWriter.FlushAsync(CancellationToken.None);
+
+				DirectoryNode directoryNode = new DirectoryNode();
+				directoryNode.AddFile("test", FileEntryFlags.None, data.Length, chunkedData);
+				handle = await writer.WriteBlobAsync(directoryNode);
+			}
+
+			DirectoryInfo tempDir = new DirectoryInfo(Path.Combine(Path.GetTempPath(), Guid.NewGuid().ToString()));
+			try
+			{
+				await handle.ExtractAsync(tempDir, new ExtractOptions(), NullLogger.Instance, CancellationToken.None);
+
+				byte[] outputData = await File.ReadAllBytesAsync(Path.Combine(tempDir.FullName, "test"));
+				Assert.IsTrue(outputData.SequenceEqual(data));
+			}
+			finally
+			{
+				tempDir.Delete(true);
+			}
+		}
+
+		[TestMethod]
+		public async Task ChunkOrderAsync()
+		{
+			InteriorChunkedDataNodeOptions interiorOptions = new InteriorChunkedDataNodeOptions(2, 2, 2);
+
+			BlobSerializerOptions serializerOptions = new BlobSerializerOptions();
+			serializerOptions.Converters.Add(new InteriorChunkedDataNodeConverter());
+
+			await using MemoryBlobWriter blobWriter = new MemoryBlobWriter(new BlobSerializerOptions());
+
+			IHashedBlobRef<LeafChunkedDataNode> leafRef = await blobWriter.WriteBlobAsync(new LeafChunkedDataNode(new byte[] { 1, 2, 3 }));
+			int leafRefIndex = MemoryBlobWriter.GetIndex((IHashedBlobRef)leafRef);
+			ChunkedDataNodeRef leafChunkedRef = new ChunkedDataNodeRef(3, 0, leafRef);
+
+			List<ChunkedDataNodeRef> leafNodeRefs = Enumerable.Repeat(leafChunkedRef, 10000).ToList();
+			ChunkedDataNodeRef root = await InteriorChunkedDataNode.CreateTreeAsync(leafNodeRefs, interiorOptions, blobWriter, CancellationToken.None);
+
+			List<IBlobRef> list = new List<IBlobRef>();
+			await GetReadOrderAsync(root.Handle, list);
+
+			int prevIndex = Int32.MaxValue;
+			for (int idx = 0; idx < list.Count; idx++)
+			{
+				int index = MemoryBlobWriter.GetIndex((IHashedBlobRef)list[idx].Innermost);
+				if (index != leafRefIndex)
+				{
+					Console.WriteLine("{0}", index);
+					Assert.IsTrue(index <= prevIndex);
+					prevIndex = index;
+				}
+			}
+		}
+
+		static async Task GetReadOrderAsync(IBlobRef handle, List<IBlobRef> list)
+		{
+			list.Add(handle);
+
+			using BlobData blobData = await handle.ReadBlobDataAsync(CancellationToken.None);
+			foreach (IBlobRef childHandle in blobData.Imports)
+			{
+				await GetReadOrderAsync(childHandle, list);
+			}
+		}
+		private struct ChunkOffsetsAndSizes
+		{
+			public List<int> _offsets;
+			public List<int> _sizes;
+		}
+
+		/// <summary>
+		/// Utility function to test low-level chunking algorithm with different chunking parameters
+		/// </summary>
+		private static ChunkOffsetsAndSizes FindChunks(ReadOnlySpan<byte> data, int minSize, int maxSize, int targetSize)
+		{
+			ChunkOffsetsAndSizes chunks;
+
+			chunks._offsets = new List<int>();
+			chunks._sizes = new List<int>();
+
+			int currentOffset = 0;
+			while (currentOffset != data.Length)
+			{
+				chunks._offsets.Add(currentOffset);
+				int chunkLength = BuzHash.FindChunkLength(data.Slice(currentOffset), minSize, maxSize, targetSize, out _);
+				currentOffset += chunkLength;
+				chunks._sizes.Add(currentOffset);
+			}
+
+			Assert.AreEqual(chunks._offsets.Count, chunks._sizes.Count);
+
+			return chunks;
+		}
+
+		/// <summary>
+		/// Utility function to test low-level chunking algorithm with varying target chunk sizes
+		/// </summary>
+		private static ChunkOffsetsAndSizes FindChunks(ReadOnlySpan<byte> data, int targetSize)
+		{
+			int minSize = targetSize / 2;
+			int maxSize = targetSize * 4;
+			return FindChunks(data, minSize, maxSize, targetSize);
+		}
+
+		/// <summary>
+		/// Verify that chunking algorithm produces chunks at predetermined offsets,
+		/// which is important for generating consistent chunking between different implementations.
+		/// </summary>
+		[TestMethod]
+		public void DeterministicChunkOffsets()
+		{
+			int dataSize = 1 << 20; // 1MB test buffer
+			byte[] data = GenerateChunkingTestBuffer(dataSize, 1234);
+
+			List<int> expectedChunkOffsets = new List<int>
+			{
+				0, 34577, 128471, 195115, 238047, 297334, 358754, 396031,
+				462359, 508658, 601550, 702021, 754650, 790285, 854987, 887998,
+				956848, 1042406
+			};
+
+			int targetSize = 65536;
+			ChunkOffsetsAndSizes chunks = FindChunks(data, targetSize);
+
+			CollectionAssert.AreEqual(expectedChunkOffsets, chunks._offsets);
+		}
+
+		/// <summary>
+		/// Verify that chunking algorithm produces chunks of expected sizes for different target chunk size configurations.
+		/// This helps validating chunking consistency between different implementations of the algorithm.
+		/// </summary>
+		[TestMethod]
+		public void DeterministicChunkSize()
+		{
+			int dataSize = 128 << 20; // 128MB test buffer
+			byte[] data = GenerateChunkingTestBuffer(dataSize, 1234);
+
+			// Chunk size in KB and expected number of generated chunks
+			List<(int, int)> configList = new List<(int, int)>
+			{
+				(8, 16442),
+				(16, 8146),
+				(32, 4089),
+				(64, 2019),
+				(96, 1362),
+				(128, 1012),
+				(160, 811),
+				(192, 681),
+				(256, 503),
+			};
+
+			Parallel.ForEach(configList, config =>
+			{
+				(int targetSizeKB, int expectedNumChunks) = config;
+
+				int targetSize = targetSizeKB * 1024;
+
+				ChunkOffsetsAndSizes chunks = FindChunks(data, targetSize);
+
+				int avgChunkSize = dataSize / chunks._offsets.Count;
+				int absError = Math.Abs(avgChunkSize - targetSize);
+				double absErrorPct = (100.0 * absError) / targetSize;
+
+				Assert.IsTrue(absErrorPct < 5.0, $"Average chunk size {avgChunkSize} is more than 5% different from target {targetSize}");
+				Assert.AreEqual(expectedNumChunks, chunks._offsets.Count);
+			});
+		}
+
+		/// <summary>
+		/// Utility function to quickly generate a deterministic byte sequence
+		/// that can be used to ensure consistency between multiple implementations 
+		/// of the chunking algorithm.
+		/// </summary>
+		static byte[] GenerateChunkingTestBuffer(int count, uint seed)
+		{
+			byte[] data = new byte[count];
+			uint rngState = seed;
+
+			for (int i = 0; i < count; ++i)
+			{
+				data[i] = (byte)XorShift32(ref rngState);
+			}
+
+			return data;
+		}
+		static uint XorShift32(ref uint state)
+		{
+			uint x = state;
+			x ^= x << 13;
+			x ^= x >> 17;
+			x ^= x << 5;
+			state = x;
+			return x;
+		}
+	}
+}
diff --git a/Engine/Source/Programs/Shared/EpicGames.Horde.Tests/Storage/CompressionTests.cs b/Engine/Source/Programs/Shared/EpicGames.Horde.Tests/Storage/CompressionTests.cs
new file mode 100644
--- /dev/null
+++ b/Engine/Source/Programs/Shared/EpicGames.Horde.Tests/Storage/CompressionTests.cs
@@ -0,0 +1,63 @@
+﻿// Copyright Epic Games, Inc. All Rights Reserved.
+
+using System;
+using System.Diagnostics;
+using System.Linq;
+using EpicGames.Core;
+using EpicGames.Horde.Storage.Bundles;
+using Microsoft.VisualStudio.TestTools.UnitTesting;
+
+namespace EpicGames.Horde.Tests.Storage
+{
+	[TestClass]
+	public class CompressionTests
+	{
+		[TestMethod]
+		[DataTestMethod]
+		[DataRow(BundleCompressionFormat.LZ4)]
+		[DataRow(BundleCompressionFormat.Gzip)]
+		[DataRow(BundleCompressionFormat.Brotli)]
+		[DataRow(BundleCompressionFormat.Zstd)]
+		public void SimpleData(BundleCompressionFormat format)
+		{
+			byte[] input = new byte[2 * 1024 * 1024];
+			for (int idx = 0; idx < input.Length; idx++)
+			{
+				input[idx] = (byte)idx;
+			}
+
+			RunTests(format, input);
+		}
+
+		[TestMethod]
+		[DataTestMethod]
+		[DataRow(BundleCompressionFormat.LZ4)]
+		[DataRow(BundleCompressionFormat.Gzip)]
+		[DataRow(BundleCompressionFormat.Brotli)]
+		[DataRow(BundleCompressionFormat.Zstd)]
+		public void RandomData(BundleCompressionFormat format)
+		{
+			byte[] data = new byte[2 * 1024 * 1024];
+			new Random(0).NextBytes(data);
+
+			RunTests(format, data);
+		}
+
+		static void RunTests(BundleCompressionFormat format, byte[] input)
+		{
+			ArrayMemoryWriter writer = new ArrayMemoryWriter(10 * 1024 * 1024);
+			Stopwatch compressTimer = Stopwatch.StartNew();
+			BundleData.Compress(format, input, writer);
+			compressTimer.Stop();
+
+			byte[] output = new byte[input.Length];
+			Stopwatch decompressTimer = Stopwatch.StartNew();
+			BundleData.Decompress(format, writer.WrittenMemory, output);
+			decompressTimer.Stop();
+
+			Assert.IsTrue(input.SequenceEqual(output));
+
+			Console.WriteLine("Output size: {0:n0} bytes, compression time: {1}ms, decompression time: {2}ms", writer.WrittenMemory.Length, compressTimer.Elapsed.TotalMilliseconds, decompressTimer.Elapsed.TotalMilliseconds);
+		}
+	}
+}
diff --git a/Engine/Source/Programs/Shared/EpicGames.Horde.Tests/Storage/ContentChunkerTests.cs b/Engine/Source/Programs/Shared/EpicGames.Horde.Tests/Storage/ContentChunkerTests.cs
new file mode 100644
--- /dev/null
+++ b/Engine/Source/Programs/Shared/EpicGames.Horde.Tests/Storage/ContentChunkerTests.cs
@@ -0,0 +1,105 @@
+// Copyright Epic Games, Inc. All Rights Reserved.
+
+using System;
+using System.Collections.Generic;
+using System.Linq;
+using System.Threading;
+using System.Threading.Tasks;
+using EpicGames.Core;
+using EpicGames.Horde.Storage;
+using EpicGames.Horde.Storage.Nodes;
+using Microsoft.VisualStudio.TestTools.UnitTesting;
+
+namespace EpicGames.Horde.Tests.Storage
+{
+	[TestClass]
+	public class ContentChunkerTests
+	{
+		record struct SplitPoint(int Index, int Length);
+
+		[TestMethod]
+		public async Task SerialTestAsync()
+		{
+			List<SplitPoint> legacyPoints;
+			await using (AsyncPipeline pipeline = new AsyncPipeline(CancellationToken.None))
+			{
+				byte[] buffer = new byte[1024 * 1024];
+				SerialBuzHashChunker chunker = new SerialBuzHashChunker(pipeline, buffer, LeafChunkedDataNodeOptions.Default);
+				legacyPoints = await TestPipelineAsync(chunker);
+			}
+		}
+
+		[TestMethod]
+		public async Task ParallelTestAsync()
+		{
+			List<SplitPoint> latestPoints;
+			await using (AsyncPipeline pipeline = new AsyncPipeline(CancellationToken.None))
+			{
+				byte[] buffer = new byte[1024 * 1024];
+				ParallelBuzHashChunker chunker = new ParallelBuzHashChunker(pipeline, buffer, LeafChunkedDataNodeOptions.Default);
+				latestPoints = await TestPipelineAsync(chunker);
+			}
+		}
+
+		[TestMethod]
+		public async Task CompareTestAsync()
+		{
+			List<SplitPoint> legacyPoints;
+			await using (AsyncPipeline pipeline = new AsyncPipeline(CancellationToken.None))
+			{
+				byte[] buffer = new byte[1024 * 1024];
+				SerialBuzHashChunker chunker = new SerialBuzHashChunker(pipeline, buffer, LeafChunkedDataNodeOptions.Default);
+				legacyPoints = await TestPipelineAsync(chunker);
+			}
+
+			List<SplitPoint> latestPoints;
+			await using (AsyncPipeline pipeline = new AsyncPipeline(CancellationToken.None))
+			{
+				byte[] buffer = new byte[1024 * 1024];
+				ParallelBuzHashChunker chunker = new ParallelBuzHashChunker(pipeline, buffer, LeafChunkedDataNodeOptions.Default);
+				latestPoints = await TestPipelineAsync(chunker);
+			}
+
+			Assert.IsTrue(Enumerable.SequenceEqual(legacyPoints, latestPoints));
+		}
+
+		static async Task<List<SplitPoint>> TestPipelineAsync(ContentChunker pipeline)
+		{
+			Random random = new Random(0);
+
+			int[] lengths = [ 0, 1, 32 * 1024, 129 * 1024, 256 * 1024, 4 * 1024 * 1024 ];
+
+			foreach (int length in lengths)
+			{
+				byte[] data = new byte[length];
+				random.NextBytes(data);
+
+				ChunkerSource input = new MemoryChunkerSource(data);
+				await pipeline.SourceWriter.WriteAsync(input);
+			}
+
+			List<SplitPoint> splitPoints = new List<SplitPoint>();
+			for (int idx = 0; idx < lengths.Length; idx++)
+			{
+				int outputLength = 0;
+
+				while (await pipeline.OutputReader.WaitToReadAsync())
+				{
+					if (pipeline.OutputReader.TryRead(out ChunkerOutput? output))
+					{
+						while (await output.MoveNextAsync())
+						{
+							splitPoints.Add(new SplitPoint(idx, output.Data.Length));
+							outputLength += output.Data.Length;
+						}
+						break;
+					}
+				}
+
+				Assert.AreEqual(outputLength, lengths[idx]);
+			}
+
+			return splitPoints;
+		}
+	}
+}
diff --git a/Engine/Source/Programs/Shared/EpicGames.Horde.Tests/Storage/FileMetadataDbTests.cs b/Engine/Source/Programs/Shared/EpicGames.Horde.Tests/Storage/FileMetadataDbTests.cs
new file mode 100644
--- /dev/null
+++ b/Engine/Source/Programs/Shared/EpicGames.Horde.Tests/Storage/FileMetadataDbTests.cs
@@ -0,0 +1,209 @@
+﻿// Copyright Epic Games, Inc. All Rights Reserved.
+
+using System;
+using System.Collections.Generic;
+using System.Linq;
+using System.Threading.Tasks;
+using EpicGames.Core;
+using EpicGames.Horde.Storage;
+using Microsoft.VisualStudio.TestTools.UnitTesting;
+
+namespace EpicGames.Horde.Tests.Storage
+{
+	[TestClass]
+	public class FileMetadataDbTests
+	{
+		[TestMethod]
+		public async Task DirectoryTestsAsync()
+		{
+			using FileMetadataDb db = await FileMetadataDb.CreateInMemoryAsync();
+
+			int a = await db.AddDirectoryAsync(FileMetadataDb.RootDirectoryId, "A");
+			int b = await db.AddDirectoryAsync(FileMetadataDb.RootDirectoryId, "B");
+			int c = await db.AddDirectoryAsync(b, "C");
+			int d = await db.AddDirectoryAsync(b, "D");
+
+			{
+				List<DirectoryRow> rows = (await db.GetDirectoriesAsync(FileMetadataDb.RootDirectoryId)).OrderBy(x => x.Id).ToList();
+				Assert.AreEqual(2, rows.Count);
+				Assert.AreEqual(a, rows[0].Id);
+				Assert.AreEqual("A", rows[0].Name);
+				Assert.AreEqual(b, rows[1].Id);
+				Assert.AreEqual("B", rows[1].Name);
+			}
+
+			{
+				List<DirectoryRow> rows = (await db.GetDirectoriesAsync(a)).OrderBy(x => x.Id).ToList();
+				Assert.AreEqual(0, rows.Count);
+			}
+
+			{
+				List<DirectoryRow> rows = (await db.GetDirectoriesAsync(b)).OrderBy(x => x.Id).ToList();
+				Assert.AreEqual(2, rows.Count);
+				Assert.AreEqual(c, rows[0].Id);
+				Assert.AreEqual("C", rows[0].Name);
+				Assert.AreEqual(d, rows[1].Id);
+				Assert.AreEqual("D", rows[1].Name);
+			}
+
+			Assert.AreEqual("", await db.GetDirectoryPathAsync(FileMetadataDb.RootDirectoryId));
+			Assert.AreEqual("A/", await db.GetDirectoryPathAsync(a));
+			Assert.AreEqual("B/", await db.GetDirectoryPathAsync(b));
+			Assert.AreEqual("B/C/", await db.GetDirectoryPathAsync(c));
+			Assert.AreEqual("B/D/", await db.GetDirectoryPathAsync(d));
+		}
+
+		[TestMethod]
+		public async Task FileTestsAsync()
+		{
+			using FileMetadataDb db = await FileMetadataDb.CreateInMemoryAsync();
+
+			int a = await db.AddDirectoryAsync(FileMetadataDb.RootDirectoryId, "A");
+
+			DateTime time = new DateTime(2024, 1, 1, 12, 30, 45, DateTimeKind.Utc);
+			int c = await db.AddFileAsync(FileMetadataDb.RootDirectoryId, "C.txt", time, 1234);
+			int d = await db.AddFileAsync(FileMetadataDb.RootDirectoryId, "D.txt", time, 2345);
+			int e = await db.AddFileAsync(a, "E.txt", time, 3456);
+			int f = await db.AddFileAsync(a, "F.txt", time, 4567);
+
+			{
+				List<FileRow> files = (await db.FindFilesInDirectoryAsync(FileMetadataDb.RootDirectoryId)).OrderBy(x => x.Id).ToList();
+				Assert.AreEqual(2, files.Count);
+
+				Assert.AreEqual(c, files[0].Id);
+				Assert.AreEqual(time, files[0].Time);
+				Assert.AreEqual(1234, files[0].Length);
+
+				Assert.AreEqual(d, files[1].Id);
+				Assert.AreEqual(time, files[1].Time);
+				Assert.AreEqual(2345, files[1].Length);
+			}
+
+			{
+				List<FileRow> files = (await db.FindFilesInDirectoryAsync(a)).OrderBy(x => x.Id).ToList();
+				Assert.AreEqual(2, files.Count);
+
+				Assert.AreEqual(e, files[0].Id);
+				Assert.AreEqual(time, files[0].Time);
+				Assert.AreEqual(3456, files[0].Length);
+
+				Assert.AreEqual(f, files[1].Id);
+				Assert.AreEqual(time, files[1].Time);
+				Assert.AreEqual(4567, files[1].Length);
+			}
+
+			Assert.AreEqual("C.txt", await db.GetFilePathAsync(c));
+			Assert.AreEqual("D.txt", await db.GetFilePathAsync(d));
+			Assert.AreEqual("A/E.txt", await db.GetFilePathAsync(e));
+			Assert.AreEqual("A/F.txt", await db.GetFilePathAsync(f));
+		}
+
+		[TestMethod]
+		public async Task BatchFileTestsAsync()
+		{
+			using FileMetadataDb db = await FileMetadataDb.CreateInMemoryAsync();
+			int root = await db.AddDirectoryAsync(FileMetadataDb.RootDirectoryId, "");
+
+			List<FileRow> files = new List<FileRow>();
+			files.Add(new FileRow(root, "A", DateTime.MinValue, 1));
+			files.Add(new FileRow(root, "B", DateTime.MinValue, 1));
+			files.Add(new FileRow(root, "C", DateTime.MinValue, 1));
+
+			await db.AddFilesAsync(files);
+			Assert.AreEqual(3, (await db.FindFilesInDirectoryAsync(root)).Count());
+		}
+
+		[TestMethod]
+		public async Task BatchDirectoryTestsAsync()
+		{
+			using FileMetadataDb db = await FileMetadataDb.CreateInMemoryAsync();
+			await db.AddDirectoriesAsync(new[] { new DirectoryRow(FileMetadataDb.RootDirectoryId, "A"), new DirectoryRow(FileMetadataDb.RootDirectoryId, "B"), new DirectoryRow(FileMetadataDb.RootDirectoryId, "C") });
+			Assert.AreEqual(3, (await db.GetDirectoriesAsync(FileMetadataDb.RootDirectoryId)).Count());
+		}
+
+		[TestMethod]
+		public async Task ChunkTestsAsync()
+		{
+			using FileMetadataDb db = await FileMetadataDb.CreateInMemoryAsync();
+
+			int d1 = await db.AddDirectoryAsync(FileMetadataDb.RootDirectoryId, "D1");
+			int d2 = await db.AddDirectoryAsync(d1, "D2");
+
+			DateTime time = new DateTime(2024, 1, 1, 12, 30, 45, DateTimeKind.Utc);
+			int f1 = await db.AddFileAsync(d1, "F1.txt", time, 1234);
+			int f2 = await db.AddFileAsync(d1, "F2.txt", time, 2345);
+			int f3 = await db.AddFileAsync(d2, "F3.txt", time, 3456);
+			int f4 = await db.AddFileAsync(d2, "F4.txt", time, 4567);
+
+			IoHash h1 = IoHash.Compute(new byte[] { 1, 2, 3 });
+			IoHash h2 = IoHash.Compute(new byte[] { 1, 2, 3, 4 });
+
+			int c1 = await db.AddChunkAsync(f1, 0, 10, h1);
+			int c2 = await db.AddChunkAsync(f2, 10, 10, h1);
+			int c3 = await db.AddChunkAsync(f3, 20, 10, h1);
+			int c4 = await db.AddChunkAsync(f4, 0, 10, h2);
+
+			{
+				List<ChunkRow> chunks = (await db.FindChunksAsync(h1, 10)).OrderBy(x => x.Id).ToList();
+				Assert.AreEqual(3, chunks.Count);
+				Assert.AreEqual(new ChunkRow(c1, f1, 0, 10, h1), chunks[0]);
+				Assert.AreEqual(new ChunkRow(c2, f2, 10, 10, h1), chunks[1]);
+				Assert.AreEqual(new ChunkRow(c3, f3, 20, 10, h1), chunks[2]);
+			}
+
+			{
+				List<ChunkRow> chunks = (await db.FindChunksAsync(h2, 10)).OrderBy(x => x.Id).ToList();
+				Assert.AreEqual(1, chunks.Count);
+				Assert.AreEqual(new ChunkRow(c4, f4, 0, 10, h2), chunks[0]);
+			}
+		}
+
+		[TestMethod]
+		public async Task RemoveTreeTestAsync()
+		{
+			using FileMetadataDb db = await FileMetadataDb.CreateInMemoryAsync();
+
+			int d1 = await db.AddDirectoryAsync(FileMetadataDb.RootDirectoryId, "D1");
+			int d2 = await db.AddDirectoryAsync(d1, "D2");
+			int d3 = await db.AddDirectoryAsync(d2, "D3");
+			int d4 = await db.AddDirectoryAsync(d2, "D4");
+			int d5 = await db.AddDirectoryAsync(d4, "D5");
+
+			DateTime time = new DateTime(2024, 1, 1, 12, 30, 45, DateTimeKind.Utc);
+			int f1 = await db.AddFileAsync(d1, "F1.txt", time, 1234);
+			int f2 = await db.AddFileAsync(d2, "F2.txt", time, 2345);
+			int f3 = await db.AddFileAsync(d3, "F3.txt", time, 3456);
+			int f4 = await db.AddFileAsync(d4, "F4.txt", time, 4567);
+			int f5 = await db.AddFileAsync(d5, "F5.txt", time, 5678);
+
+			IoHash h1 = IoHash.Compute(new byte[] { 1, 2, 3 });
+
+			int c1 = await db.AddChunkAsync(f1, 0, 10, h1);
+			int c2 = await db.AddChunkAsync(f2, 10, 10, h1);
+			int c3 = await db.AddChunkAsync(f3, 20, 10, h1);
+			int c4 = await db.AddChunkAsync(f4, 0, 10, h1);
+			int c5 = await db.AddChunkAsync(f5, 10, 10, h1);
+
+			{
+				List<ChunkRow> chunks = (await db.FindChunksAsync(h1, 10)).OrderBy(x => x.Id).ToList();
+				Assert.AreEqual(5, chunks.Count);
+				Assert.AreEqual(new ChunkRow(c1, f1, 0, 10, h1), chunks[0]);
+				Assert.AreEqual(new ChunkRow(c2, f2, 10, 10, h1), chunks[1]);
+				Assert.AreEqual(new ChunkRow(c3, f3, 20, 10, h1), chunks[2]);
+				Assert.AreEqual(new ChunkRow(c4, f4, 0, 10, h1), chunks[3]);
+				Assert.AreEqual(new ChunkRow(c5, f5, 10, 10, h1), chunks[4]);
+			}
+
+			await db.RemoveDirectoryAsync(d4);
+
+			{
+				List<ChunkRow> chunks = (await db.FindChunksAsync(h1, 10)).OrderBy(x => x.Id).ToList();
+				Assert.AreEqual(3, chunks.Count);
+
+				Assert.AreEqual(new ChunkRow(c1, f1, 0, 10, h1), chunks[0]);
+				Assert.AreEqual(new ChunkRow(c2, f2, 10, 10, h1), chunks[1]);
+				Assert.AreEqual(new ChunkRow(c3, f3, 20, 10, h1), chunks[2]);
+			}
+		}
+	}
+}
diff --git a/Engine/Source/Programs/Shared/EpicGames.Horde.Tests/Storage/StorageClientTests.cs b/Engine/Source/Programs/Shared/EpicGames.Horde.Tests/Storage/StorageClientTests.cs
new file mode 100644
--- /dev/null
+++ b/Engine/Source/Programs/Shared/EpicGames.Horde.Tests/Storage/StorageClientTests.cs
@@ -0,0 +1,165 @@
+// Copyright Epic Games, Inc. All Rights Reserved.
+
+using System;
+using System.Threading.Tasks;
+using EpicGames.Core;
+using EpicGames.Horde.Storage;
+using EpicGames.Horde.Storage.Bundles;
+using EpicGames.Horde.Storage.Bundles.V2;
+using EpicGames.Horde.Storage.Clients;
+using Microsoft.Extensions.Logging.Abstractions;
+using Microsoft.VisualStudio.TestTools.UnitTesting;
+
+namespace EpicGames.Horde.Tests.Storage
+{
+	[TestClass]
+	public class StorageNamespaceTests
+	{
+		[BlobConverter(typeof(TestNodeConverter))]
+		class TestNode
+		{
+			public static BlobType BlobType = new BlobType("{99601905-4F6E-A089-03D6-F187711BAFEE}", 1);
+
+			public int Value { get; }
+			public byte[] Padding { get; set; } = Array.Empty<byte>();
+			public IHashedBlobRef<TestNode>[] Refs { get; }
+
+			public TestNode(int value, params IHashedBlobRef<TestNode>[] refs)
+			{
+				Value = value;
+				Refs = refs;
+			}
+
+			public TestNode(int value, byte[] padding, IHashedBlobRef<TestNode>[] refs)
+			{
+				Value = value;
+				Padding = padding;
+				Refs = refs;
+			}
+		}
+
+		class TestNodeConverter : BlobConverter<TestNode>
+		{
+			public override TestNode Read(IBlobReader reader, BlobSerializerOptions options)
+			{
+				int value = reader.ReadInt32();
+				byte[] padding = reader.ReadVariableLengthBytes().ToArray();
+				IHashedBlobRef<TestNode>[] refs = reader.ReadVariableLengthArray(() => reader.ReadBlobRef<TestNode>());
+
+				return new TestNode(value, padding, refs);
+			}
+
+			public override BlobType Write(IBlobWriter writer, TestNode value, BlobSerializerOptions options)
+			{
+				writer.WriteInt32(value.Value);
+				writer.WriteVariableLengthBytes(value.Padding);
+				writer.WriteVariableLengthArray(value.Refs, x => writer.WriteBlobRef(x));
+
+				return TestNode.BlobType;
+			}
+		}
+
+		[TestMethod]
+		public async Task TestBasicAsync()
+		{
+			KeyValueStorageNamespace store = KeyValueStorageNamespace.CreateInMemory();
+			await TestBasicAsync(store);
+		}
+
+		[TestMethod]
+		public async Task TestBasicBundleV2Async()
+		{
+			BundleStorageNamespace storeV2 = BundleStorageNamespace.CreateInMemory(NullLogger.Instance);
+			await TestBasicAsync(storeV2);
+		}
+
+		static async Task TestBasicAsync(IStorageNamespace store)
+		{
+			IHashedBlobRef<TestNode> nodeRef;
+			await using (IBlobWriter writer = store.CreateBlobWriter())
+			{
+				nodeRef = await writer.WriteBlobAsync(new TestNode(123));
+			}
+			await store.AddRefAsync("hello", nodeRef);
+
+			TestNode output = await store.ReadRefTargetAsync<TestNode>("hello");
+			Assert.AreEqual(123, output.Value);
+		}
+
+		[TestMethod]
+		public async Task PendingRefsAsync()
+		{
+			await using BundleCache cache = new BundleCache();
+
+			BundleStorageNamespace store = BundleStorageNamespace.CreateInMemory(NullLogger.Instance);
+
+			IHashedBlobRef<TestNode> nodeRef2;
+			await using (IBlobWriter writer = store.CreateBlobWriter())
+			{
+				IHashedBlobRef<TestNode> nodeRef1 = await writer.WriteBlobAsync(new TestNode(123));
+				nodeRef2 = await writer.WriteBlobAsync(new TestNode(456, nodeRef1));
+			}
+			await store.AddRefAsync("hello", nodeRef2);
+
+			TestNode output2 = await store.ReadRefTargetAsync<TestNode>("hello");
+			Assert.AreEqual(456, output2.Value);
+			Assert.AreEqual(1, output2.Refs.Length);
+
+			TestNode output1 = await output2.Refs[0].ReadBlobAsync();
+			Assert.AreEqual(0, output1.Refs.Length);
+			Assert.AreEqual(123, output1.Value);
+		}
+
+		[TestMethod]
+		public async Task PacketFlushAsync()
+		{
+			await using BundleCache cache = new BundleCache();
+
+			BundleOptions bundleOptions = new BundleOptions { MinCompressionPacketSize = 100, MaxBlobSize = 1024 * 1024, MaxVersion = BundleVersion.LatestV2 };
+			BundleStorageNamespace store = BundleStorageNamespace.CreateInMemory(bundleOptions, NullLogger.Instance);
+
+			await using IBlobWriter writer = store.CreateBlobWriter();
+			IHashedBlobRef<TestNode> nodeRef1 = await writer.WriteBlobAsync(new TestNode(123) { Padding = new byte[1024] });
+			await writer.FlushAsync();
+			IHashedBlobRef<TestNode> nodeRef2 = await writer.WriteBlobAsync(new TestNode(456, nodeRef1) { Padding = new byte[1024] });
+			IHashedBlobRef<TestNode> nodeRef3 = await writer.WriteBlobAsync(new TestNode(789, nodeRef2));
+
+			// nodeRef1 is in a flushed bundle
+			ExportHandle export1 = (ExportHandle)nodeRef1.Innermost;
+			BundleWriter.PendingPacketHandle packet1 = (BundleWriter.PendingPacketHandle)export1.Packet;
+			BundleWriter.PendingBundleHandle bundle1 = (BundleWriter.PendingBundleHandle)packet1.Bundle;
+			Assert.IsNotNull(packet1.FlushedHandle);
+			Assert.IsNotNull(bundle1.FlushedHandle);
+
+			TestNode node1 = await export1.ReadBlobAsync<TestNode>();
+			Assert.AreEqual(123, node1.Value);
+
+			// nodeRef2 is in a flushed packet, unflushed bundle
+			ExportHandle export2 = (ExportHandle)nodeRef2.Innermost;
+
+			BundleWriter.PendingPacketHandle packet2 = (BundleWriter.PendingPacketHandle)export2.Packet;
+			Assert.AreNotEqual(packet1, packet2);
+			Assert.IsNotNull(packet2.FlushedHandle);
+
+			BundleWriter.PendingBundleHandle bundle2 = (BundleWriter.PendingBundleHandle)packet2.Bundle;
+			Assert.AreNotEqual(bundle1, bundle2);
+			Assert.IsNull(bundle2.FlushedHandle);
+
+			TestNode node2 = await export2.ReadBlobAsync<TestNode>();
+			Assert.AreEqual(456, node2.Value);
+
+			// nodeRef3 is in an unflushed packet, unflushed bundle
+			ExportHandle export3 = (ExportHandle)nodeRef3.Innermost;
+
+			BundleWriter.PendingPacketHandle packet3 = (BundleWriter.PendingPacketHandle)export3.Packet;
+			Assert.AreNotEqual(packet2, packet3);
+			Assert.IsNull(packet3.FlushedHandle);
+
+			BundleWriter.PendingBundleHandle bundle3 = (BundleWriter.PendingBundleHandle)packet3.Bundle;
+			Assert.AreEqual(bundle2, bundle3);
+
+			TestNode node3 = await export3.ReadBlobAsync<TestNode>();
+			Assert.AreEqual(789, node3.Value);
+		}
+	}
+}
diff --git a/Engine/Source/Programs/Shared/EpicGames.Horde.Tests/StorageClientTests.cs b/Engine/Source/Programs/Shared/EpicGames.Horde.Tests/StorageClientTests.cs
deleted file mode 100644
diff --git a/Engine/Source/Programs/Shared/EpicGames.Horde.Tests/TransportTests.cs b/Engine/Source/Programs/Shared/EpicGames.Horde.Tests/TransportTests.cs
deleted file mode 100644
diff --git a/Engine/Source/Programs/Shared/EpicGames.Horde.Tests/Utilities/IntervalJsonConverterTests.cs b/Engine/Source/Programs/Shared/EpicGames.Horde.Tests/Utilities/IntervalJsonConverterTests.cs
new file mode 100644
--- /dev/null
+++ b/Engine/Source/Programs/Shared/EpicGames.Horde.Tests/Utilities/IntervalJsonConverterTests.cs
@@ -0,0 +1,68 @@
+// Copyright Epic Games, Inc. All Rights Reserved.
+
+using System;
+using System.Text.Json;
+using System.Text.Json.Serialization;
+using EpicGames.Horde.Utilities;
+using HordeServer.Utilities;
+using Microsoft.VisualStudio.TestTools.UnitTesting;
+
+namespace EpicGames.Horde.Tests.Utilities
+{
+	[TestClass]
+	public class IntervalJsonConverterTests
+	{
+		class TestClass
+		{
+			[JsonConverter(typeof(IntervalJsonConverter))]
+			public TimeSpan Time { get; set; }
+		}
+
+		readonly JsonSerializerOptions _options = new JsonSerializerOptions();
+
+		[TestMethod]
+		public void Test()
+		{
+			TestEqual("2w", TimeSpan.FromDays(14));
+			TestEqual("2w1d", TimeSpan.FromDays(15));
+			TestEqual("5d", TimeSpan.FromDays(5));
+			TestEqual("5d4h", TimeSpan.FromDays(5) + TimeSpan.FromHours(4.0));
+			TestEqual("5d3m", TimeSpan.FromDays(5) + TimeSpan.FromMinutes(3.0));
+			TestEqual("5d4h3m1s", TimeSpan.FromDays(5) + TimeSpan.FromHours(4) + TimeSpan.FromMinutes(3) + TimeSpan.FromSeconds(1));
+
+			TestEqual("2h04m", TimeSpan.FromHours(2) + TimeSpan.FromMinutes(4.0), "2h4m");
+			TestEqual("1m30", TimeSpan.FromSeconds(90), "1m30s");
+			TestEqual("90s", TimeSpan.FromSeconds(90), "1m30s");
+			TestEqual("120s", TimeSpan.FromMinutes(2), "2m");
+
+			TestEqual("00:00:01", TimeSpan.FromSeconds(1.0), "1s");
+			TestEqual("00:02:00", TimeSpan.FromMinutes(2.0), "2m");
+			TestEqual("03:00:00", TimeSpan.FromHours(3.0), "3h");
+
+			TestInvalid("3");
+			TestInvalid("m");
+		}
+
+		void TestEqual(string text, TimeSpan time)
+		{
+			TestEqual(text, time, text);
+		}
+
+		void TestEqual(string text, TimeSpan time, string finalText)
+		{
+			string json = $"{{\"Time\":\"{text}\"}}";
+
+			TestClass? output = JsonSerializer.Deserialize<TestClass>(json, _options);
+			Assert.IsNotNull(output);
+			Assert.AreEqual(time, output.Time);
+
+			string serializedText = JsonSerializer.Serialize(output, _options);
+			Assert.AreEqual($"{{\"Time\":\"{finalText}\"}}", serializedText);
+		}
+
+		void TestInvalid(string text)
+		{
+			Assert.ThrowsException<JsonException>(() => JsonSerializer.Deserialize<TestClass>(text, _options));
+		}
+	}
+}
diff --git a/Engine/Source/Programs/Shared/EpicGames.Horde.Tests/Utilities/TimeOfDayJsonConverterTests.cs b/Engine/Source/Programs/Shared/EpicGames.Horde.Tests/Utilities/TimeOfDayJsonConverterTests.cs
new file mode 100644
--- /dev/null
+++ b/Engine/Source/Programs/Shared/EpicGames.Horde.Tests/Utilities/TimeOfDayJsonConverterTests.cs
@@ -0,0 +1,27 @@
+﻿// Copyright Epic Games, Inc. All Rights Reserved.
+
+using System;
+using EpicGames.Horde.Utilities;
+using HordeServer.Utilities;
+using Microsoft.VisualStudio.TestTools.UnitTesting;
+
+namespace EpicGames.Horde.Tests.Utilities
+{
+	[TestClass]
+	public class TimeOfDayJsonConverterTests
+	{
+		[TestMethod]
+		[DataRow("4:50", 4, 50)]
+		[DataRow("4:50am", 4, 50)]
+		[DataRow("4:50pm", 16, 50)]
+		[DataRow("4.50", 4, 50)]
+		[DataRow("4pm", 16, 0)]
+		[DataRow("16:50", 16, 50)]
+		public void TestTimeSpan(string time, int hours, int minutes)
+		{
+			TimeSpan span = TimeOfDayJsonConverter.Parse(time);
+			Assert.AreEqual(hours, span.Hours);
+			Assert.AreEqual(minutes, span.Minutes);
+		}
+	}
+}
diff --git a/Engine/Source/Programs/Shared/EpicGames.Horde/Utilities/IntervalJsonConverter.cs b/Engine/Source/Programs/Shared/EpicGames.Horde/Utilities/IntervalJsonConverter.cs
new file mode 100644
--- /dev/null
+++ b/Engine/Source/Programs/Shared/EpicGames.Horde/Utilities/IntervalJsonConverter.cs
@@ -0,0 +1,110 @@
+﻿// Copyright Epic Games, Inc. All Rights Reserved.
+
+using System;
+using System.Text;
+using System.Text.Json;
+using System.Text.Json.Serialization;
+
+namespace EpicGames.Horde.Utilities
+{
+	/// <summary>
+	/// Converts TimeSpan intervals from formats like "30m", "1h30m", etc...
+	/// </summary>
+	public class IntervalJsonConverter : JsonConverter<TimeSpan>
+	{
+		record struct TimeSuffix(string Text, long Ticks);
+
+		static readonly TimeSuffix[] s_timeSuffixes =
+		{
+			new TimeSuffix("w", TimeSpan.TicksPerDay * 7),
+			new TimeSuffix("d", TimeSpan.TicksPerDay),
+			new TimeSuffix("h", TimeSpan.TicksPerHour),
+			new TimeSuffix("m", TimeSpan.TicksPerMinute),
+			new TimeSuffix("s", TimeSpan.TicksPerSecond)
+		};
+
+		/// <inheritdoc/>
+		public override TimeSpan Read(ref Utf8JsonReader reader, Type typeToConvert, JsonSerializerOptions options)
+		{
+			return Parse(reader.GetString());
+		}
+
+		/// <summary>
+		/// Parse a string as a time interval
+		/// </summary>
+		public static TimeSpan Parse(string? time)
+		{
+			if (String.IsNullOrEmpty(time))
+			{
+				throw new FormatException("TimeSpan may not be empty");
+			}
+
+			if (TimeSpan.TryParse(time, out TimeSpan result))
+			{
+				return result;
+			}
+
+			long ticks = 0;
+			int suffixIdx = 0;
+			for (ReadOnlySpan<char> span = time.AsSpan(); span.Length > 0;)
+			{
+				long value = 0;
+
+				// Parse the value
+				int length = 0;
+				for (; length < span.Length && span[length] >= '0' && span[length] <= '9'; length++)
+				{
+					value = (value * 10) + (span[length] - '0');
+				}
+				if (length == 0)
+				{
+					throw new FormatException($"Unable to parse '{time}' as a TimeSpan. Invalid value.");
+				}
+				span = span.Slice(length);
+
+				// Parse the suffix
+				length = 0;
+				while (length < span.Length && !(span[length] >= '0' && span[length] <= '9'))
+				{
+					length++;
+				}
+				if (length != 0) // Allow '3m30' to use the next suffix in order.
+				{
+					while (suffixIdx < s_timeSuffixes.Length && !s_timeSuffixes[suffixIdx].Text.AsSpan().SequenceEqual(span[0..length]))
+					{
+						suffixIdx++;
+					}
+					if (suffixIdx == s_timeSuffixes.Length)
+					{
+						throw new FormatException($"Unable to parse '{time}' as a TimeSpan. Invalid suffix.");
+					}
+					span = span.Slice(length);
+				}
+
+				ticks += s_timeSuffixes[suffixIdx++].Ticks * value;
+			}
+
+			return new TimeSpan(ticks);
+		}
+
+		/// <inheritdoc/>
+		public override void Write(Utf8JsonWriter writer, TimeSpan value, JsonSerializerOptions options)
+		{
+			long ticks = value.Ticks;
+
+			StringBuilder builder = new StringBuilder();
+			for (int suffixIdx = 0; suffixIdx < s_timeSuffixes.Length; suffixIdx++)
+			{
+				if (ticks >= s_timeSuffixes[suffixIdx].Ticks)
+				{
+					long timeValue = ticks / s_timeSuffixes[suffixIdx].Ticks;
+					builder.Append(timeValue);
+					builder.Append(s_timeSuffixes[suffixIdx].Text);
+					ticks -= timeValue * s_timeSuffixes[suffixIdx].Ticks;
+				}
+			}
+
+			writer.WriteStringValue(builder.ToString());
+		}
+	}
+}
diff --git a/Engine/Source/Programs/Shared/EpicGames.Horde/Utilities/TimeOfDayJsonConverter.cs b/Engine/Source/Programs/Shared/EpicGames.Horde/Utilities/TimeOfDayJsonConverter.cs
new file mode 100644
--- /dev/null
+++ b/Engine/Source/Programs/Shared/EpicGames.Horde/Utilities/TimeOfDayJsonConverter.cs
@@ -0,0 +1,93 @@
+// Copyright Epic Games, Inc. All Rights Reserved.
+
+using System;
+using System.Diagnostics.CodeAnalysis;
+using System.Text.Json;
+using System.Text.Json.Serialization;
+using System.Text.RegularExpressions;
+
+namespace EpicGames.Horde.Utilities
+{
+	/// <summary>
+	/// Parses a time of day into a number of minutes since midnight
+	/// </summary>
+	public class TimeOfDayJsonConverter : JsonConverter<TimeSpan>
+	{
+		/// <inheritdoc/>
+		public override TimeSpan Read(ref Utf8JsonReader reader, Type typeToConvert, JsonSerializerOptions options)
+		{
+			return Parse(reader.GetString());
+		}
+
+		/// <inheritdoc/>
+		public override void Write(Utf8JsonWriter writer, TimeSpan value, JsonSerializerOptions options)
+		{
+			writer.WriteStringValue($"{value.Hours}:{value.Minutes}");
+		}
+
+		/// <summary>
+		/// Parse a string as a number of minutes since midnight
+		/// </summary>
+		public static TimeSpan Parse(string? text)
+		{
+			if (String.IsNullOrEmpty(text))
+			{
+				ThrowInvalidTimeException(text);
+			}
+
+			Match match = Regex.Match(text, @"(?<hours>[0-9]+)(?:[.:](?<minutes>[0-9]+))?\s*(?<morning>am|pm)?$", RegexOptions.IgnoreCase);
+			if (!match.Success)
+			{
+				ThrowInvalidTimeException(text);
+			}
+
+			int hours = Int32.Parse(match.Groups["hours"].Value);
+
+			// Parse the number of minutes
+			int minutes = 0;
+
+			Group minutesGroup = match.Groups["minutes"];
+			if (minutesGroup.Success)
+			{
+				minutes = Int32.Parse(minutesGroup.Value);
+				if (minutes >= 60)
+				{
+					ThrowInvalidTimeException(text);
+				}
+			}
+
+			// Handle 12h/24h suffix
+			Group morningGroup = match.Groups["morning"];
+			if (morningGroup.Success)
+			{
+				if (hours == 12)
+				{
+					hours = 0;
+				}
+				else if (hours == 0 || hours > 12)
+				{
+					ThrowInvalidTimeException(text);
+				}
+				if (String.Equals(morningGroup.Value, "pm", StringComparison.OrdinalIgnoreCase))
+				{
+					hours += 12;
+				}
+			}
+			else
+			{
+				if (hours >= 24)
+				{
+					ThrowInvalidTimeException(text);
+				}
+			}
+
+			return new TimeSpan(hours, minutes, 0);
+		}
+
+		[DoesNotReturn]
+		static void ThrowInvalidTimeException(string? text)
+		{
+			throw new FormatException($"Unable to parse schedule time value '{text ?? "(null)"}'");
+		}
+	}
+}
