Horde: Move schedule interfaces and messages into EpicGames.Horde.

#rnx

diff --git a/Engine/Source/Programs/Horde/Plugins/Analytics/HordeServer.Analytics/Telemetry/Metrics/MetricConfig.cs b/Engine/Source/Programs/Horde/Plugins/Analytics/HordeServer.Analytics/Telemetry/Metrics/MetricConfig.cs
--- a/Engine/Source/Programs/Horde/Plugins/Analytics/HordeServer.Analytics/Telemetry/Metrics/MetricConfig.cs
+++ b/Engine/Source/Programs/Horde/Plugins/Analytics/HordeServer.Analytics/Telemetry/Metrics/MetricConfig.cs
@@ -5,6 +5,7 @@ using System.Text.Json.Serialization;
 using EpicGames.Core;
 using EpicGames.Horde.Telemetry;
 using EpicGames.Horde.Telemetry.Metrics;
+using EpicGames.Horde.Utilities;
 using HordeServer.Acls;
 using HordeServer.Configuration;
 using HordeServer.Utilities;
diff --git a/Engine/Source/Programs/Horde/Plugins/Build/HordeServer.Build.Tests/Jobs/Schedules/SchedulerTests.cs b/Engine/Source/Programs/Horde/Plugins/Build/HordeServer.Build.Tests/Jobs/Schedules/SchedulerTests.cs
--- a/Engine/Source/Programs/Horde/Plugins/Build/HordeServer.Build.Tests/Jobs/Schedules/SchedulerTests.cs
+++ b/Engine/Source/Programs/Horde/Plugins/Build/HordeServer.Build.Tests/Jobs/Schedules/SchedulerTests.cs
@@ -3,6 +3,7 @@
 using EpicGames.Horde.Commits;
 using EpicGames.Horde.Jobs;
 using EpicGames.Horde.Jobs.Graphs;
+using EpicGames.Horde.Jobs.Schedules;
 using EpicGames.Horde.Jobs.Templates;
 using EpicGames.Horde.Projects;
 using EpicGames.Horde.Streams;
diff --git a/Engine/Source/Programs/Horde/Plugins/Build/HordeServer.Build/Jobs/Schedules/ISchedule.cs b/Engine/Source/Programs/Horde/Plugins/Build/HordeServer.Build/Jobs/Schedules/ISchedule.cs
deleted file mode 100644
diff --git a/Engine/Source/Programs/Horde/Plugins/Build/HordeServer.Build/Jobs/Schedules/ScheduleConfig.cs b/Engine/Source/Programs/Horde/Plugins/Build/HordeServer.Build/Jobs/Schedules/ScheduleConfig.cs
--- a/Engine/Source/Programs/Horde/Plugins/Build/HordeServer.Build/Jobs/Schedules/ScheduleConfig.cs
+++ b/Engine/Source/Programs/Horde/Plugins/Build/HordeServer.Build/Jobs/Schedules/ScheduleConfig.cs
@@ -1,16 +1,10 @@
 ﻿// Copyright Epic Games, Inc. All Rights Reserved.
 
 using System.ComponentModel.DataAnnotations;
-using System.Diagnostics.CodeAnalysis;
-using System.Text.Json;
-using System.Text.Json.Serialization;
-using EpicGames.Core;
 using EpicGames.Horde.Commits;
+using EpicGames.Horde.Jobs.Schedules;
 using EpicGames.Horde.Jobs.Templates;
-using EpicGames.Horde.Utilities;
-using HordeServer.Acls;
 using HordeServer.Streams;
-using HordeServer.Utilities;
 
 namespace HordeServer.Jobs.Schedules
 {
@@ -52,7 +46,7 @@ namespace HordeServer.Jobs.Schedules
 		/// <summary>
 		/// Roles to impersonate for this schedule
 		/// </summary>
-		public List<AclClaimConfig>? Claims { get; set; }
+		public List<ScheduleClaimConfig>? Claims { get; set; }
 
 		/// <summary>
 		/// The types of changes to run for
@@ -99,6 +93,20 @@ namespace HordeServer.Jobs.Schedules
 		public List<SchedulePatternConfig> Patterns { get; set; } = new List<SchedulePatternConfig>();
 	}
 
+	/// <summary>
+	/// Claim to grant to leases running a scheduled job
+	/// </summary>
+	public class ScheduleClaimConfig : IScheduleClaim
+	{
+		/// <inheritdoc/>
+		[Required]
+		public string Type { get; set; } = String.Empty;
+
+		/// <inheritdoc/>
+		[Required]
+		public string Value { get; set; } = String.Empty;
+	}
+
 	/// <summary>
 	/// Gate allowing a schedule to trigger.
 	/// </summary>
@@ -166,78 +174,4 @@ namespace HordeServer.Jobs.Schedules
 			Interval = (interval != null) ? new ScheduleInterval(interval.Value) : null;
 		}
 	}
-
-	/// <summary>
-	/// Time of day value for a schedule
-	/// </summary>
-	[JsonSchemaString]
-	[JsonConverter(typeof(ScheduleTimeOfDayJsonConverter))]
-	public record class ScheduleTimeOfDay(int Minutes)
-	{
-		/// <summary>
-		/// Parse a string as a time of day
-		/// </summary>
-		[return: NotNullIfNotNull("text")]
-		public static ScheduleTimeOfDay? Parse(string? text)
-			=> (text != null) ? new ScheduleTimeOfDay((int)TimeOfDayJsonConverter.Parse(text).TotalMinutes) : null;
-	}
-
-	class ScheduleTimeOfDayJsonConverter : JsonConverter<ScheduleTimeOfDay>
-	{
-		/// <inheritdoc/>
-		public override ScheduleTimeOfDay? Read(ref Utf8JsonReader reader, Type typeToConvert, JsonSerializerOptions options)
-		{
-			if (reader.TokenType == JsonTokenType.Number)
-			{
-				return new ScheduleTimeOfDay(reader.GetInt32());
-			}
-			else
-			{
-				return ScheduleTimeOfDay.Parse(reader.GetString());
-			}
-		}
-
-		/// <inheritdoc/>
-		public override void Write(Utf8JsonWriter writer, ScheduleTimeOfDay value, JsonSerializerOptions options)
-		{
-			writer.WriteNumberValue(value.Minutes);
-		}
-	}
-
-	/// <summary>
-	/// Time of day value for a schedule
-	/// </summary>
-	[JsonSchemaString]
-	[JsonConverter(typeof(ScheduleIntervalJsonConverter))]
-	public record class ScheduleInterval(int Minutes)
-	{
-		/// <summary>
-		/// Parse a string as a time of day
-		/// </summary>
-		[return: NotNullIfNotNull("text")]
-		public static ScheduleInterval? Parse(string? text)
-			=> (text != null) ? new ScheduleInterval((int)IntervalJsonConverter.Parse(text).TotalMinutes) : null;
-	}
-
-	class ScheduleIntervalJsonConverter : JsonConverter<ScheduleInterval>
-	{
-		/// <inheritdoc/>
-		public override ScheduleInterval? Read(ref Utf8JsonReader reader, Type typeToConvert, JsonSerializerOptions options)
-		{
-			if (reader.TokenType == JsonTokenType.Number)
-			{
-				return new ScheduleInterval(reader.GetInt32());
-			}
-			else
-			{
-				return ScheduleInterval.Parse(reader.GetString());
-			}
-		}
-
-		/// <inheritdoc/>
-		public override void Write(Utf8JsonWriter writer, ScheduleInterval value, JsonSerializerOptions options)
-		{
-			writer.WriteNumberValue(value.Minutes);
-		}
-	}
 }
diff --git a/Engine/Source/Programs/Horde/Plugins/Build/HordeServer.Build/Jobs/Schedules/ScheduleMessages.cs b/Engine/Source/Programs/Horde/Plugins/Build/HordeServer.Build/Jobs/Schedules/ScheduleMessages.cs
deleted file mode 100644
diff --git a/Engine/Source/Programs/Horde/Plugins/Build/HordeServer.Build/Jobs/Schedules/ScheduleService.cs b/Engine/Source/Programs/Horde/Plugins/Build/HordeServer.Build/Jobs/Schedules/ScheduleService.cs
--- a/Engine/Source/Programs/Horde/Plugins/Build/HordeServer.Build/Jobs/Schedules/ScheduleService.cs
+++ b/Engine/Source/Programs/Horde/Plugins/Build/HordeServer.Build/Jobs/Schedules/ScheduleService.cs
@@ -6,6 +6,7 @@ using EpicGames.Core;
 using EpicGames.Horde.Commits;
 using EpicGames.Horde.Jobs;
 using EpicGames.Horde.Jobs.Graphs;
+using EpicGames.Horde.Jobs.Schedules;
 using EpicGames.Horde.Jobs.Templates;
 using EpicGames.Horde.Streams;
 using EpicGames.Redis;
diff --git a/Engine/Source/Programs/Horde/Plugins/Build/HordeServer.Build/Streams/IStream.cs b/Engine/Source/Programs/Horde/Plugins/Build/HordeServer.Build/Streams/IStream.cs
--- a/Engine/Source/Programs/Horde/Plugins/Build/HordeServer.Build/Streams/IStream.cs
+++ b/Engine/Source/Programs/Horde/Plugins/Build/HordeServer.Build/Streams/IStream.cs
@@ -4,12 +4,11 @@ using EpicGames.Horde.Agents.Pools;
 using EpicGames.Horde.Commits;
 using EpicGames.Horde.Issues;
 using EpicGames.Horde.Jobs;
+using EpicGames.Horde.Jobs.Schedules;
 using EpicGames.Horde.Jobs.Templates;
 using EpicGames.Horde.Projects;
 using EpicGames.Horde.Streams;
 using EpicGames.Horde.Users;
-using Horde.Common.Rpc;
-using HordeServer.Jobs.Schedules;
 
 namespace HordeServer.Streams
 {
diff --git a/Engine/Source/Programs/Horde/Plugins/Build/HordeServer.Build/Streams/StreamCollection.cs b/Engine/Source/Programs/Horde/Plugins/Build/HordeServer.Build/Streams/StreamCollection.cs
--- a/Engine/Source/Programs/Horde/Plugins/Build/HordeServer.Build/Streams/StreamCollection.cs
+++ b/Engine/Source/Programs/Horde/Plugins/Build/HordeServer.Build/Streams/StreamCollection.cs
@@ -4,13 +4,12 @@ using EpicGames.Core;
 using EpicGames.Horde.Commits;
 using EpicGames.Horde.Issues;
 using EpicGames.Horde.Jobs;
+using EpicGames.Horde.Jobs.Schedules;
 using EpicGames.Horde.Jobs.Templates;
 using EpicGames.Horde.Projects;
 using EpicGames.Horde.Streams;
 using EpicGames.Horde.Users;
-using HordeServer.Acls;
 using HordeServer.Commits;
-using HordeServer.Jobs;
 using HordeServer.Jobs.Schedules;
 using HordeServer.Server;
 using HordeServer.Utilities;
@@ -43,7 +42,7 @@ namespace HordeServer.Streams
 			ProjectId IStream.ProjectId => _config.ProjectConfig.Id;
 			string IStream.Name => _config.Name;
 			string? IStream.ConfigPath => _config.Path;
-			string IStream.ConfigRevision =>_config.Revision;
+			string IStream.ConfigRevision => _config.Revision;
 			int IStream.Order => _config.Order;
 			string? IStream.NotificationChannel => _config.NotificationChannel;
 			string? IStream.NotificationChannelFilter => _config.NotificationChannelFilter;
@@ -135,7 +134,7 @@ namespace HordeServer.Streams
 			readonly ScheduleConfig _config;
 			readonly ScheduleDoc _document;
 
-			public IReadOnlyList<AclClaimConfig>? Claims => _config.Claims;
+			public IReadOnlyList<IScheduleClaim>? Claims => _config.Claims;
 			public bool Enabled => _config.Enabled;
 			public int MaxActive => _config.MaxActive;
 			public int MaxChanges => _config.MaxChanges;
diff --git a/Engine/Source/Programs/Horde/Plugins/Build/HordeServer.Build/Streams/StreamMessages.cs b/Engine/Source/Programs/Horde/Plugins/Build/HordeServer.Build/Streams/StreamMessages.cs
--- a/Engine/Source/Programs/Horde/Plugins/Build/HordeServer.Build/Streams/StreamMessages.cs
+++ b/Engine/Source/Programs/Horde/Plugins/Build/HordeServer.Build/Streams/StreamMessages.cs
@@ -5,11 +5,11 @@ using EpicGames.Core;
 using EpicGames.Horde.Agents.Pools;
 using EpicGames.Horde.Issues;
 using EpicGames.Horde.Jobs;
+using EpicGames.Horde.Jobs.Schedules;
 using EpicGames.Horde.Jobs.Templates;
 using EpicGames.Horde.Projects;
 using EpicGames.Horde.Streams;
 using EpicGames.Horde.Users;
-using HordeServer.Jobs.Schedules;
 
 #pragma warning disable CA2227 // Change x to be read-only by removing the property setter
 
diff --git a/Engine/Source/Programs/Shared/EpicGames.Horde.Tests/Utilities/IntervalJsonConverterTests.cs b/Engine/Source/Programs/Shared/EpicGames.Horde.Tests/Utilities/IntervalJsonConverterTests.cs
--- a/Engine/Source/Programs/Shared/EpicGames.Horde.Tests/Utilities/IntervalJsonConverterTests.cs
+++ b/Engine/Source/Programs/Shared/EpicGames.Horde.Tests/Utilities/IntervalJsonConverterTests.cs
@@ -4,7 +4,6 @@ using System;
 using System.Text.Json;
 using System.Text.Json.Serialization;
 using EpicGames.Horde.Utilities;
-using HordeServer.Utilities;
 using Microsoft.VisualStudio.TestTools.UnitTesting;
 
 namespace EpicGames.Horde.Tests.Utilities
diff --git a/Engine/Source/Programs/Shared/EpicGames.Horde.Tests/Utilities/TimeOfDayJsonConverterTests.cs b/Engine/Source/Programs/Shared/EpicGames.Horde.Tests/Utilities/TimeOfDayJsonConverterTests.cs
--- a/Engine/Source/Programs/Shared/EpicGames.Horde.Tests/Utilities/TimeOfDayJsonConverterTests.cs
+++ b/Engine/Source/Programs/Shared/EpicGames.Horde.Tests/Utilities/TimeOfDayJsonConverterTests.cs
@@ -2,7 +2,6 @@
 
 using System;
 using EpicGames.Horde.Utilities;
-using HordeServer.Utilities;
 using Microsoft.VisualStudio.TestTools.UnitTesting;
 
 namespace EpicGames.Horde.Tests.Utilities
diff --git a/Engine/Source/Programs/Shared/EpicGames.Horde/Jobs/Schedules/ISchedule.cs b/Engine/Source/Programs/Shared/EpicGames.Horde/Jobs/Schedules/ISchedule.cs
new file mode 100644
--- /dev/null
+++ b/Engine/Source/Programs/Shared/EpicGames.Horde/Jobs/Schedules/ISchedule.cs
@@ -0,0 +1,225 @@
+// Copyright Epic Games, Inc. All Rights Reserved.
+
+using System;
+using System.Collections.Generic;
+using System.Linq;
+using EpicGames.Horde.Commits;
+using EpicGames.Horde.Jobs.Templates;
+
+namespace EpicGames.Horde.Jobs.Schedules
+{
+	/// <summary>
+	/// Schedule for a template
+	/// </summary>
+	public interface ISchedule
+	{
+		/// <summary>
+		/// Whether the schedule should be enabled
+		/// </summary>
+		bool Enabled { get; }
+
+		/// <summary>
+		/// Maximum number of builds that can be active at once
+		/// </summary>
+		int MaxActive { get; }
+
+		/// <summary>
+		/// Maximum number of changes the schedule can fall behind head revision. If greater than zero, builds will be triggered for every submitted changelist until the backlog is this size.
+		/// </summary>
+		int MaxChanges { get; }
+
+		/// <summary>
+		/// Whether the build requires a change to be submitted
+		/// </summary>
+		bool RequireSubmittedChange { get; }
+
+		/// <summary>
+		/// Gate allowing the schedule to trigger
+		/// </summary>
+		IScheduleGate? Gate { get; }
+
+		/// <summary>
+		/// Commit tags for this schedule
+		/// </summary>
+		IReadOnlyList<CommitTag> Commits { get; }
+
+		/// <summary>
+		/// Roles to impersonate for this schedule
+		/// </summary>
+		IReadOnlyList<IScheduleClaim>? Claims { get; }
+
+		/// <summary>
+		/// Last changelist number that this was triggered for
+		/// </summary>
+		CommitIdWithOrder? LastTriggerCommitId { get; }
+
+		/// <summary>
+		/// Gets the last trigger time, in UTC
+		/// </summary>
+		DateTime LastTriggerTimeUtc { get; }
+
+		/// <summary>
+		/// List of jobs that are currently active
+		/// </summary>
+		IReadOnlyList<JobId> ActiveJobs { get; }
+
+		/// <summary>
+		/// Patterns for starting this scheduled job
+		/// </summary>
+		IReadOnlyList<ISchedulePattern> Patterns { get; }
+
+		/// <summary>
+		/// Files that should cause the job to trigger
+		/// </summary>
+		IReadOnlyList<string>? Files { get; }
+
+		/// <summary>
+		/// Parameters for the template
+		/// </summary>
+		IReadOnlyDictionary<string, string> TemplateParameters { get; }
+	}
+
+	/// <summary>
+	/// Claim granted to a schedule
+	/// </summary>
+	public interface IScheduleClaim
+	{
+		/// <summary>
+		/// The claim type
+		/// </summary>
+		string Type { get; }
+
+		/// <summary>
+		/// The claim value
+		/// </summary>
+		string Value { get; }
+	}
+
+	/// <summary>
+	/// Required gate for starting a schedule
+	/// </summary>
+	public interface IScheduleGate
+	{
+		/// <summary>
+		/// The template containing the dependency
+		/// </summary>
+		TemplateId TemplateId { get; }
+
+		/// <summary>
+		/// Target to wait for
+		/// </summary>
+		string Target { get; }
+	}
+
+	/// <summary>
+	/// Pattern for executing a schedule
+	/// </summary>
+	public interface ISchedulePattern
+	{
+		/// <summary>
+		/// Days of the week to run this schedule on. If null, the schedule will run every day.
+		/// </summary>
+		IReadOnlyList<DayOfWeek>? DaysOfWeek { get; }
+
+		/// <summary>
+		/// Time during the day for the first schedule to trigger. Measured in minutes from midnight.
+		/// </summary>
+		ScheduleTimeOfDay MinTime { get; }
+
+		/// <summary>
+		/// Time during the day for the last schedule to trigger. Measured in minutes from midnight.
+		/// </summary>
+		ScheduleTimeOfDay? MaxTime { get; }
+
+		/// <summary>
+		/// Interval between each schedule triggering
+		/// </summary>
+		ScheduleInterval? Interval { get; }
+	}
+
+	/// <summary>
+	/// Extension methods for schedules
+	/// </summary>
+	public static class ScheduleExtensions
+	{
+		/// <summary>
+		/// Gets the next trigger time for a schedule
+		/// </summary>
+		/// <param name="schedule"></param>
+		/// <param name="timeZone"></param>
+		/// <returns></returns>
+		public static DateTime? GetNextTriggerTimeUtc(this ISchedule schedule, TimeZoneInfo timeZone)
+		{
+			return schedule.GetNextTriggerTimeUtc(schedule.LastTriggerTimeUtc, timeZone);
+		}
+
+		/// <summary>
+		/// Get the next time that the schedule will trigger
+		/// </summary>
+		/// <param name="schedule">Schedule to query</param>
+		/// <param name="lastTimeUtc">Last time at which the schedule triggered</param>
+		/// <param name="timeZone">Timezone to evaluate the trigger</param>
+		/// <returns>Next time at which the schedule will trigger</returns>
+		public static DateTime? GetNextTriggerTimeUtc(this ISchedule schedule, DateTime lastTimeUtc, TimeZoneInfo timeZone)
+		{
+			DateTime? nextTriggerTimeUtc = null;
+			foreach (ISchedulePattern pattern in schedule.Patterns)
+			{
+				DateTime patternTriggerTime = pattern.GetNextTriggerTimeUtc(lastTimeUtc, timeZone);
+				if (nextTriggerTimeUtc == null || patternTriggerTime < nextTriggerTimeUtc)
+				{
+					nextTriggerTimeUtc = patternTriggerTime;
+				}
+			}
+			return nextTriggerTimeUtc;
+		}
+
+		/// <summary>
+		/// Calculates the trigger index based on the given time in minutes
+		/// </summary>
+		/// <param name="pattern">Pattern to query</param>
+		/// <param name="lastTimeUtc">Time for the last trigger</param>
+		/// <param name="timeZone">The timezone for running the schedule</param>
+		/// <returns>Index of the trigger</returns>
+		public static DateTime GetNextTriggerTimeUtc(this ISchedulePattern pattern, DateTime lastTimeUtc, TimeZoneInfo timeZone)
+		{
+			// Convert last time into the correct timezone for running the scheule
+			DateTimeOffset lastTime = TimeZoneInfo.ConvertTime((DateTimeOffset)lastTimeUtc, timeZone);
+
+			// Get the base time (ie. the start of this day) for anchoring the schedule
+			DateTimeOffset baseTime = new DateTimeOffset(lastTime.Year, lastTime.Month, lastTime.Day, 0, 0, 0, lastTime.Offset);
+			for (; ; )
+			{
+				if (pattern.DaysOfWeek == null || pattern.DaysOfWeek.Contains(baseTime.DayOfWeek))
+				{
+					// Get the last time in minutes from the start of this day
+					int lastTimeMinutes = (int)(lastTime - baseTime).TotalMinutes;
+
+					// Get the time of the first trigger of this day. If the last time is less than this, this is the next trigger.
+					if (lastTimeMinutes < pattern.MinTime.Minutes)
+					{
+						return baseTime.AddMinutes(pattern.MinTime.Minutes).UtcDateTime;
+					}
+
+					// Otherwise, get the time for the last trigger in the day.
+					if (pattern.Interval != null && pattern.Interval.Minutes > 0)
+					{
+						int actualMaxTime = pattern.MaxTime?.Minutes ?? ((24 * 60) - 1);
+						if (lastTimeMinutes < actualMaxTime)
+						{
+							int lastIndex = (lastTimeMinutes - pattern.MinTime.Minutes) / pattern.Interval.Minutes;
+							int nextIndex = lastIndex + 1;
+
+							int nextTimeMinutes = pattern.MinTime.Minutes + (nextIndex * pattern.Interval.Minutes);
+							if (nextTimeMinutes <= actualMaxTime)
+							{
+								return baseTime.AddMinutes(nextTimeMinutes).UtcDateTime;
+							}
+						}
+					}
+				}
+				baseTime = baseTime.AddDays(1.0);
+			}
+		}
+	}
+}
diff --git a/Engine/Source/Programs/Shared/EpicGames.Horde/Jobs/Schedules/ScheduleInterval.cs b/Engine/Source/Programs/Shared/EpicGames.Horde/Jobs/Schedules/ScheduleInterval.cs
new file mode 100644
--- /dev/null
+++ b/Engine/Source/Programs/Shared/EpicGames.Horde/Jobs/Schedules/ScheduleInterval.cs
@@ -0,0 +1,48 @@
+// Copyright Epic Games, Inc. All Rights Reserved.
+
+using System;
+using System.Diagnostics.CodeAnalysis;
+using System.Text.Json;
+using System.Text.Json.Serialization;
+using EpicGames.Core;
+using EpicGames.Horde.Utilities;
+
+namespace EpicGames.Horde.Jobs.Schedules
+{
+	/// <summary>
+	/// Time of day value for a schedule
+	/// </summary>
+	[JsonSchemaString]
+	[JsonConverter(typeof(ScheduleIntervalJsonConverter))]
+	public record class ScheduleInterval(int Minutes)
+	{
+		/// <summary>
+		/// Parse a string as a time of day
+		/// </summary>
+		[return: NotNullIfNotNull("text")]
+		public static ScheduleInterval? Parse(string? text)
+			=> (text != null) ? new ScheduleInterval((int)IntervalJsonConverter.Parse(text).TotalMinutes) : null;
+	}
+
+	class ScheduleIntervalJsonConverter : JsonConverter<ScheduleInterval>
+	{
+		/// <inheritdoc/>
+		public override ScheduleInterval? Read(ref Utf8JsonReader reader, Type typeToConvert, JsonSerializerOptions options)
+		{
+			if (reader.TokenType == JsonTokenType.Number)
+			{
+				return new ScheduleInterval(reader.GetInt32());
+			}
+			else
+			{
+				return ScheduleInterval.Parse(reader.GetString());
+			}
+		}
+
+		/// <inheritdoc/>
+		public override void Write(Utf8JsonWriter writer, ScheduleInterval value, JsonSerializerOptions options)
+		{
+			writer.WriteNumberValue(value.Minutes);
+		}
+	}
+}
diff --git a/Engine/Source/Programs/Shared/EpicGames.Horde/Jobs/Schedules/ScheduleMessages.cs b/Engine/Source/Programs/Shared/EpicGames.Horde/Jobs/Schedules/ScheduleMessages.cs
new file mode 100644
--- /dev/null
+++ b/Engine/Source/Programs/Shared/EpicGames.Horde/Jobs/Schedules/ScheduleMessages.cs
@@ -0,0 +1,237 @@
+// Copyright Epic Games, Inc. All Rights Reserved.
+
+using System;
+using System.Collections.Generic;
+using System.Linq;
+using EpicGames.Core;
+using EpicGames.Horde.Commits;
+using EpicGames.Horde.Jobs.Templates;
+
+#pragma warning disable CA2227 // Change x to be read-only by removing the property setter
+
+namespace EpicGames.Horde.Jobs.Schedules
+{
+	/// <summary>
+	/// Response describing a schedule
+	/// </summary>
+	public class GetScheduleResponse
+	{
+		/// <summary>
+		/// Whether the schedule is currently enabled
+		/// </summary>
+		public bool Enabled { get; set; }
+
+		/// <summary>
+		/// Maximum number of scheduled jobs at once
+		/// </summary>
+		public int MaxActive { get; set; }
+
+		/// <summary>
+		/// Maximum number of changes the schedule can fall behind head revision. If greater than zero, builds will be triggered for every submitted changelist until the backlog is this size.
+		/// </summary>
+		public int MaxChanges { get; set; }
+
+		/// <summary>
+		/// Whether the build requires a change to be submitted
+		/// </summary>
+		public bool RequireSubmittedChange { get; set; }
+
+		/// <summary>
+		/// Gate for this schedule to trigger
+		/// </summary>
+		public GetScheduleGateResponse? Gate { get; set; }
+
+		/// <summary>
+		/// Which commits to run this job for
+		/// </summary>
+		public List<CommitTag>? Commits { get; set; }
+
+		/// <summary>
+		/// Parameters for the template
+		/// </summary>
+		public Dictionary<string, string> TemplateParameters { get; set; }
+
+		/// <summary>
+		/// New patterns for the schedule
+		/// </summary>
+		public List<GetSchedulePatternResponse> Patterns { get; set; }
+
+		/// <summary>
+		/// Last changelist number that this was triggered for
+		/// </summary>
+		[Obsolete("Use LastTriggerCommitId instead")]
+		public int LastTriggerChange
+		{
+			get => _lastTriggerChange ?? _lastTriggerCommitId?.TryGetPerforceChange() ?? 0;
+			set => _lastTriggerChange = value;
+		}
+		int? _lastTriggerChange;
+
+		/// <summary>
+		/// Last changelist number that this was triggered for
+		/// </summary>
+		public CommitIdWithOrder? LastTriggerCommitId
+		{
+			get => _lastTriggerCommitId ?? CommitIdWithOrder.FromPerforceChange(_lastTriggerChange);
+			set => _lastTriggerCommitId = value;
+		}
+		CommitIdWithOrder? _lastTriggerCommitId;
+
+		/// <summary>
+		/// Last time that the schedule was triggered
+		/// </summary>
+		public DateTimeOffset LastTriggerTime { get; set; }
+
+		/// <summary>
+		/// Next trigger times for schedule
+		/// </summary>
+		public List<DateTime> NextTriggerTimesUTC { get; set; }
+
+		/// <summary>
+		/// List of active jobs
+		/// </summary>
+		public List<JobId> ActiveJobs { get; set; }
+
+		/// <summary>
+		/// Default constructor
+		/// </summary>
+		public GetScheduleResponse()
+		{
+			Patterns = new List<GetSchedulePatternResponse>();
+			TemplateParameters = new Dictionary<string, string>();
+			NextTriggerTimesUTC = new List<DateTime>();
+			ActiveJobs = new List<JobId>();
+		}
+
+		/// <summary>
+		/// Constructor
+		/// </summary>
+		/// <param name="schedule">Schedule to construct from</param>
+		/// <param name="schedulerTimeZone">The scheduler time zone</param>
+		public GetScheduleResponse(ISchedule schedule, TimeZoneInfo schedulerTimeZone)
+		{
+			Enabled = schedule.Enabled;
+			MaxActive = schedule.MaxActive;
+			MaxChanges = schedule.MaxChanges;
+			RequireSubmittedChange = schedule.RequireSubmittedChange;
+			Gate = (schedule.Gate == null) ? null : new GetScheduleGateResponse(schedule.Gate);
+			Commits = schedule.Commits?.ToList();
+			TemplateParameters = schedule.TemplateParameters.ToDictionary();
+			Patterns = schedule.Patterns.ConvertAll(x => new GetSchedulePatternResponse(x));
+			LastTriggerCommitId = schedule.LastTriggerCommitId;
+			LastTriggerTime = schedule.LastTriggerTimeUtc;
+			ActiveJobs = new List<JobId>(schedule.ActiveJobs);
+
+			DateTime curTime = schedule.LastTriggerTimeUtc;
+			NextTriggerTimesUTC = new List<DateTime>();
+			for (int i = 0; i < 16; i++)
+			{
+				DateTime? nextTime = schedule.GetNextTriggerTimeUtc(curTime, schedulerTimeZone);
+				if (nextTime == null)
+				{
+					break;
+				}
+
+				curTime = nextTime.Value;
+				NextTriggerTimesUTC.Add(curTime);
+			}
+		}
+	}
+
+	/// <summary>
+	/// Gate allowing a schedule to trigger.
+	/// </summary>
+	public class GetScheduleGateResponse
+	{
+		/// <summary>
+		/// The template containing the dependency
+		/// </summary>
+		public TemplateId TemplateId { get; set; }
+
+		/// <summary>
+		/// Target to wait for
+		/// </summary>
+		public string Target { get; set; }
+
+		/// <summary>
+		/// Default constructor
+		/// </summary>
+		public GetScheduleGateResponse()
+		{
+			Target = String.Empty;
+		}
+
+		/// <summary>
+		/// Constructor
+		/// </summary>
+		public GetScheduleGateResponse(IScheduleGate gate)
+		{
+			TemplateId = gate.TemplateId;
+			Target = gate.Target;
+		}
+	}
+
+	/// <summary>
+	/// Parameters to create a new schedule
+	/// </summary>
+	public class GetSchedulePatternResponse
+	{
+		/// <summary>
+		/// Days of the week to run this schedule on. If null, the schedule will run every day.
+		/// </summary>
+		public List<DayOfWeek>? DaysOfWeek { get; set; }
+
+		/// <summary>
+		/// Time during the day for the first schedule to trigger. Measured in minutes from midnight.
+		/// </summary>
+		public ScheduleTimeOfDay MinTime { get; set; } = new ScheduleTimeOfDay(0);
+
+		/// <summary>
+		/// Time during the day for the last schedule to trigger. Measured in minutes from midnight.
+		/// </summary>
+		public ScheduleTimeOfDay? MaxTime { get; set; }
+
+		/// <summary>
+		/// Interval between each schedule triggering
+		/// </summary>
+		public ScheduleInterval? Interval { get; set; }
+
+		/// <summary>
+		/// Constructor
+		/// </summary>
+		public GetSchedulePatternResponse()
+		{
+		}
+
+		/// <summary>
+		/// Constructor
+		/// </summary>
+		public GetSchedulePatternResponse(ISchedulePattern pattern)
+		{
+			DaysOfWeek = pattern.DaysOfWeek?.ToList();
+			MinTime = pattern.MinTime;
+			MaxTime = pattern.MaxTime;
+			Interval = pattern.Interval;
+		}
+	}
+
+	/// <summary>
+	/// Response describing when a schedule is expected to trigger
+	/// </summary>
+	public class GetScheduleForecastResponse
+	{
+		/// <summary>
+		/// Next trigger times
+		/// </summary>
+		public List<DateTime> Times { get; set; }
+
+		/// <summary>
+		/// Constructor
+		/// </summary>
+		/// <param name="times">List of trigger times</param>
+		public GetScheduleForecastResponse(List<DateTime> times)
+		{
+			Times = times;
+		}
+	}
+}
diff --git a/Engine/Source/Programs/Shared/EpicGames.Horde/Jobs/Schedules/ScheduleTimeOfDay.cs b/Engine/Source/Programs/Shared/EpicGames.Horde/Jobs/Schedules/ScheduleTimeOfDay.cs
new file mode 100644
--- /dev/null
+++ b/Engine/Source/Programs/Shared/EpicGames.Horde/Jobs/Schedules/ScheduleTimeOfDay.cs
@@ -0,0 +1,48 @@
+// Copyright Epic Games, Inc. All Rights Reserved.
+
+using System;
+using System.Diagnostics.CodeAnalysis;
+using System.Text.Json;
+using System.Text.Json.Serialization;
+using EpicGames.Core;
+using EpicGames.Horde.Utilities;
+
+namespace EpicGames.Horde.Jobs.Schedules
+{
+	/// <summary>
+	/// Time of day value for a schedule
+	/// </summary>
+	[JsonSchemaString]
+	[JsonConverter(typeof(ScheduleTimeOfDayJsonConverter))]
+	public record class ScheduleTimeOfDay(int Minutes)
+	{
+		/// <summary>
+		/// Parse a string as a time of day
+		/// </summary>
+		[return: NotNullIfNotNull("text")]
+		public static ScheduleTimeOfDay? Parse(string? text)
+			=> (text != null) ? new ScheduleTimeOfDay((int)TimeOfDayJsonConverter.Parse(text).TotalMinutes) : null;
+	}
+
+	class ScheduleTimeOfDayJsonConverter : JsonConverter<ScheduleTimeOfDay>
+	{
+		/// <inheritdoc/>
+		public override ScheduleTimeOfDay? Read(ref Utf8JsonReader reader, Type typeToConvert, JsonSerializerOptions options)
+		{
+			if (reader.TokenType == JsonTokenType.Number)
+			{
+				return new ScheduleTimeOfDay(reader.GetInt32());
+			}
+			else
+			{
+				return ScheduleTimeOfDay.Parse(reader.GetString());
+			}
+		}
+
+		/// <inheritdoc/>
+		public override void Write(Utf8JsonWriter writer, ScheduleTimeOfDay value, JsonSerializerOptions options)
+		{
+			writer.WriteNumberValue(value.Minutes);
+		}
+	}
+}
