Horde: Move stream messages and interfaces into EpicGames.Horde.

#rnx

diff --git a/Engine/Source/Programs/Horde/Plugins/Build/HordeServer.Build.Tests/BuildTestSetup.cs b/Engine/Source/Programs/Horde/Plugins/Build/HordeServer.Build.Tests/BuildTestSetup.cs
--- a/Engine/Source/Programs/Horde/Plugins/Build/HordeServer.Build.Tests/BuildTestSetup.cs
+++ b/Engine/Source/Programs/Horde/Plugins/Build/HordeServer.Build.Tests/BuildTestSetup.cs
@@ -8,6 +8,7 @@ using EpicGames.Horde.Agents;
 using EpicGames.Horde.Agents.Pools;
 using EpicGames.Horde.Artifacts;
 using EpicGames.Horde.Jobs.Templates;
+using EpicGames.Horde.Streams;
 using EpicGames.Horde.Tools;
 using HordeCommon.Rpc.Messages;
 using HordeServer.Acls;
diff --git a/Engine/Source/Programs/Horde/Plugins/Build/HordeServer.Build.Tests/Jobs/JobTaskSourceTests.cs b/Engine/Source/Programs/Horde/Plugins/Build/HordeServer.Build.Tests/Jobs/JobTaskSourceTests.cs
--- a/Engine/Source/Programs/Horde/Plugins/Build/HordeServer.Build.Tests/Jobs/JobTaskSourceTests.cs
+++ b/Engine/Source/Programs/Horde/Plugins/Build/HordeServer.Build.Tests/Jobs/JobTaskSourceTests.cs
@@ -5,12 +5,12 @@ using EpicGames.Horde.Agents.Pools;
 using EpicGames.Horde.Commits;
 using EpicGames.Horde.Jobs;
 using EpicGames.Horde.Jobs.Graphs;
+using EpicGames.Horde.Streams;
 using EpicGames.Horde.Users;
 using HordeCommon.Rpc.Messages;
 using HordeServer.Agents;
 using HordeServer.Agents.Pools;
 using HordeServer.Jobs;
-using HordeServer.Streams;
 using HordeServer.Utilities;
 
 namespace HordeServer.Tests.Jobs
diff --git a/Engine/Source/Programs/Horde/Plugins/Build/HordeServer.Build.Tests/Streams/StreamServiceTest.cs b/Engine/Source/Programs/Horde/Plugins/Build/HordeServer.Build.Tests/Streams/StreamServiceTest.cs
--- a/Engine/Source/Programs/Horde/Plugins/Build/HordeServer.Build.Tests/Streams/StreamServiceTest.cs
+++ b/Engine/Source/Programs/Horde/Plugins/Build/HordeServer.Build.Tests/Streams/StreamServiceTest.cs
@@ -1,7 +1,7 @@
 // Copyright Epic Games, Inc. All Rights Reserved.
 
 using System.Globalization;
-using HordeServer.Streams;
+using EpicGames.Horde.Streams;
 
 namespace HordeServer.Tests.Streams
 {
diff --git a/Engine/Source/Programs/Horde/Plugins/Build/HordeServer.Build/BuildPlugin.cs b/Engine/Source/Programs/Horde/Plugins/Build/HordeServer.Build/BuildPlugin.cs
--- a/Engine/Source/Programs/Horde/Plugins/Build/HordeServer.Build/BuildPlugin.cs
+++ b/Engine/Source/Programs/Horde/Plugins/Build/HordeServer.Build/BuildPlugin.cs
@@ -4,6 +4,7 @@ using System.Diagnostics.CodeAnalysis;
 using EpicGames.Core;
 using EpicGames.Horde.Artifacts;
 using EpicGames.Horde.Jobs.Templates;
+using EpicGames.Horde.Streams;
 using HordeServer.Acls;
 using HordeServer.Agents;
 using HordeServer.Agents.Fleet;
diff --git a/Engine/Source/Programs/Horde/Plugins/Build/HordeServer.Build/Devices/DeviceService.cs b/Engine/Source/Programs/Horde/Plugins/Build/HordeServer.Build/Devices/DeviceService.cs
--- a/Engine/Source/Programs/Horde/Plugins/Build/HordeServer.Build/Devices/DeviceService.cs
+++ b/Engine/Source/Programs/Horde/Plugins/Build/HordeServer.Build/Devices/DeviceService.cs
@@ -7,6 +7,7 @@ using EpicGames.Horde.Devices;
 using EpicGames.Horde.Jobs;
 using EpicGames.Horde.Jobs.Graphs;
 using EpicGames.Horde.Projects;
+using EpicGames.Horde.Streams;
 using EpicGames.Horde.Users;
 using HordeServer.Jobs;
 using HordeServer.Notifications;
diff --git a/Engine/Source/Programs/Horde/Plugins/Build/HordeServer.Build/Perforce/PerforceWorkspace.cs b/Engine/Source/Programs/Horde/Plugins/Build/HordeServer.Build/Perforce/PerforceWorkspace.cs
--- a/Engine/Source/Programs/Horde/Plugins/Build/HordeServer.Build/Perforce/PerforceWorkspace.cs
+++ b/Engine/Source/Programs/Horde/Plugins/Build/HordeServer.Build/Perforce/PerforceWorkspace.cs
@@ -2,6 +2,7 @@
 
 using System.Diagnostics.CodeAnalysis;
 using EpicGames.Horde.Agents.Pools;
+using EpicGames.Horde.Streams;
 using HordeCommon.Rpc.Messages;
 using HordeServer.Agents;
 using HordeServer.Agents.Pools;
diff --git a/Engine/Source/Programs/Horde/Plugins/Build/HordeServer.Build/Streams/IStream.cs b/Engine/Source/Programs/Horde/Plugins/Build/HordeServer.Build/Streams/IStream.cs
deleted file mode 100644
diff --git a/Engine/Source/Programs/Horde/Plugins/Build/HordeServer.Build/Streams/IStreamCollection.cs b/Engine/Source/Programs/Horde/Plugins/Build/HordeServer.Build/Streams/IStreamCollection.cs
deleted file mode 100644
diff --git a/Engine/Source/Programs/Horde/Plugins/Build/HordeServer.Build/Streams/StreamMessages.cs b/Engine/Source/Programs/Horde/Plugins/Build/HordeServer.Build/Streams/StreamMessages.cs
deleted file mode 100644
diff --git a/Engine/Source/Programs/Shared/EpicGames.Horde/Streams/IStream.cs b/Engine/Source/Programs/Shared/EpicGames.Horde/Streams/IStream.cs
new file mode 100644
--- /dev/null
+++ b/Engine/Source/Programs/Shared/EpicGames.Horde/Streams/IStream.cs
@@ -0,0 +1,579 @@
+// Copyright Epic Games, Inc. All Rights Reserved.
+
+using System;
+using System.Collections.Generic;
+using System.Linq;
+using System.Threading;
+using System.Threading.Tasks;
+using EpicGames.Horde.Agents.Pools;
+using EpicGames.Horde.Commits;
+using EpicGames.Horde.Issues;
+using EpicGames.Horde.Jobs;
+using EpicGames.Horde.Jobs.Schedules;
+using EpicGames.Horde.Jobs.Templates;
+using EpicGames.Horde.Projects;
+using EpicGames.Horde.Users;
+
+namespace EpicGames.Horde.Streams
+{
+	/// <summary>
+	/// Exception thrown when stream validation fails
+	/// </summary>
+	public class InvalidStreamException : Exception
+	{
+		/// <inheritdoc/>
+		public InvalidStreamException()
+		{
+		}
+
+		/// <inheritdoc/>
+		public InvalidStreamException(string message) : base(message)
+		{
+		}
+
+		/// <inheritdoc/>
+		public InvalidStreamException(string message, Exception innerEx) : base(message, innerEx)
+		{
+		}
+	}
+
+	/// <summary>
+	/// Information about a stream
+	/// </summary>
+	public interface IStream
+	{
+		/// <summary>
+		/// Name of the stream.
+		/// </summary>
+		StreamId Id { get; }
+
+		/// <summary>
+		/// Project that this stream belongs to
+		/// </summary>
+		ProjectId ProjectId { get; }
+
+		/// <summary>
+		/// Name of the stream
+		/// </summary>
+		string Name { get; }
+
+		/// <summary>
+		/// Path to the config file for this stream
+		/// </summary>
+		string? ConfigPath { get; }
+
+		/// <summary>
+		/// Current revision of the config file
+		/// </summary>
+		string ConfigRevision { get; }
+
+		/// <summary>
+		/// Order for this stream on the dashboard
+		/// </summary>
+		int Order { get; }
+
+		/// <summary>
+		/// Notification channel for all jobs in this stream
+		/// </summary>
+		string? NotificationChannel { get; }
+
+		/// <summary>
+		/// Notification channel filter for this template. Can be Success, Failure, or Warnings.
+		/// </summary>
+		string? NotificationChannelFilter { get; }
+
+		/// <summary>
+		/// Channel to post issue triage notifications
+		/// </summary>
+		string? TriageChannel { get; }
+
+		/// <summary>
+		/// Tabs for this stream on the dashboard
+		/// </summary>
+		IReadOnlyList<IStreamTab> Tabs { get; }
+
+		/// <summary>
+		/// Agent types configured for this stream
+		/// </summary>
+		IReadOnlyDictionary<string, IAgentType> AgentTypes { get; }
+
+		/// <summary>
+		/// Workspace types configured for this stream
+		/// </summary>
+		IReadOnlyDictionary<string, IWorkspaceType> WorkspaceTypes { get; }
+
+		/// <summary>
+		/// List of templates available for this stream
+		/// </summary>
+		IReadOnlyDictionary<TemplateId, ITemplateRef> Templates { get; }
+
+		/// <summary>
+		/// Workflows configured for this stream
+		/// </summary>
+		IReadOnlyList<IWorkflow> Workflows { get; }
+
+		/// <summary>
+		/// Default settings for preflights against this stream
+		/// </summary>
+		IDefaultPreflight? DefaultPreflight { get; }
+
+		/// <summary>
+		/// Stream is paused for builds until specified time
+		/// </summary>
+		DateTime? PausedUntil { get; }
+
+		/// <summary>
+		/// Comment/reason for why the stream was paused
+		/// </summary>
+		string? PauseComment { get; }
+
+		/// <summary>
+		/// Commits for this stream
+		/// </summary>
+		ICommitCollection Commits { get; }
+
+		/// <summary>
+		/// Get the latest stream state
+		/// </summary>
+		/// <param name="cancellationToken">Cancellation toke for this operation</param>
+		/// <returns>Updated stream, or null if it no longer exists</returns>
+		Task<IStream?> RefreshAsync(CancellationToken cancellationToken = default);
+
+		/// <summary>
+		/// Updates user-facing properties for an existing stream
+		/// </summary>
+		/// <param name="newPausedUntil">The new datetime for pausing builds</param>
+		/// <param name="newPauseComment">The reason for pausing</param>
+		/// <param name="cancellationToken">Cancellation token for the operation</param>
+		/// <returns>The updated stream if successful, null otherwise</returns>
+		Task<IStream?> TryUpdatePauseStateAsync(DateTime? newPausedUntil, string? newPauseComment, CancellationToken cancellationToken = default);
+
+		/// <summary>
+		/// Attempts to update the last trigger time for a schedule
+		/// </summary>
+		/// <param name="templateRefId">The template ref id</param>
+		/// <param name="lastTriggerTimeUtc">New last trigger time for the schedule</param>
+		/// <param name="lastTriggerCommitId">New last trigger commit for the schedule</param>
+		/// <param name="newActiveJobs">New list of active jobs</param>
+		/// <param name="cancellationToken">Cancellation token for the operation</param>
+		/// <returns>The updated stream if successful, null otherwise</returns>
+		Task<IStream?> TryUpdateScheduleTriggerAsync(TemplateId templateRefId, DateTime? lastTriggerTimeUtc, CommitIdWithOrder? lastTriggerCommitId, List<JobId> newActiveJobs, CancellationToken cancellationToken = default);
+
+		/// <summary>
+		/// Attempts to update a stream template ref
+		/// </summary>
+		/// <param name="templateRefId">The template ref to update</param>
+		/// <param name="stepStates">The stream states to update, pass an empty list to clear all step states, otherwise will be a partial update based on included step updates</param>
+		/// <param name="cancellationToken">Cancellation token for the operation</param>
+		/// <returns></returns>
+		Task<IStream?> TryUpdateTemplateRefAsync(TemplateId templateRefId, List<UpdateStepStateRequest>? stepStates = null, CancellationToken cancellationToken = default);
+	}
+
+	/// <summary>
+	/// Style for rendering a tab
+	/// </summary>
+	public enum TabStyle
+	{
+		/// <summary>
+		/// Regular job list
+		/// </summary>
+		Normal,
+
+		/// <summary>
+		/// Omit job names, show condensed view
+		/// </summary>
+		Compact,
+	}
+
+	/// <summary>
+	/// Information about a page to display in the dashboard for a stream
+	/// </summary>
+	public interface IStreamTab
+	{
+		/// <summary>
+		/// Title of this page
+		/// </summary>
+		string Title { get; }
+
+		/// <summary>
+		/// Type of this tab
+		/// </summary>
+		string Type { get; }
+
+		/// <summary>
+		/// Presentation style for this page
+		/// </summary>
+		TabStyle Style { get; }
+
+		/// <summary>
+		/// Whether to show job names on this page
+		/// </summary>
+		bool ShowNames { get; }
+
+		/// <summary>
+		/// Whether to show all user preflights 
+		/// </summary>
+		bool? ShowPreflights { get; }
+
+		/// <summary>
+		/// Names of jobs to include on this page. If there is only one name specified, the name column does not need to be displayed.
+		/// </summary>
+		IReadOnlyList<string>? JobNames { get; }
+
+		/// <summary>
+		/// List of job template names to show on this page.
+		/// </summary>
+		IReadOnlyList<TemplateId>? Templates { get; }
+
+		/// <summary>
+		/// Columns to display for different types of aggregates
+		/// </summary>
+		IReadOnlyList<IStreamTabColumn>? Columns { get; }
+	}
+
+	/// <summary>
+	/// Type of a column in a jobs tab
+	/// </summary>
+	public enum TabColumnType
+	{
+		/// <summary>
+		/// Contains labels
+		/// </summary>
+		Labels,
+
+		/// <summary>
+		/// Contains parameters
+		/// </summary>
+		Parameter
+	}
+
+	/// <summary>
+	/// Describes a column to display on the jobs page
+	/// </summary>
+	public interface IStreamTabColumn
+	{
+		/// <summary>
+		/// The type of column
+		/// </summary>
+		TabColumnType Type { get; }
+
+		/// <summary>
+		/// Heading for this column
+		/// </summary>
+		string Heading { get; }
+
+		/// <summary>
+		/// Category of aggregates to display in this column. If null, includes any aggregate not matched by another column.
+		/// </summary>
+		string? Category { get; }
+
+		/// <summary>
+		/// Parameter to show in this column
+		/// </summary>
+		string? Parameter { get; }
+
+		/// <summary>
+		/// Relative width of this column.
+		/// </summary>
+		int? RelativeWidth { get; }
+	}
+
+	/// <summary>
+	/// Mapping from a BuildGraph agent type to a set of machines on the farm
+	/// </summary>
+	public interface IAgentType
+	{
+		/// <summary>
+		/// Pool of agents to use for this agent type
+		/// </summary>
+		PoolId Pool { get; }
+
+		/// <summary>
+		/// Name of the workspace to sync
+		/// </summary>
+		string? Workspace { get; }
+
+		/// <summary>
+		/// Path to the temporary storage dir
+		/// </summary>
+		string? TempStorageDir { get; }
+
+		/// <summary>
+		/// Environment variables to be set when executing the job
+		/// </summary>
+		IReadOnlyDictionary<string, string>? Environment { get; }
+	}
+
+	/// <summary>
+	/// Information about a workspace type
+	/// </summary>
+	public interface IWorkspaceType
+	{
+		/// <summary>
+		/// Name of the Perforce server cluster to use
+		/// </summary>
+		string? Cluster { get; }
+
+		/// <summary>
+		/// The Perforce server and port (eg. perforce:1666)
+		/// </summary>
+		string? ServerAndPort { get; }
+
+		/// <summary>
+		/// User to log into Perforce with (defaults to buildmachine)
+		/// </summary>
+		string? UserName { get; }
+
+		/// <summary>
+		/// Password to use to log into the workspace
+		/// </summary>
+		string? Password { get; }
+
+		/// <summary>
+		/// Identifier to distinguish this workspace from other workspaces. Defaults to the workspace type name.
+		/// </summary>
+		string? Identifier { get; }
+
+		/// <summary>
+		/// Override for the stream to sync
+		/// </summary>
+		string? Stream { get; }
+
+		/// <summary>
+		/// Custom view for the workspace
+		/// </summary>
+		IReadOnlyList<string>? View { get; }
+
+		/// <summary>
+		/// Whether to use an incrementally synced workspace
+		/// </summary>
+		bool? Incremental { get; }
+
+		/// <summary>
+		/// Whether to use the AutoSDK
+		/// </summary>
+		bool? UseAutoSdk { get; }
+
+		/// <summary>
+		/// View for the AutoSDK paths to sync. If null, the whole thing will be synced.
+		/// </summary>
+		IReadOnlyList<string>? AutoSdkView { get; }
+
+		/// <summary>
+		/// Method to use when syncing/materializing data from Perforce
+		/// </summary>
+		string? Method { get; }
+
+		/// <summary>
+		/// Minimum disk space that must be available *after* syncing this workspace (in megabytes)
+		/// If not available, the job will be aborted.
+		/// </summary>
+		long? MinScratchSpace { get; }
+
+		/// <summary>
+		/// Threshold for when to trigger an automatic conform of agent. Measured in megabytes free on disk.
+		/// Set to null or 0 to disable.
+		/// </summary>
+		long? ConformDiskFreeSpace { get; }
+	}
+
+	/// <summary>
+	/// Specifies defaults for running a preflight
+	/// </summary>
+	public interface IDefaultPreflight
+	{
+		/// <summary>
+		/// The template id to query
+		/// </summary>
+		TemplateId? TemplateId { get; }
+
+		/// <summary>
+		/// Query for the change to use
+		/// </summary>
+		IChangeQuery? Change { get; }
+	}
+
+	/// <summary>
+	/// Job template in a stream
+	/// </summary>
+	public interface ITemplateRef
+	{
+		/// <summary>
+		/// The template id
+		/// </summary>
+		TemplateId Id { get; }
+
+		/// <summary>
+		/// Whether to show badges in UGS for these jobs
+		/// </summary>
+		bool ShowUgsBadges { get; }
+
+		/// <summary>
+		/// Whether to show alerts in UGS for these jobs
+		/// </summary>
+		bool ShowUgsAlerts { get; }
+
+		/// <summary>
+		/// Notification channel for this template. Overrides the stream channel if set.
+		/// </summary>
+		string? NotificationChannel { get; }
+
+		/// <summary>
+		/// Notification channel filter for this template. Can be a combination of "Success", "Failure" and "Warnings" separated by pipe characters.
+		/// </summary>
+		string? NotificationChannelFilter { get; }
+
+		/// <summary>
+		/// Triage channel for this template. Overrides the stream channel if set.
+		/// </summary>
+		string? TriageChannel { get; }
+
+		/// <summary>
+		/// List of schedules for this template
+		/// </summary>
+		ISchedule? Schedule { get; }
+
+		/// <summary>
+		/// List of chained job triggers
+		/// </summary>
+		IReadOnlyList<IChainedJobTemplate>? ChainedJobs { get; }
+
+		/// <summary>
+		/// List of template step states
+		/// </summary>
+		IReadOnlyList<ITemplateStep> StepStates { get; }
+
+		/// <summary>
+		/// Default change to use for this job
+		/// </summary>
+		IReadOnlyList<IChangeQuery>? DefaultChange { get; }
+	}
+
+	/// <summary>
+	/// Trigger for another template
+	/// </summary>
+	public interface IChainedJobTemplate
+	{
+		/// <summary>
+		/// Name of the target that needs to complete before starting the other template
+		/// </summary>
+		string Trigger { get; }
+
+		/// <summary>
+		/// Id of the template to trigger
+		/// </summary>
+		TemplateId TemplateId { get; }
+
+		/// <summary>
+		/// Whether to use the default change for the template rather than the change for the parent job.
+		/// </summary>
+		bool UseDefaultChangeForTemplate { get; }
+	}
+
+	/// <summary>
+	/// Information about a paused template step
+	/// </summary>
+	public interface ITemplateStep
+	{
+		/// <summary>
+		/// Name of the step
+		/// </summary>
+		string Name { get; }
+
+		/// <summary>
+		/// User who paused the step
+		/// </summary>
+		UserId PausedByUserId { get; }
+
+		/// <summary>
+		/// The UTC time when the step was paused
+		/// </summary>
+		DateTime PauseTimeUtc { get; }
+	}
+
+	/// <summary>
+	/// Extension methods for streams
+	/// </summary>
+	public static class StreamExtensions
+	{
+		/// <summary>
+		/// Updates an existing stream
+		/// </summary>
+		/// <param name="stream">The stream to update</param>
+		/// <param name="newPausedUntil">The new datetime for pausing builds</param>
+		/// <param name="newPauseComment">The reason for pausing</param>
+		/// <param name="cancellationToken">Cancellation token for the operation</param>
+		/// <returns>Async task object</returns>
+		public static async Task<IStream?> TryUpdatePauseStateAsync(this IStream? stream, DateTime? newPausedUntil = null, string? newPauseComment = null, CancellationToken cancellationToken = default)
+		{
+			for (; stream != null; stream = await stream.RefreshAsync(cancellationToken))
+			{
+				IStream? newStream = await stream.TryUpdatePauseStateAsync(newPausedUntil, newPauseComment, cancellationToken);
+				if (newStream != null)
+				{
+					return newStream;
+				}
+			}
+			return null;
+		}
+
+		/// <summary>
+		/// Attempts to update the last trigger time for a schedule
+		/// </summary>
+		/// <param name="stream">The stream to update</param>
+		/// <param name="templateRefId">The template ref id</param>
+		/// <param name="lastTriggerTimeUtc"></param>
+		/// <param name="lastTriggerCommitId"></param>
+		/// <param name="addJobs">Jobs to add</param>
+		/// <param name="removeJobs">Jobs to remove</param>
+		/// <param name="cancellationToken">Cancellation token for the operation</param>
+		/// <returns>True if the stream was updated</returns>
+		public static async Task<IStream?> TryUpdateScheduleTriggerAsync(this IStream stream, TemplateId templateRefId, DateTime? lastTriggerTimeUtc = null, CommitIdWithOrder? lastTriggerCommitId = null, List<JobId>? addJobs = null, List<JobId>? removeJobs = null, CancellationToken cancellationToken = default)
+		{
+			IStream? newStream = stream;
+			while (newStream != null)
+			{
+				ITemplateRef? templateRef;
+				if (!newStream.Templates.TryGetValue(templateRefId, out templateRef))
+				{
+					break;
+				}
+				if (templateRef.Schedule == null)
+				{
+					break;
+				}
+
+				IEnumerable<JobId> newActiveJobs = templateRef.Schedule.ActiveJobs;
+				if (removeJobs != null)
+				{
+					newActiveJobs = newActiveJobs.Except(removeJobs);
+				}
+				if (addJobs != null)
+				{
+					newActiveJobs = newActiveJobs.Union(addJobs);
+				}
+
+				newStream = await newStream.TryUpdateScheduleTriggerAsync(templateRefId, lastTriggerTimeUtc, lastTriggerCommitId, newActiveJobs.ToList(), cancellationToken);
+
+				if (newStream != null)
+				{
+					return newStream;
+				}
+
+				newStream = await stream.RefreshAsync(cancellationToken);
+			}
+			return null;
+		}
+
+		/// <summary>
+		/// Check if stream is paused for new builds
+		/// </summary>
+		/// <param name="stream">The stream object</param>
+		/// <param name="currentTime">Current time (allow tests to pass in a fake clock)</param>
+		/// <returns>If stream is paused</returns>
+		public static bool IsPaused(this IStream stream, DateTime currentTime)
+		{
+			return stream.PausedUntil != null && stream.PausedUntil > currentTime;
+		}
+	}
+}
diff --git a/Engine/Source/Programs/Shared/EpicGames.Horde/Streams/IStreamCollection.cs b/Engine/Source/Programs/Shared/EpicGames.Horde/Streams/IStreamCollection.cs
new file mode 100644
--- /dev/null
+++ b/Engine/Source/Programs/Shared/EpicGames.Horde/Streams/IStreamCollection.cs
@@ -0,0 +1,30 @@
+// Copyright Epic Games, Inc. All Rights Reserved.
+
+using System.Collections.Generic;
+using System.Threading;
+using System.Threading.Tasks;
+
+namespace EpicGames.Horde.Streams
+{
+	/// <summary>
+	/// Collection of stream documents
+	/// </summary>
+	public interface IStreamCollection
+	{
+		/// <summary>
+		/// Gets a stream by ID
+		/// </summary>
+		/// <param name="id">The stream identifier</param>
+		/// <param name="cancellationToken">Cancellation token for the operation</param>
+		/// <returns>The stream document</returns>
+		Task<IStream?> GetAsync(StreamId id, CancellationToken cancellationToken = default);
+
+		/// <summary>
+		/// Gets a stream by ID
+		/// </summary>
+		/// <param name="ids">The stream identifiers</param>
+		/// <param name="cancellationToken">Cancellation token for the operation</param>
+		/// <returns>The stream document</returns>
+		Task<IReadOnlyList<IStream>> FindAsync(IReadOnlyList<StreamId>? ids = null, CancellationToken cancellationToken = default);
+	}
+}
diff --git a/Engine/Source/Programs/Shared/EpicGames.Horde/Streams/StreamMessages.cs b/Engine/Source/Programs/Shared/EpicGames.Horde/Streams/StreamMessages.cs
new file mode 100644
--- /dev/null
+++ b/Engine/Source/Programs/Shared/EpicGames.Horde/Streams/StreamMessages.cs
@@ -0,0 +1,573 @@
+// Copyright Epic Games, Inc. All Rights Reserved.
+
+using System;
+using System.Collections.Generic;
+using System.ComponentModel.DataAnnotations;
+using System.Linq;
+using EpicGames.Core;
+using EpicGames.Horde.Agents.Pools;
+using EpicGames.Horde.Issues;
+using EpicGames.Horde.Jobs;
+using EpicGames.Horde.Jobs.Schedules;
+using EpicGames.Horde.Jobs.Templates;
+using EpicGames.Horde.Projects;
+using EpicGames.Horde.Users;
+
+#pragma warning disable CA2227 // Change x to be read-only by removing the property setter
+
+namespace EpicGames.Horde.Streams
+{
+	/// <summary>
+	/// Response describing a stream
+	/// </summary>
+	public class GetStreamResponse
+	{
+		/// <summary>
+		/// Unique id of the stream
+		/// </summary>
+		public StreamId Id { get; set; }
+
+		/// <summary>
+		/// Unique id of the project containing this stream
+		/// </summary>
+		public ProjectId ProjectId { get; set; }
+
+		/// <summary>
+		/// Name of the stream
+		/// </summary>
+		public string Name { get; set; }
+
+		/// <summary>
+		/// The config file path on the server
+		/// </summary>
+		public string ConfigPath { get; set; } = String.Empty;
+
+		/// <summary>
+		/// Revision of the config file 
+		/// </summary>
+		public string ConfigRevision { get; set; } = String.Empty;
+
+		/// <summary>
+		/// Order to display in the list
+		/// </summary>
+		public int Order { get; set; }
+
+		/// <summary>
+		/// Notification channel for all jobs in this stream
+		/// </summary>
+		public string? NotificationChannel { get; set; }
+
+		/// <summary>
+		/// Notification channel filter for this template. Can be a combination of "Success", "Failure" and "Warnings" separated by pipe characters.
+		/// </summary>
+		public string? NotificationChannelFilter { get; set; }
+
+		/// <summary>
+		/// Channel to post issue triage notifications
+		/// </summary>
+		public string? TriageChannel { get; set; }
+
+		/// <summary>
+		/// Default template for running preflights
+		/// </summary>
+		public string? DefaultPreflightTemplate { get; set; }
+
+		/// <summary>
+		/// Default template to use for preflights
+		/// </summary>
+		public DefaultPreflightMessage? DefaultPreflight { get; set; }
+
+		/// <summary>
+		/// List of tabs to display for this stream
+		/// </summary>
+		public List<GetStreamTabResponse> Tabs { get; set; }
+
+		/// <summary>
+		/// Map of agent name to type
+		/// </summary>
+		public Dictionary<string, GetAgentTypeResponse> AgentTypes { get; set; }
+
+		/// <summary>
+		/// Map of workspace name to type
+		/// </summary>
+		public Dictionary<string, GetWorkspaceTypeResponse>? WorkspaceTypes { get; set; }
+
+		/// <summary>
+		/// Templates for jobs in this stream
+		/// </summary>
+		public List<GetTemplateRefResponse> Templates { get; set; }
+
+		/// <summary>
+		/// Stream paused for new builds until this date
+		/// </summary>
+		public DateTime? PausedUntil { get; set; }
+
+		/// <summary>
+		/// Reason for stream being paused
+		/// </summary>
+		public string? PauseComment { get; set; }
+
+		/// <summary>
+		/// Workflows for this stream
+		/// </summary>
+		public List<GetWorkflowResponse> Workflows { get; set; }
+
+		/// <summary>
+		/// Constructor
+		/// </summary>
+		/// <param name="stream">The stream to construct from</param>
+		/// <param name="templates">Templates for this stream</param>
+		public GetStreamResponse(IStream stream, List<GetTemplateRefResponse> templates)
+		{
+			Id = stream.Id;
+			ProjectId = stream.ProjectId;
+			Name = stream.Name;
+			ConfigPath = stream.ConfigPath ?? String.Empty;
+			ConfigRevision = stream.ConfigRevision;
+			Order = stream.Order;
+			NotificationChannel = stream.NotificationChannel;
+			NotificationChannelFilter = stream.NotificationChannelFilter;
+			TriageChannel = stream.TriageChannel;
+			DefaultPreflightTemplate = stream.DefaultPreflight?.TemplateId?.ToString();
+			DefaultPreflight = (stream.DefaultPreflight == null) ? null : new DefaultPreflightMessage(stream.DefaultPreflight);
+			Tabs = stream.Tabs.ConvertAll(x => new GetStreamTabResponse(x));
+			AgentTypes = stream.AgentTypes.ToDictionary(x => x.Key, x => new GetAgentTypeResponse(x.Value));
+			WorkspaceTypes = stream.WorkspaceTypes.ToDictionary(x => x.Key, x => new GetWorkspaceTypeResponse(x.Value));
+			Templates = templates;
+			Workflows = stream.Workflows.ConvertAll(x => new GetWorkflowResponse(x));
+			PausedUntil = stream.PausedUntil;
+			PauseComment = stream.PauseComment;
+		}
+	}
+
+	/// <summary>
+	/// Information about the default preflight to run
+	/// </summary>
+	public class DefaultPreflightMessage : IDefaultPreflight
+	{
+		/// <inheritdoc/>
+		public TemplateId? TemplateId { get; set; }
+
+		/// <inheritdoc/>
+		public ChangeQueryMessage? Change { get; set; }
+		IChangeQuery? IDefaultPreflight.Change => Change;
+
+		/// <summary>
+		/// Constructor
+		/// </summary>
+		public DefaultPreflightMessage()
+		{ }
+
+		/// <summary>
+		/// Constructor
+		/// </summary>
+		public DefaultPreflightMessage(IDefaultPreflight other)
+		{
+			TemplateId = other.TemplateId;
+			Change = (other.Change == null) ? null : new ChangeQueryMessage(other.Change);
+		}
+	}
+
+	/// <summary>
+	/// Information about a page to display in the dashboard for a stream
+	/// </summary>
+	public class GetStreamTabResponse : IStreamTab
+	{
+		/// <inheritdoc/>
+		public string Title { get; set; }
+
+		/// <inheritdoc/>
+		public string Type { get; set; }
+
+		/// <inheritdoc/>
+		public TabStyle Style { get; set; }
+
+		/// <inheritdoc/>
+		public bool ShowNames { get; set; }
+
+		/// <inheritdoc/>
+		public bool? ShowPreflights { get; set; }
+
+		/// <inheritdoc cref="IStreamTab.JobNames"/>
+		public List<string>? JobNames { get; set; }
+		IReadOnlyList<string>? IStreamTab.JobNames => JobNames;
+
+		/// <inheritdoc cref="IStreamTab.Templates"/>
+		public List<TemplateId>? Templates { get; set; }
+		IReadOnlyList<TemplateId>? IStreamTab.Templates => Templates;
+
+		/// <inheritdoc cref="IStreamTab.Columns"/>
+		public List<GetStreamTabColumnResponse>? Columns { get; set; }
+		IReadOnlyList<IStreamTabColumn>? IStreamTab.Columns => Columns;
+
+		/// <summary>
+		/// Constructor
+		/// </summary>
+		public GetStreamTabResponse()
+		{
+			Title = String.Empty;
+			Type = String.Empty;
+		}
+
+		/// <summary>
+		/// Constructor
+		/// </summary>
+		public GetStreamTabResponse(IStreamTab other)
+		{
+			Title = other.Title;
+			Type = other.Type;
+			Style = other.Style;
+			ShowNames = other.ShowNames;
+			ShowPreflights = other.ShowPreflights;
+			JobNames = other.JobNames?.ToList();
+			Templates = other.Templates?.ToList();
+			Columns = other.Columns?.ConvertAll(x => new GetStreamTabColumnResponse(x));
+		}
+	}
+
+	/// <summary>
+	/// Describes a column to display on the jobs page
+	/// </summary>
+	public class GetStreamTabColumnResponse : IStreamTabColumn
+	{
+		/// <inheritdoc/>
+		public TabColumnType Type { get; set; }
+
+		/// <inheritdoc/>
+		public string Heading { get; set; }
+
+		/// <inheritdoc/>
+		public string? Category { get; set; }
+
+		/// <inheritdoc/>
+		public string? Parameter { get; set; }
+
+		/// <inheritdoc/>
+		public int? RelativeWidth { get; set; }
+
+		/// <summary>
+		/// Default constructor
+		/// </summary>
+		public GetStreamTabColumnResponse()
+		{
+			Heading = String.Empty;
+		}
+
+		/// <summary>
+		/// Constructor
+		/// </summary>
+		public GetStreamTabColumnResponse(IStreamTabColumn other)
+		{
+			Type = other.Type;
+			Heading = other.Heading;
+			Category = other.Category;
+			Parameter = other.Parameter;
+			RelativeWidth = other.RelativeWidth;
+		}
+	}
+
+	/// <summary>
+	/// Mapping from a BuildGraph agent type to a set of machines on the farm
+	/// </summary>
+	public class GetAgentTypeResponse
+	{
+		/// <summary>
+		/// Pool of agents to use for this agent type
+		/// </summary>
+		public PoolId Pool { get; set; }
+
+		/// <summary>
+		/// Name of the workspace to sync
+		/// </summary>
+		public string? Workspace { get; set; }
+
+		/// <summary>
+		/// Path to the temporary storage dir
+		/// </summary>
+		public string? TempStorageDir { get; set; }
+
+		/// <summary>
+		/// Environment variables to be set when executing the job
+		/// </summary>
+		public Dictionary<string, string>? Environment { get; set; }
+
+		/// <summary>
+		/// Constructor
+		/// </summary>
+		public GetAgentTypeResponse()
+		{ }
+
+		/// <summary>
+		/// Constructor
+		/// </summary>
+		public GetAgentTypeResponse(IAgentType other)
+		{
+			Pool = other.Pool;
+			Workspace = other.Workspace;
+			TempStorageDir = other.TempStorageDir;
+			Environment = other.Environment?.ToDictionary();
+		}
+	}
+
+	/// <summary>
+	/// Information about a workspace type
+	/// </summary>
+	public class GetWorkspaceTypeResponse : IWorkspaceType
+	{
+		/// <inheritdoc/>
+		public string? Cluster { get; set; }
+
+		/// <inheritdoc/>
+		public string? ServerAndPort { get; set; }
+
+		/// <inheritdoc/>
+		public string? UserName { get; set; }
+
+		/// <inheritdoc/>
+		string? IWorkspaceType.Password => throw new NotSupportedException();
+
+		/// <inheritdoc/>
+		public string? Identifier { get; set; }
+
+		/// <inheritdoc/>
+		public string? Stream { get; set; }
+
+		/// <inheritdoc cref="IWorkspaceType.View"/>
+		public List<string>? View { get; set; }
+		IReadOnlyList<string>? IWorkspaceType.View => View;
+
+		/// <inheritdoc/>
+		public bool? Incremental { get; set; }
+
+		/// <inheritdoc/>
+		public bool? UseAutoSdk { get; set; }
+
+		/// <inheritdoc cref="IWorkspaceType.AutoSdkView"/>
+		public List<string>? AutoSdkView { get; set; }
+		IReadOnlyList<string>? IWorkspaceType.AutoSdkView => AutoSdkView;
+
+		/// <inheritdoc/>
+		public string? Method { get; set; } = null;
+
+		/// <inheritdoc/>
+		public long? MinScratchSpace { get; set; } = null;
+
+		/// <inheritdoc/>
+		public long? ConformDiskFreeSpace { get; set; } = null;
+
+		/// <summary>
+		/// Default constructor
+		/// </summary>
+		public GetWorkspaceTypeResponse()
+		{ }
+
+		/// <summary>
+		/// Constructor
+		/// </summary>
+		public GetWorkspaceTypeResponse(IWorkspaceType other)
+		{
+			Cluster = other.Cluster;
+			ServerAndPort = other.ServerAndPort;
+			UserName = other.UserName;
+			Identifier = other.Identifier;
+			Stream = other.Stream;
+			View = other.View?.ToList();
+			Incremental = other.Incremental;
+			UseAutoSdk = other.UseAutoSdk;
+			AutoSdkView = other.AutoSdkView?.ToList();
+			Method = other.Method;
+			MinScratchSpace = other.MinScratchSpace;
+			ConformDiskFreeSpace = other.ConformDiskFreeSpace;
+		}
+	}
+
+	/// <summary>
+	/// State information for a step in the stream
+	/// </summary>
+	public class GetTemplateStepStateResponse
+	{
+		/// <summary>
+		/// Name of the step
+		/// </summary>
+		public string Name { get; set; } = String.Empty;
+
+		/// <summary>
+		/// User who paused the step
+		/// </summary>
+		public GetThinUserInfoResponse? PausedByUserInfo { get; set; }
+
+		/// <summary>
+		/// The UTC time when the step was paused
+		/// </summary>
+		public DateTime? PauseTimeUtc { get; set; }
+
+		/// <summary>
+		/// Default constructor for serialization
+		/// </summary>
+		private GetTemplateStepStateResponse()
+		{
+		}
+
+		/// <summary>
+		/// Constructor
+		/// </summary>
+		public GetTemplateStepStateResponse(ITemplateStep state, GetThinUserInfoResponse? pausedByUserInfo)
+		{
+			Name = state.Name;
+			PauseTimeUtc = state.PauseTimeUtc;
+			PausedByUserInfo = pausedByUserInfo;
+		}
+	}
+
+	/// <summary>
+	/// Information about a template in this stream
+	/// </summary>
+	public class GetTemplateRefResponse : GetTemplateResponseBase
+	{
+		/// <summary>
+		/// Id of the template ref
+		/// </summary>
+		public string Id { get; set; }
+
+		/// <summary>
+		/// Hash of the template definition
+		/// </summary>
+		public string Hash { get; set; }
+
+		/// <summary>
+		/// Whether to show badges in UGS for these jobs
+		/// </summary>
+		public bool ShowUgsBadges { get; set; }
+
+		/// <summary>
+		/// Whether to show alerts in UGS for these jobs
+		/// </summary>
+		public bool ShowUgsAlerts { get; set; }
+
+		/// <summary>
+		/// Notification channel for this template. Overrides the stream channel if set.
+		/// </summary>
+		public string? NotificationChannel { get; set; }
+
+		/// <summary>
+		/// Notification channel filter for this template. Can be a combination of "Success", "Failure" and "Warnings" separated by pipe characters.
+		/// </summary>
+		public string? NotificationChannelFilter { get; set; }
+
+		/// <summary>
+		/// Triage channel for this template. Overrides the stream channel if set.
+		/// </summary>
+		public string? TriageChannel { get; set; }
+
+		/// <summary>
+		/// The schedule for this ref
+		/// </summary>
+		public GetScheduleResponse? Schedule { get; set; }
+
+		/// <summary>
+		/// List of templates to trigger
+		/// </summary>
+		public List<GetChainedJobTemplateResponse>? ChainedJobs { get; set; }
+
+		/// <summary>
+		/// List of step states
+		/// </summary>
+		public List<GetTemplateStepStateResponse>? StepStates { get; set; }
+
+		/// <summary>
+		/// List of queries for the default changelist
+		/// </summary>
+		public List<ChangeQueryMessage>? DefaultChange { get; set; }
+
+		/// <summary>
+		/// Whether the user is allowed to create jobs from this template
+		/// </summary>
+		public bool CanRun { get; set; }
+
+		/// <summary>
+		/// Constructor
+		/// </summary>
+		/// <param name="id">The template ref id</param>
+		/// <param name="templateRef">The template ref</param>
+		/// <param name="template">The actual template</param>
+		/// <param name="stepStates">The template step states</param>
+		/// <param name="schedulerTimeZone">The scheduler time zone</param>
+		/// <param name="canRun">Whether the user can run this template</param>
+		public GetTemplateRefResponse(TemplateId id, ITemplateRef templateRef, ITemplate template, List<GetTemplateStepStateResponse>? stepStates, TimeZoneInfo schedulerTimeZone, bool canRun)
+			: base(template)
+		{
+			Id = id.ToString();
+			Hash = template.Hash.ToString();
+			ShowUgsBadges = templateRef.ShowUgsBadges;
+			ShowUgsAlerts = templateRef.ShowUgsAlerts;
+			NotificationChannel = templateRef.NotificationChannel;
+			NotificationChannelFilter = templateRef.NotificationChannelFilter;
+			Schedule = (templateRef.Schedule != null) ? new GetScheduleResponse(templateRef.Schedule, schedulerTimeZone) : null;
+			ChainedJobs = (templateRef.ChainedJobs != null && templateRef.ChainedJobs.Count > 0) ? templateRef.ChainedJobs.ConvertAll(x => new GetChainedJobTemplateResponse(x)) : null;
+			StepStates = stepStates;
+			DefaultChange = templateRef.DefaultChange?.ConvertAll(x => new ChangeQueryMessage(x));
+			CanRun = canRun;
+		}
+	}
+
+	/// <summary>
+	/// Trigger for another template
+	/// </summary>
+	public class GetChainedJobTemplateResponse : IChainedJobTemplate
+	{
+		/// <inheritdoc/>
+		public string Trigger { get; set; }
+
+		/// <inheritdoc/>
+		public TemplateId TemplateId { get; set; }
+
+		/// <inheritdoc/>
+		public bool UseDefaultChangeForTemplate { get; set; }
+
+		/// <summary>
+		/// Constructor
+		/// </summary>
+		public GetChainedJobTemplateResponse()
+		{
+			Trigger = String.Empty;
+		}
+
+		/// <summary>
+		/// Constructor
+		/// </summary>
+		public GetChainedJobTemplateResponse(IChainedJobTemplate other)
+		{
+			Trigger = other.Trigger;
+			TemplateId = other.TemplateId;
+			UseDefaultChangeForTemplate = other.UseDefaultChangeForTemplate;
+		}
+	}
+
+	/// <summary>
+	/// Step state update request
+	/// </summary>
+	public class UpdateStepStateRequest
+	{
+		/// <summary>
+		/// Name of the step
+		/// </summary>
+		[Required]
+		public string Name { get; set; } = String.Empty;
+
+		/// <summary>
+		/// User who paused the step
+		/// </summary>
+		public string? PausedByUserId { get; set; }
+	}
+
+	/// <summary>
+	/// Updates an existing stream template ref
+	/// </summary>
+	public class UpdateTemplateRefRequest
+	{
+		/// <summary>
+		/// Step states to update
+		/// </summary>
+		public List<UpdateStepStateRequest>? StepStates { get; set; }
+	}
+}
