Horde: Fix build error.

#rnx

diff --git a/Engine/Source/Programs/Shared/EpicGames.Horde/HordeHttpClient.cs b/Engine/Source/Programs/Shared/EpicGames.Horde/HordeHttpClient.cs
--- a/Engine/Source/Programs/Shared/EpicGames.Horde/HordeHttpClient.cs
+++ b/Engine/Source/Programs/Shared/EpicGames.Horde/HordeHttpClient.cs
@@ -532,10 +532,10 @@ namespace EpicGames.Horde
 		/// <param name="projectId"></param>
 		/// <param name="cancellationToken"></param>
 		/// <returns></returns>
-		public Task<GetUgsMetadataResponse> GetUgsMetadataAsync(StreamId streamId, CommitId commitId, ProjectId projectId, CancellationToken cancellationToken = default)
+		public Task<GetUgsMetadataListResponse> GetUgsMetadataAsync(StreamId streamId, CommitId commitId, ProjectId projectId, CancellationToken cancellationToken = default)
 		{
 			string perforceStream = $"//{streamId.ToString().Replace('-', '/')}";
-			return GetAsync<GetUgsMetadataResponse>(_httpClient, $"/ugs/api/metadata?stream={perforceStream}&change={commitId.GetPerforceChange()}&project={projectId}", cancellationToken);
+			return GetAsync<GetUgsMetadataListResponse>(_httpClient, $"/ugs/api/metadata?stream={perforceStream}&change={commitId.GetPerforceChange()}&project={projectId}", cancellationToken);
 		}
 		#endregion
 	}
