Horde: Replace AesTransport with a new implementation

- Uses random nonces for each message instead of a base nonce
- Add TCP transport integration tests

diff --git a/Engine/Source/Programs/Horde/HordeAgent/Leases/Handlers/ComputeHandler.cs b/Engine/Source/Programs/Horde/HordeAgent/Leases/Handlers/ComputeHandler.cs
--- a/Engine/Source/Programs/Horde/HordeAgent/Leases/Handlers/ComputeHandler.cs
+++ b/Engine/Source/Programs/Horde/HordeAgent/Leases/Handlers/ComputeHandler.cs
@@ -135,7 +135,7 @@ namespace HordeAgent.Leases.Handlers
 
 				case ComputeEncryption.Aes:
 #pragma warning disable CA2000 // Dispose objects before losing scope
-					return new AesTransport(new TcpTransport(socket), computeTask.Key.ToByteArray(), computeTask.Nonce.ToByteArray());
+					return new AesTransport(new TcpTransport(socket), computeTask.Key.ToByteArray());
 #pragma warning restore CA2000 // Restore CA2000
 
 				case ComputeEncryption.Unspecified:
diff --git a/Engine/Source/Programs/Shared/EpicGames.Horde.Tests/AesTransportTests.cs b/Engine/Source/Programs/Shared/EpicGames.Horde.Tests/AesTransportTests.cs
new file mode 100644
--- /dev/null
+++ b/Engine/Source/Programs/Shared/EpicGames.Horde.Tests/AesTransportTests.cs
@@ -0,0 +1,166 @@
+// Copyright Epic Games, Inc. All Rights Reserved.
+
+using System;
+using System.Buffers;
+using System.IO.Pipelines;
+using System.Net.Sockets;
+using System.Threading;
+using System.Threading.Tasks;
+using EpicGames.Horde.Compute;
+using EpicGames.Horde.Compute.Transports;
+using EpicGames.Horde.Tests.Compute;
+using Microsoft.VisualStudio.TestTools.UnitTesting;
+
+namespace EpicGames.Horde.Tests;
+
+[TestClass]
+public sealed class AesTransportTests : IAsyncDisposable
+{
+	private readonly byte[] _key;
+	private readonly MemoryComputeTransport _memoryTransport;
+	private readonly AesTransport _aesTransport;
+	private readonly Random _random = new(42);
+	private readonly int[] _bufferSizes = [1, 2, 3, 5, 10, 20, 2000, 20000, 2000000];
+	
+	public AesTransportTests()
+	{
+		_key = AesTransport.CreateKey();
+		_memoryTransport = new MemoryComputeTransport();
+		_aesTransport = new AesTransport(_memoryTransport, _key);
+	}
+	
+	public async ValueTask DisposeAsync()
+	{
+		await _memoryTransport.DisposeAsync();
+		await _aesTransport.DisposeAsync();
+	}
+	
+	[TestMethod]
+	public async Task SendAndReceive_VariousBufferSizesOverTcp_SuccessfullyTransferredAsync()
+	{
+		using CancellationTokenSource cts = new(10000);
+		(Socket clientSocket, Socket serverSocket) = await ComputeSocketTests.CreateSocketsAsync(cts.Token);
+		await using TcpTransport clientTransport = new(clientSocket);
+		await using TcpTransport serverTransport = new(serverSocket);
+		await using AesTransport encryptedClient = new(clientTransport, _key);
+		await using AesTransport encryptedServer = new(serverTransport, _key);
+		
+		foreach (int bufferSize in _bufferSizes)
+		{
+			await SendAndReceiveAsync(encryptedClient, encryptedServer, bufferSize, cts.Token);
+		}
+	}
+	
+	[TestMethod]
+	public async Task SendAndReceive_VariousBufferSizesInMemory_SuccessfullyTransferredAsync()
+	{
+		using CancellationTokenSource cts = new(10000);
+		await using MemoryComputeTransport memTransport = new();
+		await using AesTransport encryptedClient = new(memTransport, _key);
+		await using AesTransport encryptedServer = new(memTransport, _key);
+		
+		foreach (int bufferSize in _bufferSizes)
+		{
+			await SendAndReceiveAsync(encryptedClient, encryptedServer, bufferSize, cts.Token);
+		}
+	}
+	
+	private async Task SendAndReceiveAsync(AesTransport client, AesTransport server, int bufferSize, CancellationToken cancellationToken)
+	{
+		// Arrange
+		byte[] sentData = new byte[bufferSize];
+		_random.NextBytes(sentData);
+		byte[] serverReceivedData = new byte[sentData.Length];
+		byte[] clientReceivedData = new byte[sentData.Length];
+		
+		// Act
+		await client.SendAsync(sentData, cancellationToken);
+		await server.SendAsync(sentData, cancellationToken);
+		int serverBytesRead = await server.RecvAsync(serverReceivedData, cancellationToken);
+		int clientBytesRead = await client.RecvAsync(clientReceivedData, cancellationToken);
+		
+		// Assert
+		Assert.AreEqual(sentData.Length, serverBytesRead);
+		Assert.AreEqual(sentData.Length, clientBytesRead);
+		CollectionAssert.AreEqual(sentData, serverReceivedData);
+		CollectionAssert.AreEqual(sentData, clientReceivedData);
+	}
+	
+	[TestMethod]
+	public async Task Constructor_InvalidKeyLength_ThrowsArgumentExceptionAsync()
+	{
+		await Assert.ThrowsExceptionAsync<ArgumentException>(async () =>
+		{
+			byte[] invalidKey = new byte[16];
+			await using AesTransport _ = new (_memoryTransport, invalidKey);
+		});
+	}
+	
+	[TestMethod]
+	public async Task MarkComplete_PropagatedToInnerTransportAsync()
+	{
+		// Act
+		await _aesTransport.MarkCompleteAsync(CancellationToken.None);
+		
+		// Assert
+		Assert.IsTrue(_memoryTransport.IsCompleted);
+	}
+}
+
+// Helper class for testing
+internal class MemoryComputeTransport : ComputeTransport
+{
+	private readonly Pipe _pipe;
+	public bool IsCompleted { get; private set; }
+	
+	public MemoryComputeTransport()
+	{
+		PipeOptions options = new(pauseWriterThreshold: 10 * 1024 * 1024); // Must fit buffer sizes being tested
+		_pipe = new Pipe(options);
+	}
+	
+	/// <inheritdoc/>
+	public override async ValueTask<int> RecvAsync(Memory<byte> buffer, CancellationToken cancellationToken)
+	{
+		ReadResult result = await _pipe.Reader.ReadAsync(cancellationToken);
+		
+		if (result.Buffer.IsEmpty && result.IsCompleted)
+		{
+			return 0;
+		}
+		
+		int bytesToCopy = (int)Math.Min(buffer.Length, result.Buffer.Length);
+		result.Buffer.Slice(0, bytesToCopy).CopyTo(buffer.Span);
+		_pipe.Reader.AdvanceTo(result.Buffer.GetPosition(bytesToCopy));
+		
+		return bytesToCopy;
+	}
+	
+	/// <inheritdoc/>
+	public override async ValueTask SendAsync(ReadOnlySequence<byte> buffer, CancellationToken cancellationToken)
+	{
+		foreach (ReadOnlyMemory<byte> segment in buffer)
+		{
+			FlushResult result = await _pipe.Writer.WriteAsync(segment, cancellationToken);
+			if (result.IsCompleted)
+			{
+				break;
+			}
+		}
+	}
+	
+	/// <inheritdoc/>
+	public override async ValueTask MarkCompleteAsync(CancellationToken cancellationToken)
+	{
+		IsCompleted = true;
+		await _pipe.Writer.CompleteAsync();
+	}
+	
+	/// <inheritdoc/>
+	public override async ValueTask DisposeAsync()
+	{
+		await _pipe.Writer.CompleteAsync();
+		await _pipe.Reader.CompleteAsync();
+		GC.SuppressFinalize(this);
+	}
+}
\ No newline at end of file
diff --git a/Engine/Source/Programs/Shared/EpicGames.Horde/Compute/Clients/ServerComputeClient.cs b/Engine/Source/Programs/Shared/EpicGames.Horde/Compute/Clients/ServerComputeClient.cs
--- a/Engine/Source/Programs/Shared/EpicGames.Horde/Compute/Clients/ServerComputeClient.cs
+++ b/Engine/Source/Programs/Shared/EpicGames.Horde/Compute/Clients/ServerComputeClient.cs
@@ -401,7 +401,7 @@ namespace EpicGames.Horde.Compute.Clients
 				case Encryption.Aes:
 #pragma warning disable CA2000 // Dispose objects before losing scope
 					TcpTransport tcpTransport = new(socket);
-					return new AesTransport(tcpTransport, StringUtils.ParseHexString(response.Key), StringUtils.ParseHexString(response.Nonce));
+					return new AesTransport(tcpTransport, StringUtils.ParseHexString(response.Key));
 #pragma warning restore CA2000 // Restore CA2000
 
 				case Encryption.None:
diff --git a/Engine/Source/Programs/Shared/EpicGames.Horde/Compute/Transports/AesTransport.cs b/Engine/Source/Programs/Shared/EpicGames.Horde/Compute/Transports/AesTransport.cs
--- a/Engine/Source/Programs/Shared/EpicGames.Horde/Compute/Transports/AesTransport.cs
+++ b/Engine/Source/Programs/Shared/EpicGames.Horde/Compute/Transports/AesTransport.cs
@@ -2,226 +2,161 @@
 
 using System;
 using System.Buffers;
-using System.Buffers.Binary;
-using System.IO.Pipelines;
+using System.IO;
 using System.Security.Cryptography;
 using System.Threading;
 using System.Threading.Tasks;
-using EpicGames.Core;
 
-#pragma warning disable VSTHRD003 // Avoid awaiting or returning a Task representing work that was not started within your context as that can lead to deadlocks.
+namespace EpicGames.Horde.Compute.Transports;
 
-namespace EpicGames.Horde.Compute.Transports
+/// <summary>
+/// Transport layer that adds AES encryption on top of an underlying transport implementation.
+/// Key must be exchanged separately (e.g. via the HTTPS request to negotiate a lease with the server).
+/// </summary>
+public sealed class AesTransport : ComputeTransport
 {
 	/// <summary>
-	/// Transport layer that adds AES encryption on top of an underlying transport implementation. Key must be exchanged separately
-	/// (eg. via the HTTPS request to negotiate a lease with the server).
+	/// Length of the required encryption key. 
 	/// </summary>
-	public sealed class AesTransport : ComputeTransport
-	{
-		/// <summary>
-		/// Length of the required encrption key. 
-		/// </summary>
-		public const int KeyLength = 32;
-
-		/// <summary>
-		/// Length of the nonce. This should be a cryptographically random number, and does not have to be secret.
-		/// </summary>
-		public const int NonceLength = 12;
-
-		const int HeaderLength = sizeof(int); // Unencrypted size of the message
-		const int FooterLength = 16; // 16-byte auth tag for encryption
-
-		readonly ComputeTransport _inner;
-		readonly bool _leaveInnerOpen;
-		readonly Pipe _readPipe;
-		readonly AesGcm _aesGcm;
-
-		ReadOnlySequence<byte> _readBuffer;
-
-		readonly byte[] _readNonce;
-		readonly byte[] _writeNonce;
-
-		readonly IMemoryOwner<byte> _writeBufferOwner;
-		Task _lastWriteTask = Task.CompletedTask;
-		int _writeBufferOffset = 0;
-
-		const int WritePacketSize = 64 * 1024;
-		readonly BackgroundTask _backgroundReadTask;
-
-		/// <inheritdoc/>
-		public long Position { get; private set; }
+	public const int KeyLength = 32;
+	
+	/// <summary>
+	/// Length of the nonce. This should be a cryptographically random number, and does not have to be secret.
+	/// </summary>
+	public const int NonceLength = 12;
+	
+	private const int TagSize = 16; // AES-GCM tag size in bytes
+	private const int BufferLength = sizeof(int); // Length of buffer stored as int
+	private const int HeaderLength = BufferLength + NonceLength;
+	private readonly ComputeTransport _inner;
+	private readonly bool _leaveInnerOpen;
+	private readonly AesGcm _aes;
+	private readonly SemaphoreSlim _sendSemaphore = new(1, 1);
+	private readonly SemaphoreSlim _recvSemaphore = new(1, 1);
+	private bool _disposed;
 
-		/// <summary>
-		/// Constructor
-		/// </summary>
-		/// <param name="inner">The underlying transport implementation</param>
-		/// <param name="key">AES encryption key (256 bits / 32 bytes)</param>
-		/// <param name="nonce">Cryptographic nonce to identify the connection. Must be longer than <see cref="NonceLength"/>.</param>
-		/// <param name="leaveInnerOpen">Whether inner compute transport should be disposed</param>
-		public AesTransport(ComputeTransport inner, ReadOnlySpan<byte> key, ReadOnlySpan<byte> nonce, bool leaveInnerOpen = false)
+	/// <summary>
+	/// Constructor
+	/// </summary>
+	/// <param name="inner">The underlying transport implementation that will be encrypted</param>
+	/// <param name="key">Encryption key. Should be generated with a cryptographically secure random number generator</param>
+	/// <param name="leaveInnerOpen">Whether to dispose the inner transport when this instance is disposed</param>
+	public AesTransport(ComputeTransport inner, ReadOnlySpan<byte> key, bool leaveInnerOpen = false)
+	{
+		_inner = inner;
+		_leaveInnerOpen = leaveInnerOpen;
+		
+		if (key.Length != KeyLength)
 		{
-			if (key.Length != KeyLength)
-			{
-				throw new ArgumentException($"Key must be {KeyLength} bytes", nameof(key));
-			}
-			if (nonce.Length < NonceLength)
-			{
-				throw new ArgumentException($"Nonce must be at least {NonceLength} bytes", nameof(nonce));
-			}
-
-			_inner = inner;
-			_readPipe = new Pipe();
-#pragma warning disable SYSLIB0053 // Type or member is obsolete
-			_aesGcm = new AesGcm(key);
-#pragma warning restore SYSLIB0053 // Type or member is obsolete
-			_readNonce = nonce.Slice(0, NonceLength).ToArray();
-			_writeNonce = nonce.Slice(0, NonceLength).ToArray();
-			_leaveInnerOpen = leaveInnerOpen;
-			_backgroundReadTask = BackgroundTask.StartNew(BackgroundReadAsync);
-
-			_writeBufferOwner = MemoryPool<byte>.Shared.Rent(WritePacketSize * 2);
+			throw new ArgumentException("Key must be 32 bytes", nameof(key)); // AES-256
 		}
 
-		/// <inheritdoc/>
-		public override async ValueTask DisposeAsync()
+		_aes = new AesGcm(key.ToArray(), TagSize);
+	}
+	
+	/// <summary>
+	/// Create a random key for this transport
+	/// </summary>
+	/// <returns>A cryptographically random key</returns>
+	public static byte[] CreateKey() => RandomNumberGenerator.GetBytes(KeyLength);
+	
+	/// <inheritdoc/>
+	public override async ValueTask SendAsync(ReadOnlySequence<byte> buffer, CancellationToken cancellationToken)
+	{
+		await _sendSemaphore.WaitAsync(cancellationToken);
+		try
 		{
-			await _lastWriteTask;
-			await _backgroundReadTask.DisposeAsync();
-			_writeBufferOwner.Dispose();
-			_aesGcm.Dispose();
-			if (!_leaveInnerOpen)
-			{
-				await _inner.DisposeAsync();
-			}
+			int bufferLength = (int)buffer.Length;
+			
+			// Message format: [length (4 bytes)][nonce (12 bytes)][encrypted data][tag (16 bytes)]
+			int messageSize = HeaderLength + bufferLength + TagSize;
+			byte[] message = new byte[messageSize];
+			BitConverter.TryWriteBytes(message.AsSpan(0, BufferLength), bufferLength);
+			RandomNumberGenerator.Fill(message.AsSpan(BufferLength, NonceLength));
+			
+			// Copy sequence to contiguous array for encryption
+			byte[] plaintext = buffer.ToArray();
+			
+			_aes.Encrypt(
+				message.AsSpan(BufferLength, NonceLength),
+				plaintext,
+				message.AsSpan(HeaderLength, bufferLength),
+				message.AsSpan(HeaderLength + bufferLength, TagSize));
+			
+			await _inner.SendAsync(message, cancellationToken);
 		}
-
-		/// <summary>
-		/// Creates an encryption key
-		/// </summary>
-		public static byte[] CreateKey() => RandomNumberGenerator.GetBytes(KeyLength);
-
-		async Task BackgroundReadAsync(CancellationToken cancellationToken)
+		finally
 		{
-			using IMemoryOwner<byte> buffer = MemoryPool<byte>.Shared.Rent(128 * 1024);
-			Memory<byte> memory = buffer.Memory;
-
-			int size = 0;
-			for (; ; )
-			{
-				// Read more data into the buffer
-				int read = await _inner.RecvAsync(memory.Slice(size), cancellationToken);
-				if (read == 0)
-				{
-					break;
-				}
-				size += read;
-
-				// Wait until we've got a complete encryption packet
-				int prevSize = size;
-				size = Decrypt(memory.Slice(0, size).Span);
-				memory.Slice(prevSize - size, size).CopyTo(memory);
-			}
-
-			await _readPipe.Writer.CompleteAsync();
+			_sendSemaphore.Release();
 		}
-
-		int Decrypt(Span<byte> span)
+	}
+	
+	/// <inheritdoc/>
+	public override async ValueTask<int> RecvAsync(Memory<byte> buffer, CancellationToken cancellationToken)
+	{
+		await _recvSemaphore.WaitAsync(cancellationToken);
+		try
 		{
-			while (span.Length >= HeaderLength)
+			byte[] headerData = new byte[HeaderLength];
+			await _inner.RecvFullAsync(headerData, cancellationToken);
+			int bufferLength = BitConverter.ToInt32(headerData.AsSpan(0, BufferLength));
+			
+			const int MaxMessageSize = 100 * 1024 * 1024; // 100MB
+			if (bufferLength is <= 0 or > MaxMessageSize)
 			{
-				ReadOnlySpan<byte> header = span.Slice(0, HeaderLength);
-
-				int length = BinaryPrimitives.ReadInt32LittleEndian(header);
-				if (span.Length < HeaderLength + length + FooterLength)
-				{
-					break;
-				}
-
-				ReadOnlySpan<byte> cipherText = span.Slice(HeaderLength, length);
-				ReadOnlySpan<byte> tag = span.Slice(HeaderLength + length, FooterLength);
-				Span<byte> plainText = _readPipe.Writer.GetSpan(length).Slice(0, length);
-
-				_aesGcm.Decrypt(_readNonce, cipherText, tag, plainText, header);
-				IncrementNonce(_readNonce);
-
-				_readPipe.Writer.Advance(length);
-
-				span = span.Slice(HeaderLength + length + FooterLength);
+				throw new InvalidDataException($"Invalid frame length: {bufferLength}");
 			}
-			return span.Length;
-		}
-
-		/// <inheritdoc/>
-		public override async ValueTask<int> RecvAsync(Memory<byte> buffer, CancellationToken cancellationToken)
-		{
-			int sizeRead = 0;
-			while (sizeRead == 0)
+			
+			// Once we know the buffer size, read the encrypted data from underlying transport
+			byte[] encryptedBuffer = new byte[bufferLength + TagSize];
+			await _inner.RecvFullAsync(encryptedBuffer, cancellationToken);
+			
+			// Decrypt the message
+			byte[] plaintext = new byte[bufferLength];
+			try
 			{
-				// Try to get more data into the read buffer
-				if (_readBuffer.Length == 0)
-				{
-					ReadResult result = await _readPipe.Reader.ReadAsync(cancellationToken);
-					_readBuffer = result.Buffer;
-				}
-
-				// Copy as much of the next chunk as we can
-				int copySize = Math.Min(buffer.Length, _readBuffer.First.Length);
-				_readBuffer.First.Slice(0, copySize).CopyTo(buffer);
-				buffer = buffer.Slice(copySize);
-				_readBuffer = _readBuffer.Slice(copySize);
-				sizeRead += copySize;
-				Position += copySize;
+				_aes.Decrypt(
+					headerData.AsSpan(BufferLength, NonceLength),
+					encryptedBuffer.AsSpan(0, bufferLength),
+					encryptedBuffer.AsSpan(bufferLength, TagSize),
+					plaintext);
 			}
-			return sizeRead;
-		}
-
-		/// <inheritdoc/>
-		public override async ValueTask SendAsync(ReadOnlySequence<byte> buffer, CancellationToken cancellationToken)
-		{
-			while (buffer.Length > 0)
+			catch (CryptographicException)
 			{
-				ReadOnlyMemory<byte> plainText = buffer.First;
-
-				int maxSize = WritePacketSize - HeaderLength - FooterLength;
-				if (plainText.Length > maxSize)
-				{
-					plainText = plainText.Slice(0, maxSize);
-				}
-
-				Memory<byte> writeBuffer = _writeBufferOwner.Memory.Slice(_writeBufferOffset, HeaderLength + plainText.Length + FooterLength);
-				Encrypt(plainText.Span, writeBuffer.Span);
-
-				Position += plainText.Length;
-
-				await _lastWriteTask;
-				_lastWriteTask = _inner.SendAsync(writeBuffer, cancellationToken).AsTask();
-
-				buffer = buffer.Slice(plainText.Length);
-				_writeBufferOffset ^= WritePacketSize;
+				throw new InvalidDataException("Decryption failed - data may be corrupted or tampered with");
 			}
+			
+			// Copy decrypted data to the output buffer which be smaller than actual decrypted buffer
+			int bytesToCopy = Math.Min(plaintext.Length, buffer.Length);
+			plaintext.AsSpan(0, bytesToCopy).CopyTo(buffer.Span);
+			return bytesToCopy;
 		}
-
-		void Encrypt(ReadOnlySpan<byte> input, Span<byte> output)
+		finally
 		{
-			int length = input.Length;
-
-			Span<byte> header = output.Slice(0, HeaderLength);
-			BinaryPrimitives.WriteInt32LittleEndian(header, input.Length);
-
-			Span<byte> cipherText = output.Slice(HeaderLength, input.Length);
-			Span<byte> tag = output.Slice(HeaderLength + length, FooterLength);
-
-			_aesGcm.Encrypt(_writeNonce, input, cipherText, tag, header);
-			IncrementNonce(_writeNonce);
+			_recvSemaphore.Release();
 		}
-
-		static void IncrementNonce(byte[] nonce)
+	}
+	
+	/// <inheritdoc/>
+	public override ValueTask MarkCompleteAsync(CancellationToken cancellationToken)
+	{
+		return _inner.MarkCompleteAsync(cancellationToken);
+	}
+	
+	/// <inheritdoc/>
+	public override async ValueTask DisposeAsync()
+	{
+		if (!_disposed)
 		{
-			BinaryPrimitives.WriteInt64LittleEndian(nonce, BinaryPrimitives.ReadInt64LittleEndian(nonce) + 1);
+			_disposed = true;
+			_aes.Dispose();
+			_sendSemaphore.Dispose();
+			_recvSemaphore.Dispose();
+			
+			if (!_leaveInnerOpen)
+			{
+				await _inner.DisposeAsync();
+			}
 		}
-
-		/// <inheritdoc/>
-		public override ValueTask MarkCompleteAsync(CancellationToken cancellationToken) => _inner.MarkCompleteAsync(cancellationToken);
 	}
-}
+}
\ No newline at end of file
