Horde: Add missing files

diff --git a/Engine/Source/Programs/Shared/EpicGames.Horde.Tests/Compute/ComputeSocketTests.cs b/Engine/Source/Programs/Shared/EpicGames.Horde.Tests/Compute/ComputeSocketTests.cs
--- a/Engine/Source/Programs/Shared/EpicGames.Horde.Tests/Compute/ComputeSocketTests.cs
+++ b/Engine/Source/Programs/Shared/EpicGames.Horde.Tests/Compute/ComputeSocketTests.cs
@@ -118,7 +118,7 @@ namespace EpicGames.Horde.Tests.Compute
 			await RunAgentTestsAsync(localSocket, agentSocket, cts.Token);
 		}
 
-		private static async Task<(Socket client, Socket server)> CreateSocketsAsync(CancellationToken cancellationToken)
+		internal static async Task<(Socket client, Socket server)> CreateSocketsAsync(CancellationToken cancellationToken)
 		{
 			int port = GetAvailablePort();
 			using TcpListener listener = new (IPAddress.Loopback, port);
diff --git a/Engine/Source/Programs/Shared/EpicGames.Horde.Tests/Compute/TransportTests.cs b/Engine/Source/Programs/Shared/EpicGames.Horde.Tests/Compute/TransportTests.cs
--- a/Engine/Source/Programs/Shared/EpicGames.Horde.Tests/Compute/TransportTests.cs
+++ b/Engine/Source/Programs/Shared/EpicGames.Horde.Tests/Compute/TransportTests.cs
@@ -50,9 +50,6 @@ namespace EpicGames.Horde.Tests.Compute
 			byte[] key = new byte[AesTransport.KeyLength];
 			new Random(1).NextBytes(key);
 
-			byte[] nonce = new byte[AesTransport.NonceLength];
-			new Random(2).NextBytes(nonce);
-
 			byte[] input = RandomNumberGenerator.GetBytes(256 * 1024);
 			new Random(0).NextBytes(input);
 
@@ -62,7 +59,7 @@ namespace EpicGames.Horde.Tests.Compute
 				using MemoryStream memoryStream = new MemoryStream();
 				await using (StreamTransport streamTransport = new StreamTransport(memoryStream))
 				{
-					await using (AesTransport aesTransport = new AesTransport(streamTransport, key, nonce))
+					await using (AesTransport aesTransport = new AesTransport(streamTransport, key))
 					{
 						await aesTransport.SendAsync(input, CancellationToken.None);
 					}
@@ -75,7 +72,7 @@ namespace EpicGames.Horde.Tests.Compute
 				using MemoryStream memoryStream = new MemoryStream(encrypted);
 				await using (StreamTransport streamTransport = new StreamTransport(memoryStream))
 				{
-					await using (AesTransport aesTransport = new AesTransport(streamTransport, key, nonce))
+					await using (AesTransport aesTransport = new AesTransport(streamTransport, key))
 					{
 						for (int offset = 0; offset < output.Length;)
 						{
