Horde: Add a semaphore limiting number of open files on MacOS to 32, to avoid exceeding handle limits.

#rnx

diff --git a/Engine/Source/Programs/Horde/Plugins/Storage/HordeServer.Storage.Tests/Storage/StorageServiceTests.cs b/Engine/Source/Programs/Horde/Plugins/Storage/HordeServer.Storage.Tests/Storage/StorageServiceTests.cs
--- a/Engine/Source/Programs/Horde/Plugins/Storage/HordeServer.Storage.Tests/Storage/StorageServiceTests.cs
+++ b/Engine/Source/Programs/Horde/Plugins/Storage/HordeServer.Storage.Tests/Storage/StorageServiceTests.cs
@@ -2,7 +2,6 @@
 
 using EpicGames.Core;
 using EpicGames.Horde.Storage;
-using HordeServer.Server;
 using HordeServer.Storage;
 using Microsoft.Extensions.DependencyInjection;
 using MongoDB.Driver;
diff --git a/Engine/Source/Programs/Shared/EpicGames.Horde/Storage/Nodes/DirectoryNode.Extract.cs b/Engine/Source/Programs/Shared/EpicGames.Horde/Storage/Nodes/DirectoryNode.Extract.cs
--- a/Engine/Source/Programs/Shared/EpicGames.Horde/Storage/Nodes/DirectoryNode.Extract.cs
+++ b/Engine/Source/Programs/Shared/EpicGames.Horde/Storage/Nodes/DirectoryNode.Extract.cs
@@ -512,8 +512,51 @@ namespace EpicGames.Horde.Storage.Nodes
 			}
 		}
 
+		class PlatformFileLock : IDisposable
+		{
+			static readonly SemaphoreSlim? s_semaphore = CreateSemaphore();
+
+			bool _locked;
+
+			private PlatformFileLock()
+				=> _locked = true;
+
+			public static async ValueTask<PlatformFileLock?> CreateAsync(CancellationToken cancellationToken)
+			{
+				if (s_semaphore != null)
+				{
+					await s_semaphore.WaitAsync(cancellationToken);
+					return new PlatformFileLock();
+				}
+				return null;
+			}
+
+			public void Dispose()
+			{
+				if (_locked)
+				{
+					s_semaphore?.Release();
+					_locked = true;
+				}
+			}
+
+			static SemaphoreSlim? CreateSemaphore()
+			{
+				if (OperatingSystem.IsMacOS())
+				{
+					return new SemaphoreSlim(32);
+				}
+				else
+				{
+					return null;
+				}
+			}
+		}
+
 		static async Task ExtractChunksToFileAsync(OutputFile file, ArraySegment<BatchReadResponseItem<OutputChunk>> chunks, OutputStats stats, ExtractOptions options, ILogger logger, CancellationToken cancellationToken)
 		{
+			using PlatformFileLock? fileLock = await PlatformFileLock.CreateAsync(cancellationToken);
+
 			// Open the file for the current chunk
 			int remainingChunks = 0;
 			await using (FileStream stream = file.OpenStream())
