Horde: Enable parallel test execution for EpicGames.Horde.Tests

Fixed race conditions by adding unique suffixes to shared test directories.

diff --git a/Engine/Source/Programs/Shared/EpicGames.Horde.Tests/AssemblyInfo.cs b/Engine/Source/Programs/Shared/EpicGames.Horde.Tests/AssemblyInfo.cs
new file mode 100644
--- /dev/null
+++ b/Engine/Source/Programs/Shared/EpicGames.Horde.Tests/AssemblyInfo.cs
@@ -0,0 +1,5 @@
+// Copyright Epic Games, Inc. All Rights Reserved.
+
+using Microsoft.VisualStudio.TestTools.UnitTesting;
+
+[assembly: Parallelize(Workers = 0, Scope = ExecutionScope.MethodLevel)]
\ No newline at end of file
diff --git a/Engine/Source/Programs/Shared/EpicGames.Horde.Tests/Compute/ComputeSocketTests.cs b/Engine/Source/Programs/Shared/EpicGames.Horde.Tests/Compute/ComputeSocketTests.cs
--- a/Engine/Source/Programs/Shared/EpicGames.Horde.Tests/Compute/ComputeSocketTests.cs
+++ b/Engine/Source/Programs/Shared/EpicGames.Horde.Tests/Compute/ComputeSocketTests.cs
@@ -132,7 +132,7 @@ namespace EpicGames.Horde.Tests.Compute
 
 		static async Task RunAgentTestsAsync(RemoteComputeSocket localSocket, RemoteComputeSocket agentSocket, CancellationToken cancellationToken = default)
 		{
-			DirectoryReference tempDir = new DirectoryReference("test-temp");
+			DirectoryReference tempDir = new DirectoryReference("test-temp-" + DateTime.UtcNow.Ticks);
 			await using (BackgroundTask agentTask = BackgroundTask.StartNew(ctx => RunAgentAsync(agentSocket, tempDir, ctx)))
 			{
 				const int PrimaryChannelId = 0;
diff --git a/Engine/Source/Programs/Shared/EpicGames.Horde.Tests/Storage/BackendTests.cs b/Engine/Source/Programs/Shared/EpicGames.Horde.Tests/Storage/BackendTests.cs
--- a/Engine/Source/Programs/Shared/EpicGames.Horde.Tests/Storage/BackendTests.cs
+++ b/Engine/Source/Programs/Shared/EpicGames.Horde.Tests/Storage/BackendTests.cs
@@ -46,7 +46,7 @@ namespace EpicGames.Horde.Tests.Storage
 		[TestMethod]
 		public async Task TestFileBackendAsync()
 		{
-			using (TempDir tempDir = new TempDir("Cache"))
+			using (TempDir tempDir = new TempDir("Cache-" + DateTime.UtcNow.Ticks))
 			{
 				using MemoryMappedFileCache memoryMappedFileCache = new MemoryMappedFileCache();
 				FileStorageBackend backend = new FileStorageBackend(tempDir.Location, memoryMappedFileCache, NullLogger.Instance);
@@ -57,7 +57,7 @@ namespace EpicGames.Horde.Tests.Storage
 		[TestMethod]
 		public async Task TestFileBackendFileMappingAsync()
 		{
-			using (TempDir tempDir = new TempDir("Cache"))
+			using (TempDir tempDir = new TempDir("Cache-" + DateTime.UtcNow.Ticks))
 			{
 				using MemoryMappedFileCache memoryMappedFileCache = new MemoryMappedFileCache();
 				FileStorageBackend backend = new FileStorageBackend(tempDir.Location, memoryMappedFileCache, NullLogger.Instance);
@@ -81,7 +81,7 @@ namespace EpicGames.Horde.Tests.Storage
 		[TestMethod]
 		public async Task TestCacheBackendAsync()
 		{
-			using (TempDir tempDir = new TempDir("Cache"))
+			using (TempDir tempDir = new TempDir("Cache-" + DateTime.UtcNow.Ticks))
 			{
 				using StorageBackendCache cache = new StorageBackendCache(tempDir.Location, 12, NullLogger.Instance);
 
