Horde: Make port allocation for tests thread-safe

diff --git a/Engine/Source/Programs/Shared/EpicGames.Horde.Tests/Compute/ComputeSocketTests.cs b/Engine/Source/Programs/Shared/EpicGames.Horde.Tests/Compute/ComputeSocketTests.cs
--- a/Engine/Source/Programs/Shared/EpicGames.Horde.Tests/Compute/ComputeSocketTests.cs
+++ b/Engine/Source/Programs/Shared/EpicGames.Horde.Tests/Compute/ComputeSocketTests.cs
@@ -221,13 +221,33 @@ namespace EpicGames.Horde.Tests.Compute
 			await agentSocket.CloseAsync(CancellationToken.None);
 		}
 
-		static int GetAvailablePort()
+		private static readonly HashSet<int> s_usedPorts = [];
+		private static int GetAvailablePort()
 		{
-			using TcpListener listener = new(IPAddress.Loopback, 0);
-			listener.Start();
-			int port = ((IPEndPoint)listener.LocalEndpoint).Port;
-			listener.Stop();
-			return port;
+			lock (s_usedPorts)
+			{
+				for (int i = 0; i < 10; i++)
+				{
+					using TcpListener listener = new(IPAddress.Loopback, 0);
+					try
+					{
+						listener.Start();
+						int port = ((IPEndPoint)listener.LocalEndpoint).Port;
+						if (!s_usedPorts.Add(port))
+						{
+							continue;
+						}
+						
+						return port;
+					}
+					finally
+					{
+						listener.Stop();
+					}
+				}
+				
+				throw new InvalidOperationException("Unable to acquire a locally available IP port");
+			}
 		}
 
 		static async Task RunAgentAsync(ComputeSocket socket, DirectoryReference tempDir, CancellationToken cancellationToken)
