[EpicGames.UBA]
* Added code to expose vfs to c#

diff --git a/Engine/Source/Programs/Shared/EpicGames.UBA/IProcess.cs b/Engine/Source/Programs/Shared/EpicGames.UBA/IProcess.cs
--- a/Engine/Source/Programs/Shared/EpicGames.UBA/IProcess.cs
+++ b/Engine/Source/Programs/Shared/EpicGames.UBA/IProcess.cs
@@ -77,6 +77,11 @@ namespace EpicGames.UBA
 		/// Arbitary user data to pass along with the process
 		/// </summary>
 		public object? UserData { get; set; }
+
+		/// <summary>
+		/// RootsHandle for process used to convert paths
+		/// </summary>
+		public ulong RootsHandle { get; set; }
 	}
 
 	/// <summary>
diff --git a/Engine/Source/Programs/Shared/EpicGames.UBA/ISessionServer.cs b/Engine/Source/Programs/Shared/EpicGames.UBA/ISessionServer.cs
--- a/Engine/Source/Programs/Shared/EpicGames.UBA/ISessionServer.cs
+++ b/Engine/Source/Programs/Shared/EpicGames.UBA/ISessionServer.cs
@@ -178,6 +178,14 @@ namespace EpicGames.UBA
 		/// <returns>The remote process being run</returns>
 		public abstract IProcess RunProcessRemote(ProcessStartInfo info, IProcess.ExitedEventHandler? exitedEventHandler, double weight = 1.0, byte[]? knownInputs = null, uint knownInputsCount = 0, bool canExecuteCrossArchitecture = false);
 
+		/// <summary>
+		/// Register roots data and get handle back
+		/// </summary>
+		/// <param name="rootsData">Buffer containing roots used for cache and vfs. Format is multiples of byte/utf8string/utf8string. rootindex/vfspath/localpath</param>
+		/// <param name="rootsDataSize">Size of buffer in bytes</param>
+		/// <returns>The remote process being run</returns>
+		public abstract ulong RegisterRoots(byte[] rootsData, uint rootsDataSize);
+
 		/// <summary>
 		/// Refresh cached information about directories
 		/// </summary>
diff --git a/Engine/Source/Programs/Shared/EpicGames.UBA/Impl/ProcessImpl.cs b/Engine/Source/Programs/Shared/EpicGames.UBA/Impl/ProcessImpl.cs
--- a/Engine/Source/Programs/Shared/EpicGames.UBA/Impl/ProcessImpl.cs
+++ b/Engine/Source/Programs/Shared/EpicGames.UBA/Impl/ProcessImpl.cs
@@ -16,7 +16,7 @@ namespace EpicGames.UBA.Impl
 
 		#region DllImport
 		[DllImport("UbaHost", CharSet = CharSet.Auto)]
-		static extern nint ProcessStartInfo_Create(string application, string arguments, string workingDir, string description, uint priorityClass, ulong outputStatsThresholdMs, bool trackInputs, string logFile, ExitCallback? exit);
+		static extern nint ProcessStartInfo_Create2(string application, string arguments, string workingDir, string description, uint priorityClass, ulong rootsHandle, bool trackInputs, string logFile, ExitCallback? exit);
 
 		[DllImport("UbaHost", CharSet = CharSet.Auto)]
 		static extern void ProcessStartInfo_Destroy(nint server);
@@ -29,7 +29,7 @@ namespace EpicGames.UBA.Impl
 				ExitCallbackDelegate = RaiseExited;
 			}
 
-			_handle = ProcessStartInfo_Create(info.Application, info.Arguments, info.WorkingDirectory, info.Description, (uint)info.Priority, info.OutputStatsThresholdMs, info.TrackInputs, info.LogFile ?? String.Empty, ExitCallbackDelegate);
+			_handle = ProcessStartInfo_Create2(info.Application, info.Arguments, info.WorkingDirectory, info.Description, (uint)info.Priority, info.RootsHandle, info.TrackInputs, info.LogFile ?? String.Empty, ExitCallbackDelegate);
 		}
 
 		void RaiseExited(nint userData, nint handle)
diff --git a/Engine/Source/Programs/Shared/EpicGames.UBA/Impl/SessionServerImpl.cs b/Engine/Source/Programs/Shared/EpicGames.UBA/Impl/SessionServerImpl.cs
--- a/Engine/Source/Programs/Shared/EpicGames.UBA/Impl/SessionServerImpl.cs
+++ b/Engine/Source/Programs/Shared/EpicGames.UBA/Impl/SessionServerImpl.cs
@@ -109,6 +109,9 @@ namespace EpicGames.UBA.Impl
 		[DllImport("UbaHost", CharSet = CharSet.Auto)]
 		static extern nint SessionServer_RunProcessRemote(nint server, nint info, float weight, byte[]? knownInputs, uint knownInputsCount, bool allowCrossArchitecture);
 
+		[DllImport("UbaHost", CharSet = CharSet.Auto)]
+		static extern ulong SessionServer_RegisterRoots(nint server, byte[] rootsData, uint rootsDataSize);
+
 		[DllImport("UbaHost", CharSet = CharSet.Auto)]
 		static extern void SessionServer_SetMaxRemoteProcessCount(nint server, uint count);
 
@@ -209,6 +212,8 @@ namespace EpicGames.UBA.Impl
 			return process;
 		}
 
+		public ulong RegisterRoots(byte[] rootsData, uint rootsDataSize) => SessionServer_RegisterRoots(_handle, rootsData, rootsDataSize);
+
 		public void DisableRemoteExecution() => SessionServer_DisableRemoteExecution(_handle);
 
 		public void SetMaxRemoteProcessCount(uint count) => SessionServer_SetMaxRemoteProcessCount(_handle, count);
