Horde: Allow querying multiple artifacts at once.

#rnx

diff --git a/Engine/Source/Programs/Horde/Plugins/Build/HordeServer.Build/Artifacts/ArtifactCollection.cs b/Engine/Source/Programs/Horde/Plugins/Build/HordeServer.Build/Artifacts/ArtifactCollection.cs
--- a/Engine/Source/Programs/Horde/Plugins/Build/HordeServer.Build/Artifacts/ArtifactCollection.cs
+++ b/Engine/Source/Programs/Horde/Plugins/Build/HordeServer.Build/Artifacts/ArtifactCollection.cs
@@ -296,9 +296,13 @@ namespace HordeServer.Artifacts
 			await _artifactExpiryCollection.UpdateOneAsync(streamFilter, streamUpdate, new UpdateOptions { IsUpsert = true }, cancellationToken);
 		}
 		/// <inheritdoc/>
-		public async IAsyncEnumerable<IArtifact> FindAsync(StreamId? streamId = null, CommitId? minCommitId = null, CommitId? maxCommitId = null, ArtifactName? name = null, ArtifactType? type = null, IEnumerable<string>? keys = null, int maxResults = 100, [EnumeratorCancellation] CancellationToken cancellationToken = default)
+		public async IAsyncEnumerable<IArtifact> FindAsync(ArtifactId[]? ids = null, StreamId? streamId = null, CommitId? minCommitId = null, CommitId? maxCommitId = null, ArtifactName? name = null, ArtifactType? type = null, IEnumerable<string>? keys = null, int maxResults = 100, [EnumeratorCancellation] CancellationToken cancellationToken = default)
 		{
 			FilterDefinition<ArtifactDocument> filter = FilterDefinition<ArtifactDocument>.Empty;
+			if (ids != null && ids.Length > 0)
+			{
+				filter &= Builders<ArtifactDocument>.Filter.In(x => x.Id, ids.ConvertAll(x => BinaryIdUtils.ToObjectId(x.Id)));
+			}
 			if (streamId != null)
 			{
 				filter &= Builders<ArtifactDocument>.Filter.Eq(x => x.StreamId, streamId.Value);
diff --git a/Engine/Source/Programs/Horde/Plugins/Build/HordeServer.Build/Artifacts/ArtifactsController.cs b/Engine/Source/Programs/Horde/Plugins/Build/HordeServer.Build/Artifacts/ArtifactsController.cs
--- a/Engine/Source/Programs/Horde/Plugins/Build/HordeServer.Build/Artifacts/ArtifactsController.cs
+++ b/Engine/Source/Programs/Horde/Plugins/Build/HordeServer.Build/Artifacts/ArtifactsController.cs
@@ -159,7 +159,7 @@ namespace HordeServer.Artifacts
 		async Task<RefName?> GetPrevRefNameForArtifactAsync(IArtifactBuilder artifact, CancellationToken cancellationToken)
 		{
 			IStorageBackend storageBackend = _storageService.CreateBackend(artifact.NamespaceId);
-			await foreach (IArtifact prevArtifact in _artifactCollection.FindAsync(artifact.StreamId, maxCommitId: artifact.CommitId, name: artifact.Name, type: artifact.Type, cancellationToken: cancellationToken))
+			await foreach (IArtifact prevArtifact in _artifactCollection.FindAsync(streamId: artifact.StreamId, maxCommitId: artifact.CommitId, name: artifact.Name, type: artifact.Type, cancellationToken: cancellationToken))
 			{
 				if (prevArtifact.CommitId == artifact.CommitId)
 				{
@@ -183,6 +183,7 @@ namespace HordeServer.Artifacts
 		/// <summary>
 		/// Finds artifacts matching certain criteria
 		/// </summary>
+		/// <param name="ids">Artifact identifiers to filter</param>
 		/// <param name="streamId">Stream to search</param>
 		/// <param name="minChange">Minimum changelist number for artifacts to return</param>
 		/// <param name="maxChange">Maximum changelist number for artifacts to return</param>
@@ -195,10 +196,10 @@ namespace HordeServer.Artifacts
 		[HttpGet]
 		[Route("/api/v2/artifacts")]
 		[ProducesResponseType(typeof(FindArtifactsResponse), 200)]
-		public async Task<ActionResult<object>> FindArtifactsAsync([FromQuery] StreamId? streamId = null, [FromQuery] CommitId? minChange = null, [FromQuery] CommitId? maxChange = null, [FromQuery(Name = "name")] ArtifactName? name = null, [FromQuery(Name = "type")] ArtifactType? type = null, [FromQuery(Name = "key")] IEnumerable<string>? keys = null, [FromQuery] int maxResults = 100, [FromQuery] PropertyFilter? filter = null)
+		public async Task<ActionResult<object>> FindArtifactsAsync([FromQuery(Name = "id")] ArtifactId[]? ids = null, [FromQuery] StreamId? streamId = null, [FromQuery] CommitId? minChange = null, [FromQuery] CommitId? maxChange = null, [FromQuery(Name = "name")] ArtifactName? name = null, [FromQuery(Name = "type")] ArtifactType? type = null, [FromQuery(Name = "key")] IEnumerable<string>? keys = null, [FromQuery] int maxResults = 100, [FromQuery] PropertyFilter? filter = null)
 		{
 			FindArtifactsResponse response = new FindArtifactsResponse();
-			await foreach (IArtifact artifact in _artifactCollection.FindAsync(streamId, minChange, maxChange, name, type, keys, maxResults, HttpContext.RequestAborted))
+			await foreach (IArtifact artifact in _artifactCollection.FindAsync(ids, streamId, minChange, maxChange, name, type, keys, maxResults, HttpContext.RequestAborted))
 			{
 				if (_buildConfig.AuthorizeArtifact(artifact.Type, artifact.StreamId, ArtifactAclAction.ReadArtifact, User))
 				{
diff --git a/Engine/Source/Programs/Shared/EpicGames.Horde/Artifacts/HttpArtifactCollection.cs b/Engine/Source/Programs/Shared/EpicGames.Horde/Artifacts/HttpArtifactCollection.cs
--- a/Engine/Source/Programs/Shared/EpicGames.Horde/Artifacts/HttpArtifactCollection.cs
+++ b/Engine/Source/Programs/Shared/EpicGames.Horde/Artifacts/HttpArtifactCollection.cs
@@ -135,11 +135,11 @@ namespace EpicGames.Horde.Artifacts
 		}
 
 		/// <inheritdoc/>
-		public async IAsyncEnumerable<IArtifact> FindAsync(StreamId? streamId = null, CommitId? minCommitId = null, CommitId? maxCommitId = null, ArtifactName? name = null, ArtifactType? type = null, IEnumerable<string>? keys = null, int maxResults = 100, [EnumeratorCancellation] CancellationToken cancellationToken = default)
+		public async IAsyncEnumerable<IArtifact> FindAsync(ArtifactId[]? ids = null, StreamId? streamId = null, CommitId? minCommitId = null, CommitId? maxCommitId = null, ArtifactName? name = null, ArtifactType? type = null, IEnumerable<string>? keys = null, int maxResults = 100, [EnumeratorCancellation] CancellationToken cancellationToken = default)
 		{
 			HordeHttpClient hordeHttpClient = _hordeClient.CreateHttpClient();
 
-			List<GetArtifactResponse> responses = await hordeHttpClient.FindArtifactsAsync(streamId, minCommitId, maxCommitId, name, type, keys, maxResults, cancellationToken);
+			List<GetArtifactResponse> responses = await hordeHttpClient.FindArtifactsAsync(ids, streamId, minCommitId, maxCommitId, name, type, keys, maxResults, cancellationToken);
 			foreach (GetArtifactResponse response in responses)
 			{
 				yield return new Artifact(this, response);
diff --git a/Engine/Source/Programs/Shared/EpicGames.Horde/Artifacts/IArtifactCollection.cs b/Engine/Source/Programs/Shared/EpicGames.Horde/Artifacts/IArtifactCollection.cs
--- a/Engine/Source/Programs/Shared/EpicGames.Horde/Artifacts/IArtifactCollection.cs
+++ b/Engine/Source/Programs/Shared/EpicGames.Horde/Artifacts/IArtifactCollection.cs
@@ -30,6 +30,7 @@ namespace EpicGames.Horde.Artifacts
 		/// <summary>
 		/// Finds artifacts with the given keys.
 		/// </summary>
+		/// <param name="ids">Identifiers to return</param>
 		/// <param name="streamId">Stream to find artifacts for</param>
 		/// <param name="minCommitId">Minimum commit for the artifacts (inclusive)</param>
 		/// <param name="maxCommitId">Maximum commit for the artifacts (inclusive)</param>
@@ -39,7 +40,7 @@ namespace EpicGames.Horde.Artifacts
 		/// <param name="maxResults">Maximum number of results to return</param>
 		/// <param name="cancellationToken">Cancellation token for the operation</param>
 		/// <returns>Sequence of artifacts. Ordered by descending CL order, then by descending order in which they were created.</returns>
-		IAsyncEnumerable<IArtifact> FindAsync(StreamId? streamId = null, CommitId? minCommitId = null, CommitId? maxCommitId = null, ArtifactName? name = null, ArtifactType? type = null, IEnumerable<string>? keys = null, int maxResults = 100, CancellationToken cancellationToken = default);
+		IAsyncEnumerable<IArtifact> FindAsync(ArtifactId[]? ids = null, StreamId? streamId = null, CommitId? minCommitId = null, CommitId? maxCommitId = null, ArtifactName? name = null, ArtifactType? type = null, IEnumerable<string>? keys = null, int maxResults = 100, CancellationToken cancellationToken = default);
 
 		/// <summary>
 		/// Gets an artifact by ID
@@ -49,4 +50,26 @@ namespace EpicGames.Horde.Artifacts
 		/// <returns>The artifact document</returns>
 		Task<IArtifact?> GetAsync(ArtifactId artifactId, CancellationToken cancellationToken = default);
 	}
+
+	/// <summary>
+	/// Extension methods for artifacts
+	/// </summary>
+	public static class ArtifactCollectionExtensions
+	{
+		/// <summary>
+		/// Finds artifacts with the given keys.
+		/// </summary>
+		/// <param name="artifactCollection">Collection to operate on</param>
+		/// <param name="streamId">Stream to find artifacts for</param>
+		/// <param name="minCommitId">Minimum commit for the artifacts (inclusive)</param>
+		/// <param name="maxCommitId">Maximum commit for the artifacts (inclusive)</param>
+		/// <param name="name">Name of the artifact to search for</param>
+		/// <param name="type">The artifact type</param>
+		/// <param name="keys">Set of keys, all of which must all be present on any returned artifacts</param>
+		/// <param name="maxResults">Maximum number of results to return</param>
+		/// <param name="cancellationToken">Cancellation token for the operation</param>
+		/// <returns>Sequence of artifacts. Ordered by descending CL order, then by descending order in which they were created.</returns>
+		public static IAsyncEnumerable<IArtifact> FindAsync(this IArtifactCollection artifactCollection, StreamId? streamId = null, CommitId? minCommitId = null, CommitId? maxCommitId = null, ArtifactName? name = null, ArtifactType? type = null, IEnumerable<string>? keys = null, int maxResults = 100, CancellationToken cancellationToken = default)
+			=> artifactCollection.FindAsync(null, streamId, minCommitId, maxCommitId, name, type, keys, maxResults, cancellationToken);
+	}
 }
diff --git a/Engine/Source/Programs/Shared/EpicGames.Horde/HordeHttpClient.cs b/Engine/Source/Programs/Shared/EpicGames.Horde/HordeHttpClient.cs
--- a/Engine/Source/Programs/Shared/EpicGames.Horde/HordeHttpClient.cs
+++ b/Engine/Source/Programs/Shared/EpicGames.Horde/HordeHttpClient.cs
@@ -174,10 +174,31 @@ namespace EpicGames.Horde
 		/// <param name="maxResults">Maximum number of results to return</param>
 		/// <param name="cancellationToken">Cancellation token for the operation</param>
 		/// <returns>Information about all the artifacts</returns>
-		public async Task<List<GetArtifactResponse>> FindArtifactsAsync(StreamId? streamId = null, CommitId? minCommitId = null, CommitId? maxCommitId = null, ArtifactName? name = null, ArtifactType? type = null, IEnumerable<string>? keys = null, int maxResults = 100, CancellationToken cancellationToken = default)
+		public Task<List<GetArtifactResponse>> FindArtifactsAsync(StreamId? streamId = null, CommitId? minCommitId = null, CommitId? maxCommitId = null, ArtifactName? name = null, ArtifactType? type = null, IEnumerable<string>? keys = null, int maxResults = 100, CancellationToken cancellationToken = default)
+			=> FindArtifactsAsync(null, streamId, minCommitId, maxCommitId, name, type, keys, maxResults, cancellationToken);
+
+		/// <summary>
+		/// Finds artifacts with a certain type with an optional streamId
+		/// </summary>
+		/// <param name="ids">Identifiers to return</param>
+		/// <param name="streamId">Stream to look for the artifact in</param>
+		/// <param name="minCommitId">The minimum change number for the artifacts</param>
+		/// <param name="maxCommitId">The minimum change number for the artifacts</param>
+		/// <param name="name">Name of the artifact</param>
+		/// <param name="type">Type to find</param>
+		/// <param name="keys">Keys for artifacts to return</param>
+		/// <param name="maxResults">Maximum number of results to return</param>
+		/// <param name="cancellationToken">Cancellation token for the operation</param>
+		/// <returns>Information about all the artifacts</returns>
+		public async Task<List<GetArtifactResponse>> FindArtifactsAsync(IEnumerable<ArtifactId>? ids, StreamId? streamId = null, CommitId? minCommitId = null, CommitId? maxCommitId = null, ArtifactName? name = null, ArtifactType? type = null, IEnumerable<string>? keys = null, int maxResults = 100, CancellationToken cancellationToken = default)
 		{
 			QueryStringBuilder queryParams = new QueryStringBuilder();
 
+			if (ids != null)
+			{
+				queryParams.Add("id", ids.Select(x => x.ToString()));
+			}
+
 			if (streamId != null)
 			{
 				queryParams.Add("streamId", streamId.ToString()!);
