Horde: Set maximum size of channels used for buffering data read from storage. If part of this pipeline is particularly slow (eg. I/O, network, CPU), is can cause large backlogs of buffered data which can kill performance or cause machine instability. Applying some back pressure resolves the issue without any loss in performance.

diff --git a/Engine/Source/Programs/Shared/EpicGames.Horde/Storage/BatchReader.cs b/Engine/Source/Programs/Shared/EpicGames.Horde/Storage/BatchReader.cs
--- a/Engine/Source/Programs/Shared/EpicGames.Horde/Storage/BatchReader.cs
+++ b/Engine/Source/Programs/Shared/EpicGames.Horde/Storage/BatchReader.cs
@@ -127,12 +127,12 @@ namespace EpicGames.Horde.Storage
 		{
 			_verifyHashes = verifyHashes;
 
-			_blobRequests = Channel.CreateUnbounded<BatchReadRequest<TContext>>();
+			_blobRequests = Channel.CreateBounded<BatchReadRequest<TContext>>(new BoundedChannelOptions(MinQueueLength * 2));
 			_otherRequests = Channel.CreateUnbounded<BatchReadRequest<TContext>>();
 			_bundleRequests = Channel.CreateUnbounded<BundleRequest>();
-			_bundleResponses = Channel.CreateUnbounded<BundleResponse>();
-			_packetResponses = Channel.CreateUnbounded<PacketResponse>();
-			_blobResponses = Channel.CreateUnbounded<BatchReadResponse<TContext>>();
+			_bundleResponses = Channel.CreateBounded<BundleResponse>(new BoundedChannelOptions(32));
+			_packetResponses = Channel.CreateBounded<PacketResponse>(new BoundedChannelOptions(128));
+			_blobResponses = Channel.CreateBounded<BatchReadResponse<TContext>>(new BoundedChannelOptions(128));
 		}
 
 		/// <inheritdoc/>
@@ -191,7 +191,7 @@ namespace EpicGames.Horde.Storage
 					blobData = await request.BlobRef.ReadBlobDataAsync(cancellationToken);
 					VerifyHash(request.BlobRef, blobData.Data.Span, request.BlobRef.Hash);
 #pragma warning disable CA2000
-					BatchReadResponse < TContext> response = new BatchReadResponse<TContext>([new BatchReadResponseItem<TContext>(blobData, request.Context)]);
+					BatchReadResponse<TContext> response = new BatchReadResponse<TContext>([new BatchReadResponseItem<TContext>(blobData, request.Context)]);
 					await _blobResponses.Writer.WriteAsync(response, cancellationToken);
 #pragma warning restore CA2000
 				}
