Horde: Change back to 16 read tasks for batch reads from storage.

diff --git a/Engine/Source/Programs/Shared/EpicGames.Horde/Storage/Nodes/DirectoryNode.Extract.cs b/Engine/Source/Programs/Shared/EpicGames.Horde/Storage/Nodes/DirectoryNode.Extract.cs
--- a/Engine/Source/Programs/Shared/EpicGames.Horde/Storage/Nodes/DirectoryNode.Extract.cs
+++ b/Engine/Source/Programs/Shared/EpicGames.Horde/Storage/Nodes/DirectoryNode.Extract.cs
@@ -353,7 +353,7 @@ namespace EpicGames.Horde.Storage.Nodes
 		{
 			options ??= new ExtractOptions();
 
-			int numReadTasks = options.NumReadTasks ?? 8; //Slowed the read rate to avoid out of memory issues
+			int numReadTasks = options.NumReadTasks ?? 16;
 			int numDecodeTasks = options.NumDecodeTasks ?? Math.Min(Environment.ProcessorCount, 16);
 			int numWriteTasks = options.NumWriteTasks ?? Math.Min(1 + (int)(directoryNode.Length / (16 * 1024 * 1024)), 16);
 			logger.LogDebug("Using {NumReadTasks} read tasks, {NumDecodeTasks} decode tasks, {NumWriteTasks} write tasks", numReadTasks, numDecodeTasks, numWriteTasks);
