Horde: Additional logging for errors updating logs.

#rnx

diff --git a/Engine/Source/Programs/Shared/EpicGames.Horde/Logs/ServerLogger.cs b/Engine/Source/Programs/Shared/EpicGames.Horde/Logs/ServerLogger.cs
--- a/Engine/Source/Programs/Shared/EpicGames.Horde/Logs/ServerLogger.cs
+++ b/Engine/Source/Programs/Shared/EpicGames.Horde/Logs/ServerLogger.cs
@@ -411,9 +411,14 @@ namespace EpicGames.Horde.Logs
 			}
 			catch (RpcException ex) when (ex.StatusCode == StatusCode.DeadlineExceeded)
 			{
-				_internalLogger.LogDebug(ex, "Log tail deadline exceeded, ignoring.");
+				_internalLogger.LogInformation(ex, "Log tail deadline exceeded, ignoring.");
 				return -1;
 			}
+			catch (Exception ex)
+			{
+				_internalLogger.LogWarning(ex, "Exception while updating log: {Message}", ex.Message);
+				throw;
+			}
 		}
 
 		#endregion
