Horde: Forward completion event to inner in IdleTimeoutTransport

diff --git a/Engine/Source/Programs/Shared/EpicGames.Horde/Compute/Transports/IdleTimeoutTransport.cs b/Engine/Source/Programs/Shared/EpicGames.Horde/Compute/Transports/IdleTimeoutTransport.cs
--- a/Engine/Source/Programs/Shared/EpicGames.Horde/Compute/Transports/IdleTimeoutTransport.cs
+++ b/Engine/Source/Programs/Shared/EpicGames.Horde/Compute/Transports/IdleTimeoutTransport.cs
@@ -67,7 +67,10 @@ public class IdleTimeoutTransport : ComputeTransport
 	}
 
 	/// <inheritdoc/>
-	public override ValueTask MarkCompleteAsync(CancellationToken cancellationToken) => new ValueTask();
+	public override ValueTask MarkCompleteAsync(CancellationToken cancellationToken)
+	{
+		return _inner.MarkCompleteAsync(cancellationToken);
+	}
 
 	/// <summary>
 	/// Start a loop monitoring activity on the inner transport
