Horde: Handle EndOfStreamException in RecvOptionalAsync

Treat EndOfStreamException same as zero bytes read, returning false to indicate no data available rather than propagating the exception.

diff --git a/Engine/Source/Programs/Shared/EpicGames.Horde/Compute/ComputeTransport.cs b/Engine/Source/Programs/Shared/EpicGames.Horde/Compute/ComputeTransport.cs
--- a/Engine/Source/Programs/Shared/EpicGames.Horde/Compute/ComputeTransport.cs
+++ b/Engine/Source/Programs/Shared/EpicGames.Horde/Compute/ComputeTransport.cs
@@ -62,16 +62,23 @@ namespace EpicGames.Horde.Compute
 		/// <param name="cancellationToken">Cancellation token for the operation</param>
 		public async ValueTask<bool> RecvOptionalAsync(Memory<byte> buffer, CancellationToken cancellationToken)
 		{
-			int read = await RecvAsync(buffer, cancellationToken);
-			if (read == 0)
+			try
 			{
-				return false;
+				int read = await RecvAsync(buffer, cancellationToken);
+				if (read == 0)
+				{
+					return false;
+				}
+				if (read < buffer.Length)
+				{
+					await RecvFullAsync(buffer.Slice(read), cancellationToken);
+				}
+				return true;
 			}
-			if (read < buffer.Length)
+			catch (EndOfStreamException)
 			{
-				await RecvFullAsync(buffer.Slice(read), cancellationToken);
+				return false;
 			}
-			return true;
 		}
 
 		/// <summary>
