Horde: Add buffering in AesTransport for oversized decrypted messages

Store decrypted data exceeding the target buffer size for subsequent RecvAsync calls. Add torture test for validating various send/receive chunk sizes.

diff --git a/Engine/Source/Programs/Shared/EpicGames.Horde.Tests/AesTransportTests.cs b/Engine/Source/Programs/Shared/EpicGames.Horde.Tests/AesTransportTests.cs
--- a/Engine/Source/Programs/Shared/EpicGames.Horde.Tests/AesTransportTests.cs
+++ b/Engine/Source/Programs/Shared/EpicGames.Horde.Tests/AesTransportTests.cs
@@ -2,7 +2,10 @@
 
 using System;
 using System.Buffers;
+using System.Collections.Generic;
+using System.IO;
 using System.IO.Pipelines;
+using System.Linq;
 using System.Net.Sockets;
 using System.Threading;
 using System.Threading.Tasks;
@@ -16,11 +19,13 @@ namespace EpicGames.Horde.Tests;
 [TestClass]
 public sealed class AesTransportTests : IAsyncDisposable
 {
+	private const int SmallBufferThreshold = 1000;
+	private const int MinChunkSizeForLargeBuffers = 10;
 	private readonly byte[] _key;
 	private readonly MemoryComputeTransport _memoryTransport;
 	private readonly AesTransport _aesTransport;
 	private readonly Random _random = new(42);
-	private readonly int[] _bufferSizes = [1, 2, 3, 5, 10, 20, 2000, 20000, 2000000];
+	private readonly int[] _bufferSizes = [1, 2, 3, 10, 200000, 2000000];
 	
 	public AesTransportTests()
 	{
@@ -65,6 +70,32 @@ public sealed class AesTransportTests : IAsyncDisposable
 		}
 	}
 	
+	[TestMethod]
+	public async Task SendAndReceive_Torture_SuccessfullyTransferredAsync()
+	{
+		using CancellationTokenSource cts = new(20000);
+		await using MemoryComputeTransport memTransport = new();
+		await using AesTransport clientAes = new(memTransport, _key);
+		await using AesTransport serverAes = new(memTransport, _key);
+		
+		foreach (int bufferSize in _bufferSizes)
+		{
+			List<int> evilChunkSizes = [1, 2, 3, bufferSize - 1, bufferSize, bufferSize + 1, bufferSize * 5];
+			evilChunkSizes =  evilChunkSizes
+				.Where(size => size > 0)
+				.Where(size => bufferSize < SmallBufferThreshold || size >= MinChunkSizeForLargeBuffers)
+				.ToList();
+			
+			foreach (int sendChunk in evilChunkSizes)
+			{
+				foreach (int recvChunk in evilChunkSizes)
+				{
+					await SendAndMultipleReceiveAsync(clientAes, serverAes, bufferSize, () => sendChunk, () => recvChunk, cts.Token);
+				}
+			}
+		}
+	}
+	
 	private async Task SendAndReceiveAsync(AesTransport client, AesTransport server, int bufferSize, CancellationToken cancellationToken)
 	{
 		// Arrange
@@ -86,6 +117,45 @@ public sealed class AesTransportTests : IAsyncDisposable
 		CollectionAssert.AreEqual(sentData, clientReceivedData);
 	}
 	
+	private async Task SendAndMultipleReceiveAsync(ComputeTransport client, ComputeTransport server, int bufferSize,
+		Func<int> getSendChunkSize, Func<int> getRecvChunkSize, CancellationToken cancellationToken)
+	{
+		// Arrange
+		byte[] testData = new byte[bufferSize];
+		_random.NextBytes(testData);
+		using MemoryStream received = new(testData.Length);
+		int bytesSent = 0;
+		
+		
+		// Act
+		Task sendTask = Task.Run(async () =>
+		{
+			while (bytesSent < testData.Length)
+			{
+				int bytesLeft = testData.Length - bytesSent;
+				int bytesToSend = Math.Min(getSendChunkSize(), bytesLeft);
+				await client.SendAsync(testData.AsMemory(bytesSent, bytesToSend), cancellationToken);
+				bytesSent += bytesToSend;
+			}
+		}, cancellationToken);
+		
+		Task recvTask = Task.Run(async () =>
+		{
+			while (received.Length < testData.Length)
+			{
+				byte[] chunk = new byte[getRecvChunkSize()];
+				int bytesRead = await server.RecvAsync(chunk, cancellationToken);
+				await received.WriteAsync(chunk, 0, bytesRead, cancellationToken);
+			}
+		}, cancellationToken);
+		
+		await Task.WhenAll(sendTask, recvTask);
+		
+		// Assert
+		Assert.AreEqual(testData.Length, received.Length);
+		CollectionAssert.AreEqual(testData, received.ToArray());
+	}
+	
 	[TestMethod]
 	public async Task Constructor_InvalidKeyLength_ThrowsArgumentExceptionAsync()
 	{
diff --git a/Engine/Source/Programs/Shared/EpicGames.Horde.Tests/Compute/ComputeSocketTests.cs b/Engine/Source/Programs/Shared/EpicGames.Horde.Tests/Compute/ComputeSocketTests.cs
--- a/Engine/Source/Programs/Shared/EpicGames.Horde.Tests/Compute/ComputeSocketTests.cs
+++ b/Engine/Source/Programs/Shared/EpicGames.Horde.Tests/Compute/ComputeSocketTests.cs
@@ -55,7 +55,7 @@ namespace EpicGames.Horde.Tests.Compute
 			}
 		}
 
-		class TestLogger : ILogger
+		class TestLogger(string prefix) : ILogger
 		{
 			public IDisposable? BeginScope<TState>(TState state) where TState : notnull => null!;
 
@@ -63,7 +63,7 @@ namespace EpicGames.Horde.Tests.Compute
 
 			public void Log<TState>(LogLevel logLevel, EventId eventId, TState state, Exception? exception, Func<TState, Exception?, string> formatter)
 			{
-				Console.WriteLine($"{logLevel}: {formatter(state, exception)}");
+				Console.WriteLine($"{prefix} {logLevel}: {formatter(state, exception)}");
 				Assert.IsFalse(logLevel == LogLevel.Error);
 				Assert.IsFalse(logLevel == LogLevel.Warning);
 			}
@@ -76,8 +76,8 @@ namespace EpicGames.Horde.Tests.Compute
 			Pipe sendPipe = new Pipe();
 			await using PipeTransport localTransport = new(sendPipe.Reader, recvPipe.Writer);
 			await using PipeTransport agentTransport = new(recvPipe.Reader, sendPipe.Writer);
-			await using RemoteComputeSocket localSocket = new(localTransport, ComputeProtocol.Latest, new TestLogger());
-			await using RemoteComputeSocket agentSocket = new(agentTransport, ComputeProtocol.Latest, new TestLogger());
+			await using RemoteComputeSocket localSocket = new(localTransport, ComputeProtocol.Latest, new TestLogger("local"));
+			await using RemoteComputeSocket agentSocket = new(agentTransport, ComputeProtocol.Latest, new TestLogger("agent"));
 
 			await RunAgentTestsAsync(localSocket, agentSocket);
 		}
@@ -90,8 +90,8 @@ namespace EpicGames.Horde.Tests.Compute
 
 			await using TcpTransport clientTransport = new(clientSocket);
 			await using TcpTransport serverTransport = new(serverSocket);
-			await using RemoteComputeSocket localSocket = new(clientTransport, ComputeProtocol.Latest, new TestLogger());
-			await using RemoteComputeSocket agentSocket = new(serverTransport, ComputeProtocol.Latest, new TestLogger());
+			await using RemoteComputeSocket localSocket = new(clientTransport, ComputeProtocol.Latest, new TestLogger("local"));
+			await using RemoteComputeSocket agentSocket = new(serverTransport, ComputeProtocol.Latest, new TestLogger("agent"));
 
 			await RunAgentTestsAsync(localSocket, agentSocket, cts.Token);
 		}
@@ -113,10 +113,30 @@ namespace EpicGames.Horde.Tests.Compute
 			await t2;
 			await t1;
 
-			await using RemoteComputeSocket localSocket = new(clientTransport, ComputeProtocol.Latest, new TestLogger());
-			await using RemoteComputeSocket agentSocket = new(serverTransport, ComputeProtocol.Latest, new TestLogger());
+			await using RemoteComputeSocket localSocket = new(clientTransport, ComputeProtocol.Latest, new TestLogger("local"));
+			await using RemoteComputeSocket agentSocket = new(serverTransport, ComputeProtocol.Latest, new TestLogger("agent"));
 			await RunAgentTestsAsync(localSocket, agentSocket, cts.Token);
 		}
+		
+		[TestMethod]
+		public async Task TestAgentMessageLoopTcpAesAsync()
+		{
+			using CancellationTokenSource cts = new(5000);
+			(Socket localSocket, Socket agentSocket) = await CreateSocketsAsync(cts.Token);
+			
+			byte[] key = AesTransport.CreateKey();
+			await using TcpTransport localTcp = new(localSocket);
+			await using TcpTransport agentTcp = new(agentSocket);
+			await using AesTransport localAes = new(localTcp, key);
+			await using AesTransport agentAes = new(agentTcp, key);
+			
+			await using IdleTimeoutTransport localIdleTimeout = new(localAes, TimeSpan.FromSeconds(5));
+			await using IdleTimeoutTransport agentIdleTimeout = new(agentAes, TimeSpan.FromSeconds(5));
+			await using RemoteComputeSocket localComputeSocket = new(localIdleTimeout, ComputeProtocol.Latest, new TestLogger("local"));
+			await using RemoteComputeSocket agentComputeSocket = new(agentIdleTimeout, ComputeProtocol.Latest, new TestLogger("agent"));
+
+			await RunAgentTestsAsync(localComputeSocket, agentComputeSocket, cts.Token);
+		}
 
 		internal static async Task<(Socket client, Socket server)> CreateSocketsAsync(CancellationToken cancellationToken)
 		{
diff --git a/Engine/Source/Programs/Shared/EpicGames.Horde/Compute/Transports/AesTransport.cs b/Engine/Source/Programs/Shared/EpicGames.Horde/Compute/Transports/AesTransport.cs
--- a/Engine/Source/Programs/Shared/EpicGames.Horde/Compute/Transports/AesTransport.cs
+++ b/Engine/Source/Programs/Shared/EpicGames.Horde/Compute/Transports/AesTransport.cs
@@ -33,6 +33,7 @@ public sealed class AesTransport : ComputeTransport
 	private readonly AesGcm _aes;
 	private readonly SemaphoreSlim _sendSemaphore = new(1, 1);
 	private readonly SemaphoreSlim _recvSemaphore = new(1, 1);
+	private readonly MemoryStream _remainingData;
 	private bool _disposed;
 
 	/// <summary>
@@ -41,10 +42,12 @@ public sealed class AesTransport : ComputeTransport
 	/// <param name="inner">The underlying transport implementation that will be encrypted</param>
 	/// <param name="key">Encryption key. Should be generated with a cryptographically secure random number generator</param>
 	/// <param name="leaveInnerOpen">Whether to dispose the inner transport when this instance is disposed</param>
-	public AesTransport(ComputeTransport inner, ReadOnlySpan<byte> key, bool leaveInnerOpen = false)
+	/// <param name="readBufferSize">Default receive buffer size for data that needs to buffered for next read. Will automatically grow</param>
+	public AesTransport(ComputeTransport inner, ReadOnlySpan<byte> key, bool leaveInnerOpen = false, int readBufferSize = 256 * 1024)
 	{
 		_inner = inner;
 		_leaveInnerOpen = leaveInnerOpen;
+		_remainingData = new MemoryStream(readBufferSize);
 		
 		if (key.Length != KeyLength)
 		{
@@ -97,22 +100,37 @@ public sealed class AesTransport : ComputeTransport
 		await _recvSemaphore.WaitAsync(cancellationToken);
 		try
 		{
+			// Check if we have remaining data from previous recv
+			if (_remainingData.Length > 0)
+			{
+				int bytesToCopyFromPrev = Math.Min((int)(_remainingData.Length - _remainingData.Position), buffer.Length);
+				_remainingData.Read(buffer.Span[..bytesToCopyFromPrev]); // TODO: Check actual bytes read
+				
+				// If we've consumed all data, reset or dispose the stream
+				if (_remainingData.Position >= _remainingData.Length)
+				{
+					_remainingData.SetLength(0); // Reset the stream for reuse
+				}
+
+				return bytesToCopyFromPrev;
+			}
+			
 			byte[] headerData = new byte[HeaderLength];
 			await _inner.RecvFullAsync(headerData, cancellationToken);
 			int bufferLength = BitConverter.ToInt32(headerData.AsSpan(0, BufferLength));
 			
-			const int MaxMessageSize = 100 * 1024 * 1024; // 100MB
+			const int MaxMessageSize = 100 * 1024 * 1024; // 100 MB
 			if (bufferLength is <= 0 or > MaxMessageSize)
 			{
-				throw new InvalidDataException($"Invalid frame length: {bufferLength}");
+				throw new InvalidDataException($"Invalid or too large frame length: {bufferLength}");
 			}
 			
 			// Once we know the buffer size, read the encrypted data from underlying transport
-			byte[] encryptedBuffer = new byte[bufferLength + TagSize];
+			byte[] encryptedBuffer = new byte[bufferLength + TagSize]; // TODO: Use shared buffer
 			await _inner.RecvFullAsync(encryptedBuffer, cancellationToken);
 			
 			// Decrypt the message
-			byte[] plaintext = new byte[bufferLength];
+			byte[] plaintext = new byte[bufferLength]; // TODO: Use shared buffer
 			try
 			{
 				_aes.Decrypt(
@@ -129,6 +147,15 @@ public sealed class AesTransport : ComputeTransport
 			// Copy decrypted data to the output buffer which be smaller than actual decrypted buffer
 			int bytesToCopy = Math.Min(plaintext.Length, buffer.Length);
 			plaintext.AsSpan(0, bytesToCopy).CopyTo(buffer.Span);
+			
+			// Store remaining data for next receive call, if any
+			if (bytesToCopy < bufferLength)
+			{
+				_remainingData.SetLength(0); // Clear existing data
+				_remainingData.Write(plaintext.AsSpan(bytesToCopy, bufferLength - bytesToCopy));
+				_remainingData.Position = 0;
+			}
+			
 			return bytesToCopy;
 		}
 		finally
