Horde: Fix occasional cancellation exception when running chunking tests. Explicitly mark the source channel as complete.

#rnx

diff --git a/Engine/Source/Programs/Shared/EpicGames.Horde.Tests/Storage/ContentChunkerTests.cs b/Engine/Source/Programs/Shared/EpicGames.Horde.Tests/Storage/ContentChunkerTests.cs
--- a/Engine/Source/Programs/Shared/EpicGames.Horde.Tests/Storage/ContentChunkerTests.cs
+++ b/Engine/Source/Programs/Shared/EpicGames.Horde.Tests/Storage/ContentChunkerTests.cs
@@ -78,6 +78,8 @@ namespace EpicGames.Horde.Tests.Storage
 				await pipeline.SourceWriter.WriteAsync(input);
 			}
 
+			pipeline.SourceWriter.Complete();
+
 			List<SplitPoint> splitPoints = new List<SplitPoint>();
 			for (int idx = 0; idx < lengths.Length; idx++)
 			{
@@ -99,6 +101,7 @@ namespace EpicGames.Horde.Tests.Storage
 				Assert.AreEqual(outputLength, lengths[idx]);
 			}
 
+			Assert.IsFalse(await pipeline.OutputReader.WaitToReadAsync());
 			return splitPoints;
 		}
 	}
