Horde: Fix test flakiness in ComputeSocketTests

Increase test timeouts and disable parallelization to handle CPU contention better during parallel test execution.

diff --git a/Engine/Source/Programs/Shared/EpicGames.Horde.Tests/AesTransportTests.cs b/Engine/Source/Programs/Shared/EpicGames.Horde.Tests/AesTransportTests.cs
--- a/Engine/Source/Programs/Shared/EpicGames.Horde.Tests/AesTransportTests.cs
+++ b/Engine/Source/Programs/Shared/EpicGames.Horde.Tests/AesTransportTests.cs
@@ -43,7 +43,7 @@ public sealed class AesTransportTests : IAsyncDisposable
 	[TestMethod]
 	public async Task SendAndReceive_VariousBufferSizesOverTcp_SuccessfullyTransferredAsync()
 	{
-		using CancellationTokenSource cts = new(10000);
+		using CancellationTokenSource cts = new(ComputeSocketTests.TestTimoutMs);
 		(Socket clientSocket, Socket serverSocket) = await ComputeSocketTests.CreateSocketsAsync(cts.Token);
 		await using TcpTransport clientTransport = new(clientSocket);
 		await using TcpTransport serverTransport = new(serverSocket);
@@ -59,7 +59,7 @@ public sealed class AesTransportTests : IAsyncDisposable
 	[TestMethod]
 	public async Task SendAndReceive_VariousBufferSizesInMemory_SuccessfullyTransferredAsync()
 	{
-		using CancellationTokenSource cts = new(10000);
+		using CancellationTokenSource cts = new(ComputeSocketTests.TestTimoutMs);
 		await using MemoryComputeTransport memTransport = new();
 		await using AesTransport encryptedClient = new(memTransport, _key);
 		await using AesTransport encryptedServer = new(memTransport, _key);
@@ -73,7 +73,7 @@ public sealed class AesTransportTests : IAsyncDisposable
 	[TestMethod]
 	public async Task SendAndReceive_Torture_SuccessfullyTransferredAsync()
 	{
-		using CancellationTokenSource cts = new(20000);
+		using CancellationTokenSource cts = new(ComputeSocketTests.TestTimoutMs);
 		await using MemoryComputeTransport memTransport = new();
 		await using AesTransport clientAes = new(memTransport, _key);
 		await using AesTransport serverAes = new(memTransport, _key);
diff --git a/Engine/Source/Programs/Shared/EpicGames.Horde.Tests/Compute/ComputeSocketTests.cs b/Engine/Source/Programs/Shared/EpicGames.Horde.Tests/Compute/ComputeSocketTests.cs
--- a/Engine/Source/Programs/Shared/EpicGames.Horde.Tests/Compute/ComputeSocketTests.cs
+++ b/Engine/Source/Programs/Shared/EpicGames.Horde.Tests/Compute/ComputeSocketTests.cs
@@ -22,8 +22,10 @@ using Microsoft.VisualStudio.TestTools.UnitTesting;
 namespace EpicGames.Horde.Tests.Compute
 {
 	[TestClass]
+	[DoNotParallelize]
 	public class ComputeSocketTests
 	{
+		public const int TestTimoutMs = 20000;
 		class TestComputeSocket : ComputeSocket, IDisposable
 		{
 			public Dictionary<int, ComputeBufferWriter> RecvBufferWriters { get; } = new Dictionary<int, ComputeBufferWriter>();
@@ -85,7 +87,7 @@ namespace EpicGames.Horde.Tests.Compute
 		[TestMethod]
 		public async Task TestAgentMessageLoopTcpAsync()
 		{
-			using CancellationTokenSource cts = new(5000);
+			using CancellationTokenSource cts = new(TestTimoutMs);
 			(Socket clientSocket, Socket serverSocket) = await CreateSocketsAsync(cts.Token);
 
 			await using TcpTransport clientTransport = new(clientSocket);
@@ -101,7 +103,7 @@ namespace EpicGames.Horde.Tests.Compute
 		[DataRow(Encryption.SslEcdsaP256)]
 		public async Task TestAgentMessageLoopTcpSslAsync(Encryption encryption)
 		{
-			using CancellationTokenSource cts = new(5000);
+			using CancellationTokenSource cts = new(TestTimoutMs);
 			(Socket clientSocket, Socket serverSocket) = await CreateSocketsAsync(cts.Token);
 
 			byte[] certData = TcpSslTransport.GenerateCert(encryption);
@@ -121,7 +123,7 @@ namespace EpicGames.Horde.Tests.Compute
 		[TestMethod]
 		public async Task TestAgentMessageLoopTcpAesAsync()
 		{
-			using CancellationTokenSource cts = new(5000);
+			using CancellationTokenSource cts = new(TestTimoutMs);
 			(Socket localSocket, Socket agentSocket) = await CreateSocketsAsync(cts.Token);
 			
 			byte[] key = AesTransport.CreateKey();
