Horde: Force final stats to be output at the end of an extract operation.

#rnx

diff --git a/Engine/Source/Programs/Shared/EpicGames.Horde/Storage/Nodes/DirectoryNode.Extract.cs b/Engine/Source/Programs/Shared/EpicGames.Horde/Storage/Nodes/DirectoryNode.Extract.cs
--- a/Engine/Source/Programs/Shared/EpicGames.Horde/Storage/Nodes/DirectoryNode.Extract.cs
+++ b/Engine/Source/Programs/Shared/EpicGames.Horde/Storage/Nodes/DirectoryNode.Extract.cs
@@ -707,7 +707,7 @@ namespace EpicGames.Horde.Storage.Nodes
 					stats.DownloadRate = (stats.DownloadSize - firstSample.Stats.DownloadSize) / elapsedSeconds;
 				}
 
-				if (timeSeconds > lastOutputTime + frequency.TotalSeconds)
+				if (timeSeconds > lastOutputTime + frequency.TotalSeconds || completeTask == writeTask)
 				{
 					progress.Report(stats);
 					lastOutputTime = timeSeconds;
