Horde: Retry extraction of data from temp storage if a file was missing, with additional logging enabled.

diff --git a/Engine/Source/Programs/Horde/Drivers/JobDriver/Utility/TempStorage.cs b/Engine/Source/Programs/Horde/Drivers/JobDriver/Utility/TempStorage.cs
--- a/Engine/Source/Programs/Horde/Drivers/JobDriver/Utility/TempStorage.cs
+++ b/Engine/Source/Programs/Horde/Drivers/JobDriver/Utility/TempStorage.cs
@@ -779,27 +779,44 @@ namespace JobDriver.Utility
 					throw new TempStorageException($"Missing block \"{blockName}\" from node \"{nodeName}\". If running Horde against a 5.4 or earlier version of Unreal Engine, please make sure you have integrated //UE5/Main/Engine/Source/Programs/AutomationTool/BuildGraph/BuildGraph.cs#70.");
 				}
 
-				StorageStats initialStats = storageNamespace.GetStats();
-				Stopwatch timer = Stopwatch.StartNew();
-
-				// Add all the files and flush the ref
-				await rootDirEntry.Handle.ExtractAsync(rootDir.ToDirectoryInfo(), new ExtractOptions { Progress = new ExtractStatsLogger(logger) }, logger, cancellationToken);
+				for (int attempt = 0; ; attempt++)
+				{
+					StorageStats initialStats = storageNamespace.GetStats();
+					Stopwatch timer = Stopwatch.StartNew();
 
-				StorageStats deltaStats = StorageStats.GetDelta(initialStats, storageNamespace.GetStats());
-				logger.LogInformation("{Stats}", $"Elapsed: {(int)timer.Elapsed.TotalSeconds}s, {String.Join(", ", deltaStats.Values.Select(x => $"{x.Item1}: {x.Item2:n0}"))}");
+					// Add all the files and flush the ref
+					ExtractOptions options = new ExtractOptions();
+					options.Progress = new ExtractStatsLogger(logger);
+					options.VerifyOutput = options.VerboseOutput = (attempt > 0);
 
-				// Read the manifest in
-				manifest = TempStorageBlockManifest.Load(localManifestFile);
+					await rootDirEntry.Handle.ExtractAsync(rootDir.ToDirectoryInfo(), options, logger, cancellationToken);
 
-				// Update the timestamps to match the manifest.
-				foreach (TempStorageFile manifestFile in manifest.Files)
-				{
-					FileInfo fileInfo = manifestFile.ToFileReference(rootDir).ToFileInfo();
-					fileInfo.LastWriteTimeUtc = new DateTime(manifestFile.LastWriteTimeUtcTicks, DateTimeKind.Utc);
+					StorageStats deltaStats = StorageStats.GetDelta(initialStats, storageNamespace.GetStats());
+					logger.LogInformation("{Stats}", $"Elapsed: {(int)timer.Elapsed.TotalSeconds}s, {String.Join(", ", deltaStats.Values.Select(x => $"{x.Item1}: {x.Item2:n0}"))}");
 
-					if (fileInfo.Length != manifestFile.Length)
+					try
+					{
+						// Read the manifest in
+						manifest = TempStorageBlockManifest.Load(localManifestFile);
+
+						// Update the timestamps to match the manifest.
+						foreach (TempStorageFile manifestFile in manifest.Files)
+						{
+							FileInfo fileInfo = manifestFile.ToFileReference(rootDir).ToFileInfo();
+							fileInfo.LastWriteTimeUtc = new DateTime(manifestFile.LastWriteTimeUtcTicks, DateTimeKind.Utc);
+
+							if (fileInfo.Length != manifestFile.Length)
+							{
+								logger.LogError("File {File} extracted from temp storage has different size to file in manifest (manifest: {ManifestLength}, local: {LocalLength})", manifestFile.RelativePath, manifestFile.Length, fileInfo.Length);
+							}
+						}
+
+						break;
+					}
+					catch (FileNotFoundException ex) when (attempt == 0)
 					{
-						logger.LogError("File {File} extracted from temp storage has different size to file in manifest (manifest: {ManifestLength}, local: {LocalLength})", manifestFile.RelativePath, manifestFile.Length, fileInfo.Length);
+						logger.LogError(ex, "File not found: {Message}", ex.Message);
+						logger.LogInformation("Running extract again with verbose logging.");
 					}
 				}
 			}
diff --git a/Engine/Source/Programs/Shared/EpicGames.Horde/Storage/Nodes/DirectoryNode.Extract.cs b/Engine/Source/Programs/Shared/EpicGames.Horde/Storage/Nodes/DirectoryNode.Extract.cs
--- a/Engine/Source/Programs/Shared/EpicGames.Horde/Storage/Nodes/DirectoryNode.Extract.cs
+++ b/Engine/Source/Programs/Shared/EpicGames.Horde/Storage/Nodes/DirectoryNode.Extract.cs
@@ -137,6 +137,11 @@ namespace EpicGames.Horde.Storage.Nodes
 		/// Whether to hash downloaded data to ensure it's correct
 		/// </summary>
 		public bool VerifyOutput { get; set; }
+
+		/// <summary>
+		/// Output verbose logging about operations being performed
+		/// </summary>
+		public bool VerboseOutput { get; set; }
 	}
 
 	/// <summary>
@@ -362,7 +367,7 @@ namespace EpicGames.Horde.Storage.Nodes
 			{
 				using BatchReader<OutputChunk> batchReader = new BatchReader<OutputChunk>(options.VerifyOutput);
 
-				_ = pipeline.AddTask(ctx => FindOutputChunksRootAsync(directoryInfo, directoryNode, batchReader.RequestWriter, logger, ctx));
+				_ = pipeline.AddTask(ctx => FindOutputChunksRootAsync(directoryInfo, directoryNode, batchReader.RequestWriter, options, logger, ctx));
 				pipeline.AddBatchReaderTasks(numReadTasks, numDecodeTasks, batchReader);
 
 				OutputStats outputStats = new OutputStats();
@@ -381,13 +386,13 @@ namespace EpicGames.Horde.Storage.Nodes
 
 		#region Enumerate chunks
 
-		static async Task FindOutputChunksRootAsync(DirectoryInfo rootDir, DirectoryNode node, ChannelWriter<BatchReadRequest<OutputChunk>> chunks, ILogger logger, CancellationToken cancellationToken)
+		static async Task FindOutputChunksRootAsync(DirectoryInfo rootDir, DirectoryNode node, ChannelWriter<BatchReadRequest<OutputChunk>> chunks, ExtractOptions options, ILogger logger, CancellationToken cancellationToken)
 		{
-			await FindOutputChunksForDirectoryAsync(rootDir, "", node, chunks, logger, cancellationToken);
+			await FindOutputChunksForDirectoryAsync(rootDir, "", node, chunks, options, logger, cancellationToken);
 			chunks.Complete();
 		}
 
-		static async Task FindOutputChunksForDirectoryAsync(DirectoryInfo rootDir, string path, DirectoryNode node, ChannelWriter<BatchReadRequest<OutputChunk>> chunks, ILogger logger, CancellationToken cancellationToken)
+		static async Task FindOutputChunksForDirectoryAsync(DirectoryInfo rootDir, string path, DirectoryNode node, ChannelWriter<BatchReadRequest<OutputChunk>> chunks, ExtractOptions options, ILogger logger, CancellationToken cancellationToken)
 		{
 			foreach (FileEntry fileEntry in node.Files)
 			{
@@ -395,7 +400,7 @@ namespace EpicGames.Horde.Storage.Nodes
 				FileInfo fileInfo = new FileInfo(Path.Combine(rootDir.FullName, filePath));
 				OutputFile outputFile = new OutputFile(filePath, fileInfo, fileEntry);
 
-				await FindOutputChunksForFileAsync(outputFile, chunks, logger, cancellationToken);
+				await FindOutputChunksForFileAsync(outputFile, chunks, options, logger, cancellationToken);
 			}
 
 			foreach (DirectoryEntry directoryEntry in node.Directories)
@@ -404,21 +409,26 @@ namespace EpicGames.Horde.Storage.Nodes
 				TraceBlobRead("Directory", subPath, directoryEntry.Handle, logger);
 				DirectoryNode subDirectoryNode = await directoryEntry.Handle.ReadBlobAsync(cancellationToken);
 
-				await FindOutputChunksForDirectoryAsync(rootDir, subPath, subDirectoryNode, chunks, logger, cancellationToken);
+				await FindOutputChunksForDirectoryAsync(rootDir, subPath, subDirectoryNode, chunks, options, logger, cancellationToken);
 			}
 		}
 
-		static async Task FindOutputChunksForFileAsync(OutputFile outputFile, ChannelWriter<BatchReadRequest<OutputChunk>> chunks, ILogger logger, CancellationToken cancellationToken)
+		static async Task FindOutputChunksForFileAsync(OutputFile outputFile, ChannelWriter<BatchReadRequest<OutputChunk>> chunks, ExtractOptions options, ILogger logger, CancellationToken cancellationToken)
 		{
 			OutputChunkWriter outputWriter = new OutputChunkWriter(outputFile, chunks);
-			await FindOutputChunksAsync(outputWriter, 0, outputFile.FileEntry.Target, logger, cancellationToken);
+			await FindOutputChunksAsync(outputFile.Path, outputWriter, 0, outputFile.FileEntry.Target, options, logger, cancellationToken);
 			await outputWriter.FlushAsync(cancellationToken);
 		}
 
-		static async Task<long> FindOutputChunksAsync(OutputChunkWriter chunkWriter, long offset, ChunkedDataNodeRef dataRef, ILogger logger, CancellationToken cancellationToken)
+		static async Task<long> FindOutputChunksAsync(string path, OutputChunkWriter chunkWriter, long offset, ChunkedDataNodeRef dataRef, ExtractOptions options, ILogger logger, CancellationToken cancellationToken)
 		{
 			if (dataRef.Type == ChunkedDataNodeType.Leaf)
 			{
+				if (options.VerboseOutput)
+				{
+					logger.LogInformation("Queuing {Path},{Offset} ({Locator})", path, offset, dataRef.Handle.GetLocator());
+				}
+
 				await chunkWriter.WriteAsync(offset, dataRef.Length, dataRef.Handle, cancellationToken);
 				if (dataRef.Length < 0)
 				{
@@ -445,7 +455,7 @@ namespace EpicGames.Horde.Storage.Nodes
 					InteriorChunkedDataNode interiorNode = BlobSerializer.Deserialize<InteriorChunkedDataNode>(data);
 					foreach (ChunkedDataNodeRef childRef in interiorNode.Children)
 					{
-						length += await FindOutputChunksAsync(chunkWriter, offset + length, childRef, logger, cancellationToken);
+						length += await FindOutputChunksAsync(path, chunkWriter, offset + length, childRef, options, logger, cancellationToken);
 					}
 
 					return length;
@@ -474,6 +484,11 @@ namespace EpicGames.Horde.Storage.Nodes
 						}
 						await Task.WhenAll(tasks);
 					}
+					catch (Exception ex)
+					{
+						logger.LogError(ex, "Error writing to disk: {Message}", ex.Message);
+						throw;
+					}
 					finally
 					{
 						batch.Dispose();
@@ -580,6 +595,11 @@ namespace EpicGames.Horde.Storage.Nodes
 						BatchReadResponseItem<OutputChunk> chunk = chunks[chunkIdx];
 						cancellationToken.ThrowIfCancellationRequested();
 
+						if (options.VerboseOutput)
+						{
+							logger.LogInformation("Writing {File},{Offset} ({Locator})", file.Path, chunk.Context.Offset, chunk.Context.Handle.GetLocator());
+						}
+
 						// Write this chunk
 						TraceBlobRead("Leaf", chunk.Context.File.Path, chunk.Context.Handle, logger);
 						chunk.BlobData.Data.CopyTo(memoryMappedView!.GetMemory(chunk.Context.Offset, chunk.BlobData.Data.Length));
