Fix a few typos in UHT

#rnx

diff --git a/Engine/Source/Programs/Shared/EpicGames.UHT/Exporters/CodeGen/UhtHeaderCodeGeneratorCppFile.cs b/Engine/Source/Programs/Shared/EpicGames.UHT/Exporters/CodeGen/UhtHeaderCodeGeneratorCppFile.cs
--- a/Engine/Source/Programs/Shared/EpicGames.UHT/Exporters/CodeGen/UhtHeaderCodeGeneratorCppFile.cs
+++ b/Engine/Source/Programs/Shared/EpicGames.UHT/Exporters/CodeGen/UhtHeaderCodeGeneratorCppFile.cs
@@ -459,7 +459,7 @@ namespace EpicGames.UHT.Exporters.CodeGen
 								builder.Append("!int64(").Append(value.Name).Append(')');
 							}
 						}
-						builder.Append(", \"'").Append(enumObj.SourceName).Append("' does not have a 0 entry!(This is a problem when the enum is initalized by default)\");\r\n");
+						builder.Append(", \"'").Append(enumObj.SourceName).Append("' does not have a 0 entry!(This is a problem when the enum is initialized by default)\");\r\n");
 					}
 				}
 
@@ -2107,7 +2107,7 @@ namespace EpicGames.UHT.Exporters.CodeGen
 	/// <summary>
 	/// Collection of string builder extensions used to generate the cpp files for individual headers.
 	/// </summary>
-	public static class UhtHeaderCodeGeneratorCppFileStringBuilderExtensinos
+	public static class UhtHeaderCodeGeneratorCppFileStringBuilderExtensions
 	{
 
 		/// <summary>
