Horde: Prevent exceeding open file handle limit on MacOS when uploading data.

diff --git a/Engine/Source/Programs/Shared/EpicGames.Horde/Storage/ContentChunker.cs b/Engine/Source/Programs/Shared/EpicGames.Horde/Storage/ContentChunker.cs
--- a/Engine/Source/Programs/Shared/EpicGames.Horde/Storage/ContentChunker.cs
+++ b/Engine/Source/Programs/Shared/EpicGames.Horde/Storage/ContentChunker.cs
@@ -8,6 +8,7 @@ using System.Threading.Channels;
 using System.Threading.Tasks;
 using EpicGames.Core;
 using EpicGames.Horde.Storage.Nodes;
+using EpicGames.Horde.Utilities;
 
 namespace EpicGames.Horde.Storage
 {
@@ -89,6 +90,8 @@ namespace EpicGames.Horde.Storage
 		/// <inheritdoc/>
 		async Task HandleReadAsync(long offset, Memory<byte> memory, CancellationToken cancellationToken)
 		{
+			using PlatformFileLock? platformFileLock = await PlatformFileLock.CreateAsync(cancellationToken);
+
 			FileStreamOptions options = new FileStreamOptions { Mode = FileMode.Open, Access = FileAccess.Read, Options = FileOptions.Asynchronous };
 			using FileStream stream = _fileInfo.Open(options);
 			stream.Seek(offset, SeekOrigin.Begin);
diff --git a/Engine/Source/Programs/Shared/EpicGames.Horde/Storage/Nodes/DirectoryNode.Extract.cs b/Engine/Source/Programs/Shared/EpicGames.Horde/Storage/Nodes/DirectoryNode.Extract.cs
--- a/Engine/Source/Programs/Shared/EpicGames.Horde/Storage/Nodes/DirectoryNode.Extract.cs
+++ b/Engine/Source/Programs/Shared/EpicGames.Horde/Storage/Nodes/DirectoryNode.Extract.cs
@@ -10,6 +10,7 @@ using System.Threading;
 using System.Threading.Channels;
 using System.Threading.Tasks;
 using EpicGames.Core;
+using EpicGames.Horde.Utilities;
 using Microsoft.Extensions.Logging;
 
 namespace EpicGames.Horde.Storage.Nodes
@@ -527,47 +528,6 @@ namespace EpicGames.Horde.Storage.Nodes
 			}
 		}
 
-		class PlatformFileLock : IDisposable
-		{
-			static readonly SemaphoreSlim? s_semaphore = CreateSemaphore();
-
-			bool _locked;
-
-			private PlatformFileLock()
-				=> _locked = true;
-
-			public static async ValueTask<PlatformFileLock?> CreateAsync(CancellationToken cancellationToken)
-			{
-				if (s_semaphore != null)
-				{
-					await s_semaphore.WaitAsync(cancellationToken);
-					return new PlatformFileLock();
-				}
-				return null;
-			}
-
-			public void Dispose()
-			{
-				if (_locked)
-				{
-					s_semaphore?.Release();
-					_locked = true;
-				}
-			}
-
-			static SemaphoreSlim? CreateSemaphore()
-			{
-				if (OperatingSystem.IsMacOS())
-				{
-					return new SemaphoreSlim(32);
-				}
-				else
-				{
-					return null;
-				}
-			}
-		}
-
 		static async Task ExtractChunksToFileAsync(OutputFile file, ArraySegment<BatchReadResponseItem<OutputChunk>> chunks, OutputStats stats, ExtractOptions options, ILogger logger, CancellationToken cancellationToken)
 		{
 			using PlatformFileLock? fileLock = await PlatformFileLock.CreateAsync(cancellationToken);
diff --git a/Engine/Source/Programs/Shared/EpicGames.Horde/Utilities/PlatformFileLock.cs b/Engine/Source/Programs/Shared/EpicGames.Horde/Utilities/PlatformFileLock.cs
new file mode 100644
--- /dev/null
+++ b/Engine/Source/Programs/Shared/EpicGames.Horde/Utilities/PlatformFileLock.cs
@@ -0,0 +1,57 @@
+// Copyright Epic Games, Inc. All Rights Reserved.
+
+using System;
+using System.Threading;
+using System.Threading.Tasks;
+
+namespace EpicGames.Horde.Utilities
+{
+	/// <summary>
+	/// Wraps a semaphore controlling the number of open files at any time. Used to prevent exceeding handle limit on MacOS.
+	/// </summary>
+	class PlatformFileLock : IDisposable
+	{
+		static readonly SemaphoreSlim? s_semaphore = CreateSemaphore();
+
+		bool _locked;
+
+		private PlatformFileLock()
+			=> _locked = true;
+
+		/// <summary>
+		/// Acquire the platform file lock
+		/// </summary>
+		/// <param name="cancellationToken">Cancellation token for the operation</param>
+		public static async ValueTask<PlatformFileLock?> CreateAsync(CancellationToken cancellationToken)
+		{
+			if (s_semaphore != null)
+			{
+				await s_semaphore.WaitAsync(cancellationToken);
+				return new PlatformFileLock();
+			}
+			return null;
+		}
+
+		/// <inheritdoc/>
+		public void Dispose()
+		{
+			if (_locked)
+			{
+				s_semaphore?.Release();
+				_locked = true;
+			}
+		}
+
+		static SemaphoreSlim? CreateSemaphore()
+		{
+			if (OperatingSystem.IsMacOS())
+			{
+				return new SemaphoreSlim(32);
+			}
+			else
+			{
+				return null;
+			}
+		}
+	}
+}
