Horde: Allow constructing the input channel for batch blob reads outside the BatchBlobReader instance.

#rnx

diff --git a/Engine/Source/Programs/Shared/EpicGames.Horde/Storage/BatchBlobReader.cs b/Engine/Source/Programs/Shared/EpicGames.Horde/Storage/BatchBlobReader.cs
new file mode 100644
--- /dev/null
+++ b/Engine/Source/Programs/Shared/EpicGames.Horde/Storage/BatchBlobReader.cs
@@ -0,0 +1,417 @@
+// Copyright Epic Games, Inc. All Rights Reserved.
+
+using System;
+using System.Collections.Generic;
+using System.Diagnostics.CodeAnalysis;
+using System.IO;
+using System.Linq;
+using System.Threading;
+using System.Threading.Channels;
+using System.Threading.Tasks;
+using EpicGames.Core;
+using EpicGames.Horde.Storage.Bundles;
+using EpicGames.Horde.Storage.Bundles.V2;
+
+namespace EpicGames.Horde.Storage
+{
+	/// <summary>
+	/// Request for a blob to be read
+	/// </summary>
+	/// <typeparam name="TContext">User defined context type for the request</typeparam>
+	/// <param name="BlobRef">Reference to the blob data</param>
+	/// <param name="Context">User defined context</param>
+	public record class BlobReadRequest<TContext>(IHashedBlobRef BlobRef, TContext Context);
+
+	/// <summary>
+	/// Individual response for a blob that was read
+	/// </summary>
+	/// <typeparam name="TContext">User defined context type for the request</typeparam>
+	/// <param name="BlobData">Data for the blob</param>
+	/// <param name="Context">User defined context</param>
+	public sealed record class BlobReadResponse<TContext>(BlobData BlobData, TContext Context) : IDisposable
+	{
+		/// <inheritdoc/>
+		public void Dispose()
+			=> BlobData.Dispose();
+	}
+
+	/// <summary>
+	/// Response containing a batch of items that were read
+	/// </summary>
+	/// <typeparam name="TContext">User defined context type</typeparam>
+	/// <param name="Items">Items for this response</param>
+	public sealed record class BatchBlobReadResponse<TContext>(BlobReadResponse<TContext>[] Items) : IDisposable
+	{
+		/// <inheritdoc/>
+		public void Dispose()
+		{
+			foreach (BlobReadResponse<TContext> item in Items)
+			{
+				item.Dispose();
+			}
+		}
+	}
+
+	/// <summary>
+	/// Stats for a batch read
+	/// </summary>
+	public record class BatchReaderStats(int NumRequests, int NumReads, int NumBundles, int NumPackets, long BytesRead, long BytesDecoded);
+
+	/// <summary>
+	/// Implements an efficient pipeline for streaming blob data
+	/// </summary>
+	/// <typeparam name="TContext">User-defined context type</typeparam>
+	public sealed class BatchBlobReader<TContext> : IDisposable
+	{
+		record class BundleRequest(BundleHandle Handle, int MinOffset, int MaxOffset, List<PacketRequest> Packets);
+
+		record class BundleResponse(IReadOnlyMemoryOwner<byte> Data, int MinOffset, int MaxOffset, List<PacketRequest> Packets) : IDisposable
+		{
+			public void Dispose()
+				=> Data.Dispose();
+		}
+
+		record class PacketRequest(FlushedPacketHandle Handle, List<ExportRequest> Exports)
+		{
+			public int MinOffset => Handle.PacketOffset;
+			public int MaxOffset => Handle.PacketOffset + Handle.PacketLength;
+		}
+
+		record class ExportRequest(IoHash Hash, ExportHandle Handle, TContext Context);
+
+		record class PacketResponse(IReadOnlyMemoryOwner<byte> Data, List<ExportRequest> Exports) : IDisposable
+		{
+			public void Dispose()
+				=> Data.Dispose();
+		}
+
+		/// <summary>
+		/// Number of items to read from the input queue before partitioning into batches
+		/// </summary>
+		public int MinQueueLength { get; set; } = 2000;
+
+		/// <summary>
+		/// Maximum gap between reads that should be coalesced and executed together
+		/// </summary>
+		public long CoalesceReadsBelowSize { get; set; } = 2 * 1024 * 1024;
+
+		/// <summary>
+		/// Whether to verify hashes of data read from storage
+		/// </summary>
+		public bool VerifyHashes { get; set; }
+
+		readonly Channel<BlobReadRequest<TContext>> _otherRequests;
+		readonly Channel<BundleRequest> _bundleRequests;
+
+		readonly Channel<BundleResponse> _bundleResponses;
+		readonly Channel<PacketResponse> _packetResponses;
+		readonly Channel<BatchBlobReadResponse<TContext>> _blobResponses;
+
+		int _numRequests;
+		int _numReads;
+		int _numBundles;
+		int _numPackets;
+		long _bytesRead;
+		long _bytesDecoded;
+
+		/// <summary>
+		/// Access to the response channel
+		/// </summary>
+		public ChannelReader<BatchBlobReadResponse<TContext>> ResponseReader
+			=> _blobResponses.Reader;
+
+		/// <summary>
+		/// Constructor
+		/// </summary>
+		public BatchBlobReader()
+		{
+			_otherRequests = Channel.CreateUnbounded<BlobReadRequest<TContext>>();
+			_bundleRequests = Channel.CreateUnbounded<BundleRequest>();
+			_bundleResponses = Channel.CreateBounded<BundleResponse>(new BoundedChannelOptions(32));
+			_packetResponses = Channel.CreateBounded<PacketResponse>(new BoundedChannelOptions(128));
+			_blobResponses = Channel.CreateBounded<BatchBlobReadResponse<TContext>>(new BoundedChannelOptions(128));
+		}
+
+		/// <inheritdoc/>
+		public void Dispose()
+		{
+			while (_bundleResponses.Reader.TryRead(out BundleResponse? item))
+			{
+				item.Dispose();
+			}
+
+			while (_packetResponses.Reader.TryRead(out PacketResponse? item))
+			{
+				item.Dispose();
+			}
+
+			while (_blobResponses.Reader.TryRead(out BatchBlobReadResponse<TContext>? item))
+			{
+				item.Dispose();
+			}
+		}
+
+		/// <summary>
+		/// Gets stats for the reader
+		/// </summary>
+		public BatchReaderStats GetStats()
+			=> new BatchReaderStats(_numRequests, _numReads, _numBundles, _numPackets, _bytesRead, _bytesDecoded);
+
+		/// <summary>
+		/// Adds a batch reader to the given pipeline
+		/// </summary>
+		internal void AddToPipeline(AsyncPipeline pipeline, int numReadTasks, int numDecodeTasks, ChannelReader<BlobReadRequest<TContext>> requestReader)
+		{
+			_ = pipeline.AddTask(cancellationToken => CreateBundleRequestsAsync(requestReader, cancellationToken));
+			Task otherTask = pipeline.AddTask(HandleOtherRequestsAsync);
+
+			// Read bundles and mark the bundle response channel as complete once we're finished
+			Task[] bundleTasks = pipeline.AddTasks(numReadTasks, ReadBundlesAsync);
+			_ = Task.WhenAll(bundleTasks).ContinueWith(_ => _bundleResponses.Writer.TryComplete(), TaskScheduler.Default);
+
+			// Read packets
+			Task[] packetTasks = pipeline.AddTasks(numDecodeTasks, ReadPacketsAsync);
+
+			// Read all the output blobs and mark the blob response channel as complete
+			IEnumerable<Task> blobTasks = packetTasks.Append(otherTask);
+			_ = Task.WhenAll(blobTasks).ContinueWith(_ => _blobResponses.Writer.TryComplete(), TaskScheduler.Default);
+		}
+
+		// Read a blob using the naive read pipeline
+		async Task HandleOtherRequestsAsync(CancellationToken cancellationToken)
+		{
+			await foreach (BlobReadRequest<TContext> request in _otherRequests.Reader.ReadAllAsync(cancellationToken))
+			{
+				BlobData? blobData = null;
+				try
+				{
+					blobData = await request.BlobRef.ReadBlobDataAsync(cancellationToken);
+					VerifyHash(request.BlobRef, blobData.Data.Span, request.BlobRef.Hash);
+#pragma warning disable CA2000
+					BatchBlobReadResponse<TContext> response = new BatchBlobReadResponse<TContext>([new BlobReadResponse<TContext>(blobData, request.Context)]);
+					await _blobResponses.Writer.WriteAsync(response, cancellationToken);
+#pragma warning restore CA2000
+				}
+				catch
+				{
+					blobData?.Dispose();
+				}
+			}
+		}
+
+		#region Grouping requests
+
+		record class PacketExport(FlushedPacketHandle PacketHandle, IoHash ExportHash, ExportHandle ExportHandle, TContext Context);
+
+		async Task CreateBundleRequestsAsync(ChannelReader<BlobReadRequest<TContext>> requestReader, CancellationToken cancellationToken)
+		{
+			int queueLength = 0;
+			Queue<BundleHandle> bundleQueue = new Queue<BundleHandle>();
+			Dictionary<BundleHandle, List<PacketExport>> bundleHandleToExportBatch = new Dictionary<BundleHandle, List<PacketExport>>();
+
+			for (; ; )
+			{
+				// Fill the queue up to the max length
+				for (; ; )
+				{
+					BlobReadRequest<TContext>? request;
+					if (!requestReader.TryRead(out request))
+					{
+						if (queueLength >= MinQueueLength)
+						{
+							break;
+						}
+						if (!await requestReader.WaitToReadAsync(cancellationToken))
+						{
+							break;
+						}
+					}
+					else
+					{
+						Interlocked.Increment(ref _numRequests);
+
+						PacketExport? outputExport;
+						if (!TryGetPacketExport(request, out outputExport))
+						{
+							await _otherRequests.Writer.WriteAsync(request, cancellationToken);
+						}
+						else
+						{
+							BundleHandle bundleHandle = outputExport.PacketHandle.Bundle;
+							if (!bundleHandleToExportBatch.TryGetValue(bundleHandle, out List<PacketExport>? existingExportBatch))
+							{
+								existingExportBatch = new List<PacketExport>();
+								bundleHandleToExportBatch.Add(bundleHandle, existingExportBatch);
+								bundleQueue.Enqueue(bundleHandle);
+							}
+
+							existingExportBatch.Add(outputExport);
+							queueLength++;
+						}
+					}
+				}
+
+				// Exit once we've processed everything and can't get any more items to read.
+				if (queueLength == 0)
+				{
+					_bundleRequests.Writer.TryComplete();
+					_otherRequests.Writer.TryComplete();
+					break;
+				}
+
+				// Flush the first queue
+				{
+					BundleHandle bundleHandle = bundleQueue.Dequeue();
+					List<PacketExport> exportBatch = bundleHandleToExportBatch[bundleHandle];
+					queueLength -= exportBatch.Count;
+					bundleHandleToExportBatch.Remove(bundleHandle);
+
+					await WriteBundleRequestsAsync(bundleHandle, exportBatch, cancellationToken);
+				}
+			}
+		}
+
+		async Task WriteBundleRequestsAsync(BundleHandle bundleHandle, List<PacketExport> exports, CancellationToken cancellationToken)
+		{
+			Interlocked.Increment(ref _numBundles);
+
+			// Group the reads by packet
+			List<PacketRequest> packetRequests = new List<PacketRequest>();
+			foreach (IGrouping<int, PacketExport> group in exports.GroupBy(x => x.PacketHandle.PacketOffset))
+			{
+				List<ExportRequest> exportRequests = group.Select(x => new ExportRequest(x.ExportHash, x.ExportHandle, x.Context)).ToList();
+				PacketRequest packetRequest = new PacketRequest(group.First().PacketHandle, exportRequests);
+				packetRequests.Add(packetRequest);
+			}
+			packetRequests.SortBy(x => x.MinOffset);
+
+			// Split the packet requests into contiguous bundle reads
+			for (int maxIdx = 0; maxIdx < packetRequests.Count; maxIdx++)
+			{
+				int minIdx = maxIdx;
+				while (maxIdx + 1 < packetRequests.Count && packetRequests[maxIdx + 1].MinOffset < packetRequests[maxIdx].MaxOffset + CoalesceReadsBelowSize)
+				{
+					maxIdx++;
+				}
+
+				int minOffset = packetRequests[minIdx].MinOffset;
+				int maxOffset = packetRequests[maxIdx].MaxOffset;
+
+				BundleRequest request = new BundleRequest(bundleHandle, minOffset, maxOffset, packetRequests.GetRange(minIdx, (maxIdx + 1) - minIdx));
+				await _bundleRequests.Writer.WriteAsync(request, cancellationToken);
+			}
+		}
+
+		static bool TryGetPacketExport(BlobReadRequest<TContext> chunk, [NotNullWhen(true)] out PacketExport? export)
+		{
+			if (chunk.BlobRef.Innermost is ExportHandle exportHandle && exportHandle.Packet is FlushedPacketHandle packetHandle)
+			{
+				export = new PacketExport(packetHandle, chunk.BlobRef.Hash, exportHandle, chunk.Context);
+				return true;
+			}
+			else
+			{
+				export = null;
+				return false;
+			}
+		}
+
+		#endregion
+		#region Bundle Reads
+
+		async Task ReadBundlesAsync(CancellationToken cancellationToken)
+		{
+			await foreach (BundleRequest request in _bundleRequests.Reader.ReadAllAsync(cancellationToken))
+			{
+				IReadOnlyMemoryOwner<byte>? data = null;
+				try
+				{
+					data = await request.Handle.ReadAsync(request.MinOffset, request.MaxOffset - request.MinOffset, cancellationToken);
+					Interlocked.Increment(ref _numReads);
+					Interlocked.Add(ref _bytesRead, data.Memory.Length);
+#pragma warning disable CA2000
+					BundleResponse response = new BundleResponse(data, request.MinOffset, request.MaxOffset, request.Packets);
+					await _bundleResponses.Writer.WriteAsync(response, cancellationToken);
+#pragma warning restore CA2000
+				}
+				catch
+				{
+					data?.Dispose();
+					throw;
+				}
+			}
+		}
+
+		void VerifyHash(IBlobRef blobRef, ReadOnlySpan<byte> data, IoHash hash)
+		{
+			if (VerifyHashes)
+			{
+				IoHash actualHash = IoHash.Compute(data);
+				if (actualHash != hash)
+				{
+					throw new InvalidDataException($"Invalid hash while reading {blobRef.GetLocator()}: Expected {hash}, got {actualHash}.");
+				}
+			}
+		}
+
+		#endregion
+		#region Packet reads
+
+		async Task ReadPacketsAsync(CancellationToken cancellationToken)
+		{
+			await foreach (BundleResponse bundleResponse in _bundleResponses.Reader.ReadAllAsync(cancellationToken))
+			{
+				using IDisposable lifetime = bundleResponse;
+				foreach (PacketRequest packetRequest in bundleResponse.Packets)
+				{
+					// Decode the data for this packet
+					ReadOnlyMemory<byte> memory = bundleResponse.Data.Memory.Slice(packetRequest.Handle.PacketOffset - bundleResponse.MinOffset, packetRequest.Handle.PacketLength);
+
+					using PacketReader packetReader = packetRequest.Handle.CreatePacketReader(memory);
+					Interlocked.Increment(ref _numPackets);
+					Interlocked.Add(ref _bytesDecoded, packetReader.Packet.Length);
+
+					// Create responses for each blob
+					BlobReadResponse<TContext>[] items = new BlobReadResponse<TContext>[packetRequest.Exports.Count];
+					try
+					{
+#pragma warning disable CA2000
+						for (int idx = 0; idx < items.Length; idx++)
+						{
+							ExportRequest exportRequest = packetRequest.Exports[idx];
+							BlobData blobData = packetReader.ReadExport(exportRequest.Handle.ExportIdx);
+							VerifyHash(exportRequest.Handle, blobData.Data.Span, exportRequest.Hash);
+							items[idx] = new BlobReadResponse<TContext>(blobData, exportRequest.Context);
+						}
+
+						BatchBlobReadResponse<TContext> response = new BatchBlobReadResponse<TContext>(items);
+						await _blobResponses.Writer.WriteAsync(response, cancellationToken);
+#pragma warning restore CA2000
+					}
+					catch
+					{
+						for (int idx = 0; idx < items.Length; idx++)
+						{
+							items[idx]?.Dispose();
+						}
+						throw;
+					}
+				}
+			}
+		}
+
+		#endregion
+	}
+
+	/// <summary>
+	/// Extension methods for <see cref="BatchBlobReader{TContext}"/>
+	/// </summary>
+	public static class BatchReaderExtensions
+	{
+		/// <summary>
+		/// Add tasks for reading to the given async pipeline 
+		/// </summary>
+		public static void AddTasks<T>(this AsyncPipeline pipeline, int numReadTasks, int numDecodeTasks, BatchBlobReader<T> reader, ChannelReader<BlobReadRequest<T>> requestReader)
+			=> reader.AddToPipeline(pipeline, numReadTasks, numDecodeTasks, requestReader);
+	}
+}
diff --git a/Engine/Source/Programs/Shared/EpicGames.Horde/Storage/BatchReader.cs b/Engine/Source/Programs/Shared/EpicGames.Horde/Storage/BatchReader.cs
deleted file mode 100644
diff --git a/Engine/Source/Programs/Shared/EpicGames.Horde/Storage/Nodes/DirectoryNode.Extract.cs b/Engine/Source/Programs/Shared/EpicGames.Horde/Storage/Nodes/DirectoryNode.Extract.cs
--- a/Engine/Source/Programs/Shared/EpicGames.Horde/Storage/Nodes/DirectoryNode.Extract.cs
+++ b/Engine/Source/Programs/Shared/EpicGames.Horde/Storage/Nodes/DirectoryNode.Extract.cs
@@ -244,10 +244,10 @@ namespace EpicGames.Horde.Storage.Nodes
 		{
 			public OutputFile OutputFile { get; }
 
-			readonly ChannelWriter<BatchReadRequest<OutputChunk>> _chunkWriter;
-			BatchReadRequest<OutputChunk>? _bufferedChunk;
+			readonly ChannelWriter<BlobReadRequest<OutputChunk>> _chunkWriter;
+			BlobReadRequest<OutputChunk>? _bufferedChunk;
 
-			public OutputChunkWriter(OutputFile file, ChannelWriter<BatchReadRequest<OutputChunk>> chunkWriter)
+			public OutputChunkWriter(OutputFile file, ChannelWriter<BlobReadRequest<OutputChunk>> chunkWriter)
 			{
 				OutputFile = file;
 				_chunkWriter = chunkWriter;
@@ -261,7 +261,7 @@ namespace EpicGames.Horde.Storage.Nodes
 				}
 
 				OutputFile.IncrementRemaining();
-				_bufferedChunk = new BatchReadRequest<OutputChunk>(handle, new OutputChunk(OutputFile, offset, length, handle));
+				_bufferedChunk = new BlobReadRequest<OutputChunk>(handle, new OutputChunk(OutputFile, offset, length, handle));
 			}
 
 			public async Task FlushAsync(CancellationToken cancellationToken)
@@ -366,10 +366,12 @@ namespace EpicGames.Horde.Storage.Nodes
 
 			await using (AsyncPipeline pipeline = new AsyncPipeline(cancellationToken))
 			{
-				using BatchReader<OutputChunk> batchReader = new BatchReader<OutputChunk>(options.VerifyOutput);
+				using BatchBlobReader<OutputChunk> batchReader = new BatchBlobReader<OutputChunk>();
+				batchReader.VerifyHashes = options.VerifyOutput;
 
-				_ = pipeline.AddTask(ctx => FindOutputChunksRootAsync(directoryInfo, directoryNode, batchReader.RequestWriter, options, logger, ctx));
-				pipeline.AddBatchReaderTasks(numReadTasks, numDecodeTasks, batchReader);
+				Channel<BlobReadRequest<OutputChunk>> blobRequests = Channel.CreateBounded<BlobReadRequest<OutputChunk>>(new BoundedChannelOptions(batchReader.MinQueueLength * 2));
+				_ = pipeline.AddTask(ctx => FindOutputChunksRootAsync(directoryInfo, directoryNode, blobRequests.Writer, options, logger, ctx));
+				pipeline.AddTasks(numReadTasks, numDecodeTasks, batchReader, blobRequests.Reader);
 
 				OutputStats outputStats = new OutputStats();
 				Task[] writeTasks = pipeline.AddTasks(numWriteTasks, ctx => WriteAsync(batchReader.ResponseReader, outputStats, options, logger, ctx));
@@ -387,13 +389,13 @@ namespace EpicGames.Horde.Storage.Nodes
 
 		#region Enumerate chunks
 
-		static async Task FindOutputChunksRootAsync(DirectoryInfo rootDir, DirectoryNode node, ChannelWriter<BatchReadRequest<OutputChunk>> chunks, ExtractOptions options, ILogger logger, CancellationToken cancellationToken)
+		static async Task FindOutputChunksRootAsync(DirectoryInfo rootDir, DirectoryNode node, ChannelWriter<BlobReadRequest<OutputChunk>> chunks, ExtractOptions options, ILogger logger, CancellationToken cancellationToken)
 		{
 			await FindOutputChunksForDirectoryAsync(rootDir, "", node, chunks, options, logger, cancellationToken);
 			chunks.Complete();
 		}
 
-		static async Task FindOutputChunksForDirectoryAsync(DirectoryInfo rootDir, string path, DirectoryNode node, ChannelWriter<BatchReadRequest<OutputChunk>> chunks, ExtractOptions options, ILogger logger, CancellationToken cancellationToken)
+		static async Task FindOutputChunksForDirectoryAsync(DirectoryInfo rootDir, string path, DirectoryNode node, ChannelWriter<BlobReadRequest<OutputChunk>> chunks, ExtractOptions options, ILogger logger, CancellationToken cancellationToken)
 		{
 			foreach (FileEntry fileEntry in node.Files)
 			{
@@ -414,7 +416,7 @@ namespace EpicGames.Horde.Storage.Nodes
 			}
 		}
 
-		static async Task FindOutputChunksForFileAsync(OutputFile outputFile, ChannelWriter<BatchReadRequest<OutputChunk>> chunks, ExtractOptions options, ILogger logger, CancellationToken cancellationToken)
+		static async Task FindOutputChunksForFileAsync(OutputFile outputFile, ChannelWriter<BlobReadRequest<OutputChunk>> chunks, ExtractOptions options, ILogger logger, CancellationToken cancellationToken)
 		{
 			OutputChunkWriter outputWriter = new OutputChunkWriter(outputFile, chunks);
 			await FindOutputChunksAsync(outputFile.Path, outputWriter, 0, outputFile.FileEntry.Target, options, logger, cancellationToken);
@@ -468,18 +470,18 @@ namespace EpicGames.Horde.Storage.Nodes
 
 		#region Write to disk
 
-		static async Task WriteAsync(ChannelReader<BatchReadResponse<OutputChunk>> batchReader, OutputStats outputStats, ExtractOptions options, ILogger logger, CancellationToken cancellationToken)
+		static async Task WriteAsync(ChannelReader<BatchBlobReadResponse<OutputChunk>> batchReader, OutputStats outputStats, ExtractOptions options, ILogger logger, CancellationToken cancellationToken)
 		{
 			const int WriteBatchSize = 64;
 			while (await batchReader.WaitToReadAsync(cancellationToken))
 			{
-				BatchReadResponse<OutputChunk>? batch;
+				BatchBlobReadResponse<OutputChunk>? batch;
 				while (batchReader.TryRead(out batch))
 				{
 					try
 					{
 						List<Task> tasks = new List<Task>();
-						foreach (IReadOnlyList<BatchReadResponseItem<OutputChunk>> group in batch.Items.Batch(WriteBatchSize))
+						foreach (IReadOnlyList<BlobReadResponse<OutputChunk>> group in batch.Items.Batch(WriteBatchSize))
 						{
 							tasks.Add(WriteChunksAsync(group.ToArray(), outputStats, options, logger, cancellationToken));
 						}
@@ -498,7 +500,7 @@ namespace EpicGames.Horde.Storage.Nodes
 			}
 		}
 
-		static async Task WriteChunksAsync(ArraySegment<BatchReadResponseItem<OutputChunk>> chunks, OutputStats outputStats, ExtractOptions options, ILogger logger, CancellationToken cancellationToken)
+		static async Task WriteChunksAsync(ArraySegment<BlobReadResponse<OutputChunk>> chunks, OutputStats outputStats, ExtractOptions options, ILogger logger, CancellationToken cancellationToken)
 		{
 			for (int chunkIdx = 0; chunkIdx < chunks.Count;)
 			{
@@ -528,7 +530,7 @@ namespace EpicGames.Horde.Storage.Nodes
 			}
 		}
 
-		static async Task ExtractChunksToFileAsync(OutputFile file, ArraySegment<BatchReadResponseItem<OutputChunk>> chunks, OutputStats stats, ExtractOptions options, ILogger logger, CancellationToken cancellationToken)
+		static async Task ExtractChunksToFileAsync(OutputFile file, ArraySegment<BlobReadResponse<OutputChunk>> chunks, OutputStats stats, ExtractOptions options, ILogger logger, CancellationToken cancellationToken)
 		{
 			using PlatformFileLock? fileLock = await PlatformFileLock.CreateAsync(cancellationToken);
 
@@ -552,7 +554,7 @@ namespace EpicGames.Horde.Storage.Nodes
 
 					for (int chunkIdx = 0; chunkIdx < chunks.Count; chunkIdx++)
 					{
-						BatchReadResponseItem<OutputChunk> chunk = chunks[chunkIdx];
+						BlobReadResponse<OutputChunk> chunk = chunks[chunkIdx];
 						cancellationToken.ThrowIfCancellationRequested();
 
 						if (options.VerboseOutput)
@@ -616,12 +618,12 @@ namespace EpicGames.Horde.Storage.Nodes
 
 #pragma warning disable IDE0051
 		// Update counters for extracting chunks without writing any data. Useful for profiling bottlenecks in other stages of the pipeline.
-		static Task ExtractChunksToNullAsync(OutputFile file, ArraySegment<BatchReadResponseItem<OutputChunk>> chunks, OutputStats stats, ILogger logger, CancellationToken cancellationToken)
+		static Task ExtractChunksToNullAsync(OutputFile file, ArraySegment<BlobReadResponse<OutputChunk>> chunks, OutputStats stats, ILogger logger, CancellationToken cancellationToken)
 		{
 			_ = logger;
 			_ = cancellationToken;
 
-			foreach (BatchReadResponseItem<OutputChunk> chunk in chunks)
+			foreach (BlobReadResponse<OutputChunk> chunk in chunks)
 			{
 				Interlocked.Add(ref stats._writtenBytes, chunk.Context.Length);
 
@@ -649,7 +651,7 @@ namespace EpicGames.Horde.Storage.Nodes
 			public double DownloadRate { get; set; }
 		}
 
-		static async Task UpdateStatsAsync(BatchReader<OutputChunk> batchReader, OutputStats outputStats, Task writeTask, IProgress<IExtractStats> progress, TimeSpan frequency, CancellationToken cancellationToken)
+		static async Task UpdateStatsAsync(BatchBlobReader<OutputChunk> batchReader, OutputStats outputStats, Task writeTask, IProgress<IExtractStats> progress, TimeSpan frequency, CancellationToken cancellationToken)
 		{
 			Stopwatch timer = Stopwatch.StartNew();
 			double lastOutputTime = 0.0;
