Horde: Add reusable pipeline classes for reading chunks and writing chunks to disk. Includes improvements to extract performance by reducing number of file open/closes.

#rnx

diff --git a/Engine/Source/Programs/Shared/EpicGames.Horde/Storage/BatchBlobReader.cs b/Engine/Source/Programs/Shared/EpicGames.Horde/Storage/BatchBlobReader.cs
--- a/Engine/Source/Programs/Shared/EpicGames.Horde/Storage/BatchBlobReader.cs
+++ b/Engine/Source/Programs/Shared/EpicGames.Horde/Storage/BatchBlobReader.cs
@@ -17,39 +17,17 @@ namespace EpicGames.Horde.Storage
 	/// <summary>
 	/// Request for a blob to be read
 	/// </summary>
-	/// <typeparam name="TContext">User defined context type for the request</typeparam>
 	/// <param name="BlobRef">Reference to the blob data</param>
-	/// <param name="Context">User defined context</param>
-	public record class BlobReadRequest<TContext>(IHashedBlobRef BlobRef, TContext Context);
+	public record class BlobReadRequest(IHashedBlobRef BlobRef);
 
 	/// <summary>
-	/// Individual response for a blob that was read
+	/// Response from a read request
 	/// </summary>
-	/// <typeparam name="TContext">User defined context type for the request</typeparam>
-	/// <param name="BlobData">Data for the blob</param>
-	/// <param name="Context">User defined context</param>
-	public sealed record class BlobReadResponse<TContext>(BlobData BlobData, TContext Context) : IDisposable
+	public record struct BlobReadResponse<TRequest>(TRequest Request, BlobData Data) : IDisposable
 	{
 		/// <inheritdoc/>
 		public void Dispose()
-			=> BlobData.Dispose();
-	}
-
-	/// <summary>
-	/// Response containing a batch of items that were read
-	/// </summary>
-	/// <typeparam name="TContext">User defined context type</typeparam>
-	/// <param name="Items">Items for this response</param>
-	public sealed record class BatchBlobReadResponse<TContext>(BlobReadResponse<TContext>[] Items) : IDisposable
-	{
-		/// <inheritdoc/>
-		public void Dispose()
-		{
-			foreach (BlobReadResponse<TContext> item in Items)
-			{
-				item.Dispose();
-			}
-		}
+			=> Data.Dispose();
 	}
 
 	/// <summary>
@@ -60,8 +38,8 @@ namespace EpicGames.Horde.Storage
 	/// <summary>
 	/// Implements an efficient pipeline for streaming blob data
 	/// </summary>
-	/// <typeparam name="TContext">User-defined context type</typeparam>
-	public sealed class BatchBlobReader<TContext> : IDisposable
+	public abstract class BatchBlobReader<TRequest> : IDisposable
+		where TRequest : BlobReadRequest
 	{
 		record class BundleRequest(BundleHandle Handle, int MinOffset, int MaxOffset, List<PacketRequest> Packets);
 
@@ -77,7 +55,7 @@ namespace EpicGames.Horde.Storage
 			public int MaxOffset => Handle.PacketOffset + Handle.PacketLength;
 		}
 
-		record class ExportRequest(IoHash Hash, ExportHandle Handle, TContext Context);
+		record class ExportRequest(IoHash Hash, ExportHandle Handle, TRequest OriginalRequest);
 
 		record class PacketResponse(IReadOnlyMemoryOwner<byte> Data, List<ExportRequest> Exports) : IDisposable
 		{
@@ -100,12 +78,11 @@ namespace EpicGames.Horde.Storage
 		/// </summary>
 		public bool VerifyHashes { get; set; }
 
-		readonly Channel<BlobReadRequest<TContext>> _otherRequests;
+		readonly Channel<TRequest> _otherRequests;
 		readonly Channel<BundleRequest> _bundleRequests;
 
 		readonly Channel<BundleResponse> _bundleResponses;
 		readonly Channel<PacketResponse> _packetResponses;
-		readonly Channel<BatchBlobReadResponse<TContext>> _blobResponses;
 
 		int _numRequests;
 		int _numReads;
@@ -114,40 +91,40 @@ namespace EpicGames.Horde.Storage
 		long _bytesRead;
 		long _bytesDecoded;
 
-		/// <summary>
-		/// Access to the response channel
-		/// </summary>
-		public ChannelReader<BatchBlobReadResponse<TContext>> ResponseReader
-			=> _blobResponses.Reader;
-
 		/// <summary>
 		/// Constructor
 		/// </summary>
-		public BatchBlobReader()
+		protected BatchBlobReader()
 		{
-			_otherRequests = Channel.CreateUnbounded<BlobReadRequest<TContext>>();
+			_otherRequests = Channel.CreateUnbounded<TRequest>();
 			_bundleRequests = Channel.CreateUnbounded<BundleRequest>();
 			_bundleResponses = Channel.CreateBounded<BundleResponse>(new BoundedChannelOptions(32));
 			_packetResponses = Channel.CreateBounded<PacketResponse>(new BoundedChannelOptions(128));
-			_blobResponses = Channel.CreateBounded<BatchBlobReadResponse<TContext>>(new BoundedChannelOptions(128));
 		}
 
 		/// <inheritdoc/>
 		public void Dispose()
 		{
-			while (_bundleResponses.Reader.TryRead(out BundleResponse? item))
-			{
-				item.Dispose();
-			}
+			Dispose(true);
+			GC.SuppressFinalize(this);
+		}
 
-			while (_packetResponses.Reader.TryRead(out PacketResponse? item))
+		/// <summary>
+		/// Overridable dispose method
+		/// </summary>
+		protected virtual void Dispose(bool disposing)
+		{
+			if (disposing)
 			{
-				item.Dispose();
-			}
+				while (_bundleResponses.Reader.TryRead(out BundleResponse? item))
+				{
+					item.Dispose();
+				}
 
-			while (_blobResponses.Reader.TryRead(out BatchBlobReadResponse<TContext>? item))
-			{
-				item.Dispose();
+				while (_packetResponses.Reader.TryRead(out PacketResponse? item))
+				{
+					item.Dispose();
+				}
 			}
 		}
 
@@ -160,7 +137,7 @@ namespace EpicGames.Horde.Storage
 		/// <summary>
 		/// Adds a batch reader to the given pipeline
 		/// </summary>
-		internal void AddToPipeline(AsyncPipeline pipeline, int numReadTasks, int numDecodeTasks, ChannelReader<BlobReadRequest<TContext>> requestReader)
+		public Task[] AddToPipeline(AsyncPipeline pipeline, int numReadTasks, int numDecodeTasks, ChannelReader<TRequest> requestReader)
 		{
 			_ = pipeline.AddTask(cancellationToken => CreateBundleRequestsAsync(requestReader, cancellationToken));
 			Task otherTask = pipeline.AddTask(HandleOtherRequestsAsync);
@@ -173,37 +150,42 @@ namespace EpicGames.Horde.Storage
 			Task[] packetTasks = pipeline.AddTasks(numDecodeTasks, ReadPacketsAsync);
 
 			// Read all the output blobs and mark the blob response channel as complete
-			IEnumerable<Task> blobTasks = packetTasks.Append(otherTask);
-			_ = Task.WhenAll(blobTasks).ContinueWith(_ => _blobResponses.Writer.TryComplete(), TaskScheduler.Default);
+			return packetTasks.Append(otherTask).ToArray();
 		}
 
 		// Read a blob using the naive read pipeline
 		async Task HandleOtherRequestsAsync(CancellationToken cancellationToken)
 		{
-			await foreach (BlobReadRequest<TContext> request in _otherRequests.Reader.ReadAllAsync(cancellationToken))
+			await foreach (TRequest request in _otherRequests.Reader.ReadAllAsync(cancellationToken))
 			{
+#pragma warning disable CA2000
 				BlobData? blobData = null;
 				try
 				{
 					blobData = await request.BlobRef.ReadBlobDataAsync(cancellationToken);
 					VerifyHash(request.BlobRef, blobData.Data.Span, request.BlobRef.Hash);
-#pragma warning disable CA2000
-					BatchBlobReadResponse<TContext> response = new BatchBlobReadResponse<TContext>([new BlobReadResponse<TContext>(blobData, request.Context)]);
-					await _blobResponses.Writer.WriteAsync(response, cancellationToken);
-#pragma warning restore CA2000
+					await HandleResponsesAsync([new BlobReadResponse<TRequest>(request, blobData)], cancellationToken);
 				}
 				catch
 				{
 					blobData?.Dispose();
 				}
+#pragma warning restore CA2000
 			}
 		}
 
+		/// <summary>
+		/// Handle responses from the read
+		/// </summary>
+		/// <param name="responses">Responses from the read</param>
+		/// <param name="cancellationToken">Cancellation token for the operation</param>
+		protected abstract ValueTask HandleResponsesAsync(List<BlobReadResponse<TRequest>> responses, CancellationToken cancellationToken);
+
 		#region Grouping requests
 
-		record class PacketExport(FlushedPacketHandle PacketHandle, IoHash ExportHash, ExportHandle ExportHandle, TContext Context);
+		record class PacketExport(FlushedPacketHandle PacketHandle, IoHash ExportHash, ExportHandle ExportHandle, TRequest OriginalRequest);
 
-		async Task CreateBundleRequestsAsync(ChannelReader<BlobReadRequest<TContext>> requestReader, CancellationToken cancellationToken)
+		async Task CreateBundleRequestsAsync(ChannelReader<TRequest> requestReader, CancellationToken cancellationToken)
 		{
 			int queueLength = 0;
 			Queue<BundleHandle> bundleQueue = new Queue<BundleHandle>();
@@ -214,7 +196,7 @@ namespace EpicGames.Horde.Storage
 				// Fill the queue up to the max length
 				for (; ; )
 				{
-					BlobReadRequest<TContext>? request;
+					TRequest? request;
 					if (!requestReader.TryRead(out request))
 					{
 						if (queueLength >= MinQueueLength)
@@ -279,7 +261,7 @@ namespace EpicGames.Horde.Storage
 			List<PacketRequest> packetRequests = new List<PacketRequest>();
 			foreach (IGrouping<int, PacketExport> group in exports.GroupBy(x => x.PacketHandle.PacketOffset))
 			{
-				List<ExportRequest> exportRequests = group.Select(x => new ExportRequest(x.ExportHash, x.ExportHandle, x.Context)).ToList();
+				List<ExportRequest> exportRequests = group.Select(x => new ExportRequest(x.ExportHash, x.ExportHandle, x.OriginalRequest)).ToList();
 				PacketRequest packetRequest = new PacketRequest(group.First().PacketHandle, exportRequests);
 				packetRequests.Add(packetRequest);
 			}
@@ -302,11 +284,11 @@ namespace EpicGames.Horde.Storage
 			}
 		}
 
-		static bool TryGetPacketExport(BlobReadRequest<TContext> chunk, [NotNullWhen(true)] out PacketExport? export)
+		static bool TryGetPacketExport(TRequest request, [NotNullWhen(true)] out PacketExport? export)
 		{
-			if (chunk.BlobRef.Innermost is ExportHandle exportHandle && exportHandle.Packet is FlushedPacketHandle packetHandle)
+			if (request.BlobRef.Innermost is ExportHandle exportHandle && exportHandle.Packet is FlushedPacketHandle packetHandle)
 			{
-				export = new PacketExport(packetHandle, chunk.BlobRef.Hash, exportHandle, chunk.Context);
+				export = new PacketExport(packetHandle, request.BlobRef.Hash, exportHandle, request);
 				return true;
 			}
 			else
@@ -362,56 +344,43 @@ namespace EpicGames.Horde.Storage
 			await foreach (BundleResponse bundleResponse in _bundleResponses.Reader.ReadAllAsync(cancellationToken))
 			{
 				using IDisposable lifetime = bundleResponse;
-				foreach (PacketRequest packetRequest in bundleResponse.Packets)
+
+#pragma warning disable CA2000
+				List<BlobReadResponse<TRequest>> responses = new List<BlobReadResponse<TRequest>>();
+				try
 				{
-					// Decode the data for this packet
-					ReadOnlyMemory<byte> memory = bundleResponse.Data.Memory.Slice(packetRequest.Handle.PacketOffset - bundleResponse.MinOffset, packetRequest.Handle.PacketLength);
+					foreach (PacketRequest packetRequest in bundleResponse.Packets)
+					{
+						// Decode the data for this packet
+						ReadOnlyMemory<byte> memory = bundleResponse.Data.Memory.Slice(packetRequest.Handle.PacketOffset - bundleResponse.MinOffset, packetRequest.Handle.PacketLength);
 
-					using PacketReader packetReader = packetRequest.Handle.CreatePacketReader(memory);
-					Interlocked.Increment(ref _numPackets);
-					Interlocked.Add(ref _bytesDecoded, packetReader.Packet.Length);
+						using PacketReader packetReader = packetRequest.Handle.CreatePacketReader(memory);
+						Interlocked.Increment(ref _numPackets);
+						Interlocked.Add(ref _bytesDecoded, packetReader.Packet.Length);
 
-					// Create responses for each blob
-					BlobReadResponse<TContext>[] items = new BlobReadResponse<TContext>[packetRequest.Exports.Count];
-					try
-					{
-#pragma warning disable CA2000
-						for (int idx = 0; idx < items.Length; idx++)
+						// Create responses for each blob
+						for (int idx = 0; idx < packetRequest.Exports.Count; idx++)
 						{
 							ExportRequest exportRequest = packetRequest.Exports[idx];
 							BlobData blobData = packetReader.ReadExport(exportRequest.Handle.ExportIdx);
+							responses.Add(new BlobReadResponse<TRequest>(exportRequest.OriginalRequest, blobData));
 							VerifyHash(exportRequest.Handle, blobData.Data.Span, exportRequest.Hash);
-							items[idx] = new BlobReadResponse<TContext>(blobData, exportRequest.Context);
 						}
-
-						BatchBlobReadResponse<TContext> response = new BatchBlobReadResponse<TContext>(items);
-						await _blobResponses.Writer.WriteAsync(response, cancellationToken);
-#pragma warning restore CA2000
 					}
-					catch
+					await HandleResponsesAsync(responses, cancellationToken);
+				}
+				catch
+				{
+					foreach (BlobReadResponse<TRequest> response in responses)
 					{
-						for (int idx = 0; idx < items.Length; idx++)
-						{
-							items[idx]?.Dispose();
-						}
-						throw;
+						response.Dispose();
 					}
+					throw;
 				}
+#pragma warning restore CA2000
 			}
 		}
 
 		#endregion
 	}
-
-	/// <summary>
-	/// Extension methods for <see cref="BatchBlobReader{TContext}"/>
-	/// </summary>
-	public static class BatchReaderExtensions
-	{
-		/// <summary>
-		/// Add tasks for reading to the given async pipeline 
-		/// </summary>
-		public static void AddTasks<T>(this AsyncPipeline pipeline, int numReadTasks, int numDecodeTasks, BatchBlobReader<T> reader, ChannelReader<BlobReadRequest<T>> requestReader)
-			=> reader.AddToPipeline(pipeline, numReadTasks, numDecodeTasks, requestReader);
-	}
 }
diff --git a/Engine/Source/Programs/Shared/EpicGames.Horde/Storage/BatchChunkReader.cs b/Engine/Source/Programs/Shared/EpicGames.Horde/Storage/BatchChunkReader.cs
new file mode 100644
--- /dev/null
+++ b/Engine/Source/Programs/Shared/EpicGames.Horde/Storage/BatchChunkReader.cs
@@ -0,0 +1,86 @@
+// Copyright Epic Games, Inc. All Rights Reserved.
+
+using System;
+using System.Collections.Generic;
+using System.Threading;
+using System.Threading.Channels;
+using System.Threading.Tasks;
+
+namespace EpicGames.Horde.Storage
+{
+	/// <summary>
+	/// Request for a chunk to be read before being written to disk
+	/// </summary>
+	/// <param name="File">File to read from</param>
+	/// <param name="Offset">Offset within the output file</param>
+	/// <param name="Handle">Handle to the blob data</param>
+	record class ChunkReadRequest(OutputFile File, long Offset, IHashedBlobRef Handle) : BlobReadRequest(Handle);
+
+	/// <summary>
+	/// Implementation of <see cref="BatchBlobReader{TRequest}"/> designed for reading leaf chunks of data from storage
+	/// </summary>
+	class BatchChunkReader : BatchBlobReader<ChunkReadRequest>, IDisposable
+	{
+		readonly Channel<OutputChunk[]> _channel;
+
+		/// <summary>
+		/// Maximum number of exports to write in a single request
+		/// </summary>
+		public int MaxBatchSize { get; } = 20;
+
+		/// <summary>
+		/// Reader for output chunks
+		/// </summary>
+		public ChannelReader<OutputChunk[]> ResponseReader => _channel.Reader;
+
+		/// <summary>
+		/// Constructor
+		/// </summary>
+		public BatchChunkReader()
+		{
+			_channel = Channel.CreateBounded<OutputChunk[]>(new BoundedChannelOptions(128));
+		}
+
+		/// <inheritdoc/>
+		protected override void Dispose(bool disposing)
+		{
+			base.Dispose(disposing);
+
+			if (disposing)
+			{
+				while (_channel.Reader.TryRead(out OutputChunk[]? array))
+				{
+					foreach (OutputChunk chunk in array)
+					{
+						chunk.Dispose();
+					}
+				}
+			}
+		}
+
+		/// <inheritdoc/>
+		protected override async ValueTask HandleResponsesAsync(List<BlobReadResponse<ChunkReadRequest>> responses, CancellationToken cancellationToken)
+		{
+			for (int idx = 0; idx < responses.Count;)
+			{
+				int firstIdx = idx;
+
+				List<OutputChunk> batch = new List<OutputChunk>();
+				for (; idx < responses.Count && batch.Count < MaxBatchSize && responses[idx].Request.File == responses[firstIdx].Request.File; idx++)
+				{
+					BlobReadResponse<ChunkReadRequest> response = responses[idx];
+					OutputChunk chunk = new OutputChunk(response.Request.File, response.Request.Offset, response.Data.Data, response.Request.Handle, response.Data);
+					batch.Add(chunk);
+				}
+
+				await _channel.Writer.WriteAsync(batch.ToArray(), cancellationToken);
+			}
+		}
+
+		/// <summary>
+		/// Marks the chunk stream as complete
+		/// </summary>
+		public void TryComplete()
+			=> _channel.Writer.TryComplete();
+	}
+}
diff --git a/Engine/Source/Programs/Shared/EpicGames.Horde/Storage/BatchOutputWriter.cs b/Engine/Source/Programs/Shared/EpicGames.Horde/Storage/BatchOutputWriter.cs
new file mode 100644
--- /dev/null
+++ b/Engine/Source/Programs/Shared/EpicGames.Horde/Storage/BatchOutputWriter.cs
@@ -0,0 +1,358 @@
+// Copyright Epic Games, Inc. All Rights Reserved.
+
+using System;
+using System.IO;
+using System.IO.MemoryMappedFiles;
+using System.Threading;
+using System.Threading.Channels;
+using System.Threading.Tasks;
+using EpicGames.Core;
+using EpicGames.Horde.Storage.Nodes;
+using EpicGames.Horde.Utilities;
+using Microsoft.Extensions.Logging;
+
+namespace EpicGames.Horde.Storage
+{
+	/// <summary>
+	/// Represents an output file that write requests may be issued against
+	/// </summary>
+	class OutputFile
+	{
+		/// <summary>
+		/// The relative output path
+		/// </summary>
+		public string Path { get; }
+
+		/// <summary>
+		/// File metadata
+		/// </summary>
+		public FileInfo FileInfo { get; }
+
+		/// <summary>
+		/// File entry with metadata about the target file
+		/// </summary>
+		public FileEntry FileEntry { get; }
+
+		bool _createdFile;
+		int _remainingChunks;
+
+		/// <summary>
+		/// Constructor
+		/// </summary>
+		public OutputFile(string path, FileInfo fileInfo, FileEntry fileEntry)
+		{
+			Path = path;
+			FileInfo = fileInfo;
+			FileEntry = fileEntry;
+		}
+
+		public int IncrementRemaining() => Interlocked.Increment(ref _remainingChunks);
+		public int DecrementRemaining() => Interlocked.Decrement(ref _remainingChunks);
+
+		/// <summary>
+		/// Opens the file for writing, setting its length on the first run
+		/// </summary>
+		public FileStream OpenStream()
+		{
+			lock (FileEntry)
+			{
+				if (!_createdFile)
+				{
+					if (FileInfo.Exists)
+					{
+						if (FileInfo.LinkTarget != null)
+						{
+							FileInfo.Delete();
+						}
+						else if (FileInfo.IsReadOnly)
+						{
+							FileInfo.IsReadOnly = false;
+						}
+					}
+					else
+					{
+						FileInfo.Directory?.Create();
+					}
+				}
+
+				FileStream? stream = null;
+				try
+				{
+					try
+					{
+						stream = FileInfo.Open(FileMode.OpenOrCreate, FileAccess.ReadWrite, FileShare.ReadWrite);
+					}
+					catch (IOException ex)
+					{
+						string? lockInfo = FileUtils.GetFileLockInfo(FileInfo.FullName);
+						if (lockInfo == null)
+						{
+							throw;
+						}
+						else
+						{
+							throw new WrappedFileOrDirectoryException(ex, $"{ex.Message}\n{lockInfo}");
+						}
+					}
+
+					if (!_createdFile)
+					{
+						stream.SetLength(FileEntry.Length);
+						_createdFile = true;
+					}
+					return stream;
+				}
+				catch
+				{
+					stream?.Dispose();
+					throw;
+				}
+			}
+		}
+
+		/// <summary>
+		/// Callback for a file having been fully written
+		/// </summary>
+		public virtual ValueTask CompleteAsync(CancellationToken cancellationToken)
+			=> default;
+	}
+
+	/// <summary>
+	/// Request to output a block of data A chunk which needs to be written to an output file
+	/// </summary>
+	record class OutputChunk(OutputFile File, long Offset, ReadOnlyMemory<byte> Data, object Source, IDisposable? DataOwner) : IDisposable
+	{
+		/// <inheritdoc/>
+		public void Dispose()
+			=> DataOwner?.Dispose();
+	}
+
+	/// <summary>
+	/// Batches file write requests
+	/// </summary>
+	class BatchOutputWriter
+	{
+		readonly ChannelReader<OutputChunk[]> _source;
+		readonly ILogger _logger;
+
+		int _writtenFiles;
+		long _writtenBytes;
+
+		/// <summary>
+		/// Number of files that have been written
+		/// </summary>
+		public int WrittenFiles
+			=> Interlocked.CompareExchange(ref _writtenFiles, 0, 0);
+
+		/// <summary>
+		/// Total number of bytes that have been written
+		/// </summary>
+		public long WrittenBytes
+			=> Interlocked.CompareExchange(ref _writtenBytes, 0, 0);
+
+		/// <summary>
+		/// Number of writes to execute sequentially vs in parallel
+		/// </summary>
+		public int WriteBatchSize { get; set; } = 64;
+
+		/// <summary>
+		/// If false, disables output to disk. Useful for performance testing.
+		/// </summary>
+		public bool EnableOutput { get; set; } = true;
+
+		/// <summary>
+		/// If true, hashes output files after writing to verify their integrity
+		/// </summary>
+		public bool VerifyOutput { get; set; }
+
+		/// <summary>
+		/// If true, outputs verbose information to the log
+		/// </summary>
+		public bool VerboseOutput { get; set; }
+
+		/// <summary>
+		/// Constructor
+		/// </summary>
+		public BatchOutputWriter(ChannelReader<OutputChunk[]> reader, ILogger logger)
+		{
+			_source = reader;
+			_logger = logger;
+		}
+
+		/// <summary>
+		/// Processes requests from the given input channel
+		/// </summary>
+		/// <param name="cancellationToken">Cancellation token for the operation</param>
+		public async Task ProcessAsync(CancellationToken cancellationToken)
+		{
+			await foreach (OutputChunk[] chunks in _source.ReadAllAsync(cancellationToken))
+			{
+				try
+				{
+					await WriteChunksAsync(chunks, cancellationToken);
+				}
+				catch (Exception ex)
+				{
+					_logger.LogError(ex, "Error writing to disk: {Message}", ex.Message);
+					throw;
+				}
+				finally
+				{
+					foreach (OutputChunk chunk in chunks)
+					{
+						chunk.Dispose();
+					}
+				}
+			}
+		}
+
+		async Task WriteChunksAsync(OutputChunk[] chunks, CancellationToken cancellationToken)
+		{
+			for (int requestIdx = 0; requestIdx < chunks.Length;)
+			{
+				OutputFile outputFile = chunks[requestIdx].File;
+
+				int maxChunkIdx = requestIdx + 1;
+				while (maxChunkIdx < chunks.Length && chunks[maxChunkIdx].File == outputFile)
+				{
+					maxChunkIdx++;
+				}
+
+				try
+				{
+					await ExtractChunksToFileAsync(outputFile, new ArraySegment<OutputChunk>(chunks, requestIdx, maxChunkIdx - requestIdx), cancellationToken);
+				}
+				catch (OperationCanceledException)
+				{
+					throw;
+				}
+				catch (Exception ex)
+				{
+					throw new StorageException($"Unable to extract {outputFile.FileInfo.FullName}: {ex.Message}", ex);
+				}
+
+				requestIdx = maxChunkIdx;
+			}
+		}
+
+		async Task ExtractChunksToFileAsync(OutputFile outputFile, ArraySegment<OutputChunk> chunks, CancellationToken cancellationToken)
+		{
+			using PlatformFileLock? fileLock = await PlatformFileLock.CreateAsync(cancellationToken);
+
+			// Open the file for the current chunk
+			int remainingChunks = 0;
+			await using (FileStream stream = outputFile.OpenStream())
+			{
+				if (outputFile.FileEntry.Length == 0 || !EnableOutput)
+				{
+					// If this file is empty, don't write anything and just move to the next chunk
+					for (int idx = 0; idx < chunks.Count; idx++)
+					{
+						remainingChunks = outputFile.DecrementRemaining();
+					}
+				}
+				else
+				{
+					// Process as many chunks as we can for this file
+					using MemoryMappedFile memoryMappedFile = MemoryMappedFile.CreateFromFile(stream, null, outputFile.FileEntry.Length, MemoryMappedFileAccess.ReadWrite, HandleInheritability.None, false);
+					using MemoryMappedView memoryMappedView = CreateMemoryMappedView(memoryMappedFile, outputFile);
+
+					for (int chunkIdx = 0; chunkIdx < chunks.Count; chunkIdx++)
+					{
+						OutputChunk chunk = chunks[chunkIdx];
+						cancellationToken.ThrowIfCancellationRequested();
+
+						if (VerboseOutput)
+						{
+							_logger.LogInformation("Writing {File},{Offset} ({Source})", outputFile.Path, chunk.Offset, chunk.Source);
+						}
+
+						// Write this chunk
+						ReadOnlyMemory<byte> memory = chunk.Data;
+						memory.CopyTo(memoryMappedView!.GetMemory(chunk.Offset, memory.Length));
+
+						// Update the stats
+						remainingChunks = outputFile.DecrementRemaining();
+						Interlocked.Add(ref _writtenBytes, memory.Length);
+					}
+				}
+			}
+
+			// Set correct permissions on the output file
+			if (remainingChunks == 0)
+			{
+				if (VerifyOutput)
+				{
+					IoHash computedHash;
+					using (Stream stream = outputFile.OpenStream())
+					{
+						computedHash = await IoHash.ComputeAsync(stream, cancellationToken);
+					}
+					if (computedHash != outputFile.FileEntry.StreamHash)
+					{
+						throw new InvalidDataException($"File {outputFile.Path} does not have correct hash after extract; expected {outputFile.FileEntry.StreamHash}, got {computedHash}.");
+					}
+					_logger.LogDebug("Verified hash of {Path}: {Hash}", outputFile.Path, computedHash);
+				}
+
+				outputFile.FileInfo.Refresh();
+				FileEntry.SetPermissions(outputFile.FileInfo!, outputFile.FileEntry.Flags);
+
+				if ((outputFile.FileEntry.Flags & FileEntryFlags.HasModTime) != 0)
+				{
+					outputFile.FileInfo.LastWriteTimeUtc = outputFile.FileEntry.ModTime;
+				}
+
+				Interlocked.Increment(ref _writtenFiles);
+
+				await outputFile.CompleteAsync(cancellationToken);
+			}
+		}
+
+		static MemoryMappedView CreateMemoryMappedView(MemoryMappedFile memoryMappedFile, OutputFile file)
+		{
+			try
+			{
+				MemoryMappedView view = new MemoryMappedView(memoryMappedFile, 0, file.FileEntry.Length);
+				return view;
+			}
+			catch (Exception ex)
+			{
+				file.FileInfo.Refresh();
+				throw new StorageException($"Unable to create mapped view of file (entry length: {file.FileEntry.Length}, file length: {file.FileInfo.Length}): {ex.Message}", ex);
+			}
+		}
+	}
+
+	/// <summary>
+	/// Writes items to a channel. Buffers one chunk until FlushAsync() is called to ensure
+	/// the remaining chunk reference count doesn't reach zero until the last chunk has been processed.
+	/// </summary>
+	class BufferedChannelWriter<T> where T : class
+	{
+		readonly ChannelWriter<T> _writer;
+		T? _bufferedItem;
+
+		public BufferedChannelWriter(ChannelWriter<T> writer)
+			=> _writer = writer;
+
+		public async Task WriteAsync(T item, CancellationToken cancellationToken)
+		{
+			if (_bufferedItem != null)
+			{
+				await _writer.WriteAsync(_bufferedItem, cancellationToken);
+			}
+
+			_bufferedItem = item;
+		}
+
+		public async Task FlushAsync(CancellationToken cancellationToken)
+		{
+			if (_bufferedItem != null)
+			{
+				await _writer.WriteAsync(_bufferedItem, cancellationToken);
+				_bufferedItem = null;
+			}
+		}
+	}
+}
diff --git a/Engine/Source/Programs/Shared/EpicGames.Horde/Storage/Nodes/DirectoryNode.Extract.cs b/Engine/Source/Programs/Shared/EpicGames.Horde/Storage/Nodes/DirectoryNode.Extract.cs
--- a/Engine/Source/Programs/Shared/EpicGames.Horde/Storage/Nodes/DirectoryNode.Extract.cs
+++ b/Engine/Source/Programs/Shared/EpicGames.Horde/Storage/Nodes/DirectoryNode.Extract.cs
@@ -4,13 +4,10 @@ using System;
 using System.Collections.Generic;
 using System.Diagnostics;
 using System.IO;
-using System.IO.MemoryMappedFiles;
-using System.Linq;
 using System.Threading;
 using System.Threading.Channels;
 using System.Threading.Tasks;
 using EpicGames.Core;
-using EpicGames.Horde.Utilities;
 using Microsoft.Extensions.Logging;
 
 namespace EpicGames.Horde.Storage.Nodes
@@ -150,130 +147,6 @@ namespace EpicGames.Horde.Storage.Nodes
 	/// </summary>
 	public static class DirectoryNodeExtract
 	{
-		class OutputStats
-		{
-			public int _writtenFiles;
-			public long _writtenBytes;
-		}
-
-		class OutputFile
-		{
-			public string Path { get; }
-			public FileInfo FileInfo { get; }
-			public FileEntry FileEntry { get; }
-
-			bool _createdFile;
-			int _remainingChunks;
-
-			public OutputFile(string path, FileInfo fileInfo, FileEntry fileEntry)
-			{
-				Path = path;
-				FileInfo = fileInfo;
-				FileEntry = fileEntry;
-			}
-
-			public int IncrementRemaining() => Interlocked.Increment(ref _remainingChunks);
-			public int DecrementRemaining() => Interlocked.Decrement(ref _remainingChunks);
-
-			public FileStream OpenStream()
-			{
-				lock (FileEntry)
-				{
-					if (!_createdFile)
-					{
-						if (FileInfo.Exists)
-						{
-							if (FileInfo.LinkTarget != null)
-							{
-								FileInfo.Delete();
-							}
-							else if (FileInfo.IsReadOnly)
-							{
-								FileInfo.IsReadOnly = false;
-							}
-						}
-						else
-						{
-							FileInfo.Directory?.Create();
-						}
-					}
-
-					FileStream? stream = null;
-					try
-					{
-						try
-						{
-							stream = FileInfo.Open(FileMode.OpenOrCreate, FileAccess.ReadWrite, FileShare.ReadWrite);
-						}
-						catch (IOException ex)
-						{
-							string? lockInfo = FileUtils.GetFileLockInfo(FileInfo.FullName);
-							if (lockInfo == null)
-							{
-								throw;
-							}
-							else
-							{
-								throw new WrappedFileOrDirectoryException(ex, $"{ex.Message}\n{lockInfo}");
-							}
-						}
-
-						if (!_createdFile)
-						{
-							stream.SetLength(FileEntry.Length);
-							_createdFile = true;
-						}
-						return stream;
-					}
-					catch
-					{
-						stream?.Dispose();
-						throw;
-					}
-				}
-			}
-		}
-
-		record class OutputChunk(OutputFile File, long Offset, long Length, IBlobRef Handle);
-
-		record class OutputBatch(List<OutputChunk> Chunks);
-
-		// Writes output chunks to a channel. Buffers one chunk until FlushAsync() is called to ensure
-		// the remaining chunk reference count doesn't reach zero until the last chunk has been processed.
-		class OutputChunkWriter
-		{
-			public OutputFile OutputFile { get; }
-
-			readonly ChannelWriter<BlobReadRequest<OutputChunk>> _chunkWriter;
-			BlobReadRequest<OutputChunk>? _bufferedChunk;
-
-			public OutputChunkWriter(OutputFile file, ChannelWriter<BlobReadRequest<OutputChunk>> chunkWriter)
-			{
-				OutputFile = file;
-				_chunkWriter = chunkWriter;
-			}
-
-			public async Task WriteAsync(long offset, long length, IHashedBlobRef handle, CancellationToken cancellationToken)
-			{
-				if (_bufferedChunk != null)
-				{
-					await _chunkWriter.WriteAsync(_bufferedChunk, cancellationToken);
-				}
-
-				OutputFile.IncrementRemaining();
-				_bufferedChunk = new BlobReadRequest<OutputChunk>(handle, new OutputChunk(OutputFile, offset, length, handle));
-			}
-
-			public async Task FlushAsync(CancellationToken cancellationToken)
-			{
-				if (_bufferedChunk != null)
-				{
-					await _chunkWriter.WriteAsync(_bufferedChunk, cancellationToken);
-					_bufferedChunk = null;
-				}
-			}
-		}
-
 #pragma warning disable IDE0060
 		static void TraceBlobRead(string type, string path, IBlobRef handle, ILogger logger)
 		{
@@ -366,19 +239,24 @@ namespace EpicGames.Horde.Storage.Nodes
 
 			await using (AsyncPipeline pipeline = new AsyncPipeline(cancellationToken))
 			{
-				using BatchBlobReader<OutputChunk> batchReader = new BatchBlobReader<OutputChunk>();
-				batchReader.VerifyHashes = options.VerifyOutput;
+				using BatchChunkReader chunkReader = new BatchChunkReader();
+				chunkReader.VerifyHashes = options.VerifyOutput;
 
-				Channel<BlobReadRequest<OutputChunk>> blobRequests = Channel.CreateBounded<BlobReadRequest<OutputChunk>>(new BoundedChannelOptions(batchReader.MinQueueLength * 2));
+				Channel<ChunkReadRequest> blobRequests = Channel.CreateBounded<ChunkReadRequest>(new BoundedChannelOptions(chunkReader.MinQueueLength * 2));
 				_ = pipeline.AddTask(ctx => FindOutputChunksRootAsync(directoryInfo, directoryNode, blobRequests.Writer, options, logger, ctx));
-				pipeline.AddTasks(numReadTasks, numDecodeTasks, batchReader, blobRequests.Reader);
 
-				OutputStats outputStats = new OutputStats();
-				Task[] writeTasks = pipeline.AddTasks(numWriteTasks, ctx => WriteAsync(batchReader.ResponseReader, outputStats, options, logger, ctx));
+				Task[] readTasks = chunkReader.AddToPipeline(pipeline, numReadTasks, numDecodeTasks, blobRequests.Reader);
+				_ = Task.WhenAll(readTasks).ContinueWith(_ => chunkReader.TryComplete(), TaskScheduler.Default);
+
+				BatchOutputWriter outputWriter = new BatchOutputWriter(chunkReader.ResponseReader, logger);
+				outputWriter.VerifyOutput = options.VerifyOutput;
+				outputWriter.VerboseOutput = options.VerboseOutput;
+
+				Task[] writeTasks = pipeline.AddTasks(numWriteTasks, ctx => outputWriter.ProcessAsync(ctx));
 
 				if (options.Progress != null)
 				{
-					_ = pipeline.AddTask(ctx => UpdateStatsAsync(batchReader, outputStats, Task.WhenAll(writeTasks), options.Progress, options.ProgressUpdateFrequency, ctx));
+					_ = pipeline.AddTask(ctx => UpdateStatsAsync(chunkReader, outputWriter, Task.WhenAll(writeTasks), options.Progress, options.ProgressUpdateFrequency, ctx));
 				}
 
 				await pipeline.WaitForCompletionAsync();
@@ -389,13 +267,13 @@ namespace EpicGames.Horde.Storage.Nodes
 
 		#region Enumerate chunks
 
-		static async Task FindOutputChunksRootAsync(DirectoryInfo rootDir, DirectoryNode node, ChannelWriter<BlobReadRequest<OutputChunk>> chunks, ExtractOptions options, ILogger logger, CancellationToken cancellationToken)
+		static async Task FindOutputChunksRootAsync(DirectoryInfo rootDir, DirectoryNode node, ChannelWriter<ChunkReadRequest> writer, ExtractOptions options, ILogger logger, CancellationToken cancellationToken)
 		{
-			await FindOutputChunksForDirectoryAsync(rootDir, "", node, chunks, options, logger, cancellationToken);
-			chunks.Complete();
+			await FindOutputChunksForDirectoryAsync(rootDir, "", node, writer, options, logger, cancellationToken);
+			writer.Complete();
 		}
 
-		static async Task FindOutputChunksForDirectoryAsync(DirectoryInfo rootDir, string path, DirectoryNode node, ChannelWriter<BlobReadRequest<OutputChunk>> chunks, ExtractOptions options, ILogger logger, CancellationToken cancellationToken)
+		static async Task FindOutputChunksForDirectoryAsync(DirectoryInfo rootDir, string path, DirectoryNode node, ChannelWriter<ChunkReadRequest> writer, ExtractOptions options, ILogger logger, CancellationToken cancellationToken)
 		{
 			foreach (FileEntry fileEntry in node.Files)
 			{
@@ -403,7 +281,7 @@ namespace EpicGames.Horde.Storage.Nodes
 				FileInfo fileInfo = new FileInfo(Path.Combine(rootDir.FullName, filePath));
 				OutputFile outputFile = new OutputFile(filePath, fileInfo, fileEntry);
 
-				await FindOutputChunksForFileAsync(outputFile, chunks, options, logger, cancellationToken);
+				await FindOutputChunksForFileAsync(outputFile, writer, options, logger, cancellationToken);
 			}
 
 			foreach (DirectoryEntry directoryEntry in node.Directories)
@@ -412,27 +290,35 @@ namespace EpicGames.Horde.Storage.Nodes
 				TraceBlobRead("Directory", subPath, directoryEntry.Handle, logger);
 				DirectoryNode subDirectoryNode = await directoryEntry.Handle.ReadBlobAsync(cancellationToken);
 
-				await FindOutputChunksForDirectoryAsync(rootDir, subPath, subDirectoryNode, chunks, options, logger, cancellationToken);
+				await FindOutputChunksForDirectoryAsync(rootDir, subPath, subDirectoryNode, writer, options, logger, cancellationToken);
 			}
 		}
 
-		static async Task FindOutputChunksForFileAsync(OutputFile outputFile, ChannelWriter<BlobReadRequest<OutputChunk>> chunks, ExtractOptions options, ILogger logger, CancellationToken cancellationToken)
+		static async Task FindOutputChunksForFileAsync(OutputFile outputFile, ChannelWriter<ChunkReadRequest> writer, ExtractOptions options, ILogger logger, CancellationToken cancellationToken)
+		{
+			BufferedChannelWriter<ChunkReadRequest> bufferedWriter = new(writer);
+			await FindOutputChunksAsync(outputFile, bufferedWriter, 0, outputFile.FileEntry.Target, options, logger, cancellationToken);
+			await bufferedWriter.FlushAsync(cancellationToken);
+		}
+
+		static async Task WriteChunkRequestAsync(BufferedChannelWriter<ChunkReadRequest> writer, OutputFile outputFile, long offset, IHashedBlobRef blobRef, CancellationToken cancellationToken)
 		{
-			OutputChunkWriter outputWriter = new OutputChunkWriter(outputFile, chunks);
-			await FindOutputChunksAsync(outputFile.Path, outputWriter, 0, outputFile.FileEntry.Target, options, logger, cancellationToken);
-			await outputWriter.FlushAsync(cancellationToken);
+			outputFile.IncrementRemaining();
+			ChunkReadRequest request = new ChunkReadRequest(outputFile, offset, blobRef);
+			await writer.WriteAsync(request, cancellationToken);
 		}
 
-		static async Task<long> FindOutputChunksAsync(string path, OutputChunkWriter chunkWriter, long offset, ChunkedDataNodeRef dataRef, ExtractOptions options, ILogger logger, CancellationToken cancellationToken)
+		static async Task<long> FindOutputChunksAsync(OutputFile outputFile, BufferedChannelWriter<ChunkReadRequest> writer, long offset, ChunkedDataNodeRef dataRef, ExtractOptions options, ILogger logger, CancellationToken cancellationToken)
 		{
 			if (dataRef.Type == ChunkedDataNodeType.Leaf)
 			{
 				if (options.VerboseOutput)
 				{
-					logger.LogInformation("Queuing {Path},{Offset} ({Locator})", path, offset, dataRef.Handle.GetLocator());
+					logger.LogInformation("Queuing {Path},{Offset} ({Locator})", outputFile.Path, offset, dataRef.Handle.GetLocator());
 				}
 
-				await chunkWriter.WriteAsync(offset, dataRef.Length, dataRef.Handle, cancellationToken);
+				await WriteChunkRequestAsync(writer, outputFile, offset, dataRef.Handle, cancellationToken);
+
 				if (dataRef.Length < 0)
 				{
 					// Backwards compatibility hack for v2 format
@@ -444,12 +330,12 @@ namespace EpicGames.Horde.Storage.Nodes
 			else
 			{
 				using BlobData data = await dataRef.Handle.ReadBlobDataAsync(cancellationToken);
-				TraceBlobRead("Interior", chunkWriter.OutputFile.Path, dataRef.Handle, logger);
+				TraceBlobRead("Interior", outputFile.Path, dataRef.Handle, logger);
 
 				if (data.Type.Guid == LeafChunkedDataNodeConverter.BlobType.Guid)
 				{
-					await chunkWriter.WriteAsync(offset, dataRef.Length, dataRef.Handle, cancellationToken);
-					return data.Data.Length;
+					await WriteChunkRequestAsync(writer, outputFile, offset, dataRef.Handle, cancellationToken);
+					return dataRef.Length;
 				}
 				else
 				{
@@ -458,7 +344,7 @@ namespace EpicGames.Horde.Storage.Nodes
 					InteriorChunkedDataNode interiorNode = BlobSerializer.Deserialize<InteriorChunkedDataNode>(data);
 					foreach (ChunkedDataNodeRef childRef in interiorNode.Children)
 					{
-						length += await FindOutputChunksAsync(path, chunkWriter, offset + length, childRef, options, logger, cancellationToken);
+						length += await FindOutputChunksAsync(outputFile, writer, offset + length, childRef, options, logger, cancellationToken);
 					}
 
 					return length;
@@ -468,178 +354,6 @@ namespace EpicGames.Horde.Storage.Nodes
 
 		#endregion
 
-		#region Write to disk
-
-		static async Task WriteAsync(ChannelReader<BatchBlobReadResponse<OutputChunk>> batchReader, OutputStats outputStats, ExtractOptions options, ILogger logger, CancellationToken cancellationToken)
-		{
-			const int WriteBatchSize = 64;
-			while (await batchReader.WaitToReadAsync(cancellationToken))
-			{
-				BatchBlobReadResponse<OutputChunk>? batch;
-				while (batchReader.TryRead(out batch))
-				{
-					try
-					{
-						List<Task> tasks = new List<Task>();
-						foreach (IReadOnlyList<BlobReadResponse<OutputChunk>> group in batch.Items.Batch(WriteBatchSize))
-						{
-							tasks.Add(WriteChunksAsync(group.ToArray(), outputStats, options, logger, cancellationToken));
-						}
-						await Task.WhenAll(tasks);
-					}
-					catch (Exception ex)
-					{
-						logger.LogError(ex, "Error writing to disk: {Message}", ex.Message);
-						throw;
-					}
-					finally
-					{
-						batch.Dispose();
-					}
-				}
-			}
-		}
-
-		static async Task WriteChunksAsync(ArraySegment<BlobReadResponse<OutputChunk>> chunks, OutputStats outputStats, ExtractOptions options, ILogger logger, CancellationToken cancellationToken)
-		{
-			for (int chunkIdx = 0; chunkIdx < chunks.Count;)
-			{
-				OutputFile file = chunks[chunkIdx].Context.File;
-
-				int maxChunkIdx = chunkIdx + 1;
-				while (maxChunkIdx < chunks.Count && chunks[maxChunkIdx].Context.File == file)
-				{
-					maxChunkIdx++;
-				}
-
-				try
-				{
-					await ExtractChunksToFileAsync(file, chunks.Slice(chunkIdx, maxChunkIdx - chunkIdx), outputStats, options, logger, cancellationToken);
-					//await ExtractChunksToNullAsync(file, chunks.Slice(chunkIdx, maxChunkIdx - chunkIdx), outputStats, logger, cancellationToken);
-				}
-				catch (OperationCanceledException)
-				{
-					throw;
-				}
-				catch (Exception ex)
-				{
-					throw new StorageException($"Unable to extract {file?.FileInfo?.FullName}: {ex.Message}", ex);
-				}
-
-				chunkIdx = maxChunkIdx;
-			}
-		}
-
-		static async Task ExtractChunksToFileAsync(OutputFile file, ArraySegment<BlobReadResponse<OutputChunk>> chunks, OutputStats stats, ExtractOptions options, ILogger logger, CancellationToken cancellationToken)
-		{
-			using PlatformFileLock? fileLock = await PlatformFileLock.CreateAsync(cancellationToken);
-
-			// Open the file for the current chunk
-			int remainingChunks = 0;
-			await using (FileStream stream = file.OpenStream())
-			{
-				if (file.FileEntry.Length == 0)
-				{
-					// If this file is empty, don't write anything and just move to the next chunk
-					for (int idx = 0; idx < chunks.Count; idx++)
-					{
-						remainingChunks = file.DecrementRemaining();
-					}
-				}
-				else
-				{
-					// Process as many chunks as we can for this file
-					using MemoryMappedFile memoryMappedFile = MemoryMappedFile.CreateFromFile(stream, null, file.FileEntry.Length, MemoryMappedFileAccess.ReadWrite, HandleInheritability.None, false);
-					using MemoryMappedView memoryMappedView = CreateMemoryMappedView(memoryMappedFile, file);
-
-					for (int chunkIdx = 0; chunkIdx < chunks.Count; chunkIdx++)
-					{
-						BlobReadResponse<OutputChunk> chunk = chunks[chunkIdx];
-						cancellationToken.ThrowIfCancellationRequested();
-
-						if (options.VerboseOutput)
-						{
-							logger.LogInformation("Writing {File},{Offset} ({Locator})", file.Path, chunk.Context.Offset, chunk.Context.Handle.GetLocator());
-						}
-
-						// Write this chunk
-						TraceBlobRead("Leaf", chunk.Context.File.Path, chunk.Context.Handle, logger);
-						chunk.BlobData.Data.CopyTo(memoryMappedView!.GetMemory(chunk.Context.Offset, chunk.BlobData.Data.Length));
-
-						// Update the stats
-						remainingChunks = file.DecrementRemaining();
-						Interlocked.Add(ref stats._writtenBytes, chunk.Context.Length);
-					}
-				}
-			}
-
-			// Set correct permissions on the output file
-			if (remainingChunks == 0)
-			{
-				if (options.VerifyOutput)
-				{
-					IoHash computedHash;
-					using (Stream stream = file.OpenStream())
-					{
-						computedHash = await IoHash.ComputeAsync(stream, cancellationToken);
-					}
-					if (computedHash != file.FileEntry.StreamHash)
-					{
-						throw new InvalidDataException($"File {file.Path} does not have correct hash after extract; expected {file.FileEntry.StreamHash}, got {computedHash}.");
-					}
-					logger.LogDebug("Verified hash of {Path}: {Hash}", file.Path, computedHash);
-				}
-
-				file.FileInfo.Refresh();
-				FileEntry.SetPermissions(file.FileInfo!, file.FileEntry.Flags);
-
-				if ((file.FileEntry.Flags & FileEntryFlags.HasModTime) != 0)
-				{
-					file.FileInfo.LastWriteTimeUtc = file.FileEntry.ModTime;
-				}
-
-				Interlocked.Increment(ref stats._writtenFiles);
-			}
-		}
-
-		static MemoryMappedView CreateMemoryMappedView(MemoryMappedFile memoryMappedFile, OutputFile file)
-		{
-			try
-			{
-				MemoryMappedView view = new MemoryMappedView(memoryMappedFile, 0, file.FileEntry.Length);
-				return view;
-			}
-			catch (Exception ex)
-			{
-				file.FileInfo.Refresh();
-				throw new StorageException($"Unable to create mapped view of file (entry length: {file.FileEntry.Length}, file length: {file.FileInfo.Length}): {ex.Message}", ex);
-			}
-		}
-
-#pragma warning disable IDE0051
-		// Update counters for extracting chunks without writing any data. Useful for profiling bottlenecks in other stages of the pipeline.
-		static Task ExtractChunksToNullAsync(OutputFile file, ArraySegment<BlobReadResponse<OutputChunk>> chunks, OutputStats stats, ILogger logger, CancellationToken cancellationToken)
-		{
-			_ = logger;
-			_ = cancellationToken;
-
-			foreach (BlobReadResponse<OutputChunk> chunk in chunks)
-			{
-				Interlocked.Add(ref stats._writtenBytes, chunk.Context.Length);
-
-				int remainingChunks = file.DecrementRemaining();
-				if (remainingChunks == 0)
-				{
-					Interlocked.Increment(ref stats._writtenFiles);
-				}
-			}
-
-			return Task.CompletedTask;
-		}
-#pragma warning restore IDE0051
-
-		#endregion
-
 		#region Stats
 
 		class ExtractStats : IExtractStats
@@ -651,7 +365,7 @@ namespace EpicGames.Horde.Storage.Nodes
 			public double DownloadRate { get; set; }
 		}
 
-		static async Task UpdateStatsAsync(BatchBlobReader<OutputChunk> batchReader, OutputStats outputStats, Task writeTask, IProgress<IExtractStats> progress, TimeSpan frequency, CancellationToken cancellationToken)
+		static async Task UpdateStatsAsync(BatchChunkReader chunkReader, BatchOutputWriter outputWriter, Task writeTask, IProgress<IExtractStats> progress, TimeSpan frequency, CancellationToken cancellationToken)
 		{
 			Stopwatch timer = Stopwatch.StartNew();
 			double lastOutputTime = 0.0;
@@ -669,11 +383,11 @@ namespace EpicGames.Horde.Storage.Nodes
 				double timeSeconds = timer.Elapsed.TotalSeconds;
 
 				ExtractStats stats = new ExtractStats();
-				stats.NumFiles = Interlocked.CompareExchange(ref outputStats._writtenFiles, 0, 0);
-				stats.ExtractSize = Interlocked.CompareExchange(ref outputStats._writtenBytes, 0, 0);
+				stats.NumFiles = outputWriter.WrittenFiles;
+				stats.ExtractSize = outputWriter.WrittenBytes;
 				samples.Enqueue((timeSeconds, stats));
 
-				BatchReaderStats batchStats = batchReader.GetStats();
+				BatchReaderStats batchStats = chunkReader.GetStats();
 				stats.DownloadSize = batchStats.BytesRead;
 
 				(double TimeSeconds, ExtractStats Stats) firstSample;
