Horde: Decrease test time for AES transport test

Removes the largest buffer size to reduce test time.

diff --git a/Engine/Source/Programs/Shared/EpicGames.Horde.Tests/AesTransportTests.cs b/Engine/Source/Programs/Shared/EpicGames.Horde.Tests/AesTransportTests.cs
--- a/Engine/Source/Programs/Shared/EpicGames.Horde.Tests/AesTransportTests.cs
+++ b/Engine/Source/Programs/Shared/EpicGames.Horde.Tests/AesTransportTests.cs
@@ -25,7 +25,7 @@ public sealed class AesTransportTests : IAsyncDisposable
 	private readonly MemoryComputeTransport _memoryTransport;
 	private readonly AesTransport _aesTransport;
 	private readonly Random _random = new(42);
-	private readonly int[] _bufferSizes = [1, 2, 3, 10, 200000, 2000000];
+	private readonly int[] _bufferSizes = [1, 2, 3, 10, 200000];
 	
 	public AesTransportTests()
 	{
