Horde: Use optimized batch pipeline for downloading archives into workspace folders.

#rnx

diff --git a/Engine/Source/Programs/Horde/HordeCmd/Commands/Workspace/WorkspaceDrop.cs b/Engine/Source/Programs/Horde/HordeCmd/Commands/Workspace/WorkspaceDrop.cs
--- a/Engine/Source/Programs/Horde/HordeCmd/Commands/Workspace/WorkspaceDrop.cs
+++ b/Engine/Source/Programs/Horde/HordeCmd/Commands/Workspace/WorkspaceDrop.cs
@@ -26,14 +26,14 @@ namespace Horde.Commands.Workspace
 			CancellationToken cancellationToken = CancellationToken.None;
 			RootDir ??= DirectoryReference.GetCurrentDirectory();
 
-			Workspace? workspace = await Workspace.TryOpenAsync(RootDir, logger, cancellationToken);
+			Workspace? workspace = await Workspace.TryFindAndOpenAsync(RootDir, logger, cancellationToken);
 			if (workspace == null)
 			{
 				logger.LogError("No workspace has been initialized in {RootDir}. Use 'workspace init' to create a new workspace.", RootDir);
 				return 1;
 			}
 
-			await workspace.SyncAsync(LayerId, new DirectoryNode(), cancellationToken);
+			await workspace.SyncAsync(LayerId, "", new DirectoryNode(), cancellationToken);
 			workspace.RemoveLayer(LayerId);
 			await workspace.SaveAsync(cancellationToken);
 
diff --git a/Engine/Source/Programs/Horde/HordeCmd/Commands/Workspace/WorkspaceSync.cs b/Engine/Source/Programs/Horde/HordeCmd/Commands/Workspace/WorkspaceSync.cs
--- a/Engine/Source/Programs/Horde/HordeCmd/Commands/Workspace/WorkspaceSync.cs
+++ b/Engine/Source/Programs/Horde/HordeCmd/Commands/Workspace/WorkspaceSync.cs
@@ -3,9 +3,9 @@
 using System.ComponentModel;
 using System.Diagnostics;
 using EpicGames.Core;
+using EpicGames.Horde;
+using EpicGames.Horde.Artifacts;
 using EpicGames.Horde.Storage;
-using EpicGames.Horde.Storage.Backends;
-using EpicGames.Horde.Storage.Bundles;
 using EpicGames.Horde.Storage.Nodes;
 using Microsoft.Extensions.Logging;
 
@@ -14,15 +14,23 @@ namespace Horde.Commands.Workspace
 	using Workspace = EpicGames.Horde.Storage.Workspace;
 
 	[Command("workspace", "sync", "Extracts an archive into the workspace")]
-	class WorkspaceSync : StorageCommandBase
+	class WorkspaceSync : Command
 	{
 		[CommandLine("-Root=")]
 		[Description("Root directory for the managed workspace.")]
 		public DirectoryReference? RootDir { get; set; }
 
-		[CommandLine("-File=")]
-		[Description("Path to a text file containing the node to extract to this workspace.")]
-		public FileReference? File { get; set; }
+		[CommandLine("-Dir=")]
+		[Description("Directory to extract files to. Must be under the workspace root.")]
+		public DirectoryReference? OutputDir { get; set; }
+
+		[CommandLine("-Artifact=")]
+		[Description("Artifact to extract into the workspace")]
+		public string? Artifact { get; set; }
+
+		[CommandLine("-Namespace=")]
+		[Description("Namespace containing the node or ref to extract")]
+		public string Namespace { get; set; } = "default";
 
 		[CommandLine("-Ref=")]
 		[Description("Name of a ref to extract to this workspace.")]
@@ -40,55 +48,94 @@ namespace Horde.Commands.Workspace
 		[Description("Outputs stats for the extraction operation.")]
 		public bool Stats { get; set; }
 
-		public WorkspaceSync(HttpStorageClient storageClient, BundleCache bundleCache)
-			: base(storageClient, bundleCache)
+		readonly IHordeClient _hordeClient;
+
+		public WorkspaceSync(IHordeClient hordeClient)
 		{
+			_hordeClient = hordeClient;
 		}
 
 		public override async Task<int> ExecuteAsync(ILogger logger)
 		{
-			if (File != null)
-			{
-				using MemoryMappedFileCache memoryMappedFileCache = new MemoryMappedFileCache();
-				IStorageNamespace store = BundleStorageNamespace.CreateFromDirectory(File.Directory, BundleCache, memoryMappedFileCache, logger);
-				IBlobRef handle = store.CreateBlobRef(await FileStorageBackend.ReadRefAsync(File));
-				return await ExecuteInternalAsync(store, handle, logger);
-			}
-			else if (Ref != null)
+			if (Artifact != null)
 			{
-				IStorageNamespace store = GetStorageNamespace();
-				IBlobRef handle = await store.ReadRefAsync(new RefName(Ref));
-				return await ExecuteInternalAsync(store, handle, logger);
-			}
-			else if (Node != null)
-			{
-				IStorageNamespace store = GetStorageNamespace();
-				IBlobRef handle = store.CreateBlobRef(new BlobLocator(Node));
+				ArtifactId artifactId = ArtifactId.Parse(Artifact);
+
+				IArtifact? artifact = await _hordeClient.Artifacts.GetAsync(artifactId);
+				if (artifact == null)
+				{
+					logger.LogError("Artifact {Id} not found", artifactId);
+					return 1;
+				}
+
+				logger.LogInformation("Downloading artifact {Id}: {Description}", artifactId, artifact.Description);
+
+				IStorageNamespace store = _hordeClient.GetStorageNamespace(artifact.Id);
+				IBlobRef handle = await store.ReadRefAsync<DirectoryNode>(new RefName("default"));
 				return await ExecuteInternalAsync(store, handle, logger);
 			}
 			else
 			{
-				throw new CommandLineArgumentException("Either -File=... or -Ref=... must be specified");
+				IStorageNamespace store = _hordeClient.GetStorageNamespace(new NamespaceId(Namespace));
+
+				IBlobRef handle;
+				if (Ref != null)
+				{
+					handle = await store.ReadRefAsync(new RefName(Ref));
+				}
+				else if (Node != null)
+				{
+					handle = store.CreateBlobRef(new BlobLocator(Node));
+				}
+				else
+				{
+					throw new CommandLineArgumentException("Either -File=... or -Ref=... must be specified");
+				}
+
+				return await ExecuteInternalAsync(store, handle, logger);
 			}
 		}
 
 		async Task<int> ExecuteInternalAsync(IStorageNamespace store, IBlobRef handle, ILogger logger)
 		{
-			RootDir ??= DirectoryReference.GetCurrentDirectory();
 			CancellationToken cancellationToken = CancellationToken.None;
 
-			Workspace? workspace = await Workspace.TryOpenAsync(RootDir, logger, cancellationToken);
-			if (workspace == null)
+			Workspace? workspace;
+			if (RootDir != null)
+			{
+				workspace = await Workspace.CreateOrOpenAsync(RootDir, logger, cancellationToken);
+			}
+			else
+			{
+				DirectoryReference workspaceDir = OutputDir ?? DirectoryReference.GetCurrentDirectory();
+
+				workspace = await Workspace.TryFindAndOpenAsync(workspaceDir, logger, cancellationToken);
+				if (workspace == null)
+				{
+					logger.LogInformation("Creating new workspace in {Path}", workspaceDir);
+					workspace = await Workspace.CreateAsync(workspaceDir, logger, cancellationToken);
+				}
+			}
+
+			string mountPath = "";
+			if (OutputDir != null)
 			{
-				logger.LogError("No workspace has been initialized in {RootDir}. Use 'workspace init' to create a new workspace.", RootDir);
-				return 1;
+				if (!OutputDir.IsUnderDirectory(workspace.RootDir))
+				{
+					logger.LogError("Output directory ({OutputDir}) is not under workspace root ({RootDir})", OutputDir, workspace.RootDir);
+					return 1;
+				}
+				mountPath = OutputDir.MakeRelativeTo(workspace.RootDir);
 			}
 
 			Stopwatch timer = Stopwatch.StartNew();
 			logger.LogInformation("Syncing into layer '{LayerId}'...", LayerId);
 
+			logger.LogInformation("Refreshing metadata...");
+			await workspace.RefreshAsync(cancellationToken);
+
 			DirectoryNode contents = await handle.ReadBlobAsync<DirectoryNode>();
-			await workspace.SyncAsync(LayerId, contents, cancellationToken);
+			await workspace.SyncAsync(LayerId, mountPath, contents, cancellationToken);
 			await workspace.SaveAsync(cancellationToken);
 
 			logger.LogInformation("Elapsed: {Time}s", timer.Elapsed.TotalSeconds);
diff --git a/Engine/Source/Programs/Horde/HordeCmd/Commands/Workspace/WorkspaceVerify.cs b/Engine/Source/Programs/Horde/HordeCmd/Commands/Workspace/WorkspaceVerify.cs
--- a/Engine/Source/Programs/Horde/HordeCmd/Commands/Workspace/WorkspaceVerify.cs
+++ b/Engine/Source/Programs/Horde/HordeCmd/Commands/Workspace/WorkspaceVerify.cs
@@ -21,7 +21,7 @@ namespace Horde.Commands.Workspace
 			CancellationToken cancellationToken = CancellationToken.None;
 			RootDir ??= DirectoryReference.GetCurrentDirectory();
 
-			Workspace? workspace = await Workspace.TryOpenAsync(RootDir, logger, cancellationToken);
+			Workspace? workspace = await Workspace.TryFindAndOpenAsync(RootDir, logger, cancellationToken);
 			if (workspace == null)
 			{
 				logger.LogError("No workspace has been initialized in {RootDir}. Use 'workspace init' to create a new workspace.", RootDir);
diff --git a/Engine/Source/Programs/Shared/EpicGames.Horde/Storage/Nodes/DirectoryNode.Extract.cs b/Engine/Source/Programs/Shared/EpicGames.Horde/Storage/Nodes/DirectoryNode.Extract.cs
--- a/Engine/Source/Programs/Shared/EpicGames.Horde/Storage/Nodes/DirectoryNode.Extract.cs
+++ b/Engine/Source/Programs/Shared/EpicGames.Horde/Storage/Nodes/DirectoryNode.Extract.cs
@@ -111,16 +111,32 @@ namespace EpicGames.Horde.Storage.Nodes
 		/// </summary>
 		public int? NumReadTasks { get; set; }
 
+		/// <summary>
+		/// Default number of read tasks to use
+		/// </summary>
+		public static int DefaultNumReadTasks { get; } = 16;
+
 		/// <summary>
 		/// Number of async tasks to spawn for decoding data
 		/// </summary>
 		public int? NumDecodeTasks { get; set; }
 
+		/// <summary>
+		/// Default number of decode tasks to use
+		/// </summary>
+		public static int DefaultNumDecodeTasks { get; } = Math.Min(Environment.ProcessorCount, 16);
+
 		/// <summary>
 		/// Number of async tasks to spawn for writing output
 		/// </summary>
 		public int? NumWriteTasks { get; set; }
 
+		/// <summary>
+		/// Default number of write tasks to use
+		/// </summary>
+		public static int DefaultNumWriteTasks(long downloadSize)
+			=> Math.Min(1 + (int)(downloadSize / (16 * 1024 * 1024)), 16);
+
 		/// <summary>
 		/// Output for progress updates
 		/// </summary>
@@ -232,8 +248,8 @@ namespace EpicGames.Horde.Storage.Nodes
 		{
 			options ??= new ExtractOptions();
 
-			int numReadTasks = options.NumReadTasks ?? 16;
-			int numDecodeTasks = options.NumDecodeTasks ?? Math.Min(Environment.ProcessorCount, 16);
+			int numReadTasks = options.NumReadTasks ?? ExtractOptions.DefaultNumReadTasks;
+			int numDecodeTasks = options.NumDecodeTasks ?? ExtractOptions.DefaultNumDecodeTasks;
 			int numWriteTasks = options.NumWriteTasks ?? Math.Min(1 + (int)(directoryNode.Length / (16 * 1024 * 1024)), 16);
 			logger.LogDebug("Using {NumReadTasks} read tasks, {NumDecodeTasks} decode tasks, {NumWriteTasks} write tasks", numReadTasks, numDecodeTasks, numWriteTasks);
 
@@ -294,7 +310,7 @@ namespace EpicGames.Horde.Storage.Nodes
 			}
 		}
 
-		static async Task FindOutputChunksForFileAsync(OutputFile outputFile, ChannelWriter<ChunkReadRequest> writer, ExtractOptions options, ILogger logger, CancellationToken cancellationToken)
+		internal static async Task FindOutputChunksForFileAsync(OutputFile outputFile, ChannelWriter<ChunkReadRequest> writer, ExtractOptions options, ILogger logger, CancellationToken cancellationToken)
 		{
 			BufferedChannelWriter<ChunkReadRequest> bufferedWriter = new(writer);
 			await FindOutputChunksAsync(outputFile, bufferedWriter, 0, outputFile.FileEntry.Target, options, logger, cancellationToken);
@@ -365,7 +381,7 @@ namespace EpicGames.Horde.Storage.Nodes
 			public double DownloadRate { get; set; }
 		}
 
-		static async Task UpdateStatsAsync(BatchChunkReader chunkReader, BatchOutputWriter outputWriter, Task writeTask, IProgress<IExtractStats> progress, TimeSpan frequency, CancellationToken cancellationToken)
+		internal static async Task UpdateStatsAsync(BatchChunkReader chunkReader, BatchOutputWriter outputWriter, Task writeTask, IProgress<IExtractStats> progress, TimeSpan frequency, CancellationToken cancellationToken)
 		{
 			Stopwatch timer = Stopwatch.StartNew();
 			double lastOutputTime = 0.0;
diff --git a/Engine/Source/Programs/Shared/EpicGames.Horde/Storage/Workspace.cs b/Engine/Source/Programs/Shared/EpicGames.Horde/Storage/Workspace.cs
--- a/Engine/Source/Programs/Shared/EpicGames.Horde/Storage/Workspace.cs
+++ b/Engine/Source/Programs/Shared/EpicGames.Horde/Storage/Workspace.cs
@@ -8,6 +8,7 @@ using System.IO;
 using System.Linq;
 using System.Text;
 using System.Threading;
+using System.Threading.Channels;
 using System.Threading.Tasks;
 using EpicGames.Core;
 using EpicGames.Horde.Storage.Nodes;
@@ -119,6 +120,16 @@ namespace EpicGames.Horde.Storage
 				}
 			}
 
+			public DirectoryState FindOrAddDirectoryByPath(string path)
+			{
+				DirectoryState dirState = this;
+				foreach (string fragment in path.Split('/'))
+				{
+					dirState = dirState.FindOrAddDirectory(fragment);
+				}
+				return dirState;
+			}
+
 			public string GetPath()
 			{
 				StringBuilder builder = new StringBuilder();
@@ -428,6 +439,20 @@ namespace EpicGames.Horde.Storage
 			return workspace;
 		}
 
+		/// <summary>
+		/// Create a new workspace instance in the given location. Opens the existing instance if it already contains workspace data.
+		/// </summary>
+		/// <param name="rootDir">Root directory for the workspace</param>
+		/// <param name="logger">Logger for output</param>
+		/// <param name="cancellationToken">Cancellation token for the operation</param>
+		/// <returns>Workspace instance</returns>
+		public static async Task<Workspace> CreateOrOpenAsync(DirectoryReference rootDir, ILogger logger, CancellationToken cancellationToken = default)
+		{
+			Workspace? workspace = await TryOpenAsync(rootDir, logger, cancellationToken);
+			workspace ??= await CreateAsync(rootDir, logger, cancellationToken);
+			return workspace;
+		}
+
 		/// <summary>
 		/// Attempts to open an existing workspace for the current directory. 
 		/// </summary>
@@ -437,20 +462,37 @@ namespace EpicGames.Horde.Storage
 		/// <returns>Workspace instance</returns>
 		public static async Task<Workspace?> TryOpenAsync(DirectoryReference currentDir, ILogger logger, CancellationToken cancellationToken = default)
 		{
-			for (DirectoryReference? testDir = currentDir; testDir != null; testDir = testDir.ParentDirectory)
+			FileReference stateFile = FileReference.Combine(currentDir, HordeDirName, StateFileName);
+			using (FileStream? stream = FileTransaction.OpenRead(stateFile))
 			{
-				FileReference stateFile = FileReference.Combine(testDir, HordeDirName, StateFileName);
-				using (FileStream? stream = FileTransaction.OpenRead(stateFile))
+				if (stream != null)
 				{
-					if (stream != null)
-					{
-						byte[] data = await stream.ReadAllBytesAsync(cancellationToken);
+					byte[] data = await stream.ReadAllBytesAsync(cancellationToken);
 
-						Workspace workspace = new Workspace(testDir, stateFile, logger);
-						workspace.Read(new MemoryReader(data));
+					Workspace workspace = new Workspace(currentDir, stateFile, logger);
+					workspace.Read(new MemoryReader(data));
 
-						return workspace;
-					}
+					return workspace;
+				}
+			}
+			return null;
+		}
+
+		/// <summary>
+		/// Attempts to open an existing workspace for the current directory. 
+		/// </summary>
+		/// <param name="currentDir">Root directory for the workspace</param>
+		/// <param name="logger">Logger for output</param>
+		/// <param name="cancellationToken">Cancellation token for the operation</param>
+		/// <returns>Workspace instance</returns>
+		public static async Task<Workspace?> TryFindAndOpenAsync(DirectoryReference currentDir, ILogger logger, CancellationToken cancellationToken = default)
+		{
+			for (DirectoryReference? testDir = currentDir; testDir != null; testDir = testDir.ParentDirectory)
+			{
+				Workspace? workspace = await TryOpenAsync(testDir, logger, cancellationToken);
+				if (workspace != null)
+				{
+					return workspace;
 				}
 			}
 			return null;
@@ -576,14 +618,35 @@ namespace EpicGames.Horde.Storage
 
 		#region Syncing
 
+		class WorkspaceOutputFile : OutputFile
+		{
+			readonly FileState _fileState;
+
+			public WorkspaceOutputFile(string path, FileInfo fileInfo, FileEntry fileEntry, FileState fileState)
+				: base(path, fileInfo, fileEntry)
+			{
+				_fileState = fileState;
+			}
+
+			public override ValueTask CompleteAsync(CancellationToken cancellationToken)
+			{
+				_fileState.Hash = FileEntry.StreamHash;
+				_fileState.Update(FileInfo);
+				return default;
+			}
+		}
+
 		/// <summary>
 		/// Syncs a layer to the given contents
 		/// </summary>
 		/// <param name="layerId">Identifier for the layer</param>
+		/// <param name="basePath">Base path within the workspace to sync to.</param>
 		/// <param name="contents">New contents for the layer</param>
 		/// <param name="cancellationToken">Cancellation token for the operation</param>
-		public async Task SyncAsync(WorkspaceLayerId layerId, DirectoryNode? contents, CancellationToken cancellationToken = default)
+		public async Task SyncAsync(WorkspaceLayerId layerId, string? basePath, DirectoryNode? contents, CancellationToken cancellationToken = default)
 		{
+			basePath = basePath?.Replace('\\', '/').Trim('/') ?? "";
+
 			LayerState layerState = GetLayerState(layerId) ?? throw new InvalidOperationException($"Layer '{layerId}' does not exist");
 
 			DirectoryState hordeDirState = _rootDirState.FindOrAddDirectory(HordeDirName);
@@ -592,10 +655,52 @@ namespace EpicGames.Horde.Storage
 			DirectoryState cacheDirState = hordeDirState.FindOrAddDirectory(CacheDirName);
 			cacheDirState.LayerFlags |= CacheLayerFlag;
 
-			await SyncDirectoryAsync(_rootDir, _rootDirState, contents, layerState.Flag, cacheDirState, cancellationToken);
+			DirectoryState syncDirState = _rootDirState;
+			foreach (string fragment in basePath.Split('/'))
+			{
+				syncDirState.LayerFlags |= layerState.Flag;
+				syncDirState = syncDirState.FindOrAddDirectory(fragment);
+			}
+
+			ExtractOptions options = new ExtractOptions();
+			options.Progress = new ExtractStatsLogger(_logger);
+
+			_logger.LogInformation("Scanning tree to sync...");
+
+			Channel<ChunkReadRequest> chunkRequests = Channel.CreateUnbounded<ChunkReadRequest>();
+
+			using BatchChunkReader batchReader = new BatchChunkReader { VerifyHashes = options.VerifyOutput };
+			await SyncDirectoryAsync(_rootDir, basePath, syncDirState, contents, layerState.Flag, cacheDirState, chunkRequests.Writer, options, cancellationToken);
+			chunkRequests.Writer.Complete();
+
+			if (chunkRequests.Reader.Count > 0)
+			{
+				_logger.LogInformation("Updating workspace...");
+
+				int numReadTasks = options.NumReadTasks ?? ExtractOptions.DefaultNumReadTasks;
+				int numDecodeTasks = options.NumDecodeTasks ?? ExtractOptions.DefaultNumDecodeTasks;
+				int numWriteTasks = 16;
+
+				await using (AsyncPipeline pipeline = new AsyncPipeline(cancellationToken))
+				{
+					Task[] readTasks = batchReader.AddToPipeline(pipeline, numReadTasks, numDecodeTasks, chunkRequests.Reader);
+					_ = Task.WhenAll(readTasks).ContinueWith(_ => batchReader.TryComplete(), TaskScheduler.Default);
+
+					BatchOutputWriter outputWriter = new BatchOutputWriter(batchReader.ResponseReader, _logger);
+
+					Task[] writeTasks = pipeline.AddTasks(numWriteTasks, ctx => outputWriter.ProcessAsync(ctx));
+
+					if (options.Progress != null)
+					{
+						_ = pipeline.AddTask(ctx => DirectoryNodeExtract.UpdateStatsAsync(batchReader, outputWriter, Task.WhenAll(writeTasks), options.Progress, options.ProgressUpdateFrequency, ctx));
+					}
+
+					await pipeline.WaitForCompletionAsync();
+				}
+			}
 		}
 
-		async Task SyncDirectoryAsync(DirectoryReference dirRef, DirectoryState dirState, DirectoryNode? dirNode, ulong flag, DirectoryState cacheDirState, CancellationToken cancellationToken)
+		async Task SyncDirectoryAsync(DirectoryReference dirRef, string path, DirectoryState dirState, DirectoryNode? dirNode, ulong flag, DirectoryState cacheDirState, ChannelWriter<ChunkReadRequest> readRequestWriter, ExtractOptions options, CancellationToken cancellationToken)
 		{
 			// Remove any directories that no longer exist
 			for (int subDirIdx = 0; subDirIdx < dirState.Directories.Count; subDirIdx++)
@@ -606,7 +711,7 @@ namespace EpicGames.Horde.Storage
 					if (dirNode == null || !dirNode.TryGetDirectoryEntry(subDirState.Name, out _))
 					{
 						DirectoryReference subDirPath = DirectoryReference.Combine(dirRef, subDirState.Name.ToString());
-						await SyncDirectoryAsync(subDirPath, subDirState, null, flag, cacheDirState, cancellationToken);
+						await SyncDirectoryAsync(subDirPath, $"{path}{subDirState.Name}/", subDirState, null, flag, cacheDirState, readRequestWriter, options, cancellationToken);
 
 						if (subDirState.LayerFlags == 0)
 						{
@@ -651,7 +756,7 @@ namespace EpicGames.Horde.Storage
 					DirectoryState subDirState = dirState.FindOrAddDirectory(subDirEntry.Name);
 
 					DirectoryNode subDirNode = await subDirEntry.Handle.ReadBlobAsync(cancellationToken: cancellationToken);
-					await SyncDirectoryAsync(subDirPath, subDirState, subDirNode, flag, cacheDirState, cancellationToken);
+					await SyncDirectoryAsync(subDirPath, $"{path}{subDirEntry.Name}/", subDirState, subDirNode, flag, cacheDirState, readRequestWriter, options, cancellationToken);
 
 					dirState.LayerFlags |= flag;
 				}
@@ -659,7 +764,7 @@ namespace EpicGames.Horde.Storage
 				// Update files
 				foreach (FileEntry fileEntry in dirNode.Files)
 				{
-					await SyncFileAsync(dirRef, dirState, fileEntry, flag, cancellationToken);
+					await SyncFileAsync(dirRef, $"{path}{fileEntry.Name}", dirState, fileEntry, flag, readRequestWriter, options, cancellationToken);
 					dirState.LayerFlags |= flag;
 				}
 			}
@@ -671,7 +776,7 @@ namespace EpicGames.Horde.Storage
 			}
 		}
 
-		async Task SyncFileAsync(DirectoryReference dirRef, DirectoryState dirState, FileEntry fileEntry, ulong flag, CancellationToken cancellationToken)
+		async Task SyncFileAsync(DirectoryReference dirRef, string path, DirectoryState dirState, FileEntry fileEntry, ulong flag, ChannelWriter<ChunkReadRequest> readRequestWriter, ExtractOptions options, CancellationToken cancellationToken)
 		{
 			FileState? fileState;
 			if (dirState.TryGetFile(fileEntry.Name, out fileState))
@@ -694,30 +799,15 @@ namespace EpicGames.Horde.Storage
 				{
 					fileState = dirState.FindOrAddFile(fileEntry.Name);
 
-					_logger.LogInformation("Updating {File} to {Hash}", fileInfo, fileEntry.StreamHash);
-					using (FileStream stream = fileInfo.Open(FileMode.Create, FileAccess.Write, FileShare.Read))
-					{
-						await ExtractDataAsync(fileEntry.Target, stream, cancellationToken);
-						if (stream.Length != fileEntry.Length)
-						{
-							throw new EndOfStreamException($"Incorrect length for extracted file {fileInfo.FullName}");
-						}
-					}
-
-					fileState.Hash = fileEntry.StreamHash;
-
-					fileInfo.Refresh();
-					fileState.Update(fileInfo);
-
-					AddFileToHashLookup(fileState);
+					WorkspaceOutputFile outputFile = new WorkspaceOutputFile(path, fileInfo, fileEntry, fileState);
+					await DirectoryNodeExtract.FindOutputChunksForFileAsync(outputFile, readRequestWriter, options, _logger, cancellationToken);
 				}
 				else
 				{
 					fileInfo.Refresh();
 					fileState.Update(fileInfo);
+					fileState.LayerFlags |= flag;
 				}
-
-				fileState.LayerFlags |= flag;
 			}
 		}
 
@@ -825,6 +915,48 @@ namespace EpicGames.Horde.Storage
 
 		#endregion
 
+		#region Refresh
+
+		/// <summary>
+		/// Updates the status of files in this workspace based on current filesystem metadata
+		/// </summary>
+		/// <param name="cancellationToken">Cancellation token for the operation</param>
+		public async Task RefreshAsync(CancellationToken cancellationToken = default)
+		{
+			await RefreshAsync(_rootDirState, cancellationToken);
+		}
+
+		async Task RefreshAsync(DirectoryState dirState, CancellationToken cancellationToken)
+		{
+			foreach (DirectoryState subDirState in dirState.Directories)
+			{
+				await RefreshAsync(subDirState, cancellationToken);
+			}
+
+			int outputIdx = 0;
+			for(int idx = 0; idx < dirState.Files.Count; idx++)
+			{
+				FileState fileState = dirState.Files[idx];
+
+				FileInfo fileInfo = GetFileInfo(fileState);
+				if (!fileInfo.Exists)
+				{
+					_logger.LogWarning("File {Path} has been removed", fileInfo.FullName);
+				}
+				else if (fileState.IsModified(fileInfo))
+				{
+					_logger.LogWarning("File {Path} has been modified", fileInfo.FullName);
+				}
+				else
+				{
+					dirState.Files[outputIdx++] = fileState;
+				}
+			}
+			dirState.Files.RemoveRange(outputIdx, dirState.Files.Count - outputIdx);
+		}
+
+		#endregion
+
 		#region Verify
 
 		/// <summary>
