Horde: Move functionality for enumerating leaf nodes in a file into a ChunkedDataNodeRef method.

#rnx

diff --git a/Engine/Source/Programs/Shared/EpicGames.Horde/Storage/BatchOutputWriter.cs b/Engine/Source/Programs/Shared/EpicGames.Horde/Storage/BatchOutputWriter.cs
--- a/Engine/Source/Programs/Shared/EpicGames.Horde/Storage/BatchOutputWriter.cs
+++ b/Engine/Source/Programs/Shared/EpicGames.Horde/Storage/BatchOutputWriter.cs
@@ -323,36 +323,4 @@ namespace EpicGames.Horde.Storage
 			}
 		}
 	}
-
-	/// <summary>
-	/// Writes items to a channel. Buffers one chunk until FlushAsync() is called to ensure
-	/// the remaining chunk reference count doesn't reach zero until the last chunk has been processed.
-	/// </summary>
-	class BufferedChannelWriter<T> where T : class
-	{
-		readonly ChannelWriter<T> _writer;
-		T? _bufferedItem;
-
-		public BufferedChannelWriter(ChannelWriter<T> writer)
-			=> _writer = writer;
-
-		public async Task WriteAsync(T item, CancellationToken cancellationToken)
-		{
-			if (_bufferedItem != null)
-			{
-				await _writer.WriteAsync(_bufferedItem, cancellationToken);
-			}
-
-			_bufferedItem = item;
-		}
-
-		public async Task FlushAsync(CancellationToken cancellationToken)
-		{
-			if (_bufferedItem != null)
-			{
-				await _writer.WriteAsync(_bufferedItem, cancellationToken);
-				_bufferedItem = null;
-			}
-		}
-	}
 }
diff --git a/Engine/Source/Programs/Shared/EpicGames.Horde/Storage/Nodes/ChunkedDataNode.cs b/Engine/Source/Programs/Shared/EpicGames.Horde/Storage/Nodes/ChunkedDataNode.cs
--- a/Engine/Source/Programs/Shared/EpicGames.Horde/Storage/Nodes/ChunkedDataNode.cs
+++ b/Engine/Source/Programs/Shared/EpicGames.Horde/Storage/Nodes/ChunkedDataNode.cs
@@ -6,6 +6,7 @@ using System.Buffers.Binary;
 using System.Collections.Generic;
 using System.IO;
 using System.Linq;
+using System.Runtime.CompilerServices;
 using System.Threading;
 using System.Threading.Tasks;
 using EpicGames.Core;
@@ -211,8 +212,57 @@ namespace EpicGames.Horde.Storage.Nodes
 				throw new InvalidDataException();
 			}
 		}
+
+		/// <summary>
+		/// Enumerates all the leaf nodes for this node
+		/// </summary>
+		/// <param name="cancellationToken">Cancellation token for the operation</param>
+		public async IAsyncEnumerable<LeafChunkedDataNodeRef> EnumerateLeafNodesAsync([EnumeratorCancellation] CancellationToken cancellationToken)
+		{
+			Stack<ChunkedDataNodeRef> stack = new Stack<ChunkedDataNodeRef>();
+			stack.Push(this);
+
+			while (stack.TryPop(out ChunkedDataNodeRef? dataRef))
+			{
+				if (dataRef.Type == ChunkedDataNodeType.Leaf)
+				{
+					long length = dataRef.Length;
+					if (dataRef.Length < 0)
+					{
+						// Backwards compatibility hack for v2 format
+						using BlobData data = await dataRef.Handle.ReadBlobDataAsync(cancellationToken);
+						length = data.Data.Length;
+					}
+
+					yield return new LeafChunkedDataNodeRef(length, dataRef.GetLeafHandle());
+				}
+				else
+				{
+					using BlobData data = await dataRef.Handle.ReadBlobDataAsync(cancellationToken);
+					if (data.Type.Guid == LeafChunkedDataNodeConverter.BlobType.Guid)
+					{
+						yield return new LeafChunkedDataNodeRef(dataRef.Length, dataRef.GetLeafHandle());
+					}
+					else
+					{
+						InteriorChunkedDataNode interiorNode = BlobSerializer.Deserialize<InteriorChunkedDataNode>(data);
+						for (int idx = interiorNode.Children.Count - 1; idx >= 0; idx--) // Push onto stack in reverse order for correct traversal order
+						{
+							stack.Push(interiorNode.Children[idx]);
+						}
+					}
+				}
+			}
+		}
 	}
 
+	/// <summary>
+	/// Reference to a <see cref="LeafChunkedDataNode"/>
+	/// </summary>
+	/// <param name="Length">Length of the leaf data</param>
+	/// <param name="BlobRef">Reference to the blob</param>
+	public record class LeafChunkedDataNodeRef(long Length, IHashedBlobRef<LeafChunkedDataNode> BlobRef);
+
 	/// <summary>
 	/// Stores the flat list of chunks produced from chunking a single data stream
 	/// </summary>
diff --git a/Engine/Source/Programs/Shared/EpicGames.Horde/Storage/Nodes/DirectoryNode.Extract.cs b/Engine/Source/Programs/Shared/EpicGames.Horde/Storage/Nodes/DirectoryNode.Extract.cs
--- a/Engine/Source/Programs/Shared/EpicGames.Horde/Storage/Nodes/DirectoryNode.Extract.cs
+++ b/Engine/Source/Programs/Shared/EpicGames.Horde/Storage/Nodes/DirectoryNode.Extract.cs
@@ -312,59 +312,31 @@ namespace EpicGames.Horde.Storage.Nodes
 
 		internal static async Task FindOutputChunksForFileAsync(OutputFile outputFile, ChannelWriter<ChunkReadRequest> writer, ExtractOptions options, ILogger logger, CancellationToken cancellationToken)
 		{
-			BufferedChannelWriter<ChunkReadRequest> bufferedWriter = new(writer);
-			await FindOutputChunksAsync(outputFile, bufferedWriter, 0, outputFile.FileEntry.Target, options, logger, cancellationToken);
-			await bufferedWriter.FlushAsync(cancellationToken);
-		}
-
-		static async Task WriteChunkRequestAsync(BufferedChannelWriter<ChunkReadRequest> writer, OutputFile outputFile, long offset, IHashedBlobRef blobRef, CancellationToken cancellationToken)
-		{
-			outputFile.IncrementRemaining();
-			ChunkReadRequest request = new ChunkReadRequest(outputFile, offset, blobRef);
-			await writer.WriteAsync(request, cancellationToken);
-		}
+			long offset = 0;
+			ChunkReadRequest? bufferedRequest = null;
 
-		static async Task<long> FindOutputChunksAsync(OutputFile outputFile, BufferedChannelWriter<ChunkReadRequest> writer, long offset, ChunkedDataNodeRef dataRef, ExtractOptions options, ILogger logger, CancellationToken cancellationToken)
-		{
-			if (dataRef.Type == ChunkedDataNodeType.Leaf)
+			await foreach (LeafChunkedDataNodeRef leafNodeRef in outputFile.FileEntry.Target.EnumerateLeafNodesAsync(cancellationToken))
 			{
 				if (options.VerboseOutput)
 				{
-					logger.LogInformation("Queuing {Path},{Offset} ({Locator})", outputFile.Path, offset, dataRef.Handle.GetLocator());
+					logger.LogInformation("Queuing {Path},{Offset} ({Locator})", outputFile.Path, offset, leafNodeRef.BlobRef.GetLocator());
 				}
 
-				await WriteChunkRequestAsync(writer, outputFile, offset, dataRef.Handle, cancellationToken);
+				outputFile.IncrementRemaining();
+				ChunkReadRequest request = new ChunkReadRequest(outputFile, offset, leafNodeRef.BlobRef);
+				offset += leafNodeRef.Length;
 
-				if (dataRef.Length < 0)
+				if (bufferedRequest != null)
 				{
-					// Backwards compatibility hack for v2 format
-					using BlobData data = await dataRef.Handle.ReadBlobDataAsync(cancellationToken);
-					return data.Data.Length;
+					await writer.WriteAsync(bufferedRequest, cancellationToken);
 				}
-				return dataRef.Length;
-			}
-			else
-			{
-				using BlobData data = await dataRef.Handle.ReadBlobDataAsync(cancellationToken);
-				TraceBlobRead("Interior", outputFile.Path, dataRef.Handle, logger);
 
-				if (data.Type.Guid == LeafChunkedDataNodeConverter.BlobType.Guid)
-				{
-					await WriteChunkRequestAsync(writer, outputFile, offset, dataRef.Handle, cancellationToken);
-					return dataRef.Length;
-				}
-				else
-				{
-					long length = 0;
-
-					InteriorChunkedDataNode interiorNode = BlobSerializer.Deserialize<InteriorChunkedDataNode>(data);
-					foreach (ChunkedDataNodeRef childRef in interiorNode.Children)
-					{
-						length += await FindOutputChunksAsync(outputFile, writer, offset + length, childRef, options, logger, cancellationToken);
-					}
+				bufferedRequest = request;
+			}
 
-					return length;
-				}
+			if (bufferedRequest != null)
+			{
+				await writer.WriteAsync(bufferedRequest, cancellationToken);
 			}
 		}
 
