Horde: Simplify handling of chunk output pipeline. BatchOutputWriter now manages write channel and will correctly dispose of any in-flight requests on error.

#rnx

diff --git a/Engine/Source/Programs/Shared/EpicGames.Horde/Storage/BatchChunkReader.cs b/Engine/Source/Programs/Shared/EpicGames.Horde/Storage/BatchChunkReader.cs
--- a/Engine/Source/Programs/Shared/EpicGames.Horde/Storage/BatchChunkReader.cs
+++ b/Engine/Source/Programs/Shared/EpicGames.Horde/Storage/BatchChunkReader.cs
@@ -21,41 +21,19 @@ namespace EpicGames.Horde.Storage
 	/// </summary>
 	class BatchChunkReader : BatchBlobReader<ChunkReadRequest>, IDisposable
 	{
-		readonly Channel<OutputChunk[]> _channel;
-
 		/// <summary>
 		/// Maximum number of exports to write in a single request
 		/// </summary>
 		public int MaxBatchSize { get; } = 20;
 
-		/// <summary>
-		/// Reader for output chunks
-		/// </summary>
-		public ChannelReader<OutputChunk[]> ResponseReader => _channel.Reader;
+		readonly ChannelWriter<OutputChunk[]> _writer;
 
 		/// <summary>
 		/// Constructor
 		/// </summary>
-		public BatchChunkReader()
+		public BatchChunkReader(ChannelWriter<OutputChunk[]> writer)
 		{
-			_channel = Channel.CreateBounded<OutputChunk[]>(new BoundedChannelOptions(128));
-		}
-
-		/// <inheritdoc/>
-		protected override void Dispose(bool disposing)
-		{
-			base.Dispose(disposing);
-
-			if (disposing)
-			{
-				while (_channel.Reader.TryRead(out OutputChunk[]? array))
-				{
-					foreach (OutputChunk chunk in array)
-					{
-						chunk.Dispose();
-					}
-				}
-			}
+			_writer = writer;
 		}
 
 		/// <inheritdoc/>
@@ -73,14 +51,8 @@ namespace EpicGames.Horde.Storage
 					batch.Add(chunk);
 				}
 
-				await _channel.Writer.WriteAsync(batch.ToArray(), cancellationToken);
+				await _writer.WriteAsync(batch.ToArray(), cancellationToken);
 			}
 		}
-
-		/// <summary>
-		/// Marks the chunk stream as complete
-		/// </summary>
-		public void TryComplete()
-			=> _channel.Writer.TryComplete();
 	}
 }
diff --git a/Engine/Source/Programs/Shared/EpicGames.Horde/Storage/BatchOutputWriter.cs b/Engine/Source/Programs/Shared/EpicGames.Horde/Storage/BatchOutputWriter.cs
--- a/Engine/Source/Programs/Shared/EpicGames.Horde/Storage/BatchOutputWriter.cs
+++ b/Engine/Source/Programs/Shared/EpicGames.Horde/Storage/BatchOutputWriter.cs
@@ -130,9 +130,9 @@ namespace EpicGames.Horde.Storage
 	/// <summary>
 	/// Batches file write requests
 	/// </summary>
-	class BatchOutputWriter
+	sealed class BatchOutputWriter : IDisposable
 	{
-		readonly ChannelReader<OutputChunk[]> _source;
+		readonly Channel<OutputChunk[]> _channel;
 		readonly ILogger _logger;
 
 		int _writtenFiles;
@@ -170,22 +170,52 @@ namespace EpicGames.Horde.Storage
 		/// </summary>
 		public bool VerboseOutput { get; set; }
 
+		/// <summary>
+		/// Sink for write requests
+		/// </summary>
+		public ChannelWriter<OutputChunk[]> RequestWriter 
+			=> _channel.Writer;
+
 		/// <summary>
 		/// Constructor
 		/// </summary>
-		public BatchOutputWriter(ChannelReader<OutputChunk[]> reader, ILogger logger)
+		public BatchOutputWriter(ILogger logger)
 		{
-			_source = reader;
+			_channel = Channel.CreateBounded<OutputChunk[]>(new BoundedChannelOptions(128));
 			_logger = logger;
 		}
 
+		/// <summary>
+		/// Dispose remaining items in the queue
+		/// </summary>
+		public void Dispose()
+		{
+			while (_channel.Reader.TryRead(out OutputChunk[]? chunks))
+			{
+				foreach (OutputChunk chunk in chunks)
+				{
+					chunk.Dispose();
+				}
+			}
+		}
+
+		/// <summary>
+		/// Adds tasks for the writer to an async pipeline
+		/// </summary>
+		/// <param name="pipeline">Pipeline instance</param>
+		/// <param name="numWriteTasks">Number of parallel writes</param>
+		public Task[] AddToPipeline(AsyncPipeline pipeline, int numWriteTasks)
+		{
+			return pipeline.AddTasks(numWriteTasks, ProcessAsync);
+		}
+
 		/// <summary>
 		/// Processes requests from the given input channel
 		/// </summary>
 		/// <param name="cancellationToken">Cancellation token for the operation</param>
 		public async Task ProcessAsync(CancellationToken cancellationToken)
 		{
-			await foreach (OutputChunk[] chunks in _source.ReadAllAsync(cancellationToken))
+			await foreach (OutputChunk[] chunks in _channel.Reader.ReadAllAsync(cancellationToken))
 			{
 				try
 				{
diff --git a/Engine/Source/Programs/Shared/EpicGames.Horde/Storage/Nodes/DirectoryNode.Extract.cs b/Engine/Source/Programs/Shared/EpicGames.Horde/Storage/Nodes/DirectoryNode.Extract.cs
--- a/Engine/Source/Programs/Shared/EpicGames.Horde/Storage/Nodes/DirectoryNode.Extract.cs
+++ b/Engine/Source/Programs/Shared/EpicGames.Horde/Storage/Nodes/DirectoryNode.Extract.cs
@@ -255,20 +255,20 @@ namespace EpicGames.Horde.Storage.Nodes
 
 			await using (AsyncPipeline pipeline = new AsyncPipeline(cancellationToken))
 			{
-				using BatchChunkReader chunkReader = new BatchChunkReader();
+				using BatchOutputWriter outputWriter = new BatchOutputWriter(logger);
+				outputWriter.VerifyOutput = options.VerifyOutput;
+				outputWriter.VerboseOutput = options.VerboseOutput;
+
+				using BatchChunkReader chunkReader = new BatchChunkReader(outputWriter.RequestWriter);
 				chunkReader.VerifyHashes = options.VerifyOutput;
 
 				Channel<ChunkReadRequest> blobRequests = Channel.CreateBounded<ChunkReadRequest>(new BoundedChannelOptions(chunkReader.MinQueueLength * 2));
 				_ = pipeline.AddTask(ctx => FindOutputChunksRootAsync(directoryInfo, directoryNode, blobRequests.Writer, options, logger, ctx));
 
 				Task[] readTasks = chunkReader.AddToPipeline(pipeline, numReadTasks, numDecodeTasks, blobRequests.Reader);
-				_ = Task.WhenAll(readTasks).ContinueWith(_ => chunkReader.TryComplete(), TaskScheduler.Default);
-
-				BatchOutputWriter outputWriter = new BatchOutputWriter(chunkReader.ResponseReader, logger);
-				outputWriter.VerifyOutput = options.VerifyOutput;
-				outputWriter.VerboseOutput = options.VerboseOutput;
+				_ = Task.WhenAll(readTasks).ContinueWith(_ => outputWriter.RequestWriter.TryComplete(), TaskScheduler.Default);
 
-				Task[] writeTasks = pipeline.AddTasks(numWriteTasks, ctx => outputWriter.ProcessAsync(ctx));
+				Task[] writeTasks = outputWriter.AddToPipeline(pipeline, numWriteTasks);
 
 				if (options.Progress != null)
 				{
diff --git a/Engine/Source/Programs/Shared/EpicGames.Horde/Storage/Workspace.cs b/Engine/Source/Programs/Shared/EpicGames.Horde/Storage/Workspace.cs
--- a/Engine/Source/Programs/Shared/EpicGames.Horde/Storage/Workspace.cs
+++ b/Engine/Source/Programs/Shared/EpicGames.Horde/Storage/Workspace.cs
@@ -665,15 +665,19 @@ namespace EpicGames.Horde.Storage
 			ExtractOptions options = new ExtractOptions();
 			options.Progress = new ExtractStatsLogger(_logger);
 
-			_logger.LogInformation("Scanning tree to sync...");
+			_logger.LogInformation("Creating manifest...");
 
-			Channel<ChunkReadRequest> chunkRequests = Channel.CreateUnbounded<ChunkReadRequest>();
+			Channel<ChunkReadRequest> chunkReadRequests = Channel.CreateUnbounded<ChunkReadRequest>();
 
-			using BatchChunkReader batchReader = new BatchChunkReader { VerifyHashes = options.VerifyOutput };
-			await SyncDirectoryAsync(_rootDir, basePath, syncDirState, contents, layerState.Flag, cacheDirState, chunkRequests.Writer, options, cancellationToken);
-			chunkRequests.Writer.Complete();
+			using BatchOutputWriter outputWriter = new BatchOutputWriter(_logger);
 
-			if (chunkRequests.Reader.Count > 0)
+			using BatchChunkReader batchReader = new BatchChunkReader(outputWriter.RequestWriter) { VerifyHashes = options.VerifyOutput };
+			await SyncDirectoryAsync(_rootDir, basePath, syncDirState, contents, layerState.Flag, cacheDirState, chunkReadRequests.Writer, options, cancellationToken);
+			chunkReadRequests.Writer.Complete();
+
+			await SaveAsync(cancellationToken);
+
+			if (chunkReadRequests.Reader.Count > 0)
 			{
 				_logger.LogInformation("Updating workspace...");
 
@@ -683,12 +687,11 @@ namespace EpicGames.Horde.Storage
 
 				await using (AsyncPipeline pipeline = new AsyncPipeline(cancellationToken))
 				{
-					Task[] readTasks = batchReader.AddToPipeline(pipeline, numReadTasks, numDecodeTasks, chunkRequests.Reader);
-					_ = Task.WhenAll(readTasks).ContinueWith(_ => batchReader.TryComplete(), TaskScheduler.Default);
+					Task[] readTasks = batchReader.AddToPipeline(pipeline, numReadTasks, numDecodeTasks, chunkReadRequests.Reader);
 
-					BatchOutputWriter outputWriter = new BatchOutputWriter(batchReader.ResponseReader, _logger);
+					_ = Task.WhenAll(readTasks).ContinueWith(_ => outputWriter.RequestWriter.TryComplete(), TaskScheduler.Default);
 
-					Task[] writeTasks = pipeline.AddTasks(numWriteTasks, ctx => outputWriter.ProcessAsync(ctx));
+					Task[] writeTasks = outputWriter.AddToPipeline(pipeline, numWriteTasks);
 
 					if (options.Progress != null)
 					{
@@ -700,7 +703,7 @@ namespace EpicGames.Horde.Storage
 			}
 		}
 
-		async Task SyncDirectoryAsync(DirectoryReference dirRef, string path, DirectoryState dirState, DirectoryNode? dirNode, ulong flag, DirectoryState cacheDirState, ChannelWriter<ChunkReadRequest> readRequestWriter, ExtractOptions options, CancellationToken cancellationToken)
+		async Task SyncDirectoryAsync(DirectoryReference dirRef, string path, DirectoryState dirState, DirectoryNode? dirNode, ulong flag, DirectoryState cacheDirState, ChannelWriter<ChunkReadRequest> requestWriter, ExtractOptions options, CancellationToken cancellationToken)
 		{
 			// Remove any directories that no longer exist
 			for (int subDirIdx = 0; subDirIdx < dirState.Directories.Count; subDirIdx++)
@@ -711,7 +714,7 @@ namespace EpicGames.Horde.Storage
 					if (dirNode == null || !dirNode.TryGetDirectoryEntry(subDirState.Name, out _))
 					{
 						DirectoryReference subDirPath = DirectoryReference.Combine(dirRef, subDirState.Name.ToString());
-						await SyncDirectoryAsync(subDirPath, $"{path}{subDirState.Name}/", subDirState, null, flag, cacheDirState, readRequestWriter, options, cancellationToken);
+						await SyncDirectoryAsync(subDirPath, $"{path}{subDirState.Name}/", subDirState, null, flag, cacheDirState, requestWriter, options, cancellationToken);
 
 						if (subDirState.LayerFlags == 0)
 						{
@@ -756,7 +759,7 @@ namespace EpicGames.Horde.Storage
 					DirectoryState subDirState = dirState.FindOrAddDirectory(subDirEntry.Name);
 
 					DirectoryNode subDirNode = await subDirEntry.Handle.ReadBlobAsync(cancellationToken: cancellationToken);
-					await SyncDirectoryAsync(subDirPath, $"{path}{subDirEntry.Name}/", subDirState, subDirNode, flag, cacheDirState, readRequestWriter, options, cancellationToken);
+					await SyncDirectoryAsync(subDirPath, $"{path}{subDirEntry.Name}/", subDirState, subDirNode, flag, cacheDirState, requestWriter, options, cancellationToken);
 
 					dirState.LayerFlags |= flag;
 				}
@@ -764,7 +767,7 @@ namespace EpicGames.Horde.Storage
 				// Update files
 				foreach (FileEntry fileEntry in dirNode.Files)
 				{
-					await SyncFileAsync(dirRef, $"{path}{fileEntry.Name}", dirState, fileEntry, flag, readRequestWriter, options, cancellationToken);
+					await SyncFileAsync(dirRef, $"{path}{fileEntry.Name}", dirState, fileEntry, flag, requestWriter, options, cancellationToken);
 					dirState.LayerFlags |= flag;
 				}
 			}
@@ -776,7 +779,7 @@ namespace EpicGames.Horde.Storage
 			}
 		}
 
-		async Task SyncFileAsync(DirectoryReference dirRef, string path, DirectoryState dirState, FileEntry fileEntry, ulong flag, ChannelWriter<ChunkReadRequest> readRequestWriter, ExtractOptions options, CancellationToken cancellationToken)
+		async Task SyncFileAsync(DirectoryReference dirRef, string path, DirectoryState dirState, FileEntry fileEntry, ulong flag, ChannelWriter<ChunkReadRequest> requestWriter, ExtractOptions options, CancellationToken cancellationToken)
 		{
 			FileState? fileState;
 			if (dirState.TryGetFile(fileEntry.Name, out fileState))
@@ -800,7 +803,7 @@ namespace EpicGames.Horde.Storage
 					fileState = dirState.FindOrAddFile(fileEntry.Name);
 
 					WorkspaceOutputFile outputFile = new WorkspaceOutputFile(path, fileInfo, fileEntry, fileState);
-					await DirectoryNodeExtract.FindOutputChunksForFileAsync(outputFile, readRequestWriter, options, _logger, cancellationToken);
+					await DirectoryNodeExtract.FindOutputChunksForFileAsync(outputFile, requestWriter, options, _logger, cancellationToken);
 				}
 				else
 				{
@@ -934,7 +937,7 @@ namespace EpicGames.Horde.Storage
 			}
 
 			int outputIdx = 0;
-			for(int idx = 0; idx < dirState.Files.Count; idx++)
+			for (int idx = 0; idx < dirState.Files.Count; idx++)
 			{
 				FileState fileState = dirState.Files[idx];
 
