[UBT]
* Removed system roots and autosdk root now when all modules should provide the roots they need

diff --git a/Engine/Source/Programs/Shared/EpicGames.UBA/ICacheClient.cs b/Engine/Source/Programs/Shared/EpicGames.UBA/ICacheClient.cs
--- a/Engine/Source/Programs/Shared/EpicGames.UBA/ICacheClient.cs
+++ b/Engine/Source/Programs/Shared/EpicGames.UBA/ICacheClient.cs
@@ -19,13 +19,6 @@ namespace EpicGames.UBA
 		/// <returns>True if successful</returns>
 		public abstract bool RegisterRoot(string path, bool includeInKey = true, byte id = 0);
 
-		/// <summary>
-		/// Register system roots used to normalize paths in caches
-		/// </summary>
-		/// <param name="startId">Start id for system roots. On windows these take up 10 entries</param>
-		/// <returns>True if successful</returns>
-		public abstract bool RegisterSystemRoots(byte startId = 0);
-
 		/// <summary>
 		/// Create root paths instance
 		/// </summary>
diff --git a/Engine/Source/Programs/Shared/EpicGames.UBA/Impl/CacheClientImpl.cs b/Engine/Source/Programs/Shared/EpicGames.UBA/Impl/CacheClientImpl.cs
--- a/Engine/Source/Programs/Shared/EpicGames.UBA/Impl/CacheClientImpl.cs
+++ b/Engine/Source/Programs/Shared/EpicGames.UBA/Impl/CacheClientImpl.cs
@@ -20,9 +20,6 @@ namespace EpicGames.UBA.Impl
 
 		[DllImport("UbaHost", CharSet = CharSet.Auto)]
 		static extern bool RootPaths_RegisterRoot(nint rootPaths, string root, bool includeInKey, byte id);
-
-		[DllImport("UbaHost", CharSet = CharSet.Auto)]
-		static extern bool RootPaths_RegisterSystemRoots(nint rootPaths, byte id);
 		#endregion
 
 		#region IDisposable
@@ -59,11 +56,6 @@ namespace EpicGames.UBA.Impl
 			return RootPaths_RegisterRoot(_handle, path, includeInKey, id);
 		}
 
-		public bool RegisterSystemRoots(byte startId)
-		{
-			return RootPaths_RegisterSystemRoots(_handle, startId);
-		}
-
 		public nint GetHandle() => _handle;
 	}
 
diff --git a/Engine/Source/Programs/UnrealBuildTool/Executors/UnrealBuildAccelerator/UBAExecutor.cs b/Engine/Source/Programs/UnrealBuildTool/Executors/UnrealBuildAccelerator/UBAExecutor.cs
--- a/Engine/Source/Programs/UnrealBuildTool/Executors/UnrealBuildAccelerator/UBAExecutor.cs
+++ b/Engine/Source/Programs/UnrealBuildTool/Executors/UnrealBuildAccelerator/UBAExecutor.cs
@@ -707,14 +707,6 @@ namespace UnrealBuildTool
 			{
 				rootPaths.RegisterRoot(kv.local.FullName, true);
 			}
-
-			DirectoryReference? autoSdkDir;
-			if (UEBuildPlatformSDK.TryGetHostPlatformAutoSDKDir(out autoSdkDir))
-			{
-				rootPaths.RegisterRoot(autoSdkDir.FullName, true);
-			}
-
-			rootPaths.RegisterSystemRoots();
 			return rootPaths;
 		}
 
