[UBA]
* Changed caching code to use new uba path which shares logic with vfs
* Removed last parts of OutputStatsThresholdMs
* Removed code wrapping root path instance now when we use handles instead

diff --git a/Engine/Source/Programs/Shared/EpicGames.UBA/ICacheClient.cs b/Engine/Source/Programs/Shared/EpicGames.UBA/ICacheClient.cs
--- a/Engine/Source/Programs/Shared/EpicGames.UBA/ICacheClient.cs
+++ b/Engine/Source/Programs/Shared/EpicGames.UBA/ICacheClient.cs
@@ -5,29 +5,6 @@ using EpicGames.UBA.Impl;
 
 namespace EpicGames.UBA
 {
-	/// <summary>
-	/// Base interface for root paths used by cache system to normalize paths
-	/// </summary>
-	public interface IRootPaths : IBaseInterface
-	{
-		/// <summary>
-		/// Register roots used to normalize paths in caches
-		/// </summary>
-		/// <param name="path">Path of root</param>
-		/// <param name="includeInKey">set this to false if you want to ignore all files under this folder (you know they are _always_ the same for all machines)</param>
-		/// <param name="id">Id of root. On windows these numbers need to increase two at the time since double backslash paths are added automatically under the hood</param>
-		/// <returns>True if successful</returns>
-		public abstract bool RegisterRoot(string path, bool includeInKey = true, byte id = 0);
-
-		/// <summary>
-		/// Create root paths instance
-		/// </summary>
-		public static IRootPaths Create(ILogger logger)
-		{
-			return new RootPathsImpl(logger);
-		}
-	}
-
 	/// <summary>
 	/// Struct containing results from artifact fetch
 	/// </summary>
@@ -51,7 +28,6 @@ namespace EpicGames.UBA
 		/// <summary>
 		/// Write to cache
 		/// </summary>
-		/// <param name="rootPaths">RootPath instance</param>
 		/// <param name="bucket">Bucket to store cache entry</param>
 		/// <param name="process">Process</param>
 		/// <param name="inputs">Input files</param>
@@ -59,16 +35,16 @@ namespace EpicGames.UBA
 		/// <param name="outputs">Output files</param>
 		/// <param name="outputsSize">Output files size</param>
 		/// <returns>True if successful</returns>
-		public abstract bool WriteToCache(IRootPaths rootPaths, uint bucket, IProcess process, byte[] inputs, uint inputsSize, byte[] outputs, uint outputsSize);
+		public abstract bool WriteToCache(uint bucket, IProcess process, byte[] inputs, uint inputsSize, byte[] outputs, uint outputsSize);
 
 		/// <summary>
 		/// Fetch from cache
 		/// </summary>
-		/// <param name="rootPaths">RootPath instance</param>
+		/// <param name="rootsHandle">handle for roots</param>
 		/// <param name="bucket">Bucket to search for cache entry</param>
 		/// <param name="info">Process start info</param>
 		/// <returns>True if successful</returns>
-		public abstract FetchFromCacheResult FetchFromCache(IRootPaths rootPaths, uint bucket, ProcessStartInfo info);
+		public abstract FetchFromCacheResult FetchFromCache(ulong rootsHandle, uint bucket, ProcessStartInfo info);
 
 		/// <summary>
 		/// Request the connected server to shutdown
diff --git a/Engine/Source/Programs/Shared/EpicGames.UBA/IProcess.cs b/Engine/Source/Programs/Shared/EpicGames.UBA/IProcess.cs
--- a/Engine/Source/Programs/Shared/EpicGames.UBA/IProcess.cs
+++ b/Engine/Source/Programs/Shared/EpicGames.UBA/IProcess.cs
@@ -58,11 +58,6 @@ namespace EpicGames.UBA
 		/// </summary>
 		public ProcessPriorityClass Priority { get; set; }
 
-		/// <summary>
-		/// Threshold in which to report output stats
-		/// </summary>
-		public uint OutputStatsThresholdMs { get; set; }
-
 		/// <summary>
 		/// If input should be tracked
 		/// </summary>
diff --git a/Engine/Source/Programs/Shared/EpicGames.UBA/Impl/CacheClientImpl.cs b/Engine/Source/Programs/Shared/EpicGames.UBA/Impl/CacheClientImpl.cs
--- a/Engine/Source/Programs/Shared/EpicGames.UBA/Impl/CacheClientImpl.cs
+++ b/Engine/Source/Programs/Shared/EpicGames.UBA/Impl/CacheClientImpl.cs
@@ -6,59 +6,6 @@ using System.Runtime.InteropServices;
 
 namespace EpicGames.UBA.Impl
 {
-	internal class RootPathsImpl : IRootPaths
-	{
-		nint _handle = IntPtr.Zero;
-		readonly ILogger _logger;
-
-		#region DllImport
-		[DllImport("UbaHost", CharSet = CharSet.Auto)]
-		static extern nint RootPaths_Create(nint logWriter);
-
-		[DllImport("UbaHost", CharSet = CharSet.Auto)]
-		static extern void RootPaths_Destroy(nint rootPaths);
-
-		[DllImport("UbaHost", CharSet = CharSet.Auto)]
-		static extern bool RootPaths_RegisterRoot(nint rootPaths, string root, bool includeInKey, byte id);
-		#endregion
-
-		#region IDisposable
-		~RootPathsImpl() => Dispose(false);
-
-		public void Dispose()
-		{
-			Dispose(true);
-			GC.SuppressFinalize(this);
-		}
-
-		protected virtual void Dispose(bool disposing)
-		{
-			if (disposing)
-			{
-			}
-
-			if (_handle != IntPtr.Zero)
-			{
-				RootPaths_Destroy(_handle);
-				_handle = IntPtr.Zero;
-			}
-		}
-		#endregion
-
-		public RootPathsImpl(ILogger logger)
-		{
-			_logger = logger;
-			_handle = RootPaths_Create(_logger.GetHandle());
-		}
-
-		public bool RegisterRoot(string path, bool includeInKey, byte id)
-		{
-			return RootPaths_RegisterRoot(_handle, path, includeInKey, id);
-		}
-
-		public nint GetHandle() => _handle;
-	}
-
 	internal class CacheClientImpl : ICacheClient
 	{
 		nint _handle = IntPtr.Zero;
@@ -68,7 +15,7 @@ namespace EpicGames.UBA.Impl
 
 		#region DllImport
 		[DllImport("UbaHost", CharSet = CharSet.Auto)]
-		static extern nint ProcessStartInfo_Create(string application, string arguments, string workingDir, string description, uint priorityClass, ulong outputStatsThresholdMs, bool trackInputs, string logFile, ExitCallback? exit);
+		static extern nint ProcessStartInfo_Create2(string application, string arguments, string workingDir, string description, uint priorityClass, ulong rootsHandle, bool trackInputs, string logFile, ExitCallback? exit);
 
 		[DllImport("UbaHost", CharSet = CharSet.Auto)]
 		static extern void ProcessStartInfo_Destroy(nint server);
@@ -80,10 +27,10 @@ namespace EpicGames.UBA.Impl
 		static extern bool CacheClient_Connect(nint cacheClient, string host, int port);
 
 		[DllImport("UbaHost", CharSet = CharSet.Auto)]
-		static extern bool CacheClient_WriteToCache(nint cacheClient, nint rootPaths, uint bucket, nint info, byte[] inputs, uint inputsSize, byte[] outputs, uint outputsSize);
+		static extern bool CacheClient_WriteToCache2(nint cacheClient, uint bucket, nint info, byte[] inputs, uint inputsSize, byte[] outputs, uint outputsSize);
 
 		[DllImport("UbaHost", CharSet = CharSet.Auto)]
-		static extern nint CacheClient_FetchFromCache2(nint cacheClient, nint rootPaths, uint bucket, nint info);
+		static extern nint CacheClient_FetchFromCache3(nint cacheClient, ulong rootsHandle, uint bucket, nint info);
 
 		[DllImport("UbaHost", CharSet = CharSet.Auto)]
 		static extern nint CacheResult_GetLogLine(nint result, uint index);
@@ -135,15 +82,15 @@ namespace EpicGames.UBA.Impl
 			return CacheClient_Connect(_handle, host, port);
 		}
 
-		public bool WriteToCache(IRootPaths rootPaths, uint bucket, IProcess process, byte[] inputs, uint inputsSize, byte[] outputs, uint outputsSize)
+		public bool WriteToCache(uint bucket, IProcess process, byte[] inputs, uint inputsSize, byte[] outputs, uint outputsSize)
 		{
-			return CacheClient_WriteToCache(_handle, rootPaths.GetHandle(), bucket, process.GetHandle(), inputs, inputsSize, outputs, outputsSize);
+			return CacheClient_WriteToCache2(_handle, bucket, process.GetHandle(), inputs, inputsSize, outputs, outputsSize);
 		}
 
-		public FetchFromCacheResult FetchFromCache(IRootPaths rootPaths, uint bucket, ProcessStartInfo info)
+		public FetchFromCacheResult FetchFromCache(ulong rootsHandle, uint bucket, ProcessStartInfo info)
 		{
-			nint si = ProcessStartInfo_Create(info.Application, info.Arguments, info.WorkingDirectory, info.Description, (uint)info.Priority, info.OutputStatsThresholdMs, info.TrackInputs, info.LogFile ?? String.Empty, null);
-			nint result = CacheClient_FetchFromCache2(_handle, rootPaths.GetHandle(), bucket, si);
+			nint si = ProcessStartInfo_Create2(info.Application, info.Arguments, info.WorkingDirectory, info.Description, (uint)info.Priority, rootsHandle, info.TrackInputs, info.LogFile ?? String.Empty, null);
+			nint result = CacheClient_FetchFromCache3(_handle, rootsHandle, bucket, si);
 			ProcessStartInfo_Destroy(si);
 
 			if (result == IntPtr.Zero)
diff --git a/Engine/Source/Programs/UnrealBuildTool/Executors/UnrealBuildAccelerator/UBAExecutor.cs b/Engine/Source/Programs/UnrealBuildTool/Executors/UnrealBuildAccelerator/UBAExecutor.cs
--- a/Engine/Source/Programs/UnrealBuildTool/Executors/UnrealBuildAccelerator/UBAExecutor.cs
+++ b/Engine/Source/Programs/UnrealBuildTool/Executors/UnrealBuildAccelerator/UBAExecutor.cs
@@ -260,19 +260,18 @@ namespace UnrealBuildTool
 				{
 					ProcessStartInfo startInfo = _executor.GetActionStartInfo(action);
 					uint bucket = UBAExecutor.GetActionCacheBucket(action);
-					using (IRootPaths rootPaths = _executor.GetActionRootPaths(action))
+					ulong rootsHandle = _executor.GetActionRootsHandle(action);
+
+					FetchFromCacheResult result = _executor._cacheClient!.FetchFromCache(rootsHandle, bucket, startInfo);
+					int exitCode = 0;
+					RemoveLogLineSpam(result.LogLines);
+					if (result.Success && action.bForceWarningsAsError && result.LogLines.Any(x => x.Contains("): warning: ")))
 					{
-						FetchFromCacheResult result = _executor._cacheClient!.FetchFromCache(rootPaths, bucket, startInfo);
-						int exitCode = 0;
-						RemoveLogLineSpam(result.LogLines);
-						if (result.Success && action.bForceWarningsAsError && result.LogLines.Any(x => x.Contains("): warning: ")))
-						{
-							result = new(true, result.LogLines.Select(x => x.Replace("): warning: ", "): error: ", StringComparison.OrdinalIgnoreCase)).ToList());
-							exitCode = 1;
-							++_executor._errorCount;
-						}
-						return new ActionArtifactResult(result.Success, result.LogLines, exitCode);
+						result = new(true, result.LogLines.Select(x => x.Replace("): warning: ", "): error: ", StringComparison.OrdinalIgnoreCase)).ToList());
+						exitCode = 1;
+						++_executor._errorCount;
 					}
+					return new ActionArtifactResult(result.Success, result.LogLines, exitCode);
 				}, cancellationToken, TaskCreationOptions.LongRunning | TaskCreationOptions.PreferFairness, TaskScheduler.Default);
 			}
 
@@ -663,7 +662,6 @@ namespace UnrealBuildTool
 				WorkingDirectory = action.WorkingDirectory.FullName,
 				Arguments = action.CommandArguments,
 				Priority = ProcessPriority,
-				OutputStatsThresholdMs = (uint)UBAConfig.OutputStatsThresholdMs,
 				UserData = action,
 				Description = description,
 				Configuration = action.bIsGCCCompiler ? EpicGames.UBA.ProcessStartInfo.CommonProcessConfigs.CompileClang : EpicGames.UBA.ProcessStartInfo.CommonProcessConfigs.CompileMsvc,
@@ -700,30 +698,23 @@ namespace UnrealBuildTool
 			}
 		}
 
-		IRootPaths GetActionRootPaths(LinkedAction action)
-		{
-			IRootPaths rootPaths = IRootPaths.Create(_ubaLogger!);
-			foreach ((DirectoryReference vfs, DirectoryReference local) kv in action.RootPaths)
-			{
-				rootPaths.RegisterRoot(kv.local.FullName, true);
-			}
-			return rootPaths;
-		}
-
 		ulong GetActionRootsHandle(LinkedAction action)
 		{
-			if (!action.RootPaths.bUseVfs)
-			{
-				return 0;
-			}
-
 			using MemoryStream outputsMemory = new(1024);
 			using (BinaryWriter writer = new(outputsMemory, System.Text.Encoding.UTF8, true))
 			{
+				bool useVfs = action.RootPaths.bUseVfs;
 				foreach (var kv in action.RootPaths)
 				{
 					writer.Write((byte)0);// kv.Key);
-					writer.Write(kv.vfs.FullName);
+					if (useVfs)
+					{
+						writer.Write(kv.vfs.FullName);
+					}
+					else
+					{
+						writer.Write("");
+					}
 					writer.Write(kv.local.FullName);
 				}
 			}
@@ -795,9 +786,8 @@ namespace UnrealBuildTool
 			}
 
 			uint bucket = GetActionCacheBucket(action);
-			using IRootPaths rootPaths = GetActionRootPaths(action);
 
-			return _cacheClient!.WriteToCache(rootPaths, bucket, process, inputsMemory.GetBuffer(), (uint)inputsMemory.Position, outputsMemory.GetBuffer(), (uint)outputsMemory.Position);
+			return _cacheClient!.WriteToCache(bucket, process, inputsMemory.GetBuffer(), (uint)inputsMemory.Position, outputsMemory.GetBuffer(), (uint)outputsMemory.Position);
 		}
 
 		Func<Task>? RunActionLocal(ImmediateActionQueue queue, LinkedAction action)
@@ -817,7 +807,6 @@ namespace UnrealBuildTool
 				bool enableDetour = !ForceLocalNoDetour(action) && action.bCanExecuteInUBA && !_forcedRetryActions.ContainsKey(action);
 
 				ProcessStartInfo startInfo = GetActionStartInfo(action);
-
 				startInfo.RootsHandle = GetActionRootsHandle(action);
 
 				using (IProcess process = _session!.RunProcess(startInfo, false, null, enableDetour))
@@ -912,6 +901,7 @@ namespace UnrealBuildTool
 
 				ProcessStartInfo startInfo = GetActionStartInfo(action);
 				startInfo.RootsHandle = GetActionRootsHandle(action);
+
 				_session!.RunProcessRemote(startInfo, (s, e) =>
 				{
 					if (e.ExitCode == 99999) // Process was cancelled by executor
diff --git a/Engine/Source/Programs/UnrealBuildTool/Executors/UnrealBuildAccelerator/UnrealBuildAcceleratorConfig.cs b/Engine/Source/Programs/UnrealBuildTool/Executors/UnrealBuildAccelerator/UnrealBuildAcceleratorConfig.cs
--- a/Engine/Source/Programs/UnrealBuildTool/Executors/UnrealBuildAccelerator/UnrealBuildAcceleratorConfig.cs
+++ b/Engine/Source/Programs/UnrealBuildTool/Executors/UnrealBuildAccelerator/UnrealBuildAcceleratorConfig.cs
@@ -170,15 +170,9 @@ namespace UnrealBuildTool
 		/// </summary>
 		[XmlConfigFile(Category = "UnrealBuildAccelerator")]
 		[CommandLine("-UBAStoreObjFilesCompressed", Value = "true")]
+		[CommandLine("-UBADisableStoreObjFilesCompressed", Value = "false")]
 		public bool bStoreObjFilesCompressed { get; set; }
 
-		/// <summary>
-		/// Threshold for when executor should output logging for the process. Defaults to never
-		/// </summary>
-		[XmlConfigFile(Category = "UnrealBuildAccelerator")]
-		[CommandLine("-UBAOutputStatsThresholdMs")]
-		public int OutputStatsThresholdMs { get; set; } = Int32.MaxValue;
-
 		/// <summary>
 		/// Skip writing intermediate and output files to disk. Useful for validation builds where we don't need the output
 		/// </summary>
