Horde: Add support for copying chunks within the workspace to materialize files.

#rnx

diff --git a/Engine/Source/Programs/Shared/EpicGames.Horde/Storage/Workspace.cs b/Engine/Source/Programs/Shared/EpicGames.Horde/Storage/Workspace.cs
--- a/Engine/Source/Programs/Shared/EpicGames.Horde/Storage/Workspace.cs
+++ b/Engine/Source/Programs/Shared/EpicGames.Horde/Storage/Workspace.cs
@@ -308,28 +308,26 @@ namespace EpicGames.Horde.Storage
 				Hash = hash;
 				Length = length;
 			}
-		}
 
-		// Hashed chunk within another hashed object
-		[DebuggerDisplay("{Offset}+{Length}")]
-		class ChunkInfo
-		{
-			public HashInfo WithinHashInfo { get; }
-			public long Offset { get; }
-			public long Length { get; }
-
-			public ChunkInfo(HashInfo withinHashInfo, long offset, long length)
+			public void AddChunk(ChunkInfo chunk)
 			{
-				WithinHashInfo = withinHashInfo;
-				Offset = offset;
-				Length = length;
+				if (chunk.Length != chunk.WithinHashInfo.Length && !Chunks.Contains(chunk))
+				{
+					Chunks.Add(chunk);
+				}
 			}
+		}
 
-			public ChunkInfo(IMemoryReader reader, HashInfo[] hashes)
+		// Hashed chunk within another hashed object
+		[DebuggerDisplay("{Offset}+{Length}->{WithinHashInfo}")]
+		record class ChunkInfo(HashInfo WithinHashInfo, long Offset, long Length)
+		{
+			public static ChunkInfo Read(IMemoryReader reader, HashInfo[] hashes)
 			{
-				WithinHashInfo = hashes[(int)reader.ReadUnsignedVarInt()];
-				Offset = (long)reader.ReadUnsignedVarInt();
-				Length = (long)reader.ReadUnsignedVarInt();
+				HashInfo withinHashInfo = hashes[(int)reader.ReadUnsignedVarInt()];
+				long offset = (long)reader.ReadUnsignedVarInt();
+				long length = (long)reader.ReadUnsignedVarInt();
+				return new ChunkInfo(withinHashInfo, offset, length);
 			}
 
 			public void Write(IMemoryWriter writer)
@@ -520,14 +518,16 @@ namespace EpicGames.Horde.Storage
 		void Read(IMemoryReader reader)
 		{
 			_rootDirState.Read(reader);
+
+			_layers.Clear();
 			reader.ReadList(_layers, () => new LayerState(reader));
 
 			// Read the hash lookup
 			int numHashes = reader.ReadInt32();
 			HashInfo[] hashInfoArray = new HashInfo[numHashes];
 
-			_hashes.EnsureCapacity(numHashes);
 			_hashes.Clear();
+			_hashes.EnsureCapacity(numHashes);
 
 			for (int idx = 0; idx < numHashes; idx++)
 			{
@@ -539,7 +539,7 @@ namespace EpicGames.Horde.Storage
 			for (int idx = 0; idx < numHashes; idx++)
 			{
 				HashInfo hashInfo = hashInfoArray[idx];
-				reader.ReadList(hashInfoArray[idx].Chunks, () => new ChunkInfo(reader, hashInfoArray));
+				reader.ReadList(hashInfoArray[idx].Chunks, () => ChunkInfo.Read(reader, hashInfoArray));
 			}
 
 			// Add all the files to the hash lookup
@@ -548,6 +548,8 @@ namespace EpicGames.Horde.Storage
 
 		void Write(IMemoryWriter writer)
 		{
+			RemoveUnusedHashInfoObjects();
+
 			_rootDirState.Write(writer);
 			writer.WriteList(_layers, x => x.Write(writer));
 
@@ -567,6 +569,38 @@ namespace EpicGames.Horde.Storage
 			}
 		}
 
+		void RemoveUnusedHashInfoObjects()
+		{
+			const int UnknownIndex = 0;
+			const int ReferencedIndex = 1;
+			const int UnreferencedIndex = 2;
+
+			static bool IsReferenced(HashInfo hashInfo)
+			{
+				if (hashInfo.Index == UnknownIndex)
+				{
+					hashInfo.Chunks.RemoveAll(x => !IsReferenced(x.WithinHashInfo));
+					hashInfo.Index = (hashInfo.Chunks.Count > 0) ? ReferencedIndex : UnreferencedIndex;
+				}
+				return hashInfo.Index == ReferencedIndex;
+			}
+
+			List<HashInfo> hashes = new List<HashInfo>(_hashes.Values);
+			foreach (HashInfo hashInfo in hashes)
+			{
+				hashInfo.Index = (hashInfo.Files.Count > 0) ? ReferencedIndex : UnknownIndex;
+			}
+
+			_hashes.Clear();
+			foreach (HashInfo hashInfo in hashes)
+			{
+				if (IsReferenced(hashInfo))
+				{
+					_hashes.Add(hashInfo.Hash, hashInfo);
+				}
+			}
+		}
+
 		#region Layers
 
 		/// <summary>
@@ -618,24 +652,6 @@ namespace EpicGames.Horde.Storage
 
 		#region Syncing
 
-		class WorkspaceOutputFile : OutputFile
-		{
-			readonly FileState _fileState;
-
-			public WorkspaceOutputFile(string path, FileInfo fileInfo, FileEntry fileEntry, FileState fileState)
-				: base(path, fileInfo, fileEntry)
-			{
-				_fileState = fileState;
-			}
-
-			public override ValueTask CompleteAsync(CancellationToken cancellationToken)
-			{
-				_fileState.Hash = FileEntry.StreamHash;
-				_fileState.Update(FileInfo);
-				return default;
-			}
-		}
-
 		/// <summary>
 		/// Syncs a layer to the given contents
 		/// </summary>
@@ -645,7 +661,23 @@ namespace EpicGames.Horde.Storage
 		/// <param name="cancellationToken">Cancellation token for the operation</param>
 		public async Task SyncAsync(WorkspaceLayerId layerId, string? basePath, DirectoryNode? contents, CancellationToken cancellationToken = default)
 		{
-			basePath = basePath?.Replace('\\', '/').Trim('/') ?? "";
+			basePath = basePath?.Replace('\\', '/').Trim('/') ?? String.Empty;
+
+			await using IBlobWriter writer = new MemoryBlobWriter(BlobSerializerOptions.Default);
+			if (basePath.Length > 0)
+			{
+				foreach (string fragment in basePath.Split('/').Reverse())
+				{
+					DirectoryNode next = new DirectoryNode();
+					if (contents != null)
+					{
+						IHashedBlobRef<DirectoryNode> handle = await writer.WriteBlobAsync(contents, cancellationToken);
+						DirectoryEntry entry = new DirectoryEntry(fragment, contents.Length, handle);
+						next.AddDirectory(entry);
+					}
+					contents = next;
+				}
+			}
 
 			LayerState layerState = GetLayerState(layerId) ?? throw new InvalidOperationException($"Layer '{layerId}' does not exist");
 
@@ -655,55 +687,100 @@ namespace EpicGames.Horde.Storage
 			DirectoryState cacheDirState = hordeDirState.FindOrAddDirectory(CacheDirName);
 			cacheDirState.LayerFlags |= CacheLayerFlag;
 
-			DirectoryState syncDirState = _rootDirState;
-			foreach (string fragment in basePath.Split('/'))
-			{
-				syncDirState.LayerFlags |= layerState.Flag;
-				syncDirState = syncDirState.FindOrAddDirectory(fragment);
-			}
-
 			ExtractOptions options = new ExtractOptions();
 			options.Progress = new ExtractStatsLogger(_logger);
 
 			_logger.LogInformation("Creating manifest...");
 
+			// Find all the files that need updating, and move anything that needs to be overwritten into the cache
+			List<FileUpdate> fileUpdates = new List<FileUpdate>();
+			await SyncDirectoryAsync(_rootDir, "", _rootDirState, contents, layerState.Flag, cacheDirState, fileUpdates, options, cancellationToken);
+			await SaveAsync(cancellationToken);
+
+			// Split the work into copies within the workspace, and reads from the upstream store
+			Channel<ChunkCopyRequest> chunkCopyRequests = Channel.CreateUnbounded<ChunkCopyRequest>();
 			Channel<ChunkReadRequest> chunkReadRequests = Channel.CreateUnbounded<ChunkReadRequest>();
 
-			using BatchOutputWriter outputWriter = new BatchOutputWriter(_logger);
+			foreach (FileUpdate fileUpdate in fileUpdates)
+			{
+				OutputFile outputFile = fileUpdate.OutputFile;
 
-			using BatchChunkReader batchReader = new BatchChunkReader(outputWriter.RequestWriter) { VerifyHashes = options.VerifyOutput };
-			await SyncDirectoryAsync(_rootDir, basePath, syncDirState, contents, layerState.Flag, cacheDirState, chunkReadRequests.Writer, options, cancellationToken);
+				long offset = 0;
+				foreach (LeafChunkedDataNodeRef leafNodeRef in fileUpdate.Chunks)
+				{
+					outputFile.IncrementRemaining();
+
+					SourceChunk? sourceChunk = TryGetSourceChunk(leafNodeRef.BlobRef.Hash);
+					if (sourceChunk != null)
+					{
+						ChunkCopyRequest copyRequest = new ChunkCopyRequest(outputFile, offset, sourceChunk);
+						await chunkCopyRequests.Writer.WriteAsync(copyRequest, cancellationToken);
+					}
+					else
+					{
+						ChunkReadRequest readRequest = new ChunkReadRequest(outputFile, offset, leafNodeRef.BlobRef);
+						await chunkReadRequests.Writer.WriteAsync(readRequest, cancellationToken);
+					}
+
+					offset += leafNodeRef.Length;
+				}
+			}
+			chunkCopyRequests.Writer.Complete();
 			chunkReadRequests.Writer.Complete();
 
-			await SaveAsync(cancellationToken);
+			// Run the pipeline to create everything
+			_logger.LogInformation("Updating workspace...");
+
+			int numReadTasks = options.NumReadTasks ?? ExtractOptions.DefaultNumReadTasks;
+			int numDecodeTasks = options.NumDecodeTasks ?? ExtractOptions.DefaultNumDecodeTasks;
+			int numWriteTasks = 16;
+
+			using BatchOutputWriter outputWriter = new BatchOutputWriter(_logger);
+			using BatchChunkReader batchReader = new BatchChunkReader(outputWriter.RequestWriter) { VerifyHashes = options.VerifyOutput };
 
-			if (chunkReadRequests.Reader.Count > 0)
+			await using (AsyncPipeline pipeline = new AsyncPipeline(cancellationToken))
 			{
-				_logger.LogInformation("Updating workspace...");
+				Task[] copyTasks = pipeline.AddTasks(8, ct => CopyChunksAsync(chunkCopyRequests.Reader, outputWriter.RequestWriter, ct));
+				Task[] readTasks = batchReader.AddToPipeline(pipeline, numReadTasks, numDecodeTasks, chunkReadRequests.Reader);
+				_ = Task.WhenAll(Enumerable.Concat(copyTasks, readTasks)).ContinueWith(_ => outputWriter.RequestWriter.TryComplete(), TaskScheduler.Default);
 
-				int numReadTasks = options.NumReadTasks ?? ExtractOptions.DefaultNumReadTasks;
-				int numDecodeTasks = options.NumDecodeTasks ?? ExtractOptions.DefaultNumDecodeTasks;
-				int numWriteTasks = 16;
+				Task[] writeTasks = outputWriter.AddToPipeline(pipeline, numWriteTasks);
 
-				await using (AsyncPipeline pipeline = new AsyncPipeline(cancellationToken))
+				if (options.Progress != null)
 				{
-					Task[] readTasks = batchReader.AddToPipeline(pipeline, numReadTasks, numDecodeTasks, chunkReadRequests.Reader);
+					_ = pipeline.AddTask(ctx => DirectoryNodeExtract.UpdateStatsAsync(batchReader, outputWriter, Task.WhenAll(writeTasks), options.Progress, options.ProgressUpdateFrequency, ctx));
+				}
 
-					_ = Task.WhenAll(readTasks).ContinueWith(_ => outputWriter.RequestWriter.TryComplete(), TaskScheduler.Default);
+				await pipeline.WaitForCompletionAsync();
+			}
 
-					Task[] writeTasks = outputWriter.AddToPipeline(pipeline, numWriteTasks);
+			// Update metadata for the workspace
+			foreach (FileUpdate fileUpdate in fileUpdates)
+			{
+				OutputFile outputFile = fileUpdate.OutputFile;
 
-					if (options.Progress != null)
-					{
-						_ = pipeline.AddTask(ctx => DirectoryNodeExtract.UpdateStatsAsync(batchReader, outputWriter, Task.WhenAll(writeTasks), options.Progress, options.ProgressUpdateFrequency, ctx));
-					}
+				FileState fileState = fileUpdate.ParentDir.FindOrAddFile(outputFile.FileEntry.Name);
+				fileState.Hash = outputFile.FileEntry.StreamHash;
+				fileState.Update(outputFile.FileInfo);
 
-					await pipeline.WaitForCompletionAsync();
+				HashInfo hashInfo = AddFileToHashLookup(fileState);
+
+				long offset = 0;
+				foreach (LeafChunkedDataNodeRef chunkRef in fileUpdate.Chunks)
+				{
+					ChunkInfo chunkInfo = new ChunkInfo(hashInfo, offset, chunkRef.Length);
+
+					HashInfo chunkHashInfo = FindOrAddHashInfo(chunkRef.BlobRef.Hash, chunkRef.Length);
+					chunkHashInfo.AddChunk(chunkInfo);
+
+					offset += chunkRef.Length;
 				}
 			}
 		}
 
-		async Task SyncDirectoryAsync(DirectoryReference dirRef, string path, DirectoryState dirState, DirectoryNode? dirNode, ulong flag, DirectoryState cacheDirState, ChannelWriter<ChunkReadRequest> requestWriter, ExtractOptions options, CancellationToken cancellationToken)
+		record class FileUpdate(OutputFile OutputFile, DirectoryState ParentDir, List<LeafChunkedDataNodeRef> Chunks);
+
+		async Task SyncDirectoryAsync(DirectoryReference dirRef, string path, DirectoryState dirState, DirectoryNode? dirNode, ulong flag, DirectoryState cacheDirState, List<FileUpdate> fileUpdates, ExtractOptions options, CancellationToken cancellationToken)
 		{
 			// Remove any directories that no longer exist
 			for (int subDirIdx = 0; subDirIdx < dirState.Directories.Count; subDirIdx++)
@@ -714,7 +791,7 @@ namespace EpicGames.Horde.Storage
 					if (dirNode == null || !dirNode.TryGetDirectoryEntry(subDirState.Name, out _))
 					{
 						DirectoryReference subDirPath = DirectoryReference.Combine(dirRef, subDirState.Name.ToString());
-						await SyncDirectoryAsync(subDirPath, $"{path}{subDirState.Name}/", subDirState, null, flag, cacheDirState, requestWriter, options, cancellationToken);
+						await SyncDirectoryAsync(subDirPath, $"{path}{subDirState.Name}/", subDirState, null, flag, cacheDirState, fileUpdates, options, cancellationToken);
 
 						if (subDirState.LayerFlags == 0)
 						{
@@ -759,7 +836,7 @@ namespace EpicGames.Horde.Storage
 					DirectoryState subDirState = dirState.FindOrAddDirectory(subDirEntry.Name);
 
 					DirectoryNode subDirNode = await subDirEntry.Handle.ReadBlobAsync(cancellationToken: cancellationToken);
-					await SyncDirectoryAsync(subDirPath, $"{path}{subDirEntry.Name}/", subDirState, subDirNode, flag, cacheDirState, requestWriter, options, cancellationToken);
+					await SyncDirectoryAsync(subDirPath, $"{path}{subDirEntry.Name}/", subDirState, subDirNode, flag, cacheDirState, fileUpdates, options, cancellationToken);
 
 					dirState.LayerFlags |= flag;
 				}
@@ -767,7 +844,7 @@ namespace EpicGames.Horde.Storage
 				// Update files
 				foreach (FileEntry fileEntry in dirNode.Files)
 				{
-					await SyncFileAsync(dirRef, $"{path}{fileEntry.Name}", dirState, fileEntry, flag, requestWriter, options, cancellationToken);
+					await SyncFileAsync(dirRef, $"{path}{fileEntry.Name}", dirState, fileEntry, flag, fileUpdates, cancellationToken);
 					dirState.LayerFlags |= flag;
 				}
 			}
@@ -779,7 +856,7 @@ namespace EpicGames.Horde.Storage
 			}
 		}
 
-		async Task SyncFileAsync(DirectoryReference dirRef, string path, DirectoryState dirState, FileEntry fileEntry, ulong flag, ChannelWriter<ChunkReadRequest> requestWriter, ExtractOptions options, CancellationToken cancellationToken)
+		async Task SyncFileAsync(DirectoryReference dirRef, string path, DirectoryState dirState, FileEntry fileEntry, ulong flag, List<FileUpdate> fileUpdates, CancellationToken cancellationToken)
 		{
 			FileState? fileState;
 			if (dirState.TryGetFile(fileEntry.Name, out fileState))
@@ -800,10 +877,9 @@ namespace EpicGames.Horde.Storage
 				fileState = TryMoveCachedData(fileEntry.StreamHash, dirState, fileEntry.Name);
 				if (fileState == null)
 				{
-					fileState = dirState.FindOrAddFile(fileEntry.Name);
-
-					WorkspaceOutputFile outputFile = new WorkspaceOutputFile(path, fileInfo, fileEntry, fileState);
-					await DirectoryNodeExtract.FindOutputChunksForFileAsync(outputFile, requestWriter, options, _logger, cancellationToken);
+					OutputFile outputFile = new OutputFile(path, fileInfo, fileEntry);
+					List<LeafChunkedDataNodeRef> chunks = await fileEntry.Target.EnumerateLeafNodesAsync(cancellationToken).ToListAsync(cancellationToken);
+					fileUpdates.Add(new FileUpdate(outputFile, dirState, chunks));
 				}
 				else
 				{
@@ -814,41 +890,81 @@ namespace EpicGames.Horde.Storage
 			}
 		}
 
-		async Task ExtractDataAsync(ChunkedDataNodeRef nodeRef, Stream outputStream, CancellationToken cancellationToken)
+		record class ChunkCopyRequest(OutputFile OutputFile, long Offset, SourceChunk SourceChunk);
+
+		async Task CopyChunksAsync(ChannelReader<ChunkCopyRequest> copyRequests, ChannelWriter<OutputChunk[]> outputRequests, CancellationToken cancellationToken)
 		{
-			if (nodeRef.Type == ChunkedDataNodeType.Leaf)
+			await foreach (ChunkCopyRequest copyRequest in copyRequests.ReadAllAsync(cancellationToken))
 			{
-				await ExtractLeafDataAsync(nodeRef, outputStream, cancellationToken);
-			}
-			else if (nodeRef.Type == ChunkedDataNodeType.Interior)
-			{
-				InteriorChunkedDataNode interiorNode = await nodeRef.Handle.ReadBlobAsync<InteriorChunkedDataNode>(cancellationToken: cancellationToken);
-				foreach (ChunkedDataNodeRef childNodeRef in interiorNode.Children)
+				FileStreamOptions options = new FileStreamOptions
 				{
-					await ExtractDataAsync(childNodeRef, outputStream, cancellationToken);
+					Mode = FileMode.Open,
+					Share = FileShare.Read,
+					Access = FileAccess.Read,
+					Options = FileOptions.Asynchronous,
+					BufferSize = 0,
+				};
+
+				SourceChunk sourceChunk = copyRequest.SourceChunk;
+				FileReference sourceFile = sourceChunk.File.GetFileReference(RootDir);
+
+				byte[] data = new byte[sourceChunk.Length];
+				using (FileStream sourceStream = new FileStream(sourceFile.FullName, options))
+				{
+					sourceStream.Seek(sourceChunk.Offset, SeekOrigin.Begin);
+					await sourceStream.ReadExactlyAsync(data, cancellationToken);
 				}
+
+				IoHash hash = IoHash.Compute(data);
+				if (hash != sourceChunk.Hash)
+				{
+					throw new InvalidOperationException($"Source data does not have correct hash; expected {sourceChunk.Hash}, actually {hash} ({sourceFile}, offset: {sourceChunk.Offset}, length: {sourceChunk.Length})");
+				}
+
+#pragma warning disable CA2000
+				OutputChunk outputChunk = new OutputChunk(copyRequest.OutputFile, copyRequest.Offset, data, sourceChunk, null);
+				await outputRequests.WriteAsync([outputChunk], cancellationToken);
+#pragma warning restore CA2000
+			}
+		}
+
+		record class SourceChunk(FileState File, long Offset, long Length, IoHash Hash)
+		{
+			/// <inheritdoc/>
+			public override string ToString()
+				=> $"{File.GetPath()}@{Offset}";
+		}
+
+		SourceChunk? TryGetSourceChunk(IoHash hash)
+		{
+			HashInfo? hashInfo;
+			if (_hashes.TryGetValue(hash, out hashInfo))
+			{
+				return TryGetSourceChunk(hashInfo);
 			}
 			else
 			{
-				throw new InvalidDataException($"Invalid node type {nodeRef.Type}");
+				return null;
 			}
 		}
 
-		async Task ExtractLeafDataAsync(ChunkedDataNodeRef nodeRef, Stream outputStream, CancellationToken cancellationToken)
+		static SourceChunk? TryGetSourceChunk(HashInfo hashInfo)
 		{
-			// Try to copy cached data for this node
-			HashInfo? hashInfo;
-			if (_hashes.TryGetValue(nodeRef.Handle.Hash, out hashInfo))
+			if (hashInfo.Files.Count > 0)
+			{
+				return new SourceChunk(hashInfo.Files[0], 0, hashInfo.Length, hashInfo.Hash);
+			}
+
+			foreach (ChunkInfo chunkInfo in hashInfo.Chunks)
 			{
-				if (await TryCopyCachedDataAsync(hashInfo, 0, hashInfo.Length, outputStream, cancellationToken))
+				SourceChunk? otherChunk = TryGetSourceChunk(chunkInfo.WithinHashInfo);
+				if (otherChunk != null)
 				{
-					return;
+					return new SourceChunk(otherChunk.File, otherChunk.Offset + chunkInfo.Offset, hashInfo.Length, hashInfo.Hash);
 				}
 			}
 
-			// Otherwise 
-			using BlobData blobData = await nodeRef.Handle.ReadBlobDataAsync(cancellationToken);
-			await LeafChunkedDataNode.CopyToStreamAsync(blobData, outputStream, cancellationToken);
+			return null;
 		}
 
 		FileState? TryMoveCachedData(IoHash hash, DirectoryState targetDirState, string targetName)
@@ -927,6 +1043,7 @@ namespace EpicGames.Horde.Storage
 		public async Task RefreshAsync(CancellationToken cancellationToken = default)
 		{
 			await RefreshAsync(_rootDirState, cancellationToken);
+			RemoveUnusedHashInfoObjects();
 		}
 
 		async Task RefreshAsync(DirectoryState dirState, CancellationToken cancellationToken)
@@ -945,10 +1062,12 @@ namespace EpicGames.Horde.Storage
 				if (!fileInfo.Exists)
 				{
 					_logger.LogWarning("File {Path} has been removed", fileInfo.FullName);
+					RemoveFileFromHashLookup(fileState);
 				}
 				else if (fileState.IsModified(fileInfo))
 				{
 					_logger.LogWarning("File {Path} has been modified", fileInfo.FullName);
+					RemoveFileFromHashLookup(fileState);
 				}
 				else
 				{
@@ -1012,17 +1131,24 @@ namespace EpicGames.Horde.Storage
 			}
 		}
 
-		void AddFileToHashLookup(FileState file)
+		HashInfo FindOrAddHashInfo(IoHash hash, long length)
 		{
 			HashInfo? hashInfo;
-			if (!_hashes.TryGetValue(file.Hash, out hashInfo))
+			if (!_hashes.TryGetValue(hash, out hashInfo))
 			{
-				hashInfo = new HashInfo(file.Hash, file.Length);
-				_hashes.Add(file.Hash, hashInfo);
+				hashInfo = new HashInfo(hash, length);
+				_hashes.Add(hash, hashInfo);
 			}
+			return hashInfo;
+		}
+
+		HashInfo AddFileToHashLookup(FileState file)
+		{
+			HashInfo hashInfo = FindOrAddHashInfo(file.Hash, file.Length);
 			hashInfo.Files.Add(file);
+			return hashInfo;
 		}
-		/*
+
 		void RemoveFileFromHashLookup(FileState file)
 		{
 			HashInfo? hashInfo;
@@ -1031,7 +1157,7 @@ namespace EpicGames.Horde.Storage
 				hashInfo.Files.Remove(file);
 			}
 		}
-		*/
+
 		void MoveFileToCache(FileState fileState, DirectoryState cacheDirState)
 		{
 			string name = fileState.Hash.ToString();
