Rename default GetSparseClassData accessor to GetMutableSparseClassData to make mutability more obvious and to encourage use of more explicit accessor with EGetSparseClassDataMethod parameter - deprecating old version

diff --git a/Engine/Source/Programs/Shared/EpicGames.UHT/Exporters/CodeGen/UhtHeaderCodeGeneratorHFile.cs b/Engine/Source/Programs/Shared/EpicGames.UHT/Exporters/CodeGen/UhtHeaderCodeGeneratorHFile.cs
--- a/Engine/Source/Programs/Shared/EpicGames.UHT/Exporters/CodeGen/UhtHeaderCodeGeneratorHFile.cs
+++ b/Engine/Source/Programs/Shared/EpicGames.UHT/Exporters/CodeGen/UhtHeaderCodeGeneratorHFile.cs
@@ -952,6 +952,7 @@ namespace EpicGames.UHT.Exporters.CodeGen
 				{
 					string sparseDataType = sparseScriptStruct.EngineName;
 
+					builder.Append("UE_DEPRECATED(5.6, \"Use (preferably) Get").Append(sparseDataType).Append("(EGetSparseClassDataMethod::ArchetypeIfNull), or GetMutable").Append(sparseDataType).Append("() instead if you need to mutate the sparse class data (best avoided at runtime)\")\\\r\n");
 					builder.Append(api).Append('F').Append(sparseDataType).Append("* Get").Append(sparseDataType).Append("() const; \\\r\n");
 
 					builder.Append(api).Append('F').Append(sparseDataType).Append("* GetMutable").Append(sparseDataType).Append("() const; \\\r\n");
