EpicGames.Core: Prevent cancellation exceptions propagating when disposing of an async pipeline.

#rnx

diff --git a/Engine/Source/Programs/Shared/EpicGames.Core/AsyncPipeline.cs b/Engine/Source/Programs/Shared/EpicGames.Core/AsyncPipeline.cs
--- a/Engine/Source/Programs/Shared/EpicGames.Core/AsyncPipeline.cs
+++ b/Engine/Source/Programs/Shared/EpicGames.Core/AsyncPipeline.cs
@@ -45,7 +45,14 @@ namespace EpicGames.Core
 
 			if (_tasks.Count > 0)
 			{
-				await Task.WhenAll(_tasks);
+				try
+				{
+					await Task.WhenAll(_tasks);
+				}
+				catch (OperationCanceledException)
+				{
+					// Ignore during dispose
+				}
 				_tasks.Clear();
 			}
 
