Horde: Do not log cancellation exceptions from batch output pipeline.

#rnx

diff --git a/Engine/Source/Programs/Shared/EpicGames.Horde/Storage/BatchOutputWriter.cs b/Engine/Source/Programs/Shared/EpicGames.Horde/Storage/BatchOutputWriter.cs
--- a/Engine/Source/Programs/Shared/EpicGames.Horde/Storage/BatchOutputWriter.cs
+++ b/Engine/Source/Programs/Shared/EpicGames.Horde/Storage/BatchOutputWriter.cs
@@ -221,7 +221,7 @@ namespace EpicGames.Horde.Storage
 				{
 					await WriteChunksAsync(chunks, cancellationToken);
 				}
-				catch (Exception ex)
+				catch (Exception ex) when (ex is not OperationCanceledException)
 				{
 					_logger.LogError(ex, "Error writing to disk: {Message}", ex.Message);
 					throw;
