Add additional full download size info to the current BuildDiff report

diff --git a/Engine/Source/Programs/Shared/EpicGames.Slack/Blocks/SectionBlock.cs b/Engine/Source/Programs/Shared/EpicGames.Slack/Blocks/SectionBlock.cs
--- a/Engine/Source/Programs/Shared/EpicGames.Slack/Blocks/SectionBlock.cs
+++ b/Engine/Source/Programs/Shared/EpicGames.Slack/Blocks/SectionBlock.cs
@@ -60,11 +60,22 @@ namespace EpicGames.Slack.Blocks
 	{
 		/// <summary>
 		/// Add an <see cref="SectionBlock"/> to the list of blocks
+		/// <param name="fields">TextObject instance to populate the section with</param>
 		/// </summary>
 		public static void AddSection(this ISlackBlockContainer container, TextObject text)
 		{
 			SectionBlock block = new SectionBlock(text);
 			container.Blocks.Add(block);
 		}
+
+		/// <summary>
+		/// Add an <see cref="SectionBlock"/> to the list of blocks
+		/// <param name="fields">List of TextObject instances to populate the section with</param>
+		/// </summary>
+		public static void AddSection(this ISlackBlockContainer container, List<TextObject> fields)
+		{
+			SectionBlock block = new SectionBlock(fields);
+			container.Blocks.Add(block);
+		}
 	}
 }
