Merging Horde changes up to CL 38271598.

diff --git a/Engine/Source/Programs/Shared/EpicGames.Horde/Storage/Backends/JupiterStorageBackend.cs b/Engine/Source/Programs/Shared/EpicGames.Horde/Storage/Backends/JupiterStorageBackend.cs
--- a/Engine/Source/Programs/Shared/EpicGames.Horde/Storage/Backends/JupiterStorageBackend.cs
+++ b/Engine/Source/Programs/Shared/EpicGames.Horde/Storage/Backends/JupiterStorageBackend.cs
@@ -194,7 +194,7 @@ namespace EpicGames.Horde.Storage.Backends
 			}
 		}
 
-		static byte[] CreateCbObject(ReadOnlyMemory<byte> payload, IReadOnlyCollection<BlobLocator> imports)
+		static byte[] CreateCbObject(byte[] payload, IReadOnlyCollection<BlobLocator> imports)
 		{
 			CbWriter writer = new CbWriter();
 			writer.BeginObject();
