Limited a UHT warning about relative include paths to only report paths from outside of the module

#rnx

diff --git a/Engine/Source/Programs/Shared/EpicGames.UHT/Exporters/CodeGen/UhtCodeGenerator.cs b/Engine/Source/Programs/Shared/EpicGames.UHT/Exporters/CodeGen/UhtCodeGenerator.cs
--- a/Engine/Source/Programs/Shared/EpicGames.UHT/Exporters/CodeGen/UhtCodeGenerator.cs
+++ b/Engine/Source/Programs/Shared/EpicGames.UHT/Exporters/CodeGen/UhtCodeGenerator.cs
@@ -257,8 +257,13 @@ namespace EpicGames.UHT.Exporters.CodeGen
 			// This will create a "../" path which is not great
 			if (relativePath == null)
 			{
-				Session.Logger.LogWarning("There are no include paths in module {ModuleName} for {FileName}. Will create \"../\" path", headerFile.Module.ShortName, headerFile.FilePath);
-				relativePath = Path.GetRelativePath(headerFile.Module.Module.IncludePaths[0], headerFile.FilePath);
+				string? includePath = headerFile.Module.Module.IncludePaths.Find(includePath => includePath.StartsWith(headerFile.Module.Module.BaseDirectory, StringComparison.Ordinal));
+				if (includePath == null)
+				{
+					Session.Logger.LogWarning("There are no include paths in module {ModuleName} for {FileName}. It will be included by a \"../\" path external to the module.", headerFile.Module.ShortName, headerFile.FilePath);
+					includePath = headerFile.Module.Module.IncludePaths.First();
+				}
+				relativePath = Path.GetRelativePath(includePath!, headerFile.FilePath);
 			}
 
 			headerInfo.IncludePath = relativePath.Replace('\\', '/');
diff --git a/Engine/Source/Runtime/CoreUObject/CoreUObject.Build.cs b/Engine/Source/Runtime/CoreUObject/CoreUObject.Build.cs
--- a/Engine/Source/Runtime/CoreUObject/CoreUObject.Build.cs
+++ b/Engine/Source/Runtime/CoreUObject/CoreUObject.Build.cs
@@ -76,7 +76,6 @@ public class CoreUObject : ModuleRules
 			PrivateIncludePathModuleNames.Add("ApplicationCore");
 		}
 
-		PrivateIncludePaths.Add(System.IO.Path.Combine(ModuleDirectory, "Tests"));
 		PrivateDefinitions.Add("UE_DEFINE_LEGACY_MATH_CONSTANT_MACRO_NAMES=0");
 	}
 }
