Cloud DDC - Fixed build warnings with the newly added CbObjectId type

diff --git a/Engine/Source/Programs/Shared/EpicGames.Serialization/CbField.cs b/Engine/Source/Programs/Shared/EpicGames.Serialization/CbField.cs
--- a/Engine/Source/Programs/Shared/EpicGames.Serialization/CbField.cs
+++ b/Engine/Source/Programs/Shared/EpicGames.Serialization/CbField.cs
@@ -2690,9 +2690,15 @@ namespace EpicGames.Serialization
 		#endregion
 	}
 
+	/// <summary>
+	/// ObjectId used within compact binary, similar to an uuid type 4
+	/// </summary>
+	/// <param name="a">time component</param>
+	/// <param name="b">serial component</param>
+	/// <param name="c">run component</param>
 	[JsonConverter(typeof(CbObjectIdJsonConverter))]
 	[TypeConverter(typeof(CbObjectIdTypeConverter))]
-	public struct CbObjectId(uint a, uint b, uint c) : IEquatable<CbObjectId>
+	public readonly struct CbObjectId(uint a, uint b, uint c) : IEquatable<CbObjectId>
 	{
 		readonly uint _a = a;
 		readonly uint _b = b;
@@ -2804,10 +2810,20 @@ namespace EpicGames.Serialization
 		public bool Equals(CbObjectId other) => _a == other._a && _b == other._b && _c == other._c;
 
 		/// <inheritdoc/>
-		public override bool Equals(object? obj) => (obj is CbObjectId hash) && Equals(hash);
+		public override bool Equals(object? obj) => (obj is CbObjectId cid) && Equals(cid);
 
 		/// <inheritdoc/>
 		public override int GetHashCode() => (int)_a;
+
+		/// <summary>
+		/// Test two object ids for equality
+		/// </summary>
+		public static bool operator ==(CbObjectId a, CbObjectId b) => a.Equals(b);
+
+		/// <summary>
+		/// Test two object ids for equality
+		/// </summary>
+		public static bool operator !=(CbObjectId a, CbObjectId b) => !(a == b);
 	}
 
 	/// <summary>
