Fix crash in UHT by adding a default include path in case none were found

diff --git a/Engine/Source/Programs/Shared/EpicGames.UHT/Exporters/CodeGen/UhtCodeGenerator.cs b/Engine/Source/Programs/Shared/EpicGames.UHT/Exporters/CodeGen/UhtCodeGenerator.cs
--- a/Engine/Source/Programs/Shared/EpicGames.UHT/Exporters/CodeGen/UhtCodeGenerator.cs
+++ b/Engine/Source/Programs/Shared/EpicGames.UHT/Exporters/CodeGen/UhtCodeGenerator.cs
@@ -261,7 +261,7 @@ namespace EpicGames.UHT.Exporters.CodeGen
 				if (includePath == null)
 				{
 					Session.Logger.LogWarning("There are no include paths in module {ModuleName} for {FileName}. It will be included by a \"../\" path external to the module.", headerFile.Module.ShortName, headerFile.FilePath);
-					includePath = headerFile.Module.Module.IncludePaths.First();
+					includePath = headerFile.Module.Module.IncludePaths.FirstOrDefault(".");
 				}
 				relativePath = Path.GetRelativePath(includePath!, headerFile.FilePath);
 			}
