EpicGames.Slack: Fix compiler warnings

diff --git a/Engine/Source/Programs/Shared/EpicGames.Slack/Blocks/SectionBlock.cs b/Engine/Source/Programs/Shared/EpicGames.Slack/Blocks/SectionBlock.cs
--- a/Engine/Source/Programs/Shared/EpicGames.Slack/Blocks/SectionBlock.cs
+++ b/Engine/Source/Programs/Shared/EpicGames.Slack/Blocks/SectionBlock.cs
@@ -59,22 +59,24 @@ namespace EpicGames.Slack.Blocks
 	public static class SectionBlockExtensions
 	{
 		/// <summary>
-		/// Add an <see cref="SectionBlock"/> to the list of blocks
-		/// <param name="fields">TextObject instance to populate the section with</param>
+		/// Adds a new <see cref="SectionBlock"/> with a single text element to the container.
+		/// <param name="container">The block container to add the section to</param>
+		/// <param name="text">The text content to populate the section with</param>
 		/// </summary>
 		public static void AddSection(this ISlackBlockContainer container, TextObject text)
 		{
-			SectionBlock block = new SectionBlock(text);
+			SectionBlock block = new (text);
 			container.Blocks.Add(block);
 		}
 
 		/// <summary>
-		/// Add an <see cref="SectionBlock"/> to the list of blocks
-		/// <param name="fields">List of TextObject instances to populate the section with</param>
+		/// Adds a new <see cref="SectionBlock"/> with multiple field elements to the container
+		/// <param name="container">The block container to add the section to</param>
+		/// <param name="fields">The list of text fields to populate the section with</param>
 		/// </summary>
 		public static void AddSection(this ISlackBlockContainer container, List<TextObject> fields)
 		{
-			SectionBlock block = new SectionBlock(fields);
+			SectionBlock block = new (fields);
 			container.Blocks.Add(block);
 		}
 	}
