Cloud DDC - Fixed an issue were the json converted value of a object id was a base64 blob instead of a hex encoded string as we normally do.

diff --git a/Engine/Source/Programs/Shared/EpicGames.Serialization/CbField.cs b/Engine/Source/Programs/Shared/EpicGames.Serialization/CbField.cs
--- a/Engine/Source/Programs/Shared/EpicGames.Serialization/CbField.cs
+++ b/Engine/Source/Programs/Shared/EpicGames.Serialization/CbField.cs
@@ -2839,7 +2839,7 @@ namespace EpicGames.Serialization
 		public override CbObjectId Read(ref Utf8JsonReader reader, Type typeToConvert, JsonSerializerOptions options) => CbObjectId.Parse(reader.GetString()!);
 
 		/// <inheritdoc/>
-		public override void Write(Utf8JsonWriter writer, CbObjectId value, JsonSerializerOptions options) => writer.WriteBase64StringValue(value.ToByteArray());
+		public override void Write(Utf8JsonWriter writer, CbObjectId value, JsonSerializerOptions options) => writer.WriteStringValue(value.ToString());
 	}
 
 	/// <summary>
