Cloud DDC - Fixed an issue serializing a CbObject that was a part of a poco type. The collection of CbFields was picked up the serializer meaning the object was written as an array with fields (which is not valid json).

diff --git a/Engine/Source/Programs/Shared/EpicGames.Serialization/CbField.cs b/Engine/Source/Programs/Shared/EpicGames.Serialization/CbField.cs
--- a/Engine/Source/Programs/Shared/EpicGames.Serialization/CbField.cs
+++ b/Engine/Source/Programs/Shared/EpicGames.Serialization/CbField.cs
@@ -1750,6 +1750,11 @@ namespace EpicGames.Serialization
 
 		/// <inheritdoc/>
 		public override void Write(Utf8JsonWriter writer, CbField field, JsonSerializerOptions options)
+		{
+			WriteField(writer, field, options);
+		}
+
+		public static void WriteField(Utf8JsonWriter writer, CbField field, JsonSerializerOptions options)
 		{
 			switch (field.GetType())
 			{
@@ -1776,7 +1781,7 @@ namespace EpicGames.Serialization
 
 					foreach (CbField member in field.AsObject())
 					{
-						Write(writer, member, options);
+						WriteField(writer, member, options);
 					}
 					writer.WriteEndObject();
 					break;
@@ -1793,7 +1798,7 @@ namespace EpicGames.Serialization
 
 					foreach (CbField element in field.AsArray())
 					{
-						Write(writer, element, options);
+						WriteField(writer, element, options);
 					}
 					writer.WriteEndArray();
 					break;
@@ -1907,6 +1912,29 @@ namespace EpicGames.Serialization
 		}
 	}
 
+	/// <summary>
+	/// Converter to and from JSON objects
+	/// </summary>
+	public class CbObjectJsonConverter : JsonConverter<CbObject>
+	{
+		/// <inheritdoc/>
+		public override CbObject Read(ref Utf8JsonReader reader, Type typeToConvert, JsonSerializerOptions options)
+		{
+			throw new NotImplementedException();
+		}
+
+		/// <inheritdoc/>
+		public override void Write(Utf8JsonWriter writer, CbObject o, JsonSerializerOptions options)
+		{
+			writer.WriteStartObject();
+
+			foreach (CbField member in o)
+			{
+				CbFieldJsonConverter.WriteField(writer, member, options);
+			}
+			writer.WriteEndObject();
+		}
+	}
 	/// <summary>
 	/// Enumerator for contents of a field
 	/// </summary>
@@ -2365,6 +2393,7 @@ namespace EpicGames.Serialization
 	/// attempting to access the empty field is also safe, as described by FCbFieldView.
 	/// </summary>
 	[DebuggerTypeProxy(typeof(CbObjectDebugView))]
+	[JsonConverter(typeof(CbObjectJsonConverter))]
 	public class CbObject : IEnumerable<CbField>
 	{
 		/// <summary>
