Cloud DDC - Fixed build warning from change 38567473

diff --git a/Engine/Source/Programs/Shared/EpicGames.Serialization/CbField.cs b/Engine/Source/Programs/Shared/EpicGames.Serialization/CbField.cs
--- a/Engine/Source/Programs/Shared/EpicGames.Serialization/CbField.cs
+++ b/Engine/Source/Programs/Shared/EpicGames.Serialization/CbField.cs
@@ -1754,6 +1754,13 @@ namespace EpicGames.Serialization
 			WriteField(writer, field, options);
 		}
 
+		/// <summary>
+		/// Writes a CbField to a Jtf8JsonWriter
+		/// </summary>
+		/// <param name="writer">The json writer</param>
+		/// <param name="field">The field to write</param>
+		/// <param name="options">Serialization options</param>
+		/// <exception cref="NotImplementedException"></exception>
 		public static void WriteField(Utf8JsonWriter writer, CbField field, JsonSerializerOptions options)
 		{
 			switch (field.GetType())
