Horde: Add missing xml documentation
#submitool bypass
#rnx

diff --git a/Engine/Source/Programs/Shared/EpicGames.Horde/Telemetry/TelemetryMessages.cs b/Engine/Source/Programs/Shared/EpicGames.Horde/Telemetry/TelemetryMessages.cs
--- a/Engine/Source/Programs/Shared/EpicGames.Horde/Telemetry/TelemetryMessages.cs
+++ b/Engine/Source/Programs/Shared/EpicGames.Horde/Telemetry/TelemetryMessages.cs
@@ -35,7 +35,7 @@ namespace EpicGames.Horde.Telemetry
 	}
 
 	/// <summary>
-	/// 
+	/// Telemetry metrics query request
 	/// </summary>
 	public class GetTelemetryMetricsRequest
 	{
