Add specific exit code for conflicting instances of UBT. This will now cause UBT to exit with code 10, rather than the generic error code 6 (OtherCompilationError)

#rnx

diff --git a/Engine/Source/Programs/Shared/EpicGames.Build/System/GlobalSingleInstanceMutex.cs b/Engine/Source/Programs/Shared/EpicGames.Build/System/GlobalSingleInstanceMutex.cs
--- a/Engine/Source/Programs/Shared/EpicGames.Build/System/GlobalSingleInstanceMutex.cs
+++ b/Engine/Source/Programs/Shared/EpicGames.Build/System/GlobalSingleInstanceMutex.cs
@@ -42,7 +42,7 @@ namespace UnrealBuildBase
 				}
 				else
 				{
-					throw new BuildLogEventException("A conflicting instance of {Mutex} is already running.", mutexName);
+					throw new BuildLogEventException(new CompilationResultException(CompilationResult.ConflictingInstance), "A conflicting instance of {Mutex} is already running.", mutexName);
 				}
 			}
 		}
diff --git a/Engine/Source/Programs/Shared/EpicGames.Core/UnrealEngineTypes.cs b/Engine/Source/Programs/Shared/EpicGames.Core/UnrealEngineTypes.cs
--- a/Engine/Source/Programs/Shared/EpicGames.Core/UnrealEngineTypes.cs
+++ b/Engine/Source/Programs/Shared/EpicGames.Core/UnrealEngineTypes.cs
@@ -71,6 +71,11 @@ namespace EpicGames.Core
 		/// </summary>
 		ActionGraphInvalid = 9,
 
+		/// <summary>
+		/// Another instance of UBT was already running
+		/// </summary>
+		ConflictingInstance = 10,
+
 		/// <summary>
 		/// Compilation is not supported in the current build
 		/// </summary>
diff --git a/Engine/Source/Runtime/Core/Public/Misc/CompilationResult.h b/Engine/Source/Runtime/Core/Public/Misc/CompilationResult.h
--- a/Engine/Source/Runtime/Core/Public/Misc/CompilationResult.h
+++ b/Engine/Source/Runtime/Core/Public/Misc/CompilationResult.h
@@ -35,6 +35,8 @@ namespace ECompilationResult
 		RulesError = 8,
 		/** Compilation failed due to invalid action graph */
 		ActionGraphInvalid = 9,
+		/** Another instance of UBT was already running */
+		ConflictingInstance = 10,
 		/** Compilation is not supported in the current build */
 		Unsupported,
 		/** Unknown error */
