Fix crash and reliability issues with objects marked as Verse objects.

#rnx

diff --git a/Engine/Source/Programs/Shared/EpicGames.UHT/Types/UhtModule.cs b/Engine/Source/Programs/Shared/EpicGames.UHT/Types/UhtModule.cs
--- a/Engine/Source/Programs/Shared/EpicGames.UHT/Types/UhtModule.cs
+++ b/Engine/Source/Programs/Shared/EpicGames.UHT/Types/UhtModule.cs
@@ -55,9 +55,12 @@ namespace EpicGames.UHT.Types
 			get
 			{
 				yield return ScriptPackage;
-				foreach (KeyValuePair<string, UhtPackage> kvp in _packages)
+				lock (_packages)
 				{
-					yield return kvp.Value;
+					foreach (KeyValuePair<string, UhtPackage> kvp in _packages)
+					{
+						yield return kvp.Value;
+					}
 				}
 			}
 		}
@@ -196,13 +199,16 @@ namespace EpicGames.UHT.Types
 			{
 				return ScriptPackage;
 			}
-			if (_packages.TryGetValue(packageName, out UhtPackage? package))
+			lock (_packages)
 			{
+				if (_packages.TryGetValue(packageName, out UhtPackage? package))
+				{
+					return package;
+				}
+				package = new(this, packageName, packageFlags);
+				_packages.Add(packageName, package);
 				return package;
 			}
-			package = new(this, packageName, packageFlags);
-			_packages.Add(packageName, package);
-			return package;
 		}
 
 		/// <summary>
