Make the default name processing for enumerations to match structs and classes.
Convert the only remaining VNI enum in VerseTestNativeCmd to use UHT

#rnx
#okforversepublic

diff --git a/Engine/Source/Programs/Shared/EpicGames.UHT/Exporters/CodeGen/UhtHeaderCodeGeneratorCppFile.cs b/Engine/Source/Programs/Shared/EpicGames.UHT/Exporters/CodeGen/UhtHeaderCodeGeneratorCppFile.cs
--- a/Engine/Source/Programs/Shared/EpicGames.UHT/Exporters/CodeGen/UhtHeaderCodeGeneratorCppFile.cs
+++ b/Engine/Source/Programs/Shared/EpicGames.UHT/Exporters/CodeGen/UhtHeaderCodeGeneratorCppFile.cs
@@ -294,7 +294,7 @@ namespace EpicGames.UHT.Exporters.CodeGen
 								.Append("_StaticEnum, TEXT(\"")
 								.Append(enumObj.EngineName)
 								.Append("\"), &Z_Registration_Info_UEnum_")
-								.Append(enumObj.EngineName)
+								.Append(enumObj.SourceName)
 								.Append($", CONSTRUCT_RELOAD_VERSION_INFO(FEnumReloadVersionInfo, {hash}U) }},\r\n");
 							combinedHash = HashCombine(combinedHash, hash);
 						},
@@ -313,7 +313,7 @@ namespace EpicGames.UHT.Exporters.CodeGen
 								.Append("_Statics::NewStructOps, TEXT(\"")
 								.Append(scriptStruct.EngineName)
 								.Append("\"), &Z_Registration_Info_UScriptStruct_")
-								.Append(scriptStruct.EngineName)
+								.Append(scriptStruct.SourceName)
 								.Append(", CONSTRUCT_RELOAD_VERSION_INFO(FStructReloadVersionInfo, sizeof(")
 								.Append(scriptStruct.SourceName)
 								.Append($"), {hash}U) }},\r\n");
@@ -565,7 +565,7 @@ namespace EpicGames.UHT.Exporters.CodeGen
 			const string MetaDataParamsName = "Struct_MetaDataParams";
 			string singletonName = GetSingletonName(scriptStruct, true);
 			string staticsName = singletonName + "_Statics";
-			string registrationName = $"Z_Registration_Info_UScriptStruct_{scriptStruct.EngineName}";
+			string registrationName = $"Z_Registration_Info_UScriptStruct_{scriptStruct.SourceName}";
 			List<UhtScriptStruct> noExportStructs = FindNoExportStructs(scriptStruct);
 
 			if (scriptStruct.ScriptStructFlags.HasAnyFlags(EStructFlags.Native))
diff --git a/Engine/Source/Programs/Shared/EpicGames.UHT/Specifiers/UhtFieldSpecifiers.cs b/Engine/Source/Programs/Shared/EpicGames.UHT/Specifiers/UhtFieldSpecifiers.cs
--- a/Engine/Source/Programs/Shared/EpicGames.UHT/Specifiers/UhtFieldSpecifiers.cs
+++ b/Engine/Source/Programs/Shared/EpicGames.UHT/Specifiers/UhtFieldSpecifiers.cs
@@ -46,45 +46,35 @@ namespace EpicGames.UHT.Parsers
 			// If no name was specified, generate one using the following algorithm FVerseXxYyZz -> xx_yy_zz
 			if (String.IsNullOrEmpty(fieldObj.VerseName))
 			{
-				if (fieldObj is UhtEnum)
+				ReadOnlySpan<char> span = StripVersePrefix(fieldObj);
+				if (!span.IsEmpty)
 				{
-					if (String.IsNullOrEmpty(fieldObj.VerseName))
+					using BorrowStringBuilder borrowStringBuilder = new BorrowStringBuilder(StringBuilderCache.Small);
+					bool first = true;
+					foreach (char c in span)
 					{
-						fieldObj.VerseName = fieldObj.EngineName; // always use the engine name???
-					}
-				}
-				else
-				{
-					if (fieldObj.EngineName.StartsWith("Verse", StringComparison.Ordinal))
-					{
-						using BorrowStringBuilder borrowStringBuilder = new BorrowStringBuilder(StringBuilderCache.Small);
-						ReadOnlySpan<char> span = fieldObj.EngineName.AsSpan(5);
-						bool first = true;
-						foreach (char c in span)
+						if (c >= 'A' && c <= 'Z')
 						{
-							if (c >= 'A' && c <= 'Z')
+							if (!first)
 							{
-								if (!first)
-								{
-									borrowStringBuilder.StringBuilder.Append('_');
-								}
-								borrowStringBuilder.StringBuilder.Append(Char.ToLower(c));
+								borrowStringBuilder.StringBuilder.Append('_');
 							}
-							else
-							{
-								borrowStringBuilder.StringBuilder.Append(c);
-							}
-							first = false;
+							borrowStringBuilder.StringBuilder.Append(Char.ToLower(c));
+						}
+						else
+						{
+							borrowStringBuilder.StringBuilder.Append(c);
 						}
-						fieldObj.VerseName = borrowStringBuilder.StringBuilder.ToString();
+						first = false;
 					}
+					fieldObj.VerseName = borrowStringBuilder.StringBuilder.ToString();
 				}
 			}
 
 			// Final validation
 			if (!fieldObj.IsVerseField)
 			{
-				fieldObj.LogError($"Verse specifier must include the name");
+				fieldObj.LogError($"Verse specifier must include the name or the source name prefixed with EVerse, UVerse, AVerse, or FVerse depending on the type");
 			}
 			else if (fieldObj.Outer is not UhtPackage)
 			{
@@ -106,5 +96,25 @@ namespace EpicGames.UHT.Parsers
 				fieldObj.Outer = fieldObj.Module.CreatePackage(borrowBuilder.StringBuilder.ToString());
 			}
 		}
+
+		private static ReadOnlySpan<char> StripVersePrefix(UhtField fieldObj)
+		{
+			// For enums, enforce 
+			if (fieldObj is UhtEnum)
+			{
+				if (fieldObj.EngineName.StartsWith("EVerse", StringComparison.Ordinal))
+				{
+					return fieldObj.EngineName.AsSpan(6);
+				}
+			}
+			else
+			{
+				if (fieldObj.EngineName.StartsWith("Verse", StringComparison.Ordinal))
+				{
+					return fieldObj.EngineName.AsSpan(5);
+				}
+			}
+			return new();
+		}
 	}
 }
