Shared: misc intellisense updates

#rnx

diff --git a/Engine/Source/Programs/Shared/EpicGames.Build/System/CompileScriptModules.cs b/Engine/Source/Programs/Shared/EpicGames.Build/System/CompileScriptModules.cs
--- a/Engine/Source/Programs/Shared/EpicGames.Build/System/CompileScriptModules.cs
+++ b/Engine/Source/Programs/Shared/EpicGames.Build/System/CompileScriptModules.cs
@@ -124,7 +124,7 @@ namespace UnrealBuildBase
 		/// <returns>Set of target files</returns>
 		public static HashSet<FileReference> GetTargetPaths(IReadOnlyDictionary<FileReference, CsProjBuildRecordEntry> BuildRecords)
 		{
-			return new HashSet<FileReference>(BuildRecords.Select(x => GetTargetPath(x)));
+			return [.. BuildRecords.Select(GetTargetPath)];
 		}
 
 		/// <summary>
@@ -150,8 +150,8 @@ namespace UnrealBuildBase
 		/// <param name="OnBuildingProjects">Action to invoke when projects get built</param>
 		/// <param name="Logger"></param>
 		/// <returns>Collection of all the projects.  They will have been compiled.</returns>
-		public static HashSet<FileReference> InitializeScriptModules(Rules.RulesFileType RulesFileType, 
-			string? ScriptsForProjectFileName, List<string>? AdditionalScriptsFolders, bool bForceCompile, bool bNoCompile, bool bUseBuildRecords, 
+		public static HashSet<FileReference> InitializeScriptModules(Rules.RulesFileType RulesFileType,
+			string? ScriptsForProjectFileName, List<string>? AdditionalScriptsFolders, bool bForceCompile, bool bNoCompile, bool bUseBuildRecords,
 			out bool bBuildSuccess, Action<int> OnBuildingProjects, ILogger Logger)
 		{
 			List<DirectoryReference> GameDirectories = GetGameDirectories(ScriptsForProjectFileName, Logger);
@@ -167,11 +167,11 @@ namespace UnrealBuildBase
 			];
 			BaseDirectories = BaseDirectories.Distinct().ToList();
 
-			HashSet<FileReference> FoundProjects = new HashSet<FileReference>(
-				Rules.FindAllRulesSourceFiles(RulesFileType,
+			HashSet<FileReference> FoundProjects =
+				[.. Rules.FindAllRulesSourceFiles(RulesFileType,
 				// Project scripts require source engine builds
-				GameFolders: Unreal.IsEngineInstalled() ? GameDirectories : [], 
-				ForeignPlugins: null, AdditionalSearchPaths: [.. AdditionalDirectories, .. GameBuildDirectories]));
+				GameFolders: Unreal.IsEngineInstalled() ? GameDirectories : [],
+				ForeignPlugins: null, AdditionalSearchPaths: [.. AdditionalDirectories, .. GameBuildDirectories])];
 
 			BuildFlags BuildFlags = BuildFlags.None;
 			if (bForceCompile)
@@ -230,7 +230,7 @@ namespace UnrealBuildBase
 					Logger);
 
 				bBuildSuccess = bEngineBuildSuccess && bProjectBuildSuccess;
-				BuildResults = new Dictionary<FileReference, CsProjBuildRecordEntry>(EngineBuildResults);
+				BuildResults = new(EngineBuildResults);
 				foreach (KeyValuePair<FileReference, CsProjBuildRecordEntry> Item in ProjectBuildResults)
 				{
 					if (!BuildResults.ContainsKey(Item.Key))
@@ -284,7 +284,7 @@ namespace UnrealBuildBase
 		/// <param name="Logger"></param>
 		/// <returns>Collection of all the projects.  They will have been compiled.</returns>
 		public static Dictionary<FileReference, CsProjBuildRecordEntry> Build(Rules.RulesFileType RulesFileType,
-			HashSet<FileReference> FoundProjects, IEnumerable<DirectoryReference> BaseDirectories, IEnumerable<string>? DefineConstants, 
+			HashSet<FileReference> FoundProjects, IEnumerable<DirectoryReference> BaseDirectories, IEnumerable<string>? DefineConstants,
 			BuildFlags BuildFlags, out bool bBuildSuccess, Action<int> OnBuildingProjects, ILogger Logger)
 		{
 			CsProjBuildHook Hook = new Hook(Logger, RulesFileType);
@@ -383,7 +383,7 @@ namespace UnrealBuildBase
 				if (FoundProjects.All(x => ValidBuildRecords.ContainsKey(x)))
 				{
 					bBuildSuccess = true;
-					return new Dictionary<FileReference, CsProjBuildRecordEntry>(ValidBuildRecords.Where(x => FoundProjects.Contains(x.Key)));
+					return new(ValidBuildRecords.Where(x => FoundProjects.Contains(x.Key)));
 				}
 			}
 
@@ -412,7 +412,7 @@ namespace UnrealBuildBase
 		/// <param name="Logger"></param>
 		/// <returns></returns>
 		static Dictionary<FileReference, CsProjBuildRecordEntry> LoadExistingBuildRecords(CsProjBuildHook Hook, IEnumerable<DirectoryReference> BaseDirectories, ILogger Logger)
-        {
+		{
 			Dictionary<FileReference, CsProjBuildRecordEntry> LoadedBuildRecords = [];
 
 			foreach (DirectoryReference Directory in BaseDirectories)
@@ -424,7 +424,7 @@ namespace UnrealBuildBase
 				}
 
 				foreach (FileReference JsonFile in DirectoryReference.EnumerateFiles(IntermediateDirectory, "*.json"))
-                {
+				{
 					CsProjBuildRecord? BuildRecord = default;
 
 					// filesystem errors or json parsing might result in an exception. If that happens, we fall back to the
@@ -434,7 +434,7 @@ namespace UnrealBuildBase
 						BuildRecord = JsonSerializer.Deserialize<CsProjBuildRecord>(FileReference.ReadAllText(JsonFile));
 						Logger.LogDebug("Loaded script module build record {JsonFile}", JsonFile);
 					}
-					catch(Exception Ex)
+					catch (Exception Ex)
 					{
 						Logger.LogWarning("[{JsonFile}] Failed to load build record: {Message}", JsonFile, Ex.Message);
 					}
@@ -448,21 +448,21 @@ namespace UnrealBuildBase
 						LoadedBuildRecords.Add(Entry.ProjectFile, Entry);
 					}
 					else
-                    {
+					{
 						// Delete the invalid build record
 						Logger.LogWarning("Deleting invalid build record {JsonFile}", JsonFile);
-                    }
-                }
+					}
+				}
 			}
 
 			return LoadedBuildRecords;
-        }
+		}
 
-		private static bool ValidateGlobbedFiles(DirectoryReference ProjectDirectory, 
+		private static bool ValidateGlobbedFiles(DirectoryReference ProjectDirectory,
 			List<CsProjBuildRecord.Glob> Globs, HashSet<string> GlobbedDependencies, out string ValidationFailureMessage)
 		{
 			// First, evaluate globs
-			
+
 			// Files are grouped by ItemType (e.g. Compile, EmbeddedResource) to ensure that Exclude and
 			// Remove act as expected.
 			Dictionary<string, HashSet<string>> Files = [];
@@ -500,7 +500,7 @@ namespace UnrealBuildBase
 
 			bool bValid = true;
 			StringBuilder ValidationFailureText = new StringBuilder();
-			
+
 			// Look for extra files that were found
 			foreach (HashSet<string> TypedFiles in Files.Values)
 			{
@@ -513,7 +513,7 @@ namespace UnrealBuildBase
 					}
 				}
 			}
-			
+
 			// Look for files that are missing
 			foreach (string File in GlobbedDependencies)
 			{
@@ -676,7 +676,7 @@ namespace UnrealBuildBase
 		}
 
 		static List<DirectoryReference> GetGameDirectories(string? ScriptsForProjectFileName, ILogger Logger)
-        {
+		{
 			List<DirectoryReference> GameDirectories = [];
 
 			if (String.IsNullOrEmpty(ScriptsForProjectFileName))
@@ -690,7 +690,7 @@ namespace UnrealBuildBase
 				GameDirectories.Add(ScriptsDir);
 			}
 			return GameDirectories;
-        }
+		}
 
 		static List<DirectoryReference> GetAdditionalDirectories(List<string>? AdditionalScriptsFolders) =>
 			AdditionalScriptsFolders == null ? [] :
@@ -698,6 +698,6 @@ namespace UnrealBuildBase
 
 		static List<DirectoryReference> GetAdditionalBuildDirectories(List<DirectoryReference> GameDirectories) =>
 			GameDirectories.Select(x => DirectoryReference.Combine(x, "Build")).Where(x => DirectoryReference.Exists(x)).ToList();
-    }
+	}
 }
 
diff --git a/Engine/Source/Programs/Shared/EpicGames.Build/System/EnumeratePlugins.cs b/Engine/Source/Programs/Shared/EpicGames.Build/System/EnumeratePlugins.cs
--- a/Engine/Source/Programs/Shared/EpicGames.Build/System/EnumeratePlugins.cs
+++ b/Engine/Source/Programs/Shared/EpicGames.Build/System/EnumeratePlugins.cs
@@ -220,7 +220,7 @@ namespace UnrealBuildBase
 			// Build the things needing to be built
 			bool bBuildSuccess = true;
 			Dictionary<FileReference, CsProjBuildRecordEntry>? BuiltPlugins =
-				CompileScriptModule.Build(UnrealBuildBase.Rules.RulesFileType.UbtPlugin, new HashSet<FileReference>(NotBuilt),
+				CompileScriptModule.Build(UnrealBuildBase.Rules.RulesFileType.UbtPlugin, [.. NotBuilt],
 				BaseDirectories, DefineConstants, CompileScriptModule.BuildFlags.UseBuildRecords, out bBuildSuccess,
 				Count =>
 				{
diff --git a/Engine/Source/Programs/Shared/EpicGames.Core/DirectoryReference.cs b/Engine/Source/Programs/Shared/EpicGames.Core/DirectoryReference.cs
--- a/Engine/Source/Programs/Shared/EpicGames.Core/DirectoryReference.cs
+++ b/Engine/Source/Programs/Shared/EpicGames.Core/DirectoryReference.cs
@@ -132,6 +132,12 @@ namespace EpicGames.Core
 			}
 		}
 
+		/// <summary>
+		/// Returns the path of the current user's temporary folder
+		/// </summary>
+		/// <returns>The path to the temporary folder</returns>
+		public static DirectoryReference GetTempPath() => new DirectoryReference(Path.GetTempPath());
+
 		/// <summary>
 		/// Gets the path for a special folder
 		/// </summary>
diff --git a/Engine/Source/Programs/Shared/EpicGames.Horde/Storage/IStorageWriter.cs b/Engine/Source/Programs/Shared/EpicGames.Horde/Storage/IStorageWriter.cs
--- a/Engine/Source/Programs/Shared/EpicGames.Horde/Storage/IStorageWriter.cs
+++ b/Engine/Source/Programs/Shared/EpicGames.Horde/Storage/IStorageWriter.cs
@@ -102,7 +102,7 @@ namespace EpicGames.Horde.Storage
 		public IBlobWriter CreateBlobWriter(string? basePath = null, BlobSerializerOptions? serializerOptions = null)
 			=> CreateBlobWriter(basePath, serializerOptions, _cancellationToken);
 
-		/// <inheritdoc cref="IStorageWriter.CreateBlobWriter(string?, BlobSerializerOptions?)"/>
+		/// <inheritdoc cref="IStorageWriter.CreateBlobWriter(String?, BlobSerializerOptions?)"/>
 		protected abstract IBlobWriter CreateBlobWriter(string? basePath = null, BlobSerializerOptions? serializerOptions = null, CancellationToken cancellationToken = default);
 
 		#region Aliases
