Horde: Lower timeout for ComputeSocket send buffer detach

diff --git a/Engine/Source/Programs/Shared/EpicGames.Horde/Compute/ComputeSocket.cs b/Engine/Source/Programs/Shared/EpicGames.Horde/Compute/ComputeSocket.cs
--- a/Engine/Source/Programs/Shared/EpicGames.Horde/Compute/ComputeSocket.cs
+++ b/Engine/Source/Programs/Shared/EpicGames.Horde/Compute/ComputeSocket.cs
@@ -798,9 +798,10 @@ namespace EpicGames.Horde.Compute
 				}
 
 				// Try to stop the send task
+				TimeSpan timeout = TimeSpan.FromSeconds(10);
 				try
 				{
-					await sendBuffer.Task.StopAsync(cancellationToken).WaitAsync(TimeSpan.FromSeconds(30.0), cancellationToken);
+					await sendBuffer.Task.StopAsync(cancellationToken).WaitAsync(timeout, cancellationToken);
 				}
 				catch (OperationCanceledException)
 				{
@@ -808,7 +809,7 @@ namespace EpicGames.Horde.Compute
 				}
 				catch (TimeoutException ex)
 				{
-					_logger.LogWarning(ex, "Send task did not terminate gracefully in 30s. Tearing down socket buffers.");
+					_logger.LogWarning(ex, "Send task did not terminate gracefully in {Timeout} ms. Tearing down socket buffers...", timeout.TotalMilliseconds);
 				}
 				catch (Exception ex)
 				{
