EpicGames.Core: Add extension method for redacting sensitive values in request URIs

diff --git a/Engine/Source/Programs/Shared/EpicGames.Core.Tests/HttpClientTests.cs b/Engine/Source/Programs/Shared/EpicGames.Core.Tests/HttpClientTests.cs
new file mode 100644
--- /dev/null
+++ b/Engine/Source/Programs/Shared/EpicGames.Core.Tests/HttpClientTests.cs
@@ -0,0 +1,25 @@
+// Copyright Epic Games, Inc. All Rights Reserved.
+
+using System.Net.Http;
+using Microsoft.VisualStudio.TestTools.UnitTesting;
+
+namespace EpicGames.Core.Tests;
+
+[TestClass]
+public class HttpClientExtensionTests
+{
+	[TestMethod]
+	public void RedactedUriTest()
+	{
+		AssertRedactedUri("http://a.com/path?foo=bar&hello=world", "http://a.com/path?foo=bar&hello=world");
+		AssertRedactedUri("http://a.com/path?AWSAccessKeyId=redacted", "http://a.com/path?AWSAccessKeyId=bar");
+		AssertRedactedUri("http://a.com/path?AWSACCESSKEYID=redacted&foo=bar", "http://a.com/path?AWSACCESSKEYID=bar&foo=bar");
+		AssertRedactedUri("http://a.com/path", "http://a.com/path");
+	}
+	
+	private static void AssertRedactedUri(string expectedUri, string actualUri)
+	{
+		using HttpRequestMessage req = new (HttpMethod.Get, actualUri);
+		Assert.AreEqual(expectedUri, req.RedactedRequestUri()!.ToString());
+	}
+}
diff --git a/Engine/Source/Programs/Shared/EpicGames.Core/HttpClientExtensions.cs b/Engine/Source/Programs/Shared/EpicGames.Core/HttpClientExtensions.cs
--- a/Engine/Source/Programs/Shared/EpicGames.Core/HttpClientExtensions.cs
+++ b/Engine/Source/Programs/Shared/EpicGames.Core/HttpClientExtensions.cs
@@ -1,11 +1,15 @@
 // Copyright Epic Games, Inc. All Rights Reserved.
 
 using System;
+using System.Collections.Generic;
+using System.Collections.Specialized;
+using System.Linq;
 using System.Net.Http;
 using System.Text;
 using System.Text.Json;
 using System.Threading;
 using System.Threading.Tasks;
+using System.Web;
 
 #pragma warning disable CA2234 // Pass system uri objects instead of strings
 #pragma warning disable CA1054 // URI-like parameters should not be strings
@@ -170,4 +174,38 @@ namespace EpicGames.Core
 			return JsonSerializer.Deserialize<T>(bytes, new JsonSerializerOptions { PropertyNameCaseInsensitive = true })!;
 		}
 	}
+	
+	/// <summary>
+	/// Extension methods for HttpRequestMessage
+	/// </summary>
+	public static class HttpRequestMessageExtensions
+	{
+		private static readonly HashSet<string> s_sensitiveKeys = new(StringComparer.OrdinalIgnoreCase)
+		{
+			"x-amz-security-token",
+			"Signature",
+			"AWSAccessKeyId"
+		};
+		
+		/// <summary>
+		/// Returns a new URI with sensitive query parameters redacted from the request URI.
+		/// </summary>
+		public static Uri? RedactedRequestUri(this HttpRequestMessage request)
+		{
+			Uri? uri = request?.RequestUri;
+			if (uri?.Query is not { Length: > 0 })
+			{
+				return uri;
+			}
+			
+			NameValueCollection queryParams = HttpUtility.ParseQueryString(uri.Query);
+			foreach (string key in queryParams.AllKeys.OfType<string>())
+			{
+				string value = s_sensitiveKeys.Contains(key) ? "redacted" : queryParams[key] ?? String.Empty;
+				queryParams[key] = value;
+			}
+			
+			return new UriBuilder(uri) { Query = String.Join("&", queryParams) }.Uri;
+		}
+	}
 }
