Horde: Apply request URI redaction in log statements

diff --git a/Engine/Source/Programs/Horde/Drivers/JobDriver/Utility/OpenTelemetry.cs b/Engine/Source/Programs/Horde/Drivers/JobDriver/Utility/OpenTelemetry.cs
--- a/Engine/Source/Programs/Horde/Drivers/JobDriver/Utility/OpenTelemetry.cs
+++ b/Engine/Source/Programs/Horde/Drivers/JobDriver/Utility/OpenTelemetry.cs
@@ -2,6 +2,7 @@
 
 using System.Diagnostics;
 using System.Diagnostics.Metrics;
+using EpicGames.Core;
 using EpicGames.Horde.Utilities;
 using Microsoft.Extensions.DependencyInjection;
 using Microsoft.Extensions.Logging;
@@ -75,7 +76,7 @@ public static class OpenTelemetryHelper
 		{
 			activity.SetTag("service.name", settings.ServiceName + "-http-client");
 			activity.SetTag("operation.name", "http.request");
-			string url = $"{message.Method} {message.Headers.Host}{message.RequestUri?.LocalPath}";
+			string url = $"{message.Method} {message.Headers.Host}{message.RedactedRequestUri()?.LocalPath}";
 			activity.DisplayName = url;
 			activity.SetTag("resource.name", url);
 		}
diff --git a/Engine/Source/Programs/Shared/EpicGames.Horde/HordeHttpMessageHandler.cs b/Engine/Source/Programs/Shared/EpicGames.Horde/HordeHttpMessageHandler.cs
--- a/Engine/Source/Programs/Shared/EpicGames.Horde/HordeHttpMessageHandler.cs
+++ b/Engine/Source/Programs/Shared/EpicGames.Horde/HordeHttpMessageHandler.cs
@@ -46,7 +46,7 @@ namespace EpicGames.Horde
 			// Wait 30 seconds for operations to timeout
 			Task OnTimeoutAsync(Context context, TimeSpan timespan, Task timeoutTask)
 			{
-				logger.LogWarning(KnownLogEvents.Systemic_Horde_Http, "{Method} {Url} timed out after {Time}s.", request.Method, request.RequestUri, (int)timespan.TotalSeconds);
+				logger.LogWarning(KnownLogEvents.Systemic_Horde_Http, "{Method} {Url} timed out after {Time}s.", request.Method, request.RedactedRequestUri(), (int)timespan.TotalSeconds);
 				return Task.CompletedTask;
 			}
 
@@ -55,7 +55,7 @@ namespace EpicGames.Horde
 			// Retry twice after a timeout
 			void OnRetry(Exception ex, TimeSpan timespan)
 			{
-				logger.LogWarning(KnownLogEvents.Systemic_Horde_Http, ex, "{Method} {Url} retrying after {Time}s.", request.Method, request.RequestUri, timespan.TotalSeconds);
+				logger.LogWarning(KnownLogEvents.Systemic_Horde_Http, ex, "{Method} {Url} retrying after {Time}s.", request.Method, request.RedactedRequestUri(), timespan.TotalSeconds);
 			}
 
 			TimeSpan[] retryTimes = new[] { TimeSpan.FromSeconds(5.0), TimeSpan.FromSeconds(10.0) };
@@ -70,7 +70,7 @@ namespace EpicGames.Horde
 		{
 			Task OnTimeoutAsync(DelegateResult<HttpResponseMessage> outcome, TimeSpan timespan, int retryAttempt, Context context)
 			{
-				logger.LogInformation(KnownLogEvents.Systemic_Horde_Http, "{Method} {Url} failed ({Result}). Delaying for {DelayMs}ms (attempt #{RetryNum}).", request.Method, request.RequestUri, outcome.Result?.StatusCode, timespan.TotalMilliseconds, retryAttempt);
+				logger.LogInformation(KnownLogEvents.Systemic_Horde_Http, "{Method} {Url} failed ({Result}). Delaying for {DelayMs}ms (attempt #{RetryNum}).", request.Method, request.RedactedRequestUri(), outcome.Result?.StatusCode, timespan.TotalMilliseconds, retryAttempt);
 				return Task.CompletedTask;
 			}
 
