Migrate code generation for Verse definitions in UHT to the UHT generated files.
Added the ability to disable the generation of the verse namespace alias.
VNI will use the C++ name of any object and not the verse alias.

#rnx

diff --git a/Engine/Source/Programs/Shared/EpicGames.UHT/Exporters/CodeGen/UhtHeaderCodeGeneratorHFile.cs b/Engine/Source/Programs/Shared/EpicGames.UHT/Exporters/CodeGen/UhtHeaderCodeGeneratorHFile.cs
--- a/Engine/Source/Programs/Shared/EpicGames.UHT/Exporters/CodeGen/UhtHeaderCodeGeneratorHFile.cs
+++ b/Engine/Source/Programs/Shared/EpicGames.UHT/Exporters/CodeGen/UhtHeaderCodeGeneratorHFile.cs
@@ -590,10 +590,11 @@ namespace EpicGames.UHT.Exporters.CodeGen
 				// Verse declarations
 				if (scriptStruct.IsVerseField)
 				{
-					// TODO: When VNI gets deprecated, V_DECLARE_STRUCT_RUNTIME_TRAITS will be needed for both VMs and not just BPVM.
-					builder.Append("#if WITH_VERSE_BPVM\r\n");
-					builder.Append("V_DECLARE_STRUCT_RUNTIME_TRAITS(, ").Append(scriptStruct.SourceName).Append(", \"").AppendVerseUEPackageName(scriptStruct).Append("\", \"").AppendVerseName(scriptStruct).Append("\")\r\n");
-					builder.Append("#endif // WITH_VERSE_BPVM\r\n");
+					if (!scriptStruct.FieldExportFlags.HasAnyFlags(UhtFieldExportFlags.NoVerseAlias))
+					{
+						builder.Append("namespace verse { using ").Append(scriptStruct.EngineName).Append(" = ::").Append(scriptStruct.SourceName).Append("; }\r\n");
+					}
+					builder.Append("V_DECLARE_IMPORTED_STRUCT(").Append(scriptStruct.SourceName).Append(", \"").AppendVerseUEPackageName(scriptStruct).Append("\", \"").AppendVerseName(scriptStruct).Append("\")\r\n");
 					builder.Append("\r\n");
 				}
 			}
@@ -651,6 +652,12 @@ namespace EpicGames.UHT.Exporters.CodeGen
 				builder.Append("\r\n");
 			}
 
+			// Verse declarations
+			if (enumObj.IsVerseField)
+			{
+				builder.Append("V_DECLARE_IMPORTED_ENUM(").Append(enumObj.SourceName).Append(", \"").AppendVerseUEPackageName(enumObj).Append("\", \"").AppendVerseName(enumObj).Append("\")\r\n");
+				builder.Append("\r\n");
+			}
 			return builder;
 		}
 
@@ -818,6 +825,16 @@ namespace EpicGames.UHT.Exporters.CodeGen
 			// Forward declare the StaticClass specialization in the header
 			builder.Append("template<> ").Append(Module.Api).Append("UClass* StaticClass<class ").Append(classObj.SourceName).Append(">();\r\n");
 			builder.Append("\r\n");
+
+			// Verse declarations
+			if (classObj.IsVerseField)
+			{
+				if (!classObj.FieldExportFlags.HasAnyFlags(UhtFieldExportFlags.NoVerseAlias))
+				{
+					builder.Append("namespace verse { using ").Append(classObj.EngineName).Append(" = ::").Append(classObj.SourceName).Append("; }\r\n");
+				}
+				builder.Append("\r\n");
+			}
 			return builder;
 		}
 
diff --git a/Engine/Source/Programs/Shared/EpicGames.UHT/Specifiers/UhtFieldSpecifiers.cs b/Engine/Source/Programs/Shared/EpicGames.UHT/Specifiers/UhtFieldSpecifiers.cs
--- a/Engine/Source/Programs/Shared/EpicGames.UHT/Specifiers/UhtFieldSpecifiers.cs
+++ b/Engine/Source/Programs/Shared/EpicGames.UHT/Specifiers/UhtFieldSpecifiers.cs
@@ -37,6 +37,10 @@ namespace EpicGames.UHT.Parsers
 				{
 					fieldObj.VerseModule = kvp.Value.ToString();
 				}
+				else if (key.Equals("noalias", StringComparison.OrdinalIgnoreCase))
+				{
+					fieldObj.FieldExportFlags |= UhtFieldExportFlags.NoVerseAlias;
+				}
 				else
 				{
 					fieldObj.LogError($"Verse specifier option '{key}' is unknown");
diff --git a/Engine/Source/Programs/Shared/EpicGames.UHT/Types/UhtField.cs b/Engine/Source/Programs/Shared/EpicGames.UHT/Types/UhtField.cs
--- a/Engine/Source/Programs/Shared/EpicGames.UHT/Types/UhtField.cs
+++ b/Engine/Source/Programs/Shared/EpicGames.UHT/Types/UhtField.cs
@@ -1,11 +1,73 @@
 // Copyright Epic Games, Inc. All Rights Reserved.
 
+using System;
 using System.Text;
+using System.Text.Json.Serialization;
 using EpicGames.UHT.Utils;
 
 namespace EpicGames.UHT.Types
 {
 
+	/// <summary>
+	/// Series of flags not part of the engine's field flags that affect code generation or verification
+	/// </summary>
+	[Flags]
+	public enum UhtFieldExportFlags : int
+	{
+
+		/// <summary>
+		/// No export flags
+		/// </summary>
+		None = 0,
+
+		/// <summary>
+		/// Do not generate an alias in the verse namespace
+		/// </summary>
+		NoVerseAlias = 1 << 0,
+	}
+
+	/// <summary>
+	/// Helper methods for testing flags.  These methods perform better than the generic HasFlag which hits
+	/// the GC and stalls.
+	/// </summary>
+	public static class UhtFieldExportFlagsExtensions
+	{
+
+		/// <summary>
+		/// Test to see if any of the specified flags are set
+		/// </summary>
+		/// <param name="inFlags">Current flags</param>
+		/// <param name="testFlags">Flags to test for</param>
+		/// <returns>True if any of the flags are set</returns>
+		public static bool HasAnyFlags(this UhtFieldExportFlags inFlags, UhtFieldExportFlags testFlags)
+		{
+			return (inFlags & testFlags) != 0;
+		}
+
+		/// <summary>
+		/// Test to see if all of the specified flags are set
+		/// </summary>
+		/// <param name="inFlags">Current flags</param>
+		/// <param name="testFlags">Flags to test for</param>
+		/// <returns>True if all the flags are set</returns>
+		public static bool HasAllFlags(this UhtFieldExportFlags inFlags, UhtFieldExportFlags testFlags)
+		{
+			return (inFlags & testFlags) == testFlags;
+		}
+
+		/// <summary>
+		/// Test to see if a specific set of flags have a specific value.
+		/// </summary>
+		/// <param name="inFlags">Current flags</param>
+		/// <param name="testFlags">Flags to test for</param>
+		/// <param name="matchFlags">Expected value of the tested flags</param>
+		/// <returns>True if the given flags have a specific value.</returns>
+		public static bool HasExactFlags(this UhtFieldExportFlags inFlags, UhtFieldExportFlags testFlags, UhtFieldExportFlags matchFlags)
+		{
+			return (inFlags & testFlags) == matchFlags;
+		}
+	}
+
 	/// <summary>
 	/// Represents a UField
 	/// </summary>
@@ -14,6 +76,12 @@ namespace EpicGames.UHT.Types
 		/// <inheritdoc/>
 		public override string EngineClassName => "Field";
 
+		/// <summary>
+		/// UHT only field flags
+		/// </summary>
+		[JsonConverter(typeof(JsonStringEnumConverter))]
+		public UhtFieldExportFlags FieldExportFlags { get; set; } = UhtFieldExportFlags.None;
+
 		/// <summary>
 		/// Name of the module containing the type
 		/// </summary>
@@ -91,7 +159,13 @@ namespace EpicGames.UHT.Types
 				throw new UhtException(fieldObj, "Attempt to write the Verse name on a field that isn't part of Verse");
 			}
 			UhtModule module = fieldObj.Module;
-			return builder.Append('(').Append(module.Module.VersePath).Append(":)").Append(fieldObj.VerseName);
+			builder.Append('(').Append(module.Module.VersePath);
+			if (!String.IsNullOrEmpty(fieldObj.VerseModule))
+			{
+				builder.Append('/').Append(fieldObj.VerseModule);
+			}
+			builder.Append(":)").Append(fieldObj.VerseName);
+			return builder;
 		}
 	}
 }
