[UBT]
* Added hint parameter to CacheClient.Create

diff --git a/Engine/Source/Programs/Shared/EpicGames.UBA/ICacheClient.cs b/Engine/Source/Programs/Shared/EpicGames.UBA/ICacheClient.cs
--- a/Engine/Source/Programs/Shared/EpicGames.UBA/ICacheClient.cs
+++ b/Engine/Source/Programs/Shared/EpicGames.UBA/ICacheClient.cs
@@ -58,10 +58,11 @@ namespace EpicGames.UBA
 		/// <param name="session">The session</param>
 		/// <param name="reportMissReason">Output reason for cache miss to log.</param>
 		/// <param name="crypto">Enable crypto by using a 32 character crypto string (representing a 16 byte value)</param>
+		/// <param name="hint">string that will show in cache server log</param>
 		/// <returns>The ICacheClient</returns>
-		public static ICacheClient CreateCacheClient(ISessionServer session, bool reportMissReason, string crypto = "")
+		public static ICacheClient CreateCacheClient(ISessionServer session, bool reportMissReason, string crypto = "", string hint = "")
 		{
-			return new CacheClientImpl(session, reportMissReason, crypto);
+			return new CacheClientImpl(session, reportMissReason, crypto, hint);
 		}
 	}
 }
diff --git a/Engine/Source/Programs/Shared/EpicGames.UBA/Impl/CacheClientImpl.cs b/Engine/Source/Programs/Shared/EpicGames.UBA/Impl/CacheClientImpl.cs
--- a/Engine/Source/Programs/Shared/EpicGames.UBA/Impl/CacheClientImpl.cs
+++ b/Engine/Source/Programs/Shared/EpicGames.UBA/Impl/CacheClientImpl.cs
@@ -21,7 +21,7 @@ namespace EpicGames.UBA.Impl
 		static extern void ProcessStartInfo_Destroy(nint server);
 
 		[DllImport("UbaHost", CharSet = CharSet.Auto)]
-		static extern nint CacheClient_Create(nint session, bool reportMissReason, string crypto);
+		static extern nint CacheClient_Create(nint session, bool reportMissReason, string crypto, string hint);
 
 		[DllImport("UbaHost", CharSet = CharSet.Auto)]
 		static extern bool CacheClient_Connect(nint cacheClient, string host, int port);
@@ -71,10 +71,10 @@ namespace EpicGames.UBA.Impl
 		}
 		#endregion
 
-		public CacheClientImpl(ISessionServer server, bool reportMissReason, string crypto)
+		public CacheClientImpl(ISessionServer server, bool reportMissReason, string crypto, string hint)
 		{
 			_sessionServer = server;
-			_handle = CacheClient_Create(_sessionServer.GetHandle(), reportMissReason, crypto);
+			_handle = CacheClient_Create(_sessionServer.GetHandle(), reportMissReason, crypto, hint);
 		}
 
 		public bool Connect(string host, int port)
