Fixed issue where structure traits for UHT generated Verse structures with BPVM were not getting defined.

#rnx

diff --git a/Engine/Source/Programs/Shared/EpicGames.UHT/Exporters/CodeGen/UhtHeaderCodeGeneratorHFile.cs b/Engine/Source/Programs/Shared/EpicGames.UHT/Exporters/CodeGen/UhtHeaderCodeGeneratorHFile.cs
--- a/Engine/Source/Programs/Shared/EpicGames.UHT/Exporters/CodeGen/UhtHeaderCodeGeneratorHFile.cs
+++ b/Engine/Source/Programs/Shared/EpicGames.UHT/Exporters/CodeGen/UhtHeaderCodeGeneratorHFile.cs
@@ -594,7 +594,7 @@ namespace EpicGames.UHT.Exporters.CodeGen
 					{
 						builder.Append("namespace verse { using ").Append(scriptStruct.EngineName).Append(" = ::").Append(scriptStruct.SourceName).Append("; }\r\n");
 					}
-					builder.Append("V_DECLARE_IMPORTED_STRUCT(").Append(scriptStruct.SourceName).Append(", \"").AppendVerseUEPackageName(scriptStruct).Append("\", \"").AppendVerseName(scriptStruct).Append("\")\r\n");
+					builder.Append("V_DECLARE_STRUCT_RUNTIME_TRAITS(, ").Append(scriptStruct.SourceName).Append(", \"").AppendVerseUEPackageName(scriptStruct).Append("\", \"").AppendVerseName(scriptStruct).Append("\")\r\n");
 					builder.Append("\r\n");
 				}
 			}
