Horde: Include lease ID in returned properties for allocated compute tasks

diff --git a/Engine/Source/Programs/Shared/EpicGames.Horde/Agents/KnownPropertyNames.cs b/Engine/Source/Programs/Shared/EpicGames.Horde/Agents/KnownPropertyNames.cs
--- a/Engine/Source/Programs/Shared/EpicGames.Horde/Agents/KnownPropertyNames.cs
+++ b/Engine/Source/Programs/Shared/EpicGames.Horde/Agents/KnownPropertyNames.cs
@@ -91,4 +91,9 @@ public static class KnownPropertyNames
 	/// Whether the agent is trusted
 	/// </summary>
 	public const string Trusted = "Trusted";
+	
+	/// <summary>
+	/// ID of current lease
+	/// </summary>
+	public const string LeaseId = "LeaseId";
 }
\ No newline at end of file
diff --git a/Engine/Source/Programs/Shared/EpicGames.Horde/Compute/Clients/ServerComputeClient.cs b/Engine/Source/Programs/Shared/EpicGames.Horde/Compute/Clients/ServerComputeClient.cs
--- a/Engine/Source/Programs/Shared/EpicGames.Horde/Compute/Clients/ServerComputeClient.cs
+++ b/Engine/Source/Programs/Shared/EpicGames.Horde/Compute/Clients/ServerComputeClient.cs
@@ -11,6 +11,7 @@ using System.Runtime.CompilerServices;
 using System.Threading;
 using System.Threading.Tasks;
 using EpicGames.Core;
+using EpicGames.Horde.Agents;
 using EpicGames.Horde.Compute.Transports;
 using Microsoft.Extensions.Logging;
 
@@ -382,7 +383,8 @@ namespace EpicGames.Horde.Compute.Clients
 			byte[] nonce = StringUtils.ParseHexString(response.Nonce);
 			await socket.SendMessageAsync(nonce, SocketFlags.None, cancellationToken);
 			workerLogger.LogInformation("Connected to {AgentId} ({Ip}) under lease {LeaseId} (agent version: {AgentVersion})", response.AgentId, response.Ip, response.LeaseId, response.AgentVersion ?? "unknown");
-
+			
+			response.Properties = [..response.Properties, $"{KnownPropertyNames.LeaseId}={response.LeaseId}"];
 			await using ComputeTransport transport = await CreateTransportAsync(socket, response, cancellationToken);
 			await using RemoteComputeSocket computeSocket = new(transport, (ComputeProtocol)response.Protocol, workerLogger);
 			yield return new LeaseInfo(response.ClusterId, response.Properties, response.AssignedResources, computeSocket, response.Ip, response.ConnectionMode, response.Ports);
