Horde: Few small changes to handling of chunked data nodes.

- ChunkedDataNodeRef now implements IHashedNodeRef, and extension methods now allow extracting chunked data to streams or files.
- Adding a FileEntry no longer requires an explicit length parameter (this is already contained in the payload data).
- Rename FileEntry.StreamHash to FileEntry.Hash

diff --git a/Engine/Source/Programs/AutomationTool/BuildGraph/Tasks/CreateArtifactTask.cs b/Engine/Source/Programs/AutomationTool/BuildGraph/Tasks/CreateArtifactTask.cs
--- a/Engine/Source/Programs/AutomationTool/BuildGraph/Tasks/CreateArtifactTask.cs
+++ b/Engine/Source/Programs/AutomationTool/BuildGraph/Tasks/CreateArtifactTask.cs
@@ -284,7 +284,7 @@ namespace AutomationTool.Tasks
 					string fileName = file.GetFileName().ToUpperInvariant();
 					string alias = $"sym:{fileName}/{hash}/{fileName}";
 					Logger.LogInformation("Adding symbol alias: {Alias}", alias);
-					aliases.Add((fileEntry.Target.Handle, alias));
+					aliases.Add((fileEntry.Target, alias));
 				}
 			}
 		}
diff --git a/Engine/Source/Programs/Horde/Drivers/JobDriver/Utility/TempStorage.cs b/Engine/Source/Programs/Horde/Drivers/JobDriver/Utility/TempStorage.cs
--- a/Engine/Source/Programs/Horde/Drivers/JobDriver/Utility/TempStorage.cs
+++ b/Engine/Source/Programs/Horde/Drivers/JobDriver/Utility/TempStorage.cs
@@ -667,7 +667,7 @@ namespace JobDriver.Utility
 			using ChunkedDataWriter fileNodeWriter = new ChunkedDataWriter(writer, new ChunkingOptions());
 			ChunkedData fileNodeData = await fileNodeWriter.CreateAsync(localFileListLocation.ToFileInfo(), cancellationToken);
 
-			return new FileEntry(localFileListLocation.GetFileName(), FileEntryFlags.None, fileNodeWriter.Length, fileNodeData);
+			return new FileEntry(localFileListLocation.GetFileName(), FileEntryFlags.None, fileNodeData);
 		}
 
 		/// <summary>
diff --git a/Engine/Source/Programs/Horde/HordeCmd/Commands/Vcs/VcsCheckout.cs b/Engine/Source/Programs/Horde/HordeCmd/Commands/Vcs/VcsCheckout.cs
--- a/Engine/Source/Programs/Horde/HordeCmd/Commands/Vcs/VcsCheckout.cs
+++ b/Engine/Source/Programs/Horde/HordeCmd/Commands/Vcs/VcsCheckout.cs
@@ -100,7 +100,7 @@ namespace Horde.Commands.Vcs
 						FileReference.Delete(filePath);
 					}
 				}
-				else if (fileState == null || fileState.Hash != fileEntry.StreamHash)
+				else if (fileState == null || fileState.Hash != fileEntry.Hash)
 				{
 					newState.Files[name] = await CheckoutFileAsync(fileEntry, filePath.ToFileInfo(), logger);
 				}
@@ -116,10 +116,10 @@ namespace Horde.Commands.Vcs
 
 		static async Task<FileState> CheckoutFileAsync(FileEntry fileRef, FileInfo fileInfo, ILogger logger)
 		{
-			logger.LogInformation("Updating {File} to {Hash}", fileInfo, fileRef.StreamHash);
-			await ChunkedDataNode.CopyToFileAsync(fileRef.Target.Handle, fileInfo, CancellationToken.None);
+			logger.LogInformation("Updating {File} to {Hash}", fileInfo, fileRef.Hash);
+			await fileRef.CopyToFileAsync(fileInfo, CancellationToken.None);
 			fileInfo.Refresh();
-			return new FileState(fileInfo, fileRef.StreamHash);
+			return new FileState(fileInfo, fileRef.Hash);
 		}
 	}
 }
diff --git a/Engine/Source/Programs/Horde/Plugins/Build/HordeServer.Build.Tests/Artifacts/UnsyncCacheTests.cs b/Engine/Source/Programs/Horde/Plugins/Build/HordeServer.Build.Tests/Artifacts/UnsyncCacheTests.cs
--- a/Engine/Source/Programs/Horde/Plugins/Build/HordeServer.Build.Tests/Artifacts/UnsyncCacheTests.cs
+++ b/Engine/Source/Programs/Horde/Plugins/Build/HordeServer.Build.Tests/Artifacts/UnsyncCacheTests.cs
@@ -34,7 +34,7 @@ namespace HordeServer.Tests.Artifacts
 				ChunkedData data = await chunkedWriter.FlushAsync();
 
 				DirectoryNode directory = new DirectoryNode();
-				directory.AddFile("hello.txt", FileEntryFlags.None, source.Length, data);
+				directory.AddFile("hello.txt", FileEntryFlags.None, data);
 
 				directoryRef = await writer.WriteBlobAsync(directory);
 			}
diff --git a/Engine/Source/Programs/Horde/Plugins/Build/HordeServer.Build/Artifacts/ArtifactsController.cs b/Engine/Source/Programs/Horde/Plugins/Build/HordeServer.Build/Artifacts/ArtifactsController.cs
--- a/Engine/Source/Programs/Horde/Plugins/Build/HordeServer.Build/Artifacts/ArtifactsController.cs
+++ b/Engine/Source/Programs/Horde/Plugins/Build/HordeServer.Build/Artifacts/ArtifactsController.cs
@@ -390,7 +390,7 @@ namespace HordeServer.Artifacts
 			foreach (FileEntry fileEntry in directoryNode.Files)
 			{
 				response.Files ??= new List<GetArtifactFileEntryResponse>();
-				response.Files.Add(new GetArtifactFileEntryResponse(fileEntry.Name.ToString(), fileEntry.Length, fileEntry.StreamHash));
+				response.Files.Add(new GetArtifactFileEntryResponse(fileEntry.Name.ToString(), fileEntry.Length, fileEntry.Hash));
 			}
 		}
 
@@ -422,7 +422,7 @@ namespace HordeServer.Artifacts
 				if (filePath.Contains(search, StringComparison.OrdinalIgnoreCase))
 				{
 					response.Files ??= new List<GetArtifactFileEntryResponse>();
-					response.Files.Add(new GetArtifactFileEntryResponse(fileEntry.Name.ToString(), fileEntry.Length, fileEntry.StreamHash));
+					response.Files.Add(new GetArtifactFileEntryResponse(fileEntry.Name.ToString(), fileEntry.Length, fileEntry.Hash));
 				}
 			}
 		}
diff --git a/Engine/Source/Programs/Horde/Plugins/Build/HordeServer.Build/Artifacts/UnsyncCache.cs b/Engine/Source/Programs/Horde/Plugins/Build/HordeServer.Build/Artifacts/UnsyncCache.cs
--- a/Engine/Source/Programs/Horde/Plugins/Build/HordeServer.Build/Artifacts/UnsyncCache.cs
+++ b/Engine/Source/Programs/Horde/Plugins/Build/HordeServer.Build/Artifacts/UnsyncCache.cs
@@ -271,7 +271,7 @@ namespace HordeServer.Artifacts
 			}
 			else
 			{
-				InteriorChunkedDataNode interiorNode = await nodeRef.Handle.ReadBlobAsync<InteriorChunkedDataNode>(cancellationToken: cancellationToken);
+				InteriorChunkedDataNode interiorNode = await nodeRef.ReadBlobAsync<InteriorChunkedDataNode>(cancellationToken: cancellationToken);
 				foreach (ChunkedDataNodeRef childRef in interiorNode.Children)
 				{
 					await FindBlocksAsync(offset, childRef, blocks, cancellationToken);
diff --git a/Engine/Source/Programs/Horde/Plugins/Storage/HordeServer.Storage/Storage/StorageController.cs b/Engine/Source/Programs/Horde/Plugins/Storage/HordeServer.Storage/Storage/StorageController.cs
--- a/Engine/Source/Programs/Horde/Plugins/Storage/HordeServer.Storage/Storage/StorageController.cs
+++ b/Engine/Source/Programs/Horde/Plugins/Storage/HordeServer.Storage/Storage/StorageController.cs
@@ -562,7 +562,7 @@ namespace HordeServer.Storage
 				List<object> files = new List<object>();
 				foreach ((string name, FileEntry entry) in directoryNode.NameToFile)
 				{
-					files.Add(new { name = name.ToString(), length = entry.Length, flags = entry.Flags, hash = entry.StreamHash, target = GetNodeHandleLink(namespaceId, entry.Target.Handle) });
+					files.Add(new { name = name.ToString(), length = entry.Length, flags = entry.Flags, hash = entry.Hash, target = GetNodeHandleLink(namespaceId, entry.Target) });
 				}
 
 				content = new { directoryNode.Length, directories, files };
@@ -574,7 +574,7 @@ namespace HordeServer.Storage
 				List<object> children = new List<object>();
 				foreach (ChunkedDataNodeRef nodeRef in interiorNode.Children)
 				{
-					children.Add(new { nodeRef.Type, nodeRef.Length, hash = nodeRef.Handle.Hash, link = GetNodeLink(namespaceId, nodeRef.Handle) });
+					children.Add(new { nodeRef.Type, nodeRef.Length, hash = nodeRef.Hash, link = GetNodeLink(namespaceId, nodeRef) });
 				}
 
 				content = new { children };
diff --git a/Engine/Source/Programs/Horde/Plugins/Symbols/HordeServer.Symbols/Symbols/SymbolsController.cs b/Engine/Source/Programs/Horde/Plugins/Symbols/HordeServer.Symbols/Symbols/SymbolsController.cs
--- a/Engine/Source/Programs/Horde/Plugins/Symbols/HordeServer.Symbols/Symbols/SymbolsController.cs
+++ b/Engine/Source/Programs/Horde/Plugins/Symbols/HordeServer.Symbols/Symbols/SymbolsController.cs
@@ -63,7 +63,7 @@ namespace HordeServer.Tools
 			Response.StatusCode = (int)HttpStatusCode.OK;
 
 			await Response.StartAsync(cancellationToken);
-			await ChunkedDataNode.CopyToStreamAsync(blobRef, Response.Body, cancellationToken);
+			await blobRef.CopyToStreamAsync(Response.Body, cancellationToken);
 			await Response.CompleteAsync();
 
 			return NoContent();
diff --git a/Engine/Source/Programs/Shared/EpicGames.Horde.Tests/Compute/ComputeSocketTests.cs b/Engine/Source/Programs/Shared/EpicGames.Horde.Tests/Compute/ComputeSocketTests.cs
--- a/Engine/Source/Programs/Shared/EpicGames.Horde.Tests/Compute/ComputeSocketTests.cs
+++ b/Engine/Source/Programs/Shared/EpicGames.Horde.Tests/Compute/ComputeSocketTests.cs
@@ -209,7 +209,7 @@ namespace EpicGames.Horde.Tests.Compute
 						ChunkedData chunkedData = await writer.CreateAsync(data, cancellationToken);
 
 						DirectoryNode directory = new DirectoryNode();
-						directory.AddFile("hello.txt", FileEntryFlags.None, data.Length, chunkedData);
+						directory.AddFile("hello.txt", FileEntryFlags.None, chunkedData);
 
 						IHashedBlobRef<DirectoryNode> directoryRef = await blobWriter.WriteBlobAsync(directory, cancellationToken: cancellationToken);
 
diff --git a/Engine/Source/Programs/Shared/EpicGames.Horde.Tests/Storage/BundleStoreTests.cs b/Engine/Source/Programs/Shared/EpicGames.Horde.Tests/Storage/BundleStoreTests.cs
--- a/Engine/Source/Programs/Shared/EpicGames.Horde.Tests/Storage/BundleStoreTests.cs
+++ b/Engine/Source/Programs/Shared/EpicGames.Horde.Tests/Storage/BundleStoreTests.cs
@@ -400,7 +400,7 @@ namespace EpicGames.Horde.Tests.Storage
 				ChunkedData chunkedData = await fileWriter.CreateAsync(data, CancellationToken.None);
 
 				root = new DirectoryNode(DirectoryFlags.None);
-				root.AddFile("test", FileEntryFlags.None, fileWriter.Length, chunkedData);
+				root.AddFile("test", FileEntryFlags.None, chunkedData);
 
 				IHashedBlobRef<DirectoryNode> rootRef = await writer.WriteBlobAsync(root);
 				await store.AddRefAsync(new RefName("test"), rootRef);
@@ -417,7 +417,7 @@ namespace EpicGames.Horde.Tests.Storage
 				FileEntry file = root.GetFileEntry("test");
 
 				Dictionary<BlobLocator, long> locatorToSize = new Dictionary<BlobLocator, long>();
-				await GetUniqueBlobsAsync(file.Target.Handle, locatorToSize);
+				await GetUniqueBlobsAsync(file.Target, locatorToSize);
 
 				long uniqueSize = locatorToSize.Sum(x => x.Value);
 				Assert.IsTrue(uniqueSize < data.Length / 3); // random fraction meaning "lots of dedupe happened"
diff --git a/Engine/Source/Programs/Shared/EpicGames.Horde.Tests/Storage/BundleTests.cs b/Engine/Source/Programs/Shared/EpicGames.Horde.Tests/Storage/BundleTests.cs
--- a/Engine/Source/Programs/Shared/EpicGames.Horde.Tests/Storage/BundleTests.cs
+++ b/Engine/Source/Programs/Shared/EpicGames.Horde.Tests/Storage/BundleTests.cs
@@ -112,9 +112,7 @@ namespace EpicGames.Horde.Tests.Storage
 
 		private static async Task TestBufferlessReadsAsync(ChunkedDataNodeRef nodeRef, ReadOnlyMemory<byte> expected)
 		{
-			using MemoryStream memoryStream = new MemoryStream();
-			await ChunkedDataNode.CopyToStreamAsync(nodeRef.Handle, memoryStream, default);
-			ReadOnlyMemory<byte> read = memoryStream.ToArray().AsMemory();
+			ReadOnlyMemory<byte> read = await nodeRef.ReadAllBytesAsync();
 			Assert.IsTrue(read.Span.SequenceEqual(expected.Span));
 		}
 
diff --git a/Engine/Source/Programs/Shared/EpicGames.Horde.Tests/Storage/ChunkingTests.cs b/Engine/Source/Programs/Shared/EpicGames.Horde.Tests/Storage/ChunkingTests.cs
--- a/Engine/Source/Programs/Shared/EpicGames.Horde.Tests/Storage/ChunkingTests.cs
+++ b/Engine/Source/Programs/Shared/EpicGames.Horde.Tests/Storage/ChunkingTests.cs
@@ -66,7 +66,7 @@ namespace EpicGames.Horde.Tests.Storage
 				ChunkedData chunkedData = await InteriorChunkedDataNode.CreateTreeAsync(leafChunkedData, new InteriorChunkedDataNodeOptions(4, 4, 4), writer, CancellationToken.None);
 
 				DirectoryNode directory = new DirectoryNode();
-				directory.AddFile("test.foo", FileEntryFlags.None, 0, chunkedData);
+				directory.AddFile("test.foo", FileEntryFlags.None, chunkedData);
 
 				IHashedBlobRef handle = await writer.WriteBlobAsync(directory);
 				await store.AddRefAsync(RefName, handle);
@@ -188,7 +188,7 @@ namespace EpicGames.Horde.Tests.Storage
 				ChunkedData chunkedData = await chunkedWriter.FlushAsync(CancellationToken.None);
 
 				DirectoryNode directoryNode = new DirectoryNode();
-				directoryNode.AddFile("test", FileEntryFlags.None, data.Length, chunkedData);
+				directoryNode.AddFile("test", FileEntryFlags.None, chunkedData);
 				handle = await writer.WriteBlobAsync(directoryNode);
 			}
 
@@ -224,7 +224,7 @@ namespace EpicGames.Horde.Tests.Storage
 			ChunkedDataNodeRef root = await InteriorChunkedDataNode.CreateTreeAsync(leafNodeRefs, interiorOptions, blobWriter, CancellationToken.None);
 
 			List<IBlobRef> list = new List<IBlobRef>();
-			await GetReadOrderAsync(root.Handle, list);
+			await GetReadOrderAsync(root, list);
 
 			int prevIndex = Int32.MaxValue;
 			for (int idx = 0; idx < list.Count; idx++)
diff --git a/Engine/Source/Programs/Shared/EpicGames.Horde/Compute/AgentMessageHandler.cs b/Engine/Source/Programs/Shared/EpicGames.Horde/Compute/AgentMessageHandler.cs
--- a/Engine/Source/Programs/Shared/EpicGames.Horde/Compute/AgentMessageHandler.cs
+++ b/Engine/Source/Programs/Shared/EpicGames.Horde/Compute/AgentMessageHandler.cs
@@ -224,13 +224,13 @@ namespace EpicGames.Horde.Compute
 					await using FileStream stream = FileReference.Open(file, FileMode.Open, FileAccess.Read);
 					IoHash hash = await IoHash.ComputeAsync(stream, cancellationToken);
 
-					if (hash == fileEntry.StreamHash)
+					if (hash == fileEntry.Hash)
 					{
 						_logger.LogInformation("Hash of {File} is correct ({Hash})", file, hash);
 					}
-					if (hash != fileEntry.StreamHash)
+					if (hash != fileEntry.Hash)
 					{
-						_logger.LogError("Hash mismatch for {File}; expected {ExpectedHash}, got {ActualHash}", file, fileEntry.StreamHash, hash);
+						_logger.LogError("Hash mismatch for {File}; expected {ExpectedHash}, got {ActualHash}", file, fileEntry.Hash, hash);
 						result = false;
 					}
 				}
diff --git a/Engine/Source/Programs/Shared/EpicGames.Horde/Storage/BatchOutputWriter.cs b/Engine/Source/Programs/Shared/EpicGames.Horde/Storage/BatchOutputWriter.cs
--- a/Engine/Source/Programs/Shared/EpicGames.Horde/Storage/BatchOutputWriter.cs
+++ b/Engine/Source/Programs/Shared/EpicGames.Horde/Storage/BatchOutputWriter.cs
@@ -318,9 +318,9 @@ namespace EpicGames.Horde.Storage
 					{
 						computedHash = await IoHash.ComputeAsync(stream, cancellationToken);
 					}
-					if (computedHash != outputFile.FileEntry.StreamHash)
+					if (computedHash != outputFile.FileEntry.Hash)
 					{
-						throw new InvalidDataException($"File {outputFile.Path} does not have correct hash after extract; expected {outputFile.FileEntry.StreamHash}, got {computedHash}.");
+						throw new InvalidDataException($"File {outputFile.Path} does not have correct hash after extract; expected {outputFile.FileEntry.Hash}, got {computedHash}.");
 					}
 					_logger.LogDebug("Verified hash of {Path}: {Hash}", outputFile.Path, computedHash);
 				}
diff --git a/Engine/Source/Programs/Shared/EpicGames.Horde/Storage/DedupeBlobWriter.cs b/Engine/Source/Programs/Shared/EpicGames.Horde/Storage/DedupeBlobWriter.cs
--- a/Engine/Source/Programs/Shared/EpicGames.Horde/Storage/DedupeBlobWriter.cs
+++ b/Engine/Source/Programs/Shared/EpicGames.Horde/Storage/DedupeBlobWriter.cs
@@ -247,12 +247,12 @@ namespace EpicGames.Horde.Storage
 			if (dataNodeRef.Type == ChunkedDataNodeType.Leaf)
 			{
 				BlobType blobType = LeafChunkedDataNodeConverter.BlobType;
-				dedupeWriter.AddToCache(blobType, dataNodeRef.Handle);
+				dedupeWriter.AddToCache(blobType, dataNodeRef);
 			}
 			else if (dataNodeRef.Type == ChunkedDataNodeType.Interior)
 			{
-				using BlobData blobData = await dataNodeRef.Handle.ReadBlobDataAsync(cancellationToken);
-				dedupeWriter.AddToCache(blobData.Type, dataNodeRef.Handle);
+				using BlobData blobData = await dataNodeRef.ReadBlobDataAsync(cancellationToken);
+				dedupeWriter.AddToCache(blobData.Type, dataNodeRef);
 
 				InteriorChunkedDataNode interiorNode = BlobSerializer.Deserialize<InteriorChunkedDataNode>(blobData);
 				foreach (ChunkedDataNodeRef childNodeRef in interiorNode.Children)
diff --git a/Engine/Source/Programs/Shared/EpicGames.Horde/Storage/Nodes/ChunkedDataNode.cs b/Engine/Source/Programs/Shared/EpicGames.Horde/Storage/Nodes/ChunkedDataNode.cs
--- a/Engine/Source/Programs/Shared/EpicGames.Horde/Storage/Nodes/ChunkedDataNode.cs
+++ b/Engine/Source/Programs/Shared/EpicGames.Horde/Storage/Nodes/ChunkedDataNode.cs
@@ -4,7 +4,9 @@ using System;
 using System.Buffers;
 using System.Buffers.Binary;
 using System.Collections.Generic;
+using System.Diagnostics.CodeAnalysis;
 using System.IO;
+using System.IO.Pipelines;
 using System.Linq;
 using System.Runtime.CompilerServices;
 using System.Threading;
@@ -28,6 +30,20 @@ namespace EpicGames.Horde.Storage.Nodes
 		/// <param name="options">Options for controling serialization</param>
 		/// <param name="cancellationToken">Cancellation token for the operation</param>
 		public abstract Task CopyToStreamAsync(Stream outputStream, BlobSerializerOptions? options = null, CancellationToken cancellationToken = default);
+	}
+
+	/// <summary>
+	/// Extension methods for chunked data node types
+	/// </summary>
+	public static class ChunkedDataNodeExtensions
+	{
+		/// <summary>
+		/// Copy the contents of the node to the output stream without creating the intermediate FileNodes
+		/// </summary>
+		/// <param name="handle">Handle to the data to read</param>
+		/// <param name="length">Length of the stream</param>
+		public static Stream OpenAsStream(this IBlobRef<ChunkedDataNode> handle, long? length = null)
+			=> new ChunkedDataStream(handle, length);
 
 		/// <summary>
 		/// Copy the contents of the node to the output stream without creating the intermediate FileNodes
@@ -35,7 +51,7 @@ namespace EpicGames.Horde.Storage.Nodes
 		/// <param name="handle">Handle to the data to read</param>
 		/// <param name="outputStream">The output stream to receive the data</param>
 		/// <param name="cancellationToken">Cancellation token for the operation</param>
-		public static async Task CopyToStreamAsync(IBlobRef handle, Stream outputStream, CancellationToken cancellationToken)
+		public static async Task CopyToStreamAsync(this IBlobRef<ChunkedDataNode> handle, Stream outputStream, CancellationToken cancellationToken = default)
 		{
 			using BlobData blobData = await handle.ReadBlobDataAsync(cancellationToken);
 			if (blobData.Type.Guid == LeafChunkedDataNode.BlobTypeGuid)
@@ -59,7 +75,7 @@ namespace EpicGames.Horde.Storage.Nodes
 		/// <param name="file">File to write with the contents of this node</param>
 		/// <param name="cancellationToken">Cancellation token for the operation</param>
 		/// <returns></returns>
-		public static async Task CopyToFileAsync(IBlobRef handle, FileInfo file, CancellationToken cancellationToken)
+		public static async Task CopyToFileAsync(this IBlobRef<ChunkedDataNode> handle, FileInfo file, CancellationToken cancellationToken)
 		{
 			if (file.Exists && (file.Attributes & FileAttributes.ReadOnly) != 0)
 			{
@@ -72,18 +88,14 @@ namespace EpicGames.Horde.Storage.Nodes
 		}
 
 		/// <summary>
-		/// Serialize this node and its children into a byte array
+		/// Copy the contents of the node to the output stream without creating the intermediate FileNodes
 		/// </summary>
-		/// <param name="options">Options to control serialization</param>
-		/// <param name="cancellationToken"></param>
-		/// <returns>Array of data stored by the tree</returns>
-		public async Task<byte[]> ToByteArrayAsync(BlobSerializerOptions? options = null, CancellationToken cancellationToken = default)
+		/// <param name="handle">Handle to the data to read</param>
+		/// <param name="cancellationToken">Cancellation token for the operation</param>
+		public static async Task<byte[]> ReadAllBytesAsync(this IBlobRef<ChunkedDataNode> handle, CancellationToken cancellationToken = default)
 		{
-			using (MemoryStream stream = new MemoryStream())
-			{
-				await CopyToStreamAsync(stream, options, cancellationToken);
-				return stream.ToArray();
-			}
+			using Stream stream = OpenAsStream(handle);
+			return await stream.ReadAllBytesAsync(cancellationToken);
 		}
 	}
 
@@ -151,9 +163,18 @@ namespace EpicGames.Horde.Storage.Nodes
 	/// <param name="Type">Type of the referenced node</param>
 	/// <param name="Length">Length of the data stream within this node</param>
 	/// <param name="RollingHash">Rolling hash for this chunk. Only serialized for leaf node references.</param>
-	/// <param name="Handle">Handle to the target node</param>
-	public record class ChunkedDataNodeRef(ChunkedDataNodeType Type, long Length, uint RollingHash, IHashedBlobRef<ChunkedDataNode> Handle)
+	/// <param name="Inner">Handle to the target node</param>
+	public record class ChunkedDataNodeRef(ChunkedDataNodeType Type, long Length, uint RollingHash, IHashedBlobRef<ChunkedDataNode> Inner) : IHashedBlobRef<ChunkedDataNode>
 	{
+		/// <inheritdoc/>
+		public IoHash Hash => Inner.Hash;
+
+		/// <inheritdoc/>
+		public BlobSerializerOptions? SerializerOptions => Inner.SerializerOptions;
+
+		/// <inheritdoc/>
+		public IBlobRef Innermost => Inner.Innermost;
+
 		/// <summary>
 		/// Leaf node constructor
 		/// </summary>
@@ -179,7 +200,7 @@ namespace EpicGames.Horde.Storage.Nodes
 			{
 				throw new InvalidOperationException("Node is not an interior node");
 			}
-			return HashedBlobRef.Create<InteriorChunkedDataNode>(Handle.Hash, Handle, Handle.SerializerOptions);
+			return HashedBlobRef.Create<InteriorChunkedDataNode>(Inner.Hash, Inner, Inner.SerializerOptions);
 		}
 
 		/// <summary>
@@ -191,7 +212,7 @@ namespace EpicGames.Horde.Storage.Nodes
 			{
 				throw new InvalidOperationException("Node is not a leaf node");
 			}
-			return HashedBlobRef.Create<LeafChunkedDataNode>(Handle.Hash, Handle, Handle.SerializerOptions);
+			return HashedBlobRef.Create<LeafChunkedDataNode>(Inner.Hash, Inner, Inner.SerializerOptions);
 		}
 
 		/// <summary>
@@ -201,11 +222,11 @@ namespace EpicGames.Horde.Storage.Nodes
 		{
 			if (Type == ChunkedDataNodeType.Leaf)
 			{
-				return await Handle.ReadBlobAsync<LeafChunkedDataNode>(options, cancellationToken);
+				return await Inner.ReadBlobAsync<LeafChunkedDataNode>(options, cancellationToken);
 			}
 			else if (Type == ChunkedDataNodeType.Interior)
 			{
-				return await Handle.ReadBlobAsync<InteriorChunkedDataNode>(options, cancellationToken);
+				return await Inner.ReadBlobAsync<InteriorChunkedDataNode>(options, cancellationToken);
 			}
 			else
 			{
@@ -230,7 +251,7 @@ namespace EpicGames.Horde.Storage.Nodes
 					if (dataRef.Length < 0)
 					{
 						// Backwards compatibility hack for v2 format
-						using BlobData data = await dataRef.Handle.ReadBlobDataAsync(cancellationToken);
+						using BlobData data = await dataRef.ReadBlobDataAsync(cancellationToken);
 						length = data.Data.Length;
 					}
 
@@ -238,7 +259,7 @@ namespace EpicGames.Horde.Storage.Nodes
 				}
 				else
 				{
-					using BlobData data = await dataRef.Handle.ReadBlobDataAsync(cancellationToken);
+					using BlobData data = await dataRef.ReadBlobDataAsync(cancellationToken);
 					if (data.Type.Guid == LeafChunkedDataNodeConverter.BlobType.Guid)
 					{
 						yield return new LeafChunkedDataNodeRef(dataRef.Length, dataRef.GetLeafHandle());
@@ -254,14 +275,48 @@ namespace EpicGames.Horde.Storage.Nodes
 				}
 			}
 		}
+
+		/// <inheritdoc/>
+		public ValueTask FlushAsync(CancellationToken cancellationToken = default)
+			=> Inner.FlushAsync(cancellationToken);
+
+		/// <inheritdoc/>
+		public ValueTask<BlobData> ReadBlobDataAsync(CancellationToken cancellationToken = default)
+			=> Inner.ReadBlobDataAsync(cancellationToken);
+
+		/// <inheritdoc/>
+		public bool TryGetLocator([NotNullWhen(true)] out BlobLocator locator)
+			=> Inner.TryGetLocator(out locator);
 	}
 
 	/// <summary>
 	/// Reference to a <see cref="LeafChunkedDataNode"/>
 	/// </summary>
 	/// <param name="Length">Length of the leaf data</param>
-	/// <param name="BlobRef">Reference to the blob</param>
-	public record class LeafChunkedDataNodeRef(long Length, IHashedBlobRef<LeafChunkedDataNode> BlobRef);
+	/// <param name="Inner">Reference to the blob</param>
+	public record class LeafChunkedDataNodeRef(long Length, IHashedBlobRef<LeafChunkedDataNode> Inner) : IHashedBlobRef<LeafChunkedDataNode>
+	{
+		/// <inheritdoc/>
+		public BlobSerializerOptions? SerializerOptions => Inner.SerializerOptions;
+
+		/// <inheritdoc/>
+		public IBlobRef Innermost => Inner.Innermost;
+
+		/// <inheritdoc/>
+		public IoHash Hash => Inner.Hash;
+
+		/// <inheritdoc/>
+		public ValueTask FlushAsync(CancellationToken cancellationToken = default)
+			=> Inner.FlushAsync(cancellationToken);
+
+		/// <inheritdoc/>
+		public ValueTask<BlobData> ReadBlobDataAsync(CancellationToken cancellationToken = default)
+			=> Inner.ReadBlobDataAsync(cancellationToken);
+
+		/// <inheritdoc/>
+		public bool TryGetLocator([NotNullWhen(true)] out BlobLocator locator)
+			=> Inner.TryGetLocator(out locator);
+	}
 
 	/// <summary>
 	/// Stores the flat list of chunks produced from chunking a single data stream
@@ -575,7 +630,7 @@ namespace EpicGames.Horde.Storage.Nodes
 				return source;
 			}
 
-			InteriorChunkedDataNode sourceNode = await source.Handle.ReadBlobAsync<InteriorChunkedDataNode>(cancellationToken: cancellationToken);
+			InteriorChunkedDataNode sourceNode = await source.ReadBlobAsync<InteriorChunkedDataNode>(cancellationToken: cancellationToken);
 
 			List<ChunkedDataNodeRef> children = new List<ChunkedDataNodeRef>(sourceNode.Children);
 			for (int idx = children.Count - 1; idx >= 0; idx--)
@@ -604,7 +659,7 @@ namespace EpicGames.Horde.Storage.Nodes
 				index = Math.Min(index + options.MinChildCount, nodeRefs.Count);
 				for (; index < maxIndex; index++)
 				{
-					nodeRefs[index].Handle.Hash.CopyTo(buffer);
+					nodeRefs[index].Hash.CopyTo(buffer);
 
 					uint value = BinaryPrimitives.ReadUInt32LittleEndian(buffer);
 					if (value < options.SliceThreshold)
@@ -661,7 +716,7 @@ namespace EpicGames.Horde.Storage.Nodes
 					_ = nodeReader.ReadUInt32(); // Rolling hash
 				}
 
-				await ChunkedDataNode.CopyToStreamAsync(handle, outputStream, cancellationToken);
+				await handle.CopyToStreamAsync(outputStream, cancellationToken);
 			}
 		}
 	}
@@ -726,7 +781,7 @@ namespace EpicGames.Horde.Storage.Nodes
 		{
 			foreach (ChunkedDataNodeRef child in value.Children)
 			{
-				writer.WriteBlobRef(child.Handle);
+				writer.WriteBlobRef(child);
 				writer.WriteUnsignedVarInt((int)child.Type);
 
 				if (_writeVersion >= HordeApiVersion.AddLengthsToInteriorNodes)
@@ -762,4 +817,87 @@ namespace EpicGames.Horde.Storage.Nodes
 			return new BlobType(InteriorChunkedDataNode.BlobTypeGuid, blobVersion);
 		}
 	}
+
+	/// <summary>
+	/// Stream which returns the content of a chunked data stream
+	/// </summary>
+	class ChunkedDataStream : Stream
+	{
+		/// <inheritdoc/>
+		public override bool CanRead => _readStream.CanRead;
+
+		/// <inheritdoc/>
+		public override bool CanSeek => _readStream.CanSeek;
+
+		/// <inheritdoc/>
+		public override bool CanWrite => _readStream.CanWrite;
+
+		/// <inheritdoc/>
+		public override long Length => _length ?? throw new NotSupportedException();
+
+		/// <inheritdoc/>
+		public override long Position { get => _readStream.Position; set => throw new NotImplementedException(); }
+
+		readonly Pipe _pipe;
+		readonly Stream _readStream;
+		readonly BackgroundTask _writeTask;
+		readonly long? _length;
+
+		/// <summary>
+		/// Constructor
+		/// </summary>
+		/// <param name="chunkedDataNodeRef">Reference to the chunked data node</param>
+		public ChunkedDataStream(ChunkedDataNodeRef chunkedDataNodeRef)
+			: this(chunkedDataNodeRef, chunkedDataNodeRef.Length)
+		{ }
+
+		/// <summary>
+		/// Constructor
+		/// </summary>
+		/// <param name="entry">The file entry to copy from</param>
+		/// <param name="length">Length of the stream. May be null.</param>
+		public ChunkedDataStream(IBlobRef<ChunkedDataNode> entry, long? length)
+		{
+			_pipe = new Pipe();
+			_readStream = _pipe.Reader.AsStream();
+			_writeTask = BackgroundTask.StartNew(ctx => WriteDataAsync(entry, _pipe.Writer, ctx));
+			_length = length;
+		}
+
+		/// <inheritdoc/>
+		public override async ValueTask DisposeAsync()
+		{
+			await base.DisposeAsync();
+
+			await _writeTask.DisposeAsync();
+
+			_readStream.Dispose();
+		}
+
+		static async Task WriteDataAsync(IBlobRef<ChunkedDataNode> entry, PipeWriter pipeWriter, CancellationToken cancellationToken)
+		{
+			using Stream stream = pipeWriter.AsStream();
+			await entry.CopyToStreamAsync(stream, cancellationToken);
+		}
+
+		/// <inheritdoc/>
+		public override ValueTask<int> ReadAsync(Memory<byte> buffer, CancellationToken cancellationToken = default) => _readStream.ReadAsync(buffer, cancellationToken);
+
+		/// <inheritdoc/>
+		public override void Flush()
+		{
+		}
+
+		/// <inheritdoc/>
+		public override int Read(byte[] buffer, int offset, int count) => _readStream.Read(buffer, offset, count);
+
+		/// <inheritdoc/>
+		public override long Seek(long offset, SeekOrigin origin) => throw new NotSupportedException();
+
+		/// <inheritdoc/>
+		public override void SetLength(long value) => throw new NotSupportedException();
+
+		/// <inheritdoc/>
+		public override void Write(byte[] buffer, int offset, int count) => throw new NotSupportedException();
+	}
 }
diff --git a/Engine/Source/Programs/Shared/EpicGames.Horde/Storage/Nodes/ChunkedDataWriter.cs b/Engine/Source/Programs/Shared/EpicGames.Horde/Storage/Nodes/ChunkedDataWriter.cs
--- a/Engine/Source/Programs/Shared/EpicGames.Horde/Storage/Nodes/ChunkedDataWriter.cs
+++ b/Engine/Source/Programs/Shared/EpicGames.Horde/Storage/Nodes/ChunkedDataWriter.cs
@@ -359,7 +359,33 @@ namespace EpicGames.Horde.Storage.Nodes
 	/// </summary>
 	/// <param name="StreamHash">Hash of the stream as a contiguous buffer</param>
 	/// <param name="Root">Handle to the root chunk containing the data</param>
-	public record class ChunkedData(IoHash StreamHash, ChunkedDataNodeRef Root);
+	public record class ChunkedData(IoHash StreamHash, ChunkedDataNodeRef Root)
+	{
+		static readonly ChunkingOptions s_defaultChunkingOptions = new ChunkingOptions();
+
+		/// <summary>
+		/// Creates chunked data from a block of memory
+		/// </summary>
+		/// <param name="blobWriter">Writer for output blobs</param>
+		/// <param name="data">Data to create from</param>
+		/// <param name="cancellationToken">Cancellation token for the operation</param>
+		public static Task<ChunkedData> CreateAsync(IBlobWriter blobWriter, ReadOnlyMemory<byte> data, CancellationToken cancellationToken = default)
+			=> CreateAsync(blobWriter, data, s_defaultChunkingOptions, cancellationToken);
+
+		/// <summary>
+		/// Creates chunked data from a block of memory
+		/// </summary>
+		/// <param name="blobWriter">Writer for output blobs</param>
+		/// <param name="data">Data to create from</param>
+		/// <param name="chunkingOptions">Options controlling the chunking</param>
+		/// <param name="cancellationToken">Cancellation token for the operation</param>
+		public static async Task<ChunkedData> CreateAsync(IBlobWriter blobWriter, ReadOnlyMemory<byte> data, ChunkingOptions chunkingOptions, CancellationToken cancellationToken = default)
+		{
+			using ChunkedDataWriter chunkedWriter = new ChunkedDataWriter(blobWriter, chunkingOptions);
+			await chunkedWriter.AppendAsync(data, cancellationToken);
+			return await chunkedWriter.FlushAsync(cancellationToken);
+		}
+	}
 
 	/// <summary>
 	/// Writes chunked data to an output writer
diff --git a/Engine/Source/Programs/Shared/EpicGames.Horde/Storage/Nodes/DirectoryNode.Extract.cs b/Engine/Source/Programs/Shared/EpicGames.Horde/Storage/Nodes/DirectoryNode.Extract.cs
--- a/Engine/Source/Programs/Shared/EpicGames.Horde/Storage/Nodes/DirectoryNode.Extract.cs
+++ b/Engine/Source/Programs/Shared/EpicGames.Horde/Storage/Nodes/DirectoryNode.Extract.cs
@@ -319,11 +319,11 @@ namespace EpicGames.Horde.Storage.Nodes
 			{
 				if (options.VerboseOutput)
 				{
-					logger.LogInformation("Queuing {Path},{Offset} ({Locator})", outputFile.Path, offset, leafNodeRef.BlobRef.GetLocator());
+					logger.LogInformation("Queuing {Path},{Offset} ({Locator})", outputFile.Path, offset, leafNodeRef.GetLocator());
 				}
 
 				outputFile.IncrementRemaining();
-				ChunkReadRequest request = new ChunkReadRequest(outputFile, offset, leafNodeRef.BlobRef);
+				ChunkReadRequest request = new ChunkReadRequest(outputFile, offset, leafNodeRef);
 				offset += leafNodeRef.Length;
 
 				if (bufferedRequest != null)
diff --git a/Engine/Source/Programs/Shared/EpicGames.Horde/Storage/Nodes/DirectoryNode.cs b/Engine/Source/Programs/Shared/EpicGames.Horde/Storage/Nodes/DirectoryNode.cs
--- a/Engine/Source/Programs/Shared/EpicGames.Horde/Storage/Nodes/DirectoryNode.cs
+++ b/Engine/Source/Programs/Shared/EpicGames.Horde/Storage/Nodes/DirectoryNode.cs
@@ -109,12 +109,11 @@ namespace EpicGames.Horde.Storage.Nodes
 		/// </summary>
 		/// <param name="name">Name of the new directory</param>
 		/// <param name="flags">Flags for the new file</param>
-		/// <param name="length">Length of the file</param>
 		/// <param name="data">Chunked data for the file</param>
 		/// <returns>The new directory object</returns>
-		public FileEntry AddFile(string name, FileEntryFlags flags, long length, ChunkedData data)
+		public FileEntry AddFile(string name, FileEntryFlags flags, ChunkedData data)
 		{
-			FileEntry entry = new FileEntry(name, flags, length, data);
+			FileEntry entry = new FileEntry(name, flags, data);
 			AddFile(entry);
 			return entry;
 		}
@@ -420,7 +419,7 @@ namespace EpicGames.Horde.Storage.Nodes
 			writer.WriteUnsignedVarInt(value.Files.Count);
 			foreach (FileEntry fileEntry in value.Files)
 			{
-				writer.WriteBlobRef(fileEntry.Target.Handle);
+				writer.WriteBlobRef(fileEntry.Target);
 				writer.WriteUnsignedVarInt((int)fileEntry.Target.Type);
 
 				if (_writeVersion >= HordeApiVersion.AddRollingHashesForLeafNodes && fileEntry.Target.Type == ChunkedDataNodeType.Leaf)
@@ -441,7 +440,7 @@ namespace EpicGames.Horde.Storage.Nodes
 				writer.WriteString(fileEntry.Name);
 				writer.WriteUnsignedVarInt((ulong)flags);
 				writer.WriteUnsignedVarInt((ulong)fileEntry.Length);
-				writer.WriteIoHash(fileEntry.StreamHash);
+				writer.WriteIoHash(fileEntry.Hash);
 
 				if ((flags & FileEntryFlags.HasModTime) != 0)
 				{
diff --git a/Engine/Source/Programs/Shared/EpicGames.Horde/Storage/Nodes/FileEntry.cs b/Engine/Source/Programs/Shared/EpicGames.Horde/Storage/Nodes/FileEntry.cs
--- a/Engine/Source/Programs/Shared/EpicGames.Horde/Storage/Nodes/FileEntry.cs
+++ b/Engine/Source/Programs/Shared/EpicGames.Horde/Storage/Nodes/FileEntry.cs
@@ -2,7 +2,6 @@
 
 using System;
 using System.IO;
-using System.IO.Pipelines;
 using System.Runtime.InteropServices;
 using System.Threading;
 using System.Threading.Tasks;
@@ -75,7 +74,7 @@ namespace EpicGames.Horde.Storage.Nodes
 		/// <summary>
 		/// Hash of the file as a contiguous stream. This differs from individual node hashes which hash the Merkle tree of chunks forming it.
 		/// </summary>
-		public IoHash StreamHash { get; }
+		public IoHash Hash { get; }
 
 		/// <summary>
 		/// Reference to the chunked data for the file
@@ -95,15 +94,15 @@ namespace EpicGames.Horde.Storage.Nodes
 		/// <summary>
 		/// Constructor
 		/// </summary>
-		public FileEntry(string name, FileEntryFlags flags, long length, ChunkedData chunkedData, DateTime modTime = default, ReadOnlyMemory<byte> customData = default)
-			: this(name, flags, length, chunkedData.StreamHash, chunkedData.Root, modTime, customData)
+		public FileEntry(string name, FileEntryFlags flags, ChunkedData chunkedData, DateTime modTime = default, ReadOnlyMemory<byte> customData = default)
+			: this(name, flags, chunkedData.Root.Length, chunkedData.StreamHash, chunkedData.Root, modTime, customData)
 		{
 		}
 
 		/// <summary>
 		/// Constructor
 		/// </summary>
-		public FileEntry(string name, FileEntryFlags flags, long length, IoHash streamHash, ChunkedDataNodeRef target, DateTime modTime = default, ReadOnlyMemory<byte> customData = default)
+		public FileEntry(string name, FileEntryFlags flags, long length, IoHash hash, ChunkedDataNodeRef target, DateTime modTime = default, ReadOnlyMemory<byte> customData = default)
 		{
 			flags &= ~(FileEntryFlags.HasModTime | FileEntryFlags.HasCustomData);
 			if (modTime != default)
@@ -118,7 +117,7 @@ namespace EpicGames.Horde.Storage.Nodes
 			Name = name;
 			Flags = flags;
 			Length = length;
-			StreamHash = streamHash;
+			Hash = hash;
 			Target = target;
 			ModTime = modTime;
 			CustomData = customData;
@@ -128,17 +127,23 @@ namespace EpicGames.Horde.Storage.Nodes
 		/// Creates a stream that returns the contents of this file
 		/// </summary>
 		/// <returns>The content stream</returns>
-		public Stream OpenAsStream() => new FileEntryContentStream(this);
+		public Stream OpenAsStream()
+			=> Target.OpenAsStream();
 
 		/// <summary>
 		/// Copies the contents of this node and its children to the given output stream
 		/// </summary>
 		/// <param name="outputStream">The output stream to receive the data</param>
 		/// <param name="cancellationToken">Cancellation token for the operation</param>
-		public async Task CopyToStreamAsync(Stream outputStream, CancellationToken cancellationToken)
-		{
-			await ChunkedDataNode.CopyToStreamAsync(Target.Handle, outputStream, cancellationToken);
-		}
+		public Task CopyToStreamAsync(Stream outputStream, CancellationToken cancellationToken = default)
+			=> Target.CopyToStreamAsync(outputStream, cancellationToken);
+
+		/// <summary>
+		/// Reads this entire file as a byte array
+		/// </summary>
+		/// <param name="cancellationToken">Cancellation token for the operation</param>
+		public Task<byte[]> ReadAllBytesAsync(CancellationToken cancellationToken = default)
+			=> Target.ReadAllBytesAsync(cancellationToken);
 
 		/// <summary>
 		/// Extracts the contents of this node to a file
@@ -150,7 +155,7 @@ namespace EpicGames.Horde.Storage.Nodes
 		{
 			try
 			{
-				await ChunkedDataNode.CopyToFileAsync(Target.Handle, file, cancellationToken);
+				await Target.CopyToFileAsync(file, cancellationToken);
 				SetPermissions(file, Flags);
 			}
 			catch (Exception ex) when (ex is not OperationCanceledException)
@@ -247,78 +252,4 @@ namespace EpicGames.Horde.Storage.Nodes
 		/// <inheritdoc/>
 		public override string ToString() => Name.ToString();
 	}
-
-	/// <summary>
-	/// Stream which returns the content of a file
-	/// </summary>
-	class FileEntryContentStream : Stream
-	{
-		/// <inheritdoc/>
-		public override bool CanRead => _readStream.CanRead;
-
-		/// <inheritdoc/>
-		public override bool CanSeek => _readStream.CanSeek;
-
-		/// <inheritdoc/>
-		public override bool CanWrite => _readStream.CanWrite;
-
-		/// <inheritdoc/>
-		public override long Length { get; }
-
-		/// <inheritdoc/>
-		public override long Position { get => _readStream.Position; set => throw new NotImplementedException(); }
-
-		readonly Pipe _pipe;
-		readonly Stream _readStream;
-		readonly BackgroundTask _writeTask;
-
-		/// <summary>
-		/// Constructor
-		/// </summary>
-		/// <param name="entry">The file entry to copy from</param>
-		public FileEntryContentStream(FileEntry entry)
-		{
-			_pipe = new Pipe();
-			_readStream = _pipe.Reader.AsStream();
-			_writeTask = BackgroundTask.StartNew(ctx => WriteDataAsync(entry, _pipe.Writer, ctx));
-
-			Length = entry.Length;
-		}
-
-		/// <inheritdoc/>
-		public override async ValueTask DisposeAsync()
-		{
-			await base.DisposeAsync();
-
-			await _writeTask.DisposeAsync();
-
-			_readStream.Dispose();
-		}
-
-		static async Task WriteDataAsync(FileEntry entry, PipeWriter pipeWriter, CancellationToken cancellationToken)
-		{
-			using Stream stream = pipeWriter.AsStream();
-			await entry.CopyToStreamAsync(stream, cancellationToken);
-		}
-
-		/// <inheritdoc/>
-		public override ValueTask<int> ReadAsync(Memory<byte> buffer, CancellationToken cancellationToken = default) => _readStream.ReadAsync(buffer, cancellationToken);
-
-		/// <inheritdoc/>
-		public override void Flush()
-		{
-		}
-
-		/// <inheritdoc/>
-		public override int Read(byte[] buffer, int offset, int count) => _readStream.Read(buffer, offset, count);
-
-		/// <inheritdoc/>
-		public override long Seek(long offset, SeekOrigin origin) => throw new NotSupportedException();
-
-		/// <inheritdoc/>
-		public override void SetLength(long value) => throw new NotSupportedException();
-
-		/// <inheritdoc/>
-		public override void Write(byte[] buffer, int offset, int count) => throw new NotSupportedException();
-	}
 }
diff --git a/Engine/Source/Programs/Shared/EpicGames.Horde/Storage/Workspace.cs b/Engine/Source/Programs/Shared/EpicGames.Horde/Storage/Workspace.cs
--- a/Engine/Source/Programs/Shared/EpicGames.Horde/Storage/Workspace.cs
+++ b/Engine/Source/Programs/Shared/EpicGames.Horde/Storage/Workspace.cs
@@ -710,7 +710,7 @@ namespace EpicGames.Horde.Storage
 				{
 					outputFile.IncrementRemaining();
 
-					SourceChunk? sourceChunk = TryGetSourceChunk(leafNodeRef.BlobRef.Hash);
+					SourceChunk? sourceChunk = TryGetSourceChunk(leafNodeRef.Hash);
 					if (sourceChunk != null)
 					{
 						ChunkCopyRequest copyRequest = new ChunkCopyRequest(outputFile, offset, sourceChunk);
@@ -718,7 +718,7 @@ namespace EpicGames.Horde.Storage
 					}
 					else
 					{
-						ChunkReadRequest readRequest = new ChunkReadRequest(outputFile, offset, leafNodeRef.BlobRef);
+						ChunkReadRequest readRequest = new ChunkReadRequest(outputFile, offset, leafNodeRef);
 						await chunkReadRequests.Writer.WriteAsync(readRequest, cancellationToken);
 					}
 
@@ -760,7 +760,7 @@ namespace EpicGames.Horde.Storage
 				OutputFile outputFile = fileUpdate.OutputFile;
 
 				FileState fileState = fileUpdate.ParentDir.FindOrAddFile(outputFile.FileEntry.Name);
-				fileState.Hash = outputFile.FileEntry.StreamHash;
+				fileState.Hash = outputFile.FileEntry.Hash;
 				fileState.Update(outputFile.FileInfo);
 
 				HashInfo hashInfo = AddFileToHashLookup(fileState);
@@ -770,7 +770,7 @@ namespace EpicGames.Horde.Storage
 				{
 					ChunkInfo chunkInfo = new ChunkInfo(hashInfo, offset, chunkRef.Length);
 
-					HashInfo chunkHashInfo = FindOrAddHashInfo(chunkRef.BlobRef.Hash, chunkRef.Length);
+					HashInfo chunkHashInfo = FindOrAddHashInfo(chunkRef.Hash, chunkRef.Length);
 					chunkHashInfo.AddChunk(chunkInfo);
 
 					offset += chunkRef.Length;
@@ -808,7 +808,7 @@ namespace EpicGames.Horde.Storage
 				FileState fileState = dirState.Files[fileIdx];
 				if ((fileState.LayerFlags & flag) != 0)
 				{
-					if (dirNode == null || !dirNode.TryGetFileEntry(fileState.Name, out FileEntry? entry) || entry.StreamHash != fileState.Hash)
+					if (dirNode == null || !dirNode.TryGetFileEntry(fileState.Name, out FileEntry? entry) || entry.Hash != fileState.Hash)
 					{
 						fileState.LayerFlags &= ~flag;
 
@@ -861,7 +861,7 @@ namespace EpicGames.Horde.Storage
 			FileState? fileState;
 			if (dirState.TryGetFile(fileEntry.Name, out fileState))
 			{
-				if (fileState.Hash == fileEntry.StreamHash)
+				if (fileState.Hash == fileEntry.Hash)
 				{
 					fileState.LayerFlags |= flag;
 				}
@@ -874,7 +874,7 @@ namespace EpicGames.Horde.Storage
 			{
 				FileInfo fileInfo = FileReference.Combine(dirRef, fileEntry.Name.ToString()).ToFileInfo();
 
-				fileState = TryMoveCachedData(fileEntry.StreamHash, dirState, fileEntry.Name);
+				fileState = TryMoveCachedData(fileEntry.Hash, dirState, fileEntry.Name);
 				if (fileState == null)
 				{
 					OutputFile outputFile = new OutputFile(path, fileInfo, fileEntry);
diff --git a/Engine/Source/Programs/UnrealBuildTool/Artifacts/HordeStorageArtifactCache.cs b/Engine/Source/Programs/UnrealBuildTool/Artifacts/HordeStorageArtifactCache.cs
--- a/Engine/Source/Programs/UnrealBuildTool/Artifacts/HordeStorageArtifactCache.cs
+++ b/Engine/Source/Programs/UnrealBuildTool/Artifacts/HordeStorageArtifactCache.cs
@@ -81,7 +81,7 @@ namespace UnrealBuildTool.Artifacts
 				using FileStream stream = new(outputName, FileMode.Open, FileAccess.Read, FileShare.Read);
 				LeafChunkedData leafChunkedData = await fileWriter.CreateAsync(stream, leafOptions.TargetSize, cancellationToken);
 				ChunkedData chunkedData = await InteriorChunkedDataNode.CreateTreeAsync(leafChunkedData, interiorOptions, writer, cancellationToken);
-				OutputRefs[index++] = chunkedData.Root.Handle;
+				OutputRefs[index++] = chunkedData.Root;
 			}
 		}
 	}
@@ -318,7 +318,7 @@ namespace UnrealBuildTool.Artifacts
 							{
 								string outputName = hordeArtifactAction.ArtifactAction.Outputs[refIndex++].GetFullPath(artifactAction.DirectoryMapping);
 								using FileStream stream = new(outputName, FileMode.Create, FileAccess.Write, FileShare.ReadWrite);
-								await ChunkedDataNode.CopyToStreamAsync(artifactRef, stream, cancellationToken);
+								await artifactRef.CopyToStreamAsync(stream, cancellationToken);
 							}
 							catch (Exception)
 							{
