Horde: Fix computation of hash values for streams passed through ChunkedDataWriter.

#rnx

diff --git a/Engine/Source/Programs/Shared/EpicGames.Horde.Tests/Storage/ChunkingTests.cs b/Engine/Source/Programs/Shared/EpicGames.Horde.Tests/Storage/ChunkingTests.cs
--- a/Engine/Source/Programs/Shared/EpicGames.Horde.Tests/Storage/ChunkingTests.cs
+++ b/Engine/Source/Programs/Shared/EpicGames.Horde.Tests/Storage/ChunkingTests.cs
@@ -4,6 +4,7 @@ using System;
 using System.Collections.Generic;
 using System.IO;
 using System.Linq;
+using System.Security.Cryptography;
 using System.Text;
 using System.Threading;
 using System.Threading.Tasks;
@@ -206,6 +207,26 @@ namespace EpicGames.Horde.Tests.Storage
 			}
 		}
 
+		[TestMethod]
+		public async Task HashAsync()
+		{
+			ChunkingOptions chunkingOptions = new ChunkingOptions();
+			chunkingOptions.LeafOptions = new LeafChunkedDataNodeOptions(64, 128, 92);
+
+			await using MemoryBlobWriter blobWriter = new MemoryBlobWriter(new BlobSerializerOptions());
+			byte[] data = RandomNumberGenerator.GetBytes(4096);
+
+			ChunkedData chunkedData;
+			using (ChunkedDataWriter writer = new ChunkedDataWriter(blobWriter, chunkingOptions))
+			{
+				await writer.AppendAsync(data, default);
+				chunkedData = await writer.FlushAsync();
+			}
+
+			IoHash hash = IoHash.Compute(data);
+			Assert.AreEqual(hash, chunkedData.StreamHash);
+		}
+
 		[TestMethod]
 		public async Task ChunkOrderAsync()
 		{
diff --git a/Engine/Source/Programs/Shared/EpicGames.Horde/Storage/Nodes/ChunkedDataWriter.cs b/Engine/Source/Programs/Shared/EpicGames.Horde/Storage/Nodes/ChunkedDataWriter.cs
--- a/Engine/Source/Programs/Shared/EpicGames.Horde/Storage/Nodes/ChunkedDataWriter.cs
+++ b/Engine/Source/Programs/Shared/EpicGames.Horde/Storage/Nodes/ChunkedDataWriter.cs
@@ -227,6 +227,8 @@ namespace EpicGames.Horde.Storage.Nodes
 		/// <param name="cancellationToken">Cancellation token for the operation</param>
 		public async Task AppendAsync(ReadOnlyMemory<byte> data, CancellationToken cancellationToken)
 		{
+			_hasher.Update(data.Span);
+
 			for (; ; )
 			{
 				int appendLength = AppendToLeafNode(_writer.WrittenMemory.Span, data.Span, ref _leafHash, _leafChunkOptions);
@@ -243,7 +245,6 @@ namespace EpicGames.Horde.Storage.Nodes
 
 				await FlushLeafNodeAsync(cancellationToken);
 			}
-			_hasher.Update(data.Span);
 		}
 
 		/// <summary>
@@ -330,8 +331,9 @@ namespace EpicGames.Horde.Storage.Nodes
 		/// <returns>Handle to the root node</returns>
 		public async Task<LeafChunkedData> CompleteAsync(CancellationToken cancellationToken)
 		{
+			IoHash hash = IoHash.FromBlake3(_hasher);
 			await FlushLeafNodeAsync(cancellationToken);
-			LeafChunkedData leafChunkedData = new LeafChunkedData(IoHash.FromBlake3(_hasher), new List<ChunkedDataNodeRef>(_leafHandles));
+			LeafChunkedData leafChunkedData = new LeafChunkedData(hash, new List<ChunkedDataNodeRef>(_leafHandles));
 
 			_leafHandles.Clear();
 			_hasher.Reset();
@@ -454,7 +456,7 @@ namespace EpicGames.Horde.Storage.Nodes
 		/// </summary>
 		/// <param name="data">Data to append</param>
 		/// <param name="cancellationToken">Cancellation token for the operation</param>
-		public async Task AppendAsync(ReadOnlyMemory<byte> data, CancellationToken cancellationToken)
+		public async Task AppendAsync(ReadOnlyMemory<byte> data, CancellationToken cancellationToken = default)
 		{
 			await _leafWriter.AppendAsync(data, cancellationToken);
 		}
